[Backout] - CL29948158
Original CL Desc
-----------------------------------------------------------------
[FScriptSetHelper and FScriptMapHelper iterators]
simple index to iterator replacements

diff --git a/Engine/Plugins/Animation/GameplayInsights/Source/GameplayInsights/Private/ObjectPropertyTrace.cpp b/Engine/Plugins/Animation/GameplayInsights/Source/GameplayInsights/Private/ObjectPropertyTrace.cpp
--- a/Engine/Plugins/Animation/GameplayInsights/Source/GameplayInsights/Private/ObjectPropertyTrace.cpp
+++ b/Engine/Plugins/Animation/GameplayInsights/Source/GameplayInsights/Private/ObjectPropertyTrace.cpp
@@ -130,24 +130,31 @@ namespace ObjectPropertyTrace
 			
 			InFunction(InProperty->GetCPPType(), SizeString, MapRootId, InParentId, InPropertyNameBuilder);
 			
-			for (FScriptMapHelper::FIterator It(Helper); It; ++It)
+			int32 Num = Helper.Num();
+			int32 MapIndex = 0;
+			for (int32 DynamicIndex = 0; Num; ++DynamicIndex)
 			{
-				int32 MapEntryId = ++InId;
-				FString KeyString = FString::Printf(TEXT("[%d]"), It.GetLogicalIndex());
-				FString TypeString = FString::Printf(TEXT("{%s, %s}"), *MapProperty->KeyProp->GetCPPType(), *MapProperty->ValueProp->GetCPPType());
-
-				// Build pair parent name id
-				FScopedPropertyName ScopedParentPairName(InPropertyNameBuilder);
-				ScopedParentPairName.Builder.AppendChar(TEXT('.'));
-				ScopedParentPairName.Builder.Append(KeyString);
-
-				InFunction(TypeString, TEXT("{...}"), MapEntryId, MapRootId, InPropertyNameBuilder);
+				if (Helper.IsValidIndex(DynamicIndex))
+				{
+					int32 MapEntryId = ++InId;
+					FString KeyString = FString::Printf(TEXT("[%d]"), MapIndex++);
+					FString TypeString = FString::Printf(TEXT("{%s, %s}"), *MapProperty->KeyProp->GetCPPType(), *MapProperty->ValueProp->GetCPPType());
+
+					// Build pair parent name id
+					FScopedPropertyName ScopedParentPairName(InPropertyNameBuilder);
+					ScopedParentPairName.Builder.AppendChar(TEXT('.'));
+					ScopedParentPairName.Builder.Append(KeyString);
+					
+					InFunction(TypeString, TEXT("{...}"), MapEntryId, MapRootId, InPropertyNameBuilder);
+					
+					const void* PairPtr = Helper.GetPairPtr(DynamicIndex);
 
-				const void* PairPtr = Helper.GetPairPtr(It);
+					IteratePropertiesRecursive(MapProperty->KeyProp, PairPtr, MapProperty->KeyProp->GetName(), InFunction, InId, MapEntryId, InPropertyNameBuilder);
 
-				IteratePropertiesRecursive(MapProperty->KeyProp, PairPtr, MapProperty->KeyProp->GetName(), InFunction, InId, MapEntryId, InPropertyNameBuilder);
+					IteratePropertiesRecursive(MapProperty->ValueProp, PairPtr, MapProperty->ValueProp->GetName(), InFunction, InId, MapEntryId, InPropertyNameBuilder);
 
-				IteratePropertiesRecursive(MapProperty->ValueProp, PairPtr, MapProperty->ValueProp->GetName(), InFunction, InId, MapEntryId, InPropertyNameBuilder);
+					--Num;
+				}
 			}
 		}
 		else if (const FSetProperty* SetProperty = CastField<FSetProperty>(InProperty))
@@ -166,12 +173,19 @@ namespace ObjectPropertyTrace
 			InFunction(InProperty->GetCPPType(), SizeString, SetRootId, InParentId, InPropertyNameBuilder);
 
 			// Trace container items
-			for (FScriptSetHelper::FIterator It(Helper); It; ++It)
+			int32 Num = Helper.Num();
+			int32 SetIndex = 0;
+			for (int32 DynamicIndex = 0; Num; ++DynamicIndex)
 			{
-				const void* ValuePtr = Helper.GetElementPtr(It);
-				FString KeyString = FString::Printf(TEXT("[%d]"), It.GetLogicalIndex());
+				if (Helper.IsValidIndex(DynamicIndex))
+				{
+					const void* ValuePtr = Helper.GetElementPtr(DynamicIndex);
+					FString KeyString = FString::Printf(TEXT("[%d]"), SetIndex++);
+					
+					IteratePropertiesRecursive(SetProperty->ElementProp, ValuePtr, KeyString, InFunction, InId, SetRootId, InPropertyNameBuilder);
 
-				IteratePropertiesRecursive(SetProperty->ElementProp, ValuePtr, KeyString, InFunction, InId, SetRootId, InPropertyNameBuilder);
+					--Num;
+				}
 			}
 		}
 		else if (const FStructProperty* StructProperty = CastField<FStructProperty>(InProperty))
diff --git a/Engine/Plugins/Enterprise/GLTFExporter/Source/GLTFExporter/Private/Exporters/GLTFExporterAnalytics.cpp b/Engine/Plugins/Enterprise/GLTFExporter/Source/GLTFExporter/Private/Exporters/GLTFExporterAnalytics.cpp
--- a/Engine/Plugins/Enterprise/GLTFExporter/Source/GLTFExporter/Private/Exporters/GLTFExporterAnalytics.cpp
+++ b/Engine/Plugins/Enterprise/GLTFExporter/Source/GLTFExporter/Private/Exporters/GLTFExporterAnalytics.cpp
@@ -110,14 +110,21 @@ void FGLTFExporterAnalytics::GetAttributesFromProperty(const FProperty* Property
 	else if (const FMapProperty* MapProperty = CastField<FMapProperty>(Property))
 	{
 		FScriptMapHelper Helper(MapProperty, ValuePtr);
-		OutAttributes.Emplace(AttributeName, Helper.Num() > 0);
+		int32 RemainingPairs = Helper.Num();
 
-		for (FScriptMapHelper::FIterator It(Helper); It; ++It)
+		OutAttributes.Emplace(AttributeName, RemainingPairs > 0);
+
+		for (int32 Index = 0; RemainingPairs > 0; ++Index)
 		{
-			FString KeyString;
-			MapProperty->KeyProp->ExportTextItem_Direct(KeyString, Helper.GetKeyPtr(It), nullptr, nullptr, PPF_None);
+			if (Helper.IsValidIndex(Index))
+			{
+				FString KeyString;
+				MapProperty->KeyProp->ExportTextItem_Direct(KeyString, Helper.GetKeyPtr(Index), nullptr, nullptr, PPF_None);
 
-			GetAttributesFromProperty(MapProperty->ValueProp, Helper.GetValuePtr(It), AttributeName + TEXT(".") + KeyString, OutAttributes);
+				GetAttributesFromProperty(MapProperty->ValueProp, Helper.GetValuePtr(Index), AttributeName + TEXT(".") + KeyString, OutAttributes);
+
+				--RemainingPairs;
+			}
 		}
 	}
 	else if (const FStructProperty* StructProperty = CastField<FStructProperty>(Property))
diff --git a/Engine/Plugins/Experimental/PCG/Source/PCG/Private/Helpers/PCGHelpers.cpp b/Engine/Plugins/Experimental/PCG/Source/PCG/Private/Helpers/PCGHelpers.cpp
--- a/Engine/Plugins/Experimental/PCG/Source/PCG/Private/Helpers/PCGHelpers.cpp
+++ b/Engine/Plugins/Experimental/PCG/Source/PCG/Private/Helpers/PCGHelpers.cpp
@@ -378,22 +378,34 @@ namespace PCGHelpers
 		else if (FMapProperty* MapProperty = CastField<FMapProperty>(Property))
 		{
 			FScriptMapHelper_InContainer Helper(MapProperty, InContainer);
-			for (FScriptMapHelper::FIterator It(Helper); It; ++It)
+			int32 Num = Helper.Num();
+			for (int32 DynamicIndex = 0; Num; ++DynamicIndex)
 			{
+				if (Helper.IsValidIndex(DynamicIndex))
+				{
 					// Key and Value are stored next to each other in memory.
 					// ValueProp has an offset, so we should use the same starting address for both.
-					const void* PairKeyValuePtr = Helper.GetKeyPtr(It);
+					const void* PairKeyValuePtr = Helper.GetKeyPtr(DynamicIndex);
 					GatherDependencies(MapProperty->KeyProp, PairKeyValuePtr, OutDependencies, MaxDepth, InExcludedClasses);
 					GatherDependencies(MapProperty->ValueProp, PairKeyValuePtr, OutDependencies, MaxDepth, InExcludedClasses);
+
+					--Num;
+				}
 			}
 		}
 		else if (FSetProperty* SetProperty = CastField<FSetProperty>(Property))
 		{
 			FScriptSetHelper_InContainer Helper(SetProperty, InContainer);
-			for (FScriptSetHelper::FIterator It(Helper); It; ++It)
+			int32 Num = Helper.Num();
+			for (int32 DynamicIndex = 0; Num; ++DynamicIndex)
 			{
-				const void* ValuePtr = Helper.GetElementPtr(It);
-				GatherDependencies(SetProperty->ElementProp, ValuePtr, OutDependencies, MaxDepth, InExcludedClasses);
+				if (Helper.IsValidIndex(DynamicIndex))
+				{
+					const void* ValuePtr = Helper.GetElementPtr(DynamicIndex);
+					GatherDependencies(SetProperty->ElementProp, ValuePtr, OutDependencies, MaxDepth, InExcludedClasses);
+
+					--Num;
+				}
 			}
 		}
 	}
diff --git a/Engine/Plugins/Experimental/PythonScriptPlugin/Source/PythonScriptPlugin/Private/PyGenUtil.cpp b/Engine/Plugins/Experimental/PythonScriptPlugin/Source/PythonScriptPlugin/Private/PyGenUtil.cpp
--- a/Engine/Plugins/Experimental/PythonScriptPlugin/Source/PythonScriptPlugin/Private/PyGenUtil.cpp
+++ b/Engine/Plugins/Experimental/PythonScriptPlugin/Source/PythonScriptPlugin/Private/PyGenUtil.cpp
@@ -2315,13 +2315,16 @@ void PythonizeValueImpl(const FProperty* InProp, const void* InPropValue, const
 				: TEXT("[");
 			{
 				FScriptSetHelper ScriptSetHelper(SetProperty, PropArrValue);
-				for (FScriptSetHelper::FIterator It(ScriptSetHelper); It; ++It)
+				for (int32 SparseElementIndex = 0, ElementIndex = 0; SparseElementIndex < ScriptSetHelper.GetMaxIndex(); ++SparseElementIndex)
 				{
-					if (It.GetLogicalIndex() > 0)
+					if (ScriptSetHelper.IsValidIndex(SparseElementIndex))
 					{
-						OutPythonDefaultValue += TEXT(", ");
+						if (ElementIndex++ > 0)
+						{
+							OutPythonDefaultValue += TEXT(", ");
+						}
+						PythonizeValueImpl(ScriptSetHelper.GetElementProperty(), ScriptSetHelper.GetElementPtr(SparseElementIndex), InFlags, OutPythonDefaultValue);
 					}
-					PythonizeValueImpl(ScriptSetHelper.GetElementProperty(), ScriptSetHelper.GetElementPtr(It), InFlags, OutPythonDefaultValue);
 				}
 			}
 			OutPythonDefaultValue += bUseStrictTyping
@@ -2335,15 +2338,18 @@ void PythonizeValueImpl(const FProperty* InProp, const void* InPropValue, const
 				: TEXT("{");
 			{
 				FScriptMapHelper ScriptMapHelper(MapProperty, PropArrValue);
-				for (FScriptMapHelper::FIterator It(ScriptMapHelper); It; ++It)
+				for (int32 SparseElementIndex = 0, ElementIndex = 0; SparseElementIndex < ScriptMapHelper.GetMaxIndex(); ++SparseElementIndex)
 				{
-					if (It.GetLogicalIndex() > 0)
+					if (ScriptMapHelper.IsValidIndex(SparseElementIndex))
 					{
-						OutPythonDefaultValue += TEXT(", ");
+						if (ElementIndex++ > 0)
+						{
+							OutPythonDefaultValue += TEXT(", ");
+						}
+						PythonizeValueImpl(ScriptMapHelper.GetKeyProperty(), ScriptMapHelper.GetKeyPtr(SparseElementIndex), InFlags, OutPythonDefaultValue);
+						OutPythonDefaultValue += TEXT(": ");
+						PythonizeValueImpl(ScriptMapHelper.GetValueProperty(), ScriptMapHelper.GetValuePtr(SparseElementIndex), InFlags, OutPythonDefaultValue);
 					}
-					PythonizeValueImpl(ScriptMapHelper.GetKeyProperty(), ScriptMapHelper.GetKeyPtr(It), InFlags, OutPythonDefaultValue);
-					OutPythonDefaultValue += TEXT(": ");
-					PythonizeValueImpl(ScriptMapHelper.GetValueProperty(), ScriptMapHelper.GetValuePtr(It), InFlags, OutPythonDefaultValue);
 				}
 			}
 			OutPythonDefaultValue += bUseStrictTyping
diff --git a/Engine/Plugins/Experimental/PythonScriptPlugin/Source/PythonScriptPlugin/Private/PyReferenceCollector.cpp b/Engine/Plugins/Experimental/PythonScriptPlugin/Source/PythonScriptPlugin/Private/PyReferenceCollector.cpp
--- a/Engine/Plugins/Experimental/PythonScriptPlugin/Source/PythonScriptPlugin/Private/PyReferenceCollector.cpp
+++ b/Engine/Plugins/Experimental/PythonScriptPlugin/Source/PythonScriptPlugin/Private/PyReferenceCollector.cpp
@@ -345,9 +345,13 @@ void FPyReferenceCollector::AddReferencedObjectsFromPropertyInternal(FReferenceC
 		{
 			bool bSetValuesChanged = false;
 			FScriptSetHelper_InContainer ScriptSetHelper(CastProp, InBaseAddr, ArrIndex);
-			for (FScriptSetHelper::FIterator It(ScriptSetHelper); It; ++It)
+
+			for (int32 SparseElementIndex = 0; SparseElementIndex < ScriptSetHelper.GetMaxIndex(); ++SparseElementIndex)
 			{
-				AddReferencedObjectsFromPropertyInternal(InCollector, ScriptSetHelper.GetElementProperty(), ScriptSetHelper.GetElementPtr(It), InFlags, bSetValuesChanged);
+				if (ScriptSetHelper.IsValidIndex(SparseElementIndex))
+				{
+						AddReferencedObjectsFromPropertyInternal(InCollector, ScriptSetHelper.GetElementProperty(), ScriptSetHelper.GetElementPtr(SparseElementIndex), InFlags, bSetValuesChanged);
+				}
 			}
 
 			if (bSetValuesChanged)
@@ -370,11 +374,14 @@ void FPyReferenceCollector::AddReferencedObjectsFromPropertyInternal(FReferenceC
 			bool bMapValuesChanged = false;
 			FScriptMapHelper_InContainer ScriptMapHelper(CastProp, InBaseAddr, ArrIndex);
 
-			for (FScriptMapHelper::FIterator It(ScriptMapHelper); It; ++It)
+			for (int32 SparseElementIndex = 0; SparseElementIndex < ScriptMapHelper.GetMaxIndex(); ++SparseElementIndex)
 			{
+				if (ScriptMapHelper.IsValidIndex(SparseElementIndex))
+				{
 					// Note: We use the pair pointer below as AddReferencedObjectsFromPropertyInternal expects a base address and the key/value property will apply the correct offset from the base
-					AddReferencedObjectsFromPropertyInternal(InCollector, ScriptMapHelper.GetKeyProperty(), ScriptMapHelper.GetPairPtr(It), InFlags, bMapKeysChanged);
-					AddReferencedObjectsFromPropertyInternal(InCollector, ScriptMapHelper.GetValueProperty(), ScriptMapHelper.GetPairPtr(It), InFlags, bMapValuesChanged);
+					AddReferencedObjectsFromPropertyInternal(InCollector, ScriptMapHelper.GetKeyProperty(), ScriptMapHelper.GetPairPtr(SparseElementIndex), InFlags, bMapKeysChanged);
+					AddReferencedObjectsFromPropertyInternal(InCollector, ScriptMapHelper.GetValueProperty(), ScriptMapHelper.GetPairPtr(SparseElementIndex), InFlags, bMapValuesChanged);
+				}
 			}
 
 			if (bMapKeysChanged || bMapValuesChanged)
diff --git a/Engine/Plugins/Experimental/PythonScriptPlugin/Source/PythonScriptPlugin/Private/PyWrapperMap.cpp b/Engine/Plugins/Experimental/PythonScriptPlugin/Source/PythonScriptPlugin/Private/PyWrapperMap.cpp
--- a/Engine/Plugins/Experimental/PythonScriptPlugin/Source/PythonScriptPlugin/Private/PyWrapperMap.cpp
+++ b/Engine/Plugins/Experimental/PythonScriptPlugin/Source/PythonScriptPlugin/Private/PyWrapperMap.cpp
@@ -222,17 +222,17 @@ struct FPyWrapperMapItemIterator : public TPyWrapperMapIterator<FPyWrapperMapIte
 		return &PyWrapperMapItemIteratorType;
 	}
 
-	static PyObject* GetItem(FPyWrapperMapItemIterator* InSelf, FScriptMapHelper& InScriptMapHelper, const int32 InValidatedInternalIndex)
+	static PyObject* GetItem(FPyWrapperMapItemIterator* InSelf, FScriptMapHelper& InScriptMapHelper, int32 InElementIndex)
 	{
 		FPyObjectPtr PyKeyObj;
-		if (!PyConversion::PythonizeProperty(InScriptMapHelper.GetKeyProperty(), InScriptMapHelper.GetKeyPtr(InValidatedInternalIndex), PyKeyObj.Get()))
+		if (!PyConversion::PythonizeProperty(InScriptMapHelper.GetKeyProperty(), InScriptMapHelper.GetKeyPtr(InElementIndex), PyKeyObj.Get()))
 		{
 			PyUtil::SetPythonError(PyExc_TypeError, InSelf, *FString::Printf(TEXT("Failed to convert key property '%s' (%s)"), *InScriptMapHelper.GetKeyProperty()->GetName(), *InScriptMapHelper.GetKeyProperty()->GetClass()->GetName()));
 			return nullptr;
 		}
 
 		FPyObjectPtr PyValueObj;
-		if (!PyConversion::PythonizeProperty(InScriptMapHelper.GetValueProperty(), InScriptMapHelper.GetValuePtr(InValidatedInternalIndex), PyValueObj.Get()))
+		if (!PyConversion::PythonizeProperty(InScriptMapHelper.GetValueProperty(), InScriptMapHelper.GetValuePtr(InElementIndex), PyValueObj.Get()))
 		{
 			PyUtil::SetPythonError(PyExc_TypeError, InSelf, *FString::Printf(TEXT("Failed to convert value property '%s' (%s)"), *InScriptMapHelper.GetValueProperty()->GetName(), *InScriptMapHelper.GetValueProperty()->GetClass()->GetName()));
 			return nullptr;
@@ -250,10 +250,10 @@ struct FPyWrapperMapKeyIterator : public TPyWrapperMapIterator<FPyWrapperMapKeyI
 		return &PyWrapperMapKeyIteratorType;
 	}
 
-	static PyObject* GetItem(FPyWrapperMapKeyIterator* InSelf, FScriptMapHelper& InScriptMapHelper, const int32 InValidatedInternalIndex)
+	static PyObject* GetItem(FPyWrapperMapKeyIterator* InSelf, FScriptMapHelper& InScriptMapHelper, int32 InElementIndex)
 	{
 		PyObject* PyItemObj = nullptr;
-		if (!PyConversion::PythonizeProperty(InScriptMapHelper.GetKeyProperty(), InScriptMapHelper.GetKeyPtr(InValidatedInternalIndex), PyItemObj))
+		if (!PyConversion::PythonizeProperty(InScriptMapHelper.GetKeyProperty(), InScriptMapHelper.GetKeyPtr(InElementIndex), PyItemObj))
 		{
 			PyUtil::SetPythonError(PyExc_TypeError, InSelf, *FString::Printf(TEXT("Failed to convert key property '%s' (%s)"), *InScriptMapHelper.GetKeyProperty()->GetName(), *InScriptMapHelper.GetKeyProperty()->GetClass()->GetName()));
 			return nullptr;
@@ -270,10 +270,10 @@ struct FPyWrapperMapValueIterator : public TPyWrapperMapIterator<FPyWrapperMapVa
 		return &PyWrapperMapValueIteratorType;
 	}
 
-	static PyObject* GetItem(FPyWrapperMapValueIterator* InSelf, FScriptMapHelper& InScriptMapHelper, const int32 InValidatedInternalIndex)
+	static PyObject* GetItem(FPyWrapperMapValueIterator* InSelf, FScriptMapHelper& InScriptMapHelper, int32 InElementIndex)
 	{
 		PyObject* PyItemObj = nullptr;
-		if (!PyConversion::PythonizeProperty(InScriptMapHelper.GetValueProperty(), InScriptMapHelper.GetValuePtr(InValidatedInternalIndex), PyItemObj))
+		if (!PyConversion::PythonizeProperty(InScriptMapHelper.GetValueProperty(), InScriptMapHelper.GetValuePtr(InElementIndex), PyItemObj))
 		{
 			PyUtil::SetPythonError(PyExc_TypeError, InSelf, *FString::Printf(TEXT("Failed to convert value property '%s' (%s)"), *InScriptMapHelper.GetValueProperty()->GetName(), *InScriptMapHelper.GetValueProperty()->GetClass()->GetName()));
 			return nullptr;
@@ -661,38 +661,47 @@ FPyWrapperMap* FPyWrapperMap::CastPyObject(PyObject* InPyObject, PyTypeObject* I
 
 		// Attempt to convert the entries in the map to the native format of the new map
 		{
-			const FScriptMapHelper SelfScriptMapHelper(Self->MapProp, Self->MapInstance);
+			FScriptMapHelper SelfScriptMapHelper(Self->MapProp, Self->MapInstance);
 			FScriptMapHelper NewScriptMapHelper(NewMap->MapProp, NewMap->MapInstance);
 
+			const int32 ElementCount = SelfScriptMapHelper.Num();
+
 			FString ExportedKey;
 			FString ExportedValue;
-			for (FScriptMapHelper::FIterator It(SelfScriptMapHelper); It; ++It)
+			for (int32 ElementIndex = 0, SparseIndex = 0; ElementIndex < ElementCount; ++SparseIndex)
 			{
+				if (!SelfScriptMapHelper.IsValidIndex(SparseIndex))
+				{
+					continue;
+				}
+
 				ExportedKey.Reset();
-				if (!SelfScriptMapHelper.GetKeyProperty()->ExportText_Direct(ExportedKey, SelfScriptMapHelper.GetKeyPtr(It), SelfScriptMapHelper.GetKeyPtr(It), nullptr, PPF_None))
+				if (!SelfScriptMapHelper.GetKeyProperty()->ExportText_Direct(ExportedKey, SelfScriptMapHelper.GetKeyPtr(SparseIndex), SelfScriptMapHelper.GetKeyPtr(SparseIndex), nullptr, PPF_None))
 				{
-					PyUtil::SetPythonError(PyExc_Exception, Self, *FString::Printf(TEXT("Failed to export text for key property '%s' (%s) at index %d"), *SelfScriptMapHelper.GetKeyProperty()->GetName(), *SelfScriptMapHelper.GetKeyProperty()->GetClass()->GetName(), It.GetLogicalIndex()));
+					PyUtil::SetPythonError(PyExc_Exception, Self, *FString::Printf(TEXT("Failed to export text for key property '%s' (%s) at index %d"), *SelfScriptMapHelper.GetKeyProperty()->GetName(), *SelfScriptMapHelper.GetKeyProperty()->GetClass()->GetName(), ElementIndex));
 					return nullptr;
 				}
 
 				ExportedValue.Reset();
-				if (!SelfScriptMapHelper.GetValueProperty()->ExportText_Direct(ExportedValue, SelfScriptMapHelper.GetValuePtr(It), SelfScriptMapHelper.GetValuePtr(It), nullptr, PPF_None))
+				if (!SelfScriptMapHelper.GetValueProperty()->ExportText_Direct(ExportedValue, SelfScriptMapHelper.GetValuePtr(SparseIndex), SelfScriptMapHelper.GetValuePtr(SparseIndex), nullptr, PPF_None))
 				{
-					PyUtil::SetPythonError(PyExc_Exception, Self, *FString::Printf(TEXT("Failed to export text for value property '%s' (%s) at index %d"), *SelfScriptMapHelper.GetValueProperty()->GetName(), *SelfScriptMapHelper.GetValueProperty()->GetClass()->GetName(), It.GetLogicalIndex()));
+					PyUtil::SetPythonError(PyExc_Exception, Self, *FString::Printf(TEXT("Failed to export text for value property '%s' (%s) at index %d"), *SelfScriptMapHelper.GetValueProperty()->GetName(), *SelfScriptMapHelper.GetValueProperty()->GetClass()->GetName(), ElementIndex));
 					return nullptr;
 				}
 
 				const int32 NewElementIndex = NewScriptMapHelper.AddDefaultValue_Invalid_NeedsRehash();
 				if (!NewScriptMapHelper.GetKeyProperty()->ImportText_Direct(*ExportedKey, NewScriptMapHelper.GetKeyPtr(NewElementIndex), nullptr, PPF_None))
 				{
-					PyUtil::SetPythonError(PyExc_Exception, Self, *FString::Printf(TEXT("Failed to import text '%s' key for property '%s' (%s) at index %d"), *ExportedKey, *NewScriptMapHelper.GetKeyProperty()->GetName(), *NewScriptMapHelper.GetKeyProperty()->GetClass()->GetName(), NewElementIndex));
+					PyUtil::SetPythonError(PyExc_Exception, Self, *FString::Printf(TEXT("Failed to import text '%s' key for property '%s' (%s) at index %d"), *ExportedKey, *NewScriptMapHelper.GetKeyProperty()->GetName(), *NewScriptMapHelper.GetKeyProperty()->GetClass()->GetName(), ElementIndex));
 					return nullptr;
 				}
 				if (!NewScriptMapHelper.GetValueProperty()->ImportText_Direct(*ExportedValue, NewScriptMapHelper.GetValuePtr(NewElementIndex), nullptr, PPF_None))
 				{
-					PyUtil::SetPythonError(PyExc_Exception, Self, *FString::Printf(TEXT("Failed to import text '%s' value for property '%s' (%s) at index %d"), *ExportedValue, *NewScriptMapHelper.GetValueProperty()->GetName(), *NewScriptMapHelper.GetValueProperty()->GetClass()->GetName(), NewElementIndex));
+					PyUtil::SetPythonError(PyExc_Exception, Self, *FString::Printf(TEXT("Failed to import text '%s' value for property '%s' (%s) at index %d"), *ExportedValue, *NewScriptMapHelper.GetValueProperty()->GetName(), *NewScriptMapHelper.GetValueProperty()->GetClass()->GetName(), ElementIndex));
 					return nullptr;
 				}
+
+				++ElementIndex;
 			}
 
 			NewScriptMapHelper.Rehash();
@@ -1032,29 +1041,40 @@ PyObject* FPyWrapperMap::PopItem(FPyWrapperMap* InSelf)
 	}
 
 	FScriptMapHelper SelfScriptMapHelper(InSelf->MapProp, InSelf->MapInstance);
-	const FScriptMapHelper::FIterator It(SelfScriptMapHelper);
-	if (It)
+	const int32 SelfElementCount = SelfScriptMapHelper.Num();
+
+	if (SelfElementCount == 0)
+	{
+		PyUtil::SetPythonError(PyExc_KeyError, InSelf, TEXT("Cannot pop from an empty map"));
+		return nullptr;
+	}
+
+	for (int32 SelfSparseIndex = 0; ; ++SelfSparseIndex)
 	{
+		if (!SelfScriptMapHelper.IsValidIndex(SelfSparseIndex))
+		{
+			continue;
+		}
+
 		FPyObjectPtr PyReturnKey;
-		if (!PyConversion::PythonizeProperty(SelfScriptMapHelper.GetKeyProperty(), SelfScriptMapHelper.GetKeyPtr(It), PyReturnKey.Get()))
+		if (!PyConversion::PythonizeProperty(SelfScriptMapHelper.GetKeyProperty(), SelfScriptMapHelper.GetKeyPtr(SelfSparseIndex), PyReturnKey.Get()))
 		{
 			PyUtil::SetPythonError(PyExc_TypeError, InSelf, *FString::Printf(TEXT("Failed to convert key property '%s' (%s) at index 0"), *SelfScriptMapHelper.GetKeyProperty()->GetName(), *SelfScriptMapHelper.GetKeyProperty()->GetClass()->GetName()));
 			return nullptr;
 		}
 
 		FPyObjectPtr PyReturnValue;
-		if (!PyConversion::PythonizeProperty(SelfScriptMapHelper.GetValueProperty(), SelfScriptMapHelper.GetValuePtr(It), PyReturnValue.Get()))
+		if (!PyConversion::PythonizeProperty(SelfScriptMapHelper.GetValueProperty(), SelfScriptMapHelper.GetValuePtr(SelfSparseIndex), PyReturnValue.Get()))
 		{
 			PyUtil::SetPythonError(PyExc_TypeError, InSelf, *FString::Printf(TEXT("Failed to convert value property '%s' (%s) at index 0"), *SelfScriptMapHelper.GetValueProperty()->GetName(), *SelfScriptMapHelper.GetValueProperty()->GetClass()->GetName()));
 			return nullptr;
 		}
 
-		SelfScriptMapHelper.RemoveAt(It.GetInternalIndex());
+		SelfScriptMapHelper.RemoveAt(SelfSparseIndex);
 		
 		return PyTuple_Pack(2, PyReturnKey.Get(), PyReturnValue.Get());
 	}
 
-	PyUtil::SetPythonError(PyExc_KeyError, InSelf, TEXT("Cannot pop from an empty map"));
 	return nullptr;
 }
 
@@ -1075,11 +1095,17 @@ int FPyWrapperMap::Update(FPyWrapperMap* InSelf, PyObject* InOther)
 	}
 
 	FScriptMapHelper SelfScriptMapHelper(InSelf->MapProp, InSelf->MapInstance);
-	const FScriptMapHelper OtherScriptMapHelper(Other->MapProp, Other->MapInstance);
+	FScriptMapHelper OtherScriptMapHelper(Other->MapProp, Other->MapInstance);
 
-	for (FScriptMapHelper::FIterator It(OtherScriptMapHelper); It; ++It)
+	const int32 OtherSparseCount = OtherScriptMapHelper.Num();
+	for (int32 OtherSparseIndex = 0; OtherSparseIndex < OtherSparseCount; ++OtherSparseIndex)
 	{
-		SelfScriptMapHelper.AddPair(OtherScriptMapHelper.GetKeyPtr(It), OtherScriptMapHelper.GetValuePtr(It));
+		if (!OtherScriptMapHelper.IsValidIndex(OtherSparseIndex))
+		{
+			continue;
+		}
+
+		SelfScriptMapHelper.AddPair(OtherScriptMapHelper.GetKeyPtr(OtherSparseIndex), OtherScriptMapHelper.GetValuePtr(OtherSparseIndex));
 	}
 
 	return 0;
@@ -1214,18 +1240,25 @@ PyTypeObject InitializePyWrapperMapType()
 				return nullptr;
 			}
 
-			const FScriptMapHelper SelfScriptMapHelper(InSelf->MapProp, InSelf->MapInstance);
+			FScriptMapHelper SelfScriptMapHelper(InSelf->MapProp, InSelf->MapInstance);
+			const int32 ElementCount = SelfScriptMapHelper.Num();
 
 			FString ExportedMap;
-			for (FScriptMapHelper::FIterator It(SelfScriptMapHelper); It; ++It)
+			for (int32 ElementIndex = 0, SparseIndex = 0; ElementIndex < ElementCount; ++SparseIndex)
 			{
-				if (It.GetLogicalIndex() > 0)
+				if (!SelfScriptMapHelper.IsValidIndex(SparseIndex))
+				{
+					continue;
+				}
+
+				if (ElementIndex > 0)
 				{
 					ExportedMap += TEXT(", ");
 				}
-				ExportedMap += PyUtil::GetFriendlyPropertyValue(SelfScriptMapHelper.GetKeyProperty(), SelfScriptMapHelper.GetKeyPtr(It), PPF_Delimited | PPF_IncludeTransient);
+				ExportedMap += PyUtil::GetFriendlyPropertyValue(SelfScriptMapHelper.GetKeyProperty(), SelfScriptMapHelper.GetKeyPtr(SparseIndex), PPF_Delimited | PPF_IncludeTransient);
 				ExportedMap += TEXT(": ");
-				ExportedMap += PyUtil::GetFriendlyPropertyValue(SelfScriptMapHelper.GetValueProperty(), SelfScriptMapHelper.GetValuePtr(It), PPF_Delimited | PPF_IncludeTransient);
+				ExportedMap += PyUtil::GetFriendlyPropertyValue(SelfScriptMapHelper.GetValueProperty(), SelfScriptMapHelper.GetValuePtr(SparseIndex), PPF_Delimited | PPF_IncludeTransient);
+				++ElementIndex;
 			}
 			return PyUnicode_FromFormat("{%s}", TCHAR_TO_UTF8(*ExportedMap));
 		}
diff --git a/Engine/Plugins/Experimental/PythonScriptPlugin/Source/PythonScriptPlugin/Private/PyWrapperSet.cpp b/Engine/Plugins/Experimental/PythonScriptPlugin/Source/PythonScriptPlugin/Private/PyWrapperSet.cpp
--- a/Engine/Plugins/Experimental/PythonScriptPlugin/Source/PythonScriptPlugin/Private/PyWrapperSet.cpp
+++ b/Engine/Plugins/Experimental/PythonScriptPlugin/Source/PythonScriptPlugin/Private/PyWrapperSet.cpp
@@ -381,22 +381,31 @@ FPyWrapperSet* FPyWrapperSet::CastPyObject(PyObject* InPyObject, PyTypeObject* I
 			FScriptSetHelper SelfScriptSetHelper(Self->SetProp, Self->SetInstance);
 			FScriptSetHelper NewScriptSetHelper(NewSet->SetProp, NewSet->SetInstance);
 
+			const int32 ElementCount = SelfScriptSetHelper.Num();
+
 			FString ExportedEntry;
-			for (FScriptSetHelper::FIterator It(SelfScriptSetHelper); It; ++It)
+			for (int32 ElementIndex = 0, SparseIndex = 0; ElementIndex < ElementCount; ++SparseIndex)
 			{
+				if (!SelfScriptSetHelper.IsValidIndex(SparseIndex))
+				{
+					continue;
+				}
+
 				ExportedEntry.Reset();
-				if (!SelfScriptSetHelper.GetElementProperty()->ExportText_Direct(ExportedEntry, SelfScriptSetHelper.GetElementPtr(It), SelfScriptSetHelper.GetElementPtr(It), nullptr, PPF_None))
+				if (!SelfScriptSetHelper.GetElementProperty()->ExportText_Direct(ExportedEntry, SelfScriptSetHelper.GetElementPtr(SparseIndex), SelfScriptSetHelper.GetElementPtr(SparseIndex), nullptr, PPF_None))
 				{
-					PyUtil::SetPythonError(PyExc_Exception, Self, *FString::Printf(TEXT("Failed to export text for element property '%s' (%s) at index %d"), *SelfScriptSetHelper.GetElementProperty()->GetName(), *SelfScriptSetHelper.GetElementProperty()->GetClass()->GetName(), It.GetLogicalIndex()));
+					PyUtil::SetPythonError(PyExc_Exception, Self, *FString::Printf(TEXT("Failed to export text for element property '%s' (%s) at index %d"), *SelfScriptSetHelper.GetElementProperty()->GetName(), *SelfScriptSetHelper.GetElementProperty()->GetClass()->GetName(), ElementIndex));
 					return nullptr;
 				}
 
 				const int32 NewElementIndex = NewScriptSetHelper.AddDefaultValue_Invalid_NeedsRehash();
 				if (!NewScriptSetHelper.GetElementProperty()->ImportText_Direct(*ExportedEntry, NewScriptSetHelper.GetElementPtr(NewElementIndex), nullptr, PPF_None))
 				{
-					PyUtil::SetPythonError(PyExc_Exception, Self, *FString::Printf(TEXT("Failed to import text '%s' element for property '%s' (%s) at index %d"), *ExportedEntry, *NewScriptSetHelper.GetElementProperty()->GetName(), *NewScriptSetHelper.GetElementProperty()->GetClass()->GetName(), NewElementIndex));
+					PyUtil::SetPythonError(PyExc_Exception, Self, *FString::Printf(TEXT("Failed to import text '%s' element for property '%s' (%s) at index %d"), *ExportedEntry, *NewScriptSetHelper.GetElementProperty()->GetName(), *NewScriptSetHelper.GetElementProperty()->GetClass()->GetName(), ElementIndex));
 					return nullptr;
 				}
+
+				++ElementIndex;
 			}
 
 			NewScriptSetHelper.Rehash();
@@ -554,22 +563,33 @@ PyObject* FPyWrapperSet::Pop(FPyWrapperSet* InSelf)
 	}
 
 	FScriptSetHelper SelfScriptSetHelper(InSelf->SetProp, InSelf->SetInstance);
-	const FScriptSetHelper::FIterator It(SelfScriptSetHelper);
-	if (It)
+	const int32 SelfElementCount = SelfScriptSetHelper.Num();
+
+	if (SelfElementCount == 0)
 	{
+		PyUtil::SetPythonError(PyExc_KeyError, InSelf, TEXT("Cannot pop from an empty set"));
+		return nullptr;
+	}
+
+	for (int32 SelfSparseIndex = 0; ; ++SelfSparseIndex)
+	{
+		if (!SelfScriptSetHelper.IsValidIndex(SelfSparseIndex))
+		{
+			continue;
+		}
+
 		PyObject* PyReturnValue = nullptr;
-		if (!PyConversion::PythonizeProperty(SelfScriptSetHelper.GetElementProperty(), SelfScriptSetHelper.GetElementPtr(It), PyReturnValue))
+		if (!PyConversion::PythonizeProperty(SelfScriptSetHelper.GetElementProperty(), SelfScriptSetHelper.GetElementPtr(SelfSparseIndex), PyReturnValue))
 		{
 			PyUtil::SetPythonError(PyExc_TypeError, InSelf, *FString::Printf(TEXT("Failed to convert element property '%s' (%s) at index 0"), *SelfScriptSetHelper.GetElementProperty()->GetName(), *SelfScriptSetHelper.GetElementProperty()->GetClass()->GetName()));
 			return nullptr;
 		}
 
-		SelfScriptSetHelper.RemoveAt(It.GetInternalIndex());
+		SelfScriptSetHelper.RemoveAt(SelfSparseIndex);
 
 		return PyReturnValue;
 	}
 
-	PyUtil::SetPythonError(PyExc_KeyError, InSelf, TEXT("Cannot pop from an empty set"));
 	return nullptr;
 }
 
@@ -621,10 +641,17 @@ int FPyWrapperSet::DifferenceUpdate(FPyWrapperSet* InSelf, PyObject* InOthers)
 			return false;
 		}
 
-		const FScriptSetHelper OtherScriptSetHelper(Other->SetProp, Other->SetInstance);
-		for (FScriptSetHelper::FIterator It(OtherScriptSetHelper); It; ++It)
+		FScriptSetHelper OtherScriptSetHelper(Other->SetProp, Other->SetInstance);
+		const int32 OtherSparseCount = OtherScriptSetHelper.GetMaxIndex();
+
+		for (int32 OtherSparseIndex = 0; OtherSparseIndex < OtherSparseCount; ++OtherSparseIndex)
 		{
-			const void* OtherElementPtr = OtherScriptSetHelper.GetElementPtr(It);
+			if (!OtherScriptSetHelper.IsValidIndex(OtherSparseIndex))
+			{
+				continue;
+			}
+
+			const void* OtherElementPtr = OtherScriptSetHelper.GetElementPtr(OtherSparseIndex);
 			SelfScriptSetHelper.RemoveElement(OtherElementPtr);
 		}
 
@@ -770,12 +797,18 @@ int FPyWrapperSet::SymmetricDifferenceUpdate(FPyWrapperSet* InSelf, PyObject* In
 		return -1;
 	}
 
-	const FScriptSetHelper OtherScriptSetHelper(Other->SetProp, Other->SetInstance);
+	FScriptSetHelper OtherScriptSetHelper(Other->SetProp, Other->SetInstance);
+	const int32 OtherSparseCount = OtherScriptSetHelper.GetMaxIndex();
 
 	// We need to go through the other set and remove any values from Self that are present in Other, and add any values from Other that aren't present in Self
-	for (FScriptSetHelper::FIterator It(OtherScriptSetHelper); It; ++It)
+	for (int32 OtherSparseIndex = 0; OtherSparseIndex < OtherSparseCount; ++OtherSparseIndex)
 	{
-		const void* OtherElementPtr = OtherScriptSetHelper.GetElementPtr(It);
+		if (!OtherScriptSetHelper.IsValidIndex(OtherSparseIndex))
+		{
+			continue;
+		}
+
+		const void* OtherElementPtr = OtherScriptSetHelper.GetElementPtr(OtherSparseIndex);
 		if (SelfScriptSetHelper.FindElementIndexFromHash(OtherElementPtr) == INDEX_NONE)
 		{
 			SelfScriptSetHelper.AddElement(OtherElementPtr);
@@ -824,10 +857,17 @@ int FPyWrapperSet::Update(FPyWrapperSet* InSelf, PyObject* InOthers)
 			return false;
 		}
 
-		const FScriptSetHelper OtherScriptSetHelper(Other->SetProp, Other->SetInstance);
-		for (FScriptSetHelper::FIterator It(OtherScriptSetHelper); It; ++It)
+		FScriptSetHelper OtherScriptSetHelper(Other->SetProp, Other->SetInstance);
+		const int32 OtherSparseCount = OtherScriptSetHelper.GetMaxIndex();
+
+		for (int32 OtherSparseIndex = 0; OtherSparseIndex < OtherSparseCount; ++OtherSparseIndex)
 		{
-			const void* OtherElementPtr = OtherScriptSetHelper.GetElementPtr(It);
+			if (!OtherScriptSetHelper.IsValidIndex(OtherSparseIndex))
+			{
+				continue;
+			}
+
+			const void* OtherElementPtr = OtherScriptSetHelper.GetElementPtr(OtherSparseIndex);
 			SelfScriptSetHelper.AddElement(OtherElementPtr);
 		}
 
@@ -896,12 +936,19 @@ int FPyWrapperSet::IsSubset(FPyWrapperSet* InSelf, PyObject* InOther)
 		return -1;
 	}
 
-	const FScriptSetHelper SelfScriptSetHelper(InSelf->SetProp, InSelf->SetInstance);
-	const FScriptSetHelper OtherScriptSetHelper(Other->SetProp, Other->SetInstance);
+	FScriptSetHelper SelfScriptSetHelper(InSelf->SetProp, InSelf->SetInstance);
+	const int32 SelfSparseCount = SelfScriptSetHelper.Num();
 
-	for (FScriptSetHelper::FIterator It(SelfScriptSetHelper); It; ++It)
+	FScriptSetHelper OtherScriptSetHelper(Other->SetProp, Other->SetInstance);
+
+	for (int32 SelfSparseIndex = 0; SelfSparseIndex < SelfSparseCount; ++SelfSparseIndex)
 	{
-		const void* SelfElementPtr = SelfScriptSetHelper.GetElementPtr(It);
+		if (!SelfScriptSetHelper.IsValidIndex(SelfSparseIndex))
+		{
+			continue;
+		}
+
+		const void* SelfElementPtr = SelfScriptSetHelper.GetElementPtr(SelfSparseIndex);
 		if (OtherScriptSetHelper.FindElementIndexFromHash(SelfElementPtr) == INDEX_NONE)
 		{
 			return 0;
@@ -926,12 +973,19 @@ int FPyWrapperSet::IsSuperset(FPyWrapperSet* InSelf, PyObject* InOther)
 		return -1;
 	}
 
-	const FScriptSetHelper SelfScriptSetHelper(InSelf->SetProp, InSelf->SetInstance);
-	const FScriptSetHelper OtherScriptSetHelper(Other->SetProp, Other->SetInstance);
+	FScriptSetHelper SelfScriptSetHelper(InSelf->SetProp, InSelf->SetInstance);
 
-	for (FScriptSetHelper::FIterator It(OtherScriptSetHelper); It; ++It)
+	FScriptSetHelper OtherScriptSetHelper(Other->SetProp, Other->SetInstance);
+	const int32 OtherSparseCount = OtherScriptSetHelper.Num();
+
+	for (int32 OtherSparseIndex = 0; OtherSparseIndex < OtherSparseCount; ++OtherSparseIndex)
 	{
-		const void* OtherElementPtr = OtherScriptSetHelper.GetElementPtr(It);
+		if (!OtherScriptSetHelper.IsValidIndex(OtherSparseIndex))
+		{
+			continue;
+		}
+
+		const void* OtherElementPtr = OtherScriptSetHelper.GetElementPtr(OtherSparseIndex);
 		if (SelfScriptSetHelper.FindElementIndexFromHash(OtherElementPtr) == INDEX_NONE)
 		{
 			return 0;
@@ -982,16 +1036,23 @@ PyTypeObject InitializePyWrapperSetType()
 				return nullptr;
 			}
 
-			const FScriptSetHelper SelfScriptSetHelper(InSelf->SetProp, InSelf->SetInstance);
+			FScriptSetHelper SelfScriptSetHelper(InSelf->SetProp, InSelf->SetInstance);
+			const int32 ElementCount = SelfScriptSetHelper.Num();
 
 			FString ExportedSet;
-			for (FScriptSetHelper::FIterator It(SelfScriptSetHelper); It; ++It)
+			for (int32 ElementIndex = 0, SparseIndex = 0; ElementIndex < ElementCount; ++SparseIndex)
 			{
-				if (It.GetLogicalIndex() > 0)
+				if (!SelfScriptSetHelper.IsValidIndex(SparseIndex))
+				{
+					continue;
+				}
+
+				if (ElementIndex > 0)
 				{
 					ExportedSet += TEXT(", ");
 				}
-				ExportedSet += PyUtil::GetFriendlyPropertyValue(SelfScriptSetHelper.GetElementProperty(), SelfScriptSetHelper.GetElementPtr(It), PPF_Delimited | PPF_IncludeTransient);
+				ExportedSet += PyUtil::GetFriendlyPropertyValue(SelfScriptSetHelper.GetElementProperty(), SelfScriptSetHelper.GetElementPtr(SparseIndex), PPF_Delimited | PPF_IncludeTransient);
+				++ElementIndex;
 			}
 			return PyUnicode_FromFormat("set([%s])", TCHAR_TO_UTF8(*ExportedSet));
 		}
diff --git a/Engine/Source/Editor/KismetCompiler/Private/KismetCompilerVMBackend.cpp b/Engine/Source/Editor/KismetCompiler/Private/KismetCompilerVMBackend.cpp
--- a/Engine/Source/Editor/KismetCompiler/Private/KismetCompilerVMBackend.cpp
+++ b/Engine/Source/Editor/KismetCompiler/Private/KismetCompilerVMBackend.cpp
@@ -1018,10 +1018,15 @@ public:
 				Writer << InnerProp;
 				Writer << ElementNum;
 
-				for (FScriptSetHelper::FIterator It(ScriptSetHelper); It; ++It)
+				for (int32 ElemIdx = 0, SparseIndex = 0; ElemIdx < ElementNum; ++SparseIndex)
 				{
-					uint8* RawElemData = ScriptSetHelper.GetElementPtr(It);
-					EmitInnerElementExpr(Term, InnerProp, RawElemData);
+					if (ScriptSet.IsValidIndex(SparseIndex))
+					{
+						uint8* RawElemData = ScriptSetHelper.GetElementPtr(SparseIndex);
+						EmitInnerElementExpr(Term, InnerProp, RawElemData);
+
+						++ElemIdx;
+					}
 				}
 				Writer << EX_EndSetConst;
 			}
@@ -1042,10 +1047,15 @@ public:
 				Writer << ValProp;
 				Writer << ElementNum;
 
-				for (FScriptMapHelper::FIterator It(ScriptMapHelper); It; ++It)
+				for (int32 ElemIdx = 0, SparseIndex = 0; ElemIdx < ElementNum; ++SparseIndex)
 				{
-					EmitInnerElementExpr(Term, KeyProp, ScriptMapHelper.GetKeyPtr(It));
-					EmitInnerElementExpr(Term, ValProp, ScriptMapHelper.GetValuePtr(It));
+					if (ScriptMap.IsValidIndex(SparseIndex))
+					{
+						EmitInnerElementExpr(Term, KeyProp, ScriptMapHelper.GetKeyPtr(SparseIndex));
+						EmitInnerElementExpr(Term, ValProp, ScriptMapHelper.GetValuePtr(SparseIndex));
+
+						++ElemIdx;
+					}
 				}
 				Writer << EX_EndMapConst;
 			}
diff --git a/Engine/Source/Editor/PropertyEditor/Private/PropertyHandleImpl.cpp b/Engine/Source/Editor/PropertyEditor/Private/PropertyHandleImpl.cpp
--- a/Engine/Source/Editor/PropertyEditor/Private/PropertyHandleImpl.cpp
+++ b/Engine/Source/Editor/PropertyEditor/Private/PropertyHandleImpl.cpp
@@ -366,13 +366,18 @@ FPropertyAccess::Result FPropertyValueImpl::ImportText( const TArray<FObjectBase
 							ElementProp->DestroyAndFreeValue(TempElementStorage);
 						};
 
-						for (FScriptSetHelper::FIterator It(Helper); It; ++It)
+						for (int32 Index = 0, ItemsLeft = Helper.Num(); ItemsLeft > 0; ++Index)
 						{
-							const uint8* Element = Helper.GetElementPtr(It);
-
-							if (Element != InBaseAddress && ElementProp->ImportText_Direct(*InElementValue, TempElementStorage, nullptr, 0) && ElementProp->Identical(Element, TempElementStorage))
+							if (Helper.IsValidIndex(Index))
 							{
-								return true;
+								--ItemsLeft;
+
+								const uint8* Element = Helper.GetElementPtr(Index);
+
+								if (Element != InBaseAddress && ElementProp->ImportText_Direct(*InElementValue, TempElementStorage, nullptr, 0) && ElementProp->Identical(Element, TempElementStorage))
+								{
+									return true;
+								}
 							}
 						}
 
@@ -417,14 +422,19 @@ FPropertyAccess::Result FPropertyValueImpl::ImportText( const TArray<FObjectBase
 							KeyProp->DestroyAndFreeValue(TempKeyStorage);
 						};
 
-						for (FScriptMapHelper::FIterator It(Helper); It; ++It)
+						for (int32 Index = 0, ItemsLeft = Helper.Num(); ItemsLeft > 0; ++Index)
 						{
-							const uint8* PairPtr = Helper.GetPairPtr(It);
-							const uint8* KeyPtr = KeyProp->ContainerPtrToValuePtr<const uint8>(PairPtr);
-
-							if (KeyPtr != InBaseAddress && KeyProp->ImportText_Direct(*InKeyValue, TempKeyStorage, nullptr, 0) && KeyProp->Identical(KeyPtr, TempKeyStorage))
+							if (Helper.IsValidIndex(Index))
 							{
-								return true;
+								--ItemsLeft;
+
+								const uint8* PairPtr = Helper.GetPairPtr(Index);
+								const uint8* KeyPtr = KeyProp->ContainerPtrToValuePtr<const uint8>(PairPtr);
+
+								if (KeyPtr != InBaseAddress && KeyProp->ImportText_Direct(*InKeyValue, TempKeyStorage, nullptr, 0) && KeyProp->Identical(KeyPtr, TempKeyStorage))
+								{
+									return true;
+								}
 							}
 						}
 
@@ -1329,13 +1339,20 @@ void FPropertyValueImpl::ClearChildren()
 							FObjectProperty* ElementObjectProperty = CastField<FObjectProperty>(SetProperty->ElementProp);
 							if (ShouldOwnInstance(ElementObjectProperty, PropertyNodePin.Get()))
 							{
-								for (FScriptSetHelper::FIterator It(SetHelper); It; ++It)
+								int32 ElementsToRemove = SetHelper.Num();
+								int32 Index = 0;
+								while (ElementsToRemove > 0)
 								{
-									if (UObject* InstancedObject = *reinterpret_cast<UObject**>(SetHelper.GetElementPtr(It)))
+									if (SetHelper.IsValidIndex(Index))
 									{
-										InstancedObject->Modify();
-										InstancedObject->Rename(nullptr, GetTransientPackage(), REN_DontCreateRedirectors);
+										if (UObject* InstancedObject = *reinterpret_cast<UObject**>(SetHelper.GetElementPtr(Index)))
+										{
+											InstancedObject->Modify();
+											InstancedObject->Rename(nullptr, GetTransientPackage(), REN_DontCreateRedirectors);
+										}
+										--ElementsToRemove;
 									}
+									++Index;
 								}
 							}
 
@@ -1350,13 +1367,20 @@ void FPropertyValueImpl::ClearChildren()
 							FObjectProperty* ValueObjectProperty = CastField<FObjectProperty>(MapProperty->ValueProp);
 							if (ShouldOwnInstance(ValueObjectProperty, PropertyNodePin.Get()))
 							{
-								for (FScriptMapHelper::FIterator It(MapHelper); It; ++It)
+								int32 ElementsToRemove = MapHelper.Num();
+								int32 Index = 0;
+								while (ElementsToRemove > 0)
 								{
-									if (UObject* InstancedObject = *reinterpret_cast<UObject**>(MapHelper.GetValuePtr(It)))
+									if (MapHelper.IsValidIndex(Index))
 									{
-										InstancedObject->Modify();
-										InstancedObject->Rename(nullptr, GetTransientPackage(), REN_DontCreateRedirectors);
+										if (UObject* InstancedObject = *reinterpret_cast<UObject**>(MapHelper.GetValuePtr(Index)))
+										{
+											InstancedObject->Modify();
+											InstancedObject->Rename(nullptr, GetTransientPackage(), REN_DontCreateRedirectors);
+										}
+										--ElementsToRemove;
 									}
+									++Index;
 								}
 							}
 
diff --git a/Engine/Source/Editor/PropertyEditor/Private/PropertyNode.cpp b/Engine/Source/Editor/PropertyEditor/Private/PropertyNode.cpp
--- a/Engine/Source/Editor/PropertyEditor/Private/PropertyNode.cpp
+++ b/Engine/Source/Editor/PropertyEditor/Private/PropertyNode.cpp
@@ -2114,12 +2114,19 @@ private:
 		if (MapProp != NULL)
 		{
 			FScriptMapHelper MapHelper(MapProp, PropertyValueAddress);
-			for (FScriptMapHelper::FIterator It(MapHelper); It; ++It)
+
+			int32 ItemsLeft = MapHelper.Num();
+			for (int32 Index = 0; ItemsLeft > 0; ++Index)
 			{
-				uint8* Data = MapHelper.GetPairPtr(It);
+				if (MapHelper.IsValidIndex(Index))
+				{
+					--ItemsLeft;
+
+					uint8* Data = MapHelper.GetPairPtr(Index);
 
-				ProcessProperty(MapProp->KeyProp, MapProp->KeyProp->ContainerPtrToValuePtr<uint8>(Data));
-				ProcessProperty(MapProp->ValueProp, MapProp->ValueProp->ContainerPtrToValuePtr<uint8>(Data));
+					ProcessProperty(MapProp->KeyProp, MapProp->KeyProp->ContainerPtrToValuePtr<uint8>(Data));
+					ProcessProperty(MapProp->ValueProp, MapProp->ValueProp->ContainerPtrToValuePtr<uint8>(Data));
+				}
 			}
 
 			bResult = true;
diff --git a/Engine/Source/Editor/UnrealEd/Private/Kismet2/ComponentEditorUtils.cpp b/Engine/Source/Editor/UnrealEd/Private/Kismet2/ComponentEditorUtils.cpp
--- a/Engine/Source/Editor/UnrealEd/Private/Kismet2/ComponentEditorUtils.cpp
+++ b/Engine/Source/Editor/UnrealEd/Private/Kismet2/ComponentEditorUtils.cpp
@@ -241,14 +241,17 @@ FProperty* FComponentEditorUtils::GetPropertyForEditableNativeComponent(const UA
 				}
 
 				FScriptMapHelper MapHelper(TestProperty, MapPropInstAddress);
-				for (FScriptMapHelper::FIterator It(MapHelper); It; ++It)
+				for (int32 MapSparseIndex = 0; MapSparseIndex < MapHelper.GetMaxIndex(); ++MapSparseIndex)
 				{
 					// For each value in the map (don't bother checking the keys, they won't be what the user can edit in this case)
-					const uint8* MapValueData = MapHelper.GetValuePtr(It);
-					UObject* ValueElement = MapValProp->GetObjectPropertyValue(MapValueData);
-					if (ValueElement != nullptr && ValueElement->GetFName() == NativeComponent->GetFName())
+					if (MapHelper.IsValidIndex(MapSparseIndex))
 					{
-						return MapValProp;
+						const uint8* MapValueData = MapHelper.GetValuePtr(MapSparseIndex);						
+						UObject* ValueElement = MapValProp->GetObjectPropertyValue(MapValueData);
+						if (ValueElement != nullptr && ValueElement->GetFName() == NativeComponent->GetFName())
+						{
+							return MapValProp;
+						}
 					}
 				}
 			}
@@ -268,14 +271,17 @@ FProperty* FComponentEditorUtils::GetPropertyForEditableNativeComponent(const UA
 				
 				// For each item in the set
 				FScriptSetHelper SetHelper(TestProperty, SetPropInstAddress);
-				for (FScriptSetHelper::FIterator It(SetHelper); It; ++It)
+				for (int32 i = 0; i < SetHelper.Num(); ++i)
 				{
-					const uint8* SetValData = SetHelper.GetElementPtr(It);
-					UObject* SetValueElem = SetValProp->GetObjectPropertyValue(SetValData);
-
-					if (SetValueElem != nullptr && SetValueElem->GetFName() == NativeComponent->GetFName())
+					if (SetHelper.IsValidIndex(i))
 					{
-						return SetValProp;
+						const uint8* SetValData = SetHelper.GetElementPtr(i);
+						UObject* SetValueElem = SetValProp->GetObjectPropertyValue(SetValData);
+						
+						if (SetValueElem != nullptr && SetValueElem->GetFName() == NativeComponent->GetFName())
+						{
+							return SetValProp;
+						}
 					}
 				}
 			}
diff --git a/Engine/Source/Runtime/CoreUObject/Private/Internationalization/StabilizeLocalizationKeys.cpp b/Engine/Source/Runtime/CoreUObject/Private/Internationalization/StabilizeLocalizationKeys.cpp
--- a/Engine/Source/Runtime/CoreUObject/Private/Internationalization/StabilizeLocalizationKeys.cpp
+++ b/Engine/Source/Runtime/CoreUObject/Private/Internationalization/StabilizeLocalizationKeys.cpp
@@ -109,15 +109,23 @@ void StabilizeLocalizationKeys::StabilizeLocalizationKeysForProperty(FProperty*
 		{
 			FScriptSetHelper ScriptSetHelper(SetProp, InPropData);
 
-			for (FScriptSetHelper::FIterator It(ScriptSetHelper); It; ++It)
+			const int32 ElementCount = ScriptSetHelper.Num();
+			for (int32 RawIndex = 0, ElementIndex = 0; ElementIndex < ElementCount; ++RawIndex)
 			{
+				if (!ScriptSetHelper.IsValidIndex(RawIndex))
+				{
+					continue;
+				}
+
 				StabilizeLocalizationKeysForProperty(
 					SetProp->ElementProp,
-					ScriptSetHelper.GetElementPtr(It),
+					ScriptSetHelper.GetElementPtr(RawIndex),
 					InNamespace,
-					FString::Printf(TEXT("%s_Index%d"), *PropKeyRoot, It.GetLogicalIndex()),
+					FString::Printf(TEXT("%s_Index%d"), *PropKeyRoot, ElementIndex),
 					/*bAppendPropertyNameToKey*/false
 					);
+
+				++ElementIndex;
 			}
 
 			ScriptSetHelper.Rehash();
@@ -131,15 +139,21 @@ void StabilizeLocalizationKeys::StabilizeLocalizationKeysForProperty(FProperty*
 		{
 			FScriptMapHelper ScriptMapHelper(MapProp, InPropData);
 
-			for (FScriptMapHelper::FIterator It(ScriptMapHelper); It; ++It)
+			const int32 ElementCount = ScriptMapHelper.Num();
+			for (int32 RawIndex = 0, ElementIndex = 0; ElementIndex < ElementCount; ++RawIndex)
 			{
+				if (!ScriptMapHelper.IsValidIndex(RawIndex))
+				{
+					continue;
+				}
+
 				if (ShouldStabilizeInnerProperty(MapProp->KeyProp))
 				{
 					StabilizeLocalizationKeysForProperty(
 						MapProp->KeyProp,
-						ScriptMapHelper.GetKeyPtr(It),
+						ScriptMapHelper.GetKeyPtr(RawIndex),
 						InNamespace,
-						FString::Printf(TEXT("%s_KeyIndex%d"), *PropKeyRoot, It.GetLogicalIndex()),
+						FString::Printf(TEXT("%s_KeyIndex%d"), *PropKeyRoot, ElementIndex),
 						/*bAppendPropertyNameToKey*/false
 						);
 				}
@@ -148,12 +162,14 @@ void StabilizeLocalizationKeys::StabilizeLocalizationKeysForProperty(FProperty*
 				{
 					StabilizeLocalizationKeysForProperty(
 						MapProp->ValueProp,
-						ScriptMapHelper.GetValuePtr(It),
+						ScriptMapHelper.GetValuePtr(RawIndex),
 						InNamespace,
-						FString::Printf(TEXT("%s_ValueIndex%d"), *PropKeyRoot, It.GetLogicalIndex()),
+						FString::Printf(TEXT("%s_ValueIndex%d"), *PropKeyRoot, ElementIndex),
 						/*bAppendPropertyNameToKey*/false
 						);
 				}
+
+				++ElementIndex;
 			}
 
 			if (ShouldStabilizeInnerProperty(MapProp->KeyProp))
diff --git a/Engine/Source/Runtime/CoreUObject/Private/Serialization/PropertyLocalizationDataGathering.cpp b/Engine/Source/Runtime/CoreUObject/Private/Serialization/PropertyLocalizationDataGathering.cpp
--- a/Engine/Source/Runtime/CoreUObject/Private/Serialization/PropertyLocalizationDataGathering.cpp
+++ b/Engine/Source/Runtime/CoreUObject/Private/Serialization/PropertyLocalizationDataGathering.cpp
@@ -443,9 +443,15 @@ void FPropertyLocalizationDataGatherer::GatherLocalizationDataFromChildTextPrope
 
 			// Iterate over all elements of the map.
 			FScriptMapHelper ScriptMapHelper(MapProperty, ElementValueAddress);
-			for (FScriptMapHelper::FIterator It(ScriptMapHelper); It; ++It)
+			const int32 ElementCount = ScriptMapHelper.Num();
+			for(int32 j = 0, ElementIndex = 0; ElementIndex < ElementCount; ++j)
 			{
-				const uint8* MapPairPtr = ScriptMapHelper.GetPairPtr(It);
+				if (!ScriptMapHelper.IsValidIndex(j))
+				{
+					continue;
+				}
+
+				const uint8* MapPairPtr = ScriptMapHelper.GetPairPtr(j);
 
 				if (bGatherMapKey)
 				{
@@ -454,7 +460,7 @@ void FPropertyLocalizationDataGatherer::GatherLocalizationDataFromChildTextPrope
 						PathToInnerElement.Reserve(PathToElement.Len() + 20); // +20 for some slack for the number, braces, and description
 						PathToInnerElement += PathToElement;
 						PathToInnerElement += TEXT('(');
-						PathToInnerElement.AppendInt(It.GetLogicalIndex());
+						PathToInnerElement.AppendInt(ElementIndex);
 						PathToInnerElement += TEXT(" - Key)");
 					}
 
@@ -469,13 +475,15 @@ void FPropertyLocalizationDataGatherer::GatherLocalizationDataFromChildTextPrope
 						PathToInnerElement.Reserve(PathToElement.Len() + 20); // +20 for some slack for the number, braces, and description
 						PathToInnerElement += PathToElement;
 						PathToInnerElement += TEXT('(');
-						PathToInnerElement.AppendInt(It.GetLogicalIndex());
+						PathToInnerElement.AppendInt(ElementIndex);
 						PathToInnerElement += TEXT(" - Value)");
 					}
 
 					const uint8* MapValuePtr = MapPairPtr + MapProperty->MapLayout.ValueOffset;
 					GatherLocalizationDataFromChildTextProperties(PathToInnerElement, MapProperty->ValueProp, MapValuePtr, nullptr, ElementChildPropertyGatherTextFlags);
 				}
+
+				++ElementIndex;
 			}
 		}
 		// Property is a set property.
@@ -483,19 +491,26 @@ void FPropertyLocalizationDataGatherer::GatherLocalizationDataFromChildTextPrope
 		{
 			// Iterate over all elements of the Set.
 			FScriptSetHelper ScriptSetHelper(SetProperty, ElementValueAddress);
-			for (FScriptSetHelper::FIterator It(ScriptSetHelper); It; ++It)
+			const int32 ElementCount = ScriptSetHelper.Num();
+			for(int32 j = 0, ElementIndex = 0; ElementIndex < ElementCount; ++j)
 			{
+				if (!ScriptSetHelper.IsValidIndex(j))
+				{
+					continue;
+				}
+
 				FString PathToInnerElement;
 				{
 					PathToInnerElement.Reserve(PathToElement.Len() + 10); // +10 for some slack for the number and braces
 					PathToInnerElement += PathToElement;
 					PathToInnerElement += TEXT('(');
-					PathToInnerElement.AppendInt(It.GetLogicalIndex());
+					PathToInnerElement.AppendInt(ElementIndex);
 					PathToInnerElement += TEXT(')');
 				}
 
-				const uint8* ElementPtr = ScriptSetHelper.GetElementPtr(It);
+				const uint8* ElementPtr = ScriptSetHelper.GetElementPtr(j);
 				GatherLocalizationDataFromChildTextProperties(PathToInnerElement, SetProperty->ElementProp, ElementPtr, nullptr, ElementChildPropertyGatherTextFlags);
+				++ElementIndex;
 			}
 		}
 		// Property is a struct property.
diff --git a/Engine/Source/Runtime/CoreUObject/Private/UObject/LinkerPlaceholderBase.cpp b/Engine/Source/Runtime/CoreUObject/Private/UObject/LinkerPlaceholderBase.cpp
--- a/Engine/Source/Runtime/CoreUObject/Private/UObject/LinkerPlaceholderBase.cpp
+++ b/Engine/Source/Runtime/CoreUObject/Private/UObject/LinkerPlaceholderBase.cpp
@@ -129,10 +129,15 @@ int32 FLinkerPlaceholderObjectImpl::ResolvePlaceholderValues(const TArray<FField
 			// because we can't know which set entry was set with a reference 
 			// to this object, we have to comb through them all
 			FScriptSetHelper SetHelper(SetProperty, ValueAddress);
-			for (FScriptSetHelper::FIterator It(SetHelper); It; ++It)
+			int32 Num = SetHelper.Num();
+			for (int32 SetIndex = 0; Num; ++SetIndex)
 			{
-				uint8* ElementAddress = SetHelper.GetElementPtr(It);
-				ReplacementCount += ResolvePlaceholderValues(PropertyChain, PropertyIndex - 1, ElementAddress, OldValue, ReplacementValue);
+				if (SetHelper.IsValidIndex(SetIndex))
+				{
+					--Num;
+					uint8* ElementAddress = SetHelper.GetElementPtr(SetIndex);
+					ReplacementCount += ResolvePlaceholderValues(PropertyChain, PropertyIndex - 1, ElementAddress, OldValue, ReplacementValue);
+				}
 			}
 
 			// the above recursive call chewed through the rest of the
@@ -153,17 +158,23 @@ int32 FLinkerPlaceholderObjectImpl::ResolvePlaceholderValues(const TArray<FField
 			// because we can't know which map entry was set with a reference 
 			// to this object, we have to comb through them all
 			FScriptMapHelper MapHelper(MapProperty, ValueAddress);
-			for (FScriptMapHelper::FIterator It(MapHelper); It; ++It)
+			int32 Num = MapHelper.Num();
+			for (int32 MapIndex = 0; Num; ++MapIndex)
 			{
-				if (NextProperty == MapProperty->KeyProp)
+				if (MapHelper.IsValidIndex(MapIndex))
 				{
-					uint8* KeyAddress = MapHelper.GetKeyPtr(It);
-					ReplacementCount += ResolvePlaceholderValues(PropertyChain, PropertyIndex - 1, KeyAddress, OldValue, ReplacementValue);
-				}
-				else if (NextProperty == MapProperty->ValueProp)
-				{
-					uint8* MapValueAddress = MapHelper.GetValuePtr(It);
-					ReplacementCount += ResolvePlaceholderValues(PropertyChain, PropertyIndex - 1, MapValueAddress, OldValue, ReplacementValue);
+					--Num;
+
+					if (NextProperty == MapProperty->KeyProp)
+					{
+						uint8* KeyAddress = MapHelper.GetKeyPtr(MapIndex);
+						ReplacementCount += ResolvePlaceholderValues(PropertyChain, PropertyIndex - 1, KeyAddress, OldValue, ReplacementValue);
+					}
+					else if (NextProperty == MapProperty->ValueProp)
+					{
+						uint8* MapValueAddress = MapHelper.GetValuePtr(MapIndex);
+						ReplacementCount += ResolvePlaceholderValues(PropertyChain, PropertyIndex - 1, MapValueAddress, OldValue, ReplacementValue);
+					}
 				}
 			}
 
diff --git a/Engine/Source/Runtime/CoreUObject/Private/UObject/PropertyMap.cpp b/Engine/Source/Runtime/CoreUObject/Private/UObject/PropertyMap.cpp
--- a/Engine/Source/Runtime/CoreUObject/Private/UObject/PropertyMap.cpp
+++ b/Engine/Source/Runtime/CoreUObject/Private/UObject/PropertyMap.cpp
@@ -19,41 +19,45 @@ namespace UEMapProperty_Private
 	/**
 	 * Checks if any of the pairs in the map compare equal to the one passed.
 	 *
-	 * @param  MapHelper    The map to search through.
-	 * @param  LogicalIndex The index in the map to start searching from.
-	 * @param  Num          The number of elements to compare.
+	 * @param  MapHelper  The map to search through.
+	 * @param  Index      The index in the map to start searching from.
+	 * @param  Num        The number of elements to compare.
 	 */
-	bool AnyEqual(const FScriptMapHelper& MapHelper, const int32 LogicalIndex, int32 Num, const uint8* PairToCompare, const uint32 PortFlags)
+	bool AnyEqual(const FScriptMapHelper& MapHelper, int32 Index, int32 Num, const uint8* PairToCompare, uint32 PortFlags)
 	{
-		const FProperty* KeyProp   = MapHelper.GetKeyProperty();
-		const FProperty* ValueProp = MapHelper.GetValueProperty();
-		const int32 ValueOffset = MapHelper.MapLayout.ValueOffset;
+		FProperty* KeyProp   = MapHelper.GetKeyProperty();
+		FProperty* ValueProp = MapHelper.GetValueProperty();
 
-		FScriptMapHelper::FIterator IteratorA(MapHelper, LogicalIndex);
-		for (; IteratorA && Num; --Num, ++IteratorA)
+		int32 ValueOffset = MapHelper.MapLayout.ValueOffset;
+
+		for (; Num; --Num)
 		{
-			if (KeyProp->Identical(MapHelper.GetPairPtr(IteratorA), PairToCompare, PortFlags) && ValueProp->Identical(MapHelper.GetPairPtr(IteratorA) + ValueOffset, PairToCompare + ValueOffset, PortFlags))
+			while (!MapHelper.IsValidIndex(Index))
+			{
+				++Index;
+			}
+
+			if (KeyProp->Identical(MapHelper.GetPairPtr(Index), PairToCompare, PortFlags) && ValueProp->Identical(MapHelper.GetPairPtr(Index) + ValueOffset, PairToCompare + ValueOffset, PortFlags))
 			{
 				return true;
 			}
+
+			++Index;
 		}
 
 		return false;
 	}
 
-	bool RangesContainSameAmountsOfVal(const FScriptMapHelper& MapHelperA, const int32 LogicalIndexA, const FScriptMapHelper& MapHelperB, const int32 LogicalIndexB, int32 Num, const uint8* PairToCompare, const uint32 PortFlags)
+	bool RangesContainSameAmountsOfVal(const FScriptMapHelper& MapHelperA, int32 IndexA, const FScriptMapHelper& MapHelperB, int32 IndexB, int32 Num, const uint8* PairToCompare, uint32 PortFlags)
 	{
-		const FProperty* KeyProp   = MapHelperA.GetKeyProperty();
-		const FProperty* ValueProp = MapHelperA.GetValueProperty();
+		FProperty* KeyProp   = MapHelperA.GetKeyProperty();
+		FProperty* ValueProp = MapHelperA.GetValueProperty();
 
 		// Ensure that both maps are the same type
 		check(KeyProp   == MapHelperB.GetKeyProperty());
 		check(ValueProp == MapHelperB.GetValueProperty());
 
-		const int32 ValueOffset = MapHelperA.MapLayout.ValueOffset;
-
-		FScriptMapHelper::FIterator IteratorA(MapHelperA, LogicalIndexA);
-		FScriptMapHelper::FIterator IteratorB(MapHelperB, LogicalIndexB);
+		int32 ValueOffset = MapHelperA.MapLayout.ValueOffset;
 
 		int32 CountA = 0;
 		int32 CountB = 0;
@@ -64,8 +68,18 @@ namespace UEMapProperty_Private
 				return CountA == CountB;
 			}
 
-			const uint8* PairA = MapHelperA.GetPairPtr(IteratorA);
-			const uint8* PairB = MapHelperB.GetPairPtr(IteratorB);
+			while (!MapHelperA.IsValidIndex(IndexA))
+			{
+				++IndexA;
+			}
+
+			while (!MapHelperB.IsValidIndex(IndexB))
+			{
+				++IndexB;
+			}
+
+			const uint8* PairA = MapHelperA.GetPairPtr(IndexA);
+			const uint8* PairB = MapHelperB.GetPairPtr(IndexB);
 			if (PairA == PairToCompare || (KeyProp->Identical(PairA, PairToCompare, PortFlags) && ValueProp->Identical(PairA + ValueOffset, PairToCompare + ValueOffset, PortFlags)))
 			{
 				++CountA;
@@ -76,16 +90,16 @@ namespace UEMapProperty_Private
 				++CountB;
 			}
 
-			++IteratorA;
-			++IteratorB;
+			++IndexA;
+			++IndexB;
 			--Num;
 		}
 	}
 
-	bool IsPermutation(const FScriptMapHelper& MapHelperA, const FScriptMapHelper& MapHelperB, const uint32 PortFlags)
+	bool IsPermutation(const FScriptMapHelper& MapHelperA, const FScriptMapHelper& MapHelperB, uint32 PortFlags)
 	{
-		const FProperty* KeyProp   = MapHelperA.GetKeyProperty();
-		const FProperty* ValueProp = MapHelperA.GetValueProperty();
+		FProperty* KeyProp   = MapHelperA.GetKeyProperty();
+		FProperty* ValueProp = MapHelperA.GetValueProperty();
 
 		// Ensure that both maps are the same type
 		check(KeyProp   == MapHelperB.GetKeyProperty());
@@ -97,11 +111,11 @@ namespace UEMapProperty_Private
 			return false;
 		}
 
-		const int32 ValueOffset = MapHelperA.MapLayout.ValueOffset;
+		int32 ValueOffset = MapHelperA.MapLayout.ValueOffset;
 
 		// Skip over common initial sequence
-		FScriptMapHelper::FIterator IteratorA(MapHelperA);
-		FScriptMapHelper::FIterator IteratorB(MapHelperB);
+		int32 IndexA = 0;
+		int32 IndexB = 0;
 		for (;;)
 		{
 			if (Num == 0)
@@ -109,8 +123,18 @@ namespace UEMapProperty_Private
 				return true;
 			}
 
-			const uint8* PairA = MapHelperA.GetPairPtr(IteratorA);
-			const uint8* PairB = MapHelperB.GetPairPtr(IteratorB);
+			while (!MapHelperA.IsValidIndex(IndexA))
+			{
+				++IndexA;
+			}
+
+			while (!MapHelperB.IsValidIndex(IndexB))
+			{
+				++IndexB;
+			}
+
+			const uint8* PairA = MapHelperA.GetPairPtr(IndexA);
+			const uint8* PairB = MapHelperB.GetPairPtr(IndexB);
 			if (!KeyProp->Identical(PairA, PairB, PortFlags))
 			{
 				break;
@@ -121,17 +145,17 @@ namespace UEMapProperty_Private
 				break;
 			}
 
-			++IteratorA;
-			++IteratorB;
+			++IndexA;
+			++IndexB;
 			--Num;
 		}
 
-		const int32 FirstIndexA = IteratorA.GetLogicalIndex();
-		const int32 FirstIndexB = IteratorB.GetLogicalIndex();
-		const int32 FirstNum    = Num;
+		int32 FirstIndexA = IndexA;
+		int32 FirstIndexB = IndexB;
+		int32 FirstNum    = Num;
 		for (;;)
 		{
-			const uint8* PairA = MapHelperA.GetPairPtr(IteratorA);
+			const uint8* PairA = MapHelperA.GetPairPtr(IndexA);
 			if (!AnyEqual(MapHelperA, FirstIndexA, FirstNum - Num, PairA, PortFlags) && !RangesContainSameAmountsOfVal(MapHelperA, FirstIndexA, MapHelperB, FirstIndexB, FirstNum, PairA, PortFlags))
 			{
 				return false;
@@ -143,7 +167,11 @@ namespace UEMapProperty_Private
 				return true;
 			}
 
-			++IteratorA;
+			++IndexA;
+			while (!MapHelperA.IsValidIndex(IndexA))
+			{
+				++IndexA;
+			}
 		}
 	}
 }
@@ -464,12 +492,18 @@ void FMapProperty::SerializeItem(FStructuredArchive::FSlot Slot, void* Value, co
 		// Determine how many keys are missing from the object
 		if (Defaults)
 		{
-			for (FScriptMapHelper::FIterator Iterator(DefaultsHelper); Iterator; ++Iterator)
+			for (int32 Index = 0, Count = DefaultsHelper.Num(); Count; ++Index)
 			{
-				uint8* DefaultPairPtr = DefaultsHelper.GetPairPtr(Iterator);
-				if (!MapHelper.FindMapPairPtrWithKey(DefaultPairPtr))
+				uint8* DefaultPairPtr = DefaultsHelper.GetPairPtrWithoutCheck(Index);
+
+				if (DefaultsHelper.IsValidIndex(Index))
 				{
-					Indices.Add(Iterator.GetInternalIndex());
+					if (!MapHelper.FindMapPairPtrWithKey(DefaultPairPtr))
+					{
+						Indices.Add(Index);
+					}
+
+					--Count;
 				}
 			}
 		}
@@ -489,14 +523,19 @@ void FMapProperty::SerializeItem(FStructuredArchive::FSlot Slot, void* Value, co
 		if (Defaults)
 		{
 			Indices.Empty(Indices.Num());
-			for (FScriptMapHelper::FIterator Iterator(MapHelper); Iterator; ++Iterator)
+			for (int32 Index = 0, Count = MapHelper.Num(); Count; ++Index)
 			{
-				uint8* ValuePairPtr   = MapHelper.GetPairPtr(Iterator);
-				uint8* DefaultPairPtr = DefaultsHelper.FindMapPairPtrWithKey(ValuePairPtr);
-
-				if (!DefaultPairPtr || !ValueProp->Identical(ValuePairPtr + MapLayout.ValueOffset, DefaultPairPtr + MapLayout.ValueOffset))
+				if (MapHelper.IsValidIndex(Index))
 				{
-					Indices.Add(Iterator.GetInternalIndex());
+					uint8* ValuePairPtr   = MapHelper.GetPairPtrWithoutCheck(Index);
+					uint8* DefaultPairPtr = DefaultsHelper.FindMapPairPtrWithKey(ValuePairPtr);
+
+					if (!DefaultPairPtr || !ValueProp->Identical(ValuePairPtr + MapLayout.ValueOffset, DefaultPairPtr + MapLayout.ValueOffset))
+					{
+						Indices.Add(Index);
+					}
+
+					--Count;
 				}
 			}
 
@@ -523,18 +562,24 @@ void FMapProperty::SerializeItem(FStructuredArchive::FSlot Slot, void* Value, co
 			int32 Num = MapHelper.Num();
 			FStructuredArchive::FArray EntriesArray = Record.EnterArray(TEXT("Entries"), Num);
 
-			for (FScriptMapHelper::FIterator Iterator(MapHelper); Iterator; ++Iterator)
+			for (int32 Index = 0; Num; ++Index)
 			{
-				FStructuredArchive::FRecord EntryRecord = EntriesArray.EnterElement().EnterRecord();
-
-				uint8* ValuePairPtr = MapHelper.GetPairPtr(Iterator);
+				if (MapHelper.IsValidIndex(Index))
 				{
-					FSerializedPropertyScope SerializedProperty(UnderlyingArchive, KeyProp, this);
-					KeyProp->SerializeItem(EntryRecord.EnterField(TEXT("Key")), ValuePairPtr);
-				}
-				{
-					FSerializedPropertyScope SerializedProperty(UnderlyingArchive, ValueProp, this);
-					ValueProp->SerializeItem(EntryRecord.EnterField(TEXT("Value")), ValuePairPtr + MapLayout.ValueOffset);
+					FStructuredArchive::FRecord EntryRecord = EntriesArray.EnterElement().EnterRecord();
+
+					uint8* ValuePairPtr = MapHelper.GetPairPtrWithoutCheck(Index);
+
+					{
+						FSerializedPropertyScope SerializedProperty(UnderlyingArchive, KeyProp, this);
+						KeyProp->SerializeItem(EntryRecord.EnterField(TEXT("Key")), ValuePairPtr);
+					}
+					{
+						FSerializedPropertyScope SerializedProperty(UnderlyingArchive, ValueProp, this);
+						ValueProp->SerializeItem(EntryRecord.EnterField(TEXT("Value")), ValuePairPtr + MapLayout.ValueOffset);
+					}
+
+					--Num;
 				}
 			}
 		}
@@ -978,7 +1023,7 @@ bool FMapProperty::PassCPPArgsByRef() const
  * 
  * @param	Data				pointer to the address of the instanced object referenced by this UComponentProperty
  * @param	DefaultData			pointer to the address of the default value of the instanced object referenced by this UComponentProperty
- * @param	InOwner				the object that contains this property's data
+ * @param	Owner				the object that contains this property's data
  * @param	InstanceGraph		contains the mappings of instanced objects and components to their templates
  */
 void FMapProperty::InstanceSubobjects(void* Data, void const* DefaultData, UObject* InOwner, FObjectInstancingGraph* InstanceGraph)
@@ -988,8 +1033,8 @@ void FMapProperty::InstanceSubobjects(void* Data, void const* DefaultData, UObje
 		return;
 	}
 
-	const bool bInstancedKey   = KeyProp  ->ContainsInstancedObjectProperty();
-	const bool bInstancedValue = ValueProp->ContainsInstancedObjectProperty();
+	bool bInstancedKey   = KeyProp  ->ContainsInstancedObjectProperty();
+	bool bInstancedValue = ValueProp->ContainsInstancedObjectProperty();
 
 	if (!bInstancedKey && !bInstancedValue)
 	{
@@ -1001,36 +1046,48 @@ void FMapProperty::InstanceSubobjects(void* Data, void const* DefaultData, UObje
 	if (DefaultData)
 	{
 		FScriptMapHelper DefaultMapHelper(this, DefaultData);
-		for (FScriptMapHelper::FIterator It(MapHelper.CreateIterator()); It; ++It)
-		{
-			uint8* PairPtr = MapHelper.GetPairPtr(It);
-			const uint8* DefaultPairPtr = DefaultMapHelper.FindMapPairPtrWithKey(PairPtr, /*IndexHint*/ It.GetLogicalIndex());
+		int32            DefaultNum = DefaultMapHelper.Num();
 
-			if (bInstancedKey)
+		for (int32 Index = 0, Num = MapHelper.Num(); Num; ++Index)
+		{
+			if (MapHelper.IsValidIndex(Index))
 			{
-				KeyProp->InstanceSubobjects(PairPtr, DefaultPairPtr, InOwner, InstanceGraph);
-			}
+				uint8* PairPtr        = MapHelper.GetPairPtr(Index);
+				uint8* DefaultPairPtr = DefaultMapHelper.FindMapPairPtrWithKey(PairPtr, Index);
 
-			if (bInstancedValue)
-			{
-				ValueProp->InstanceSubobjects(PairPtr + MapLayout.ValueOffset, DefaultPairPtr ? DefaultPairPtr + MapLayout.ValueOffset : nullptr, InOwner, InstanceGraph);
+				if (bInstancedKey)
+				{
+					KeyProp->InstanceSubobjects(PairPtr, DefaultPairPtr, InOwner, InstanceGraph);
+				}
+
+				if (bInstancedValue)
+				{
+					ValueProp->InstanceSubobjects(PairPtr + MapLayout.ValueOffset, DefaultPairPtr ? DefaultPairPtr + MapLayout.ValueOffset : nullptr, InOwner, InstanceGraph);
+				}
+
+				--Num;
 			}
 		}
 	}
 	else
 	{
-		for (FScriptMapHelper::FIterator It(MapHelper.CreateIterator()); It; ++It)
+		for (int32 Index = 0, Num = MapHelper.Num(); Num; ++Index)
 		{
-			uint8* PairPtr = MapHelper.GetPairPtr(It);
-
-			if (bInstancedKey)
+			if (MapHelper.IsValidIndex(Index))
 			{
-				KeyProp->InstanceSubobjects(PairPtr, nullptr, InOwner, InstanceGraph);
-			}
+				uint8* PairPtr = MapHelper.GetPairPtr(Index);
 
-			if (bInstancedValue)
-			{
-				ValueProp->InstanceSubobjects(PairPtr + MapLayout.ValueOffset, nullptr, InOwner, InstanceGraph);
+				if (bInstancedKey)
+				{
+					KeyProp->InstanceSubobjects(PairPtr, nullptr, InOwner, InstanceGraph);
+				}
+
+				if (bInstancedValue)
+				{
+					ValueProp->InstanceSubobjects(PairPtr + MapLayout.ValueOffset, nullptr, InOwner, InstanceGraph);
+				}
+
+				--Num;
 			}
 		}
 	}
@@ -1293,20 +1350,29 @@ void FMapProperty::GetInnerFields(TArray<FField*>& OutFields)
 	}
 }
 
-void* FMapProperty::GetValueAddressAtIndex_Direct(const FProperty* Inner, void* InValueAddress, const int32 LogicalIndex) const
+void* FMapProperty::GetValueAddressAtIndex_Direct(const FProperty* Inner, void* InValueAddress, int32 Index) const
 {
+	FScriptMapHelper MapHelper(this, InValueAddress);
 	checkf(Inner == KeyProp || Inner == ValueProp, TEXT("Inner property must be either KeyProp or ValueProp"));
 
-	FScriptMapHelper MapHelper(this, InValueAddress);
-	const int32 InternalIndex = MapHelper.FindInternalIndex(LogicalIndex);
-	if (InternalIndex != INDEX_NONE)
+	for (int32 MapIndex = 0, Num = MapHelper.Num(), LocalIndex = Index; LocalIndex >= 0 && Num > 0; ++MapIndex)
 	{
-		if (Inner == KeyProp)
+		if (MapHelper.IsValidIndex(MapIndex))
 		{
-			return MapHelper.GetKeyPtr(InternalIndex);
+			if (LocalIndex == 0)
+			{
+				if (Inner == KeyProp)
+				{
+					return MapHelper.GetKeyPtr(MapIndex);
+				}
+				else
+				{
+					return MapHelper.GetValuePtr(MapIndex);
+				}
+			}
+			LocalIndex--;
+			Num--;
 		}
-
-		return MapHelper.GetValuePtr(InternalIndex);
 	}
 	return nullptr;
 }
diff --git a/Engine/Source/Runtime/CoreUObject/Private/UObject/PropertySet.cpp b/Engine/Source/Runtime/CoreUObject/Private/UObject/PropertySet.cpp
--- a/Engine/Source/Runtime/CoreUObject/Private/UObject/PropertySet.cpp
+++ b/Engine/Source/Runtime/CoreUObject/Private/UObject/PropertySet.cpp
@@ -17,36 +17,39 @@ namespace UESetProperty_Private
 	/**
 	 * Checks if any of the elements in the set compare equal to the one passed.
 	 *
-	 * @param  SetHelper    The set to search through.
-	 * @param  LogicalIndex The index in the set to start searching from.
-	 * @param  Num          The number of elements to compare.
+	 * @param  SetHelper  The set to search through.
+	 * @param  Index      The index in the set to start searching from.
+	 * @param  Num        The number of elements to compare.
 	 */
-	bool AnyEqual(const FScriptSetHelper& SetHelper, const int32 LogicalIndex, int32 Num, const uint8* ElementToCompare, uint32 PortFlags)
+	bool AnyEqual(const FScriptSetHelper& SetHelper, int32 Index, int32 Num, const uint8* ElementToCompare, uint32 PortFlags)
 	{
-		const FProperty* ElementProp = SetHelper.GetElementProperty();
+		FProperty* ElementProp = SetHelper.GetElementProperty();
 
-		FScriptSetHelper::FIterator Iterator(SetHelper, LogicalIndex);
-		for (; Iterator && Num; --Num, ++Iterator)
+		for (; Num; --Num)
 		{
-			if (ElementProp->Identical(SetHelper.GetElementPtr(Iterator), ElementToCompare, PortFlags))
+			while (!SetHelper.IsValidIndex(Index))
+			{
+				++Index;
+			}
+
+			if (ElementProp->Identical(SetHelper.GetElementPtr(Index), ElementToCompare, PortFlags))
 			{
 				return true;
 			}
+
+			++Index;
 		}
 
 		return false;
 	}
 
-	bool RangesContainSameAmountsOfVal(const FScriptSetHelper& SetHelperA, const int32 LogicalIndexA, const FScriptSetHelper& SetHelperB, const int32 LogicalIndexB, int32 Num, const uint8* ElementToCompare, uint32 PortFlags)
+	bool RangesContainSameAmountsOfVal(const FScriptSetHelper& SetHelperA, int32 IndexA, const FScriptSetHelper& SetHelperB, int32 IndexB, int32 Num, const uint8* ElementToCompare, uint32 PortFlags)
 	{
-		const FProperty* ElementProp = SetHelperA.GetElementProperty();
+		FProperty* ElementProp = SetHelperA.GetElementProperty();
 
 		// Ensure that both sets are the same type
 		check(ElementProp == SetHelperB.GetElementProperty());
 
-		FScriptSetHelper::FIterator IteratorA(SetHelperA, LogicalIndexA);
-		FScriptSetHelper::FIterator IteratorB(SetHelperB, LogicalIndexB);
-
 		int32 CountA = 0;
 		int32 CountB = 0;
 		for (;;)
@@ -56,8 +59,18 @@ namespace UESetProperty_Private
 				return CountA == CountB;
 			}
 
-			const uint8* ElementA = SetHelperA.GetElementPtr(IteratorA);
-			const uint8* ElementB = SetHelperB.GetElementPtr(IteratorB);
+			while (!SetHelperA.IsValidIndex(IndexA))
+			{
+				++IndexA;
+			}
+
+			while (!SetHelperB.IsValidIndex(IndexB))
+			{
+				++IndexB;
+			}
+
+			const uint8* ElementA = SetHelperA.GetElementPtr(IndexA);
+			const uint8* ElementB = SetHelperB.GetElementPtr(IndexB);
 			if (ElementProp->Identical(ElementA, ElementToCompare, PortFlags))
 			{
 				++CountA;
@@ -68,15 +81,15 @@ namespace UESetProperty_Private
 				++CountB;
 			}
 
-			++IteratorA;
-			++IteratorB;
+			++IndexA;
+			++IndexB;
 			--Num;
 		}
 	}
 
-	bool IsPermutation(const FScriptSetHelper& SetHelperA, const FScriptSetHelper& SetHelperB, const uint32 PortFlags)
+	bool IsPermutation(const FScriptSetHelper& SetHelperA, const FScriptSetHelper& SetHelperB, uint32 PortFlags)
 	{
-		const FProperty* ElementProp = SetHelperA.GetElementProperty();
+		FProperty* ElementProp = SetHelperA.GetElementProperty();
 
 		// Ensure that both maps are the same type
 		check(ElementProp == SetHelperB.GetElementProperty());
@@ -88,8 +101,8 @@ namespace UESetProperty_Private
 		}
 
 		// Skip over common initial sequence
-		FScriptSetHelper::FIterator IteratorA(SetHelperA);
-		FScriptSetHelper::FIterator IteratorB(SetHelperB);
+		int32 IndexA = 0;
+		int32 IndexB = 0;
 		for (;;)
 		{
 			if (Num == 0)
@@ -97,24 +110,34 @@ namespace UESetProperty_Private
 				return true;
 			}
 
-			const uint8* ElementA = SetHelperA.GetElementPtr(IteratorA);
-			const uint8* ElementB = SetHelperB.GetElementPtr(IteratorB);
+			while (!SetHelperA.IsValidIndex(IndexA))
+			{
+				++IndexA;
+			}
+
+			while (!SetHelperB.IsValidIndex(IndexB))
+			{
+				++IndexB;
+			}
+
+			const uint8* ElementA = SetHelperA.GetElementPtr(IndexA);
+			const uint8* ElementB = SetHelperB.GetElementPtr(IndexB);
 			if (!ElementProp->Identical(ElementA, ElementB, PortFlags))
 			{
 				break;
 			}
 
-			++IteratorA;
-			++IteratorB;
+			++IndexA;
+			++IndexB;
 			--Num;
 		}
 
-		const int32 FirstIndexA = IteratorA.GetLogicalIndex();
-		const int32 FirstIndexB = IteratorB.GetLogicalIndex();
-		const int32 FirstNum    = Num;
+		int32 FirstIndexA = IndexA;
+		int32 FirstIndexB = IndexB;
+		int32 FirstNum    = Num;
 		for (;;)
 		{
-			const uint8* ElementA = SetHelperA.GetElementPtr(IteratorA);
+			const uint8* ElementA = SetHelperA.GetElementPtr(IndexA);
 			if (!AnyEqual(SetHelperA, FirstIndexA, FirstNum - Num, ElementA, PortFlags) && !RangesContainSameAmountsOfVal(SetHelperA, FirstIndexA, SetHelperB, FirstIndexB, FirstNum, ElementA, PortFlags))
 			{
 				return false;
@@ -126,7 +149,12 @@ namespace UESetProperty_Private
 				return true;
 			}
 
-			++IteratorA;
+			++IndexA;
+			while (!SetHelperA.IsValidIndex(IndexA))
+			{
+				++IndexA;
+			}
+
 		}
 	}
 }
@@ -404,12 +432,18 @@ void FSetProperty::SerializeItem(FStructuredArchive::FSlot Slot, void* Value, co
 		// Determine how many keys are missing from the object
 		if (Defaults)
 		{
-			for (FScriptSetHelper::FIterator Iterator(DefaultsHelper); Iterator; ++Iterator)
+			for (int32 Index = 0, Count = DefaultsHelper.Num(); Count; ++Index)
 			{
-				const uint8* DefaultElementPtr = DefaultsHelper.GetElementPtr(Iterator);
-				if (SetHelper.FindElementIndex(DefaultElementPtr) == INDEX_NONE)
+				uint8* DefaultElementPtr = DefaultsHelper.GetElementPtrWithoutCheck(Index);
+
+				if (DefaultsHelper.IsValidIndex(Index))
 				{
-					Indices.Add(Iterator.GetInternalIndex());
+					if (SetHelper.FindElementIndex(DefaultElementPtr) == INDEX_NONE)
+					{
+						Indices.Add(Index);
+					}
+
+					--Count;
 				}
 			}
 		}
@@ -430,14 +464,19 @@ void FSetProperty::SerializeItem(FStructuredArchive::FSlot Slot, void* Value, co
 		if (Defaults)
 		{
 			Indices.Reset();
-			for (FScriptSetHelper::FIterator Iterator(SetHelper); Iterator; ++Iterator)
+			for (int32 Index = 0, Count = SetHelper.Num(); Count; ++Index)
 			{
-				const uint8* ValueElement = SetHelper.GetElementPtr(Iterator);
-				const uint8* DefaultElement = DefaultsHelper.FindElementPtr(ValueElement);
-
-				if (!DefaultElement)
+				if (SetHelper.IsValidIndex(Index))
 				{
-					Indices.Add(Iterator.GetInternalIndex());
+					uint8* ValueElement   = SetHelper.GetElementPtrWithoutCheck(Index);
+					uint8* DefaultElement = DefaultsHelper.FindElementPtr(ValueElement);
+
+					if (!DefaultElement)
+					{
+						Indices.Add(Index);
+					}
+
+					--Count;
 				}
 			}
 
@@ -459,11 +498,16 @@ void FSetProperty::SerializeItem(FStructuredArchive::FSlot Slot, void* Value, co
 			FStructuredArchive::FArray ElementsArray = Record.EnterArray(TEXT("Elements"), Num);
 
 			FSerializedPropertyScope SerializedProperty(UnderlyingArchive, ElementProp, this);
-
-			for (FScriptSetHelper::FIterator Iterator(SetHelper); Iterator; ++Iterator)
+			for (int32 Index = 0; Num; ++Index)
 			{
-				uint8* ElementPtr = SetHelper.GetElementPtr(Iterator);
-				ElementProp->SerializeItem(ElementsArray.EnterElement(), ElementPtr);
+				if (SetHelper.IsValidIndex(Index))
+				{
+					uint8* ElementPtr = SetHelper.GetElementPtrWithoutCheck(Index);
+
+					ElementProp->SerializeItem(ElementsArray.EnterElement(), ElementPtr);
+
+					--Num;
+				}
 			}
 		}
 	}
@@ -859,20 +903,31 @@ void FSetProperty::InstanceSubobjects(void* Data, void const* DefaultData, UObje
 	{
 		FScriptSetHelper DefaultSetHelper(this, DefaultData);
 
-		for (FScriptSetHelper::FIterator It(SetHelper); It; ++It)
+		for (int32 Index = 0, Num = SetHelper.Num(); Num; ++Index)
 		{
-			uint8* ElementPtr = SetHelper.GetElementPtr(It);
-			const uint8* DefaultElementPtr = DefaultSetHelper.FindElementPtr(ElementPtr, It.GetLogicalIndex());
+			if (SetHelper.IsValidIndex(Index))
+			{
+				uint8* ElementPtr        = SetHelper.GetElementPtr(Index);
+				uint8* DefaultElementPtr = DefaultSetHelper.FindElementPtr(ElementPtr, Index);
+
+				ElementProp->InstanceSubobjects(ElementPtr, DefaultElementPtr, InOwner, InstanceGraph);
 
-			ElementProp->InstanceSubobjects(ElementPtr, DefaultElementPtr, InOwner, InstanceGraph);
+				--Num;
+			}
 		}
 	}
 	else
 	{
-		for (FScriptSetHelper::FIterator It(SetHelper); It; ++It)
+		for (int32 Index = 0, Num = SetHelper.Num(); Num; ++Index)
 		{
-			uint8* ElementPtr = SetHelper.GetElementPtr(It);
-			ElementProp->InstanceSubobjects(ElementPtr, nullptr, InOwner, InstanceGraph);
+			if (SetHelper.IsValidIndex(Index))
+			{
+				uint8* ElementPtr = SetHelper.GetElementPtr(Index);
+
+				ElementProp->InstanceSubobjects(ElementPtr, nullptr, InOwner, InstanceGraph);
+
+				--Num;
+			}
 		}
 	}
 }
@@ -1069,16 +1124,22 @@ void FSetProperty::GetInnerFields(TArray<FField*>& OutFields)
 	}
 }
 
-void* FSetProperty::GetValueAddressAtIndex_Direct(const FProperty* Inner, void* InValueAddress, const int32 LogicalIndex) const
+void* FSetProperty::GetValueAddressAtIndex_Direct(const FProperty* Inner, void* InValueAddress, int32 Index) const
 {
+	FScriptSetHelper SetHelper(this, InValueAddress);
 	checkf(Inner == ElementProp, TEXT("Inner property must be identical to ElementProp"));
 
-	FScriptSetHelper SetHelper(this, InValueAddress);
-	const int32 InternalIndex = SetHelper.FindInternalIndex(LogicalIndex);
-	if (InternalIndex != INDEX_NONE)
+	for (int32 SetIndex = 0, Num = SetHelper.Num(), LocalIndex = Index; LocalIndex >= 0 && Num > 0; ++SetIndex)
 	{
-		return SetHelper.GetElementPtrWithoutCheck(InternalIndex);
+		if (SetHelper.IsValidIndex(SetIndex))
+		{
+			if (LocalIndex == 0)
+			{
+				return SetHelper.GetElementPtr(SetIndex);
+			}
+			LocalIndex--;
+			Num--;
+		}
 	}
-
 	return nullptr;
 }
diff --git a/Engine/Source/Runtime/Engine/Private/BlueprintMapLibrary.cpp b/Engine/Source/Runtime/Engine/Private/BlueprintMapLibrary.cpp
--- a/Engine/Source/Runtime/Engine/Private/BlueprintMapLibrary.cpp
+++ b/Engine/Source/Runtime/Engine/Private/BlueprintMapLibrary.cpp
@@ -61,11 +61,17 @@ void UBlueprintMapLibrary::GenericMap_Keys(const void* TargetMap, const FMapProp
 		FScriptArrayHelper ArrayHelper(ArrayProperty, TargetArray);
 		ArrayHelper.EmptyValues();
 
-		const FProperty* InnerProp = ArrayProperty->Inner;
-		for (FScriptMapHelper::FIterator It(MapHelper); It; ++It)
+		FProperty* InnerProp = ArrayProperty->Inner;
+
+		int32 Size = MapHelper.Num();
+		for( int32 I = 0; Size; ++I )
 		{
-			const int32 LastIndex = ArrayHelper.AddValue();
-			InnerProp->CopySingleValueToScriptVM(ArrayHelper.GetRawPtr(LastIndex), MapHelper.GetKeyPtr(It));
+			if(MapHelper.IsValidIndex(I))
+			{
+				int32 LastIndex = ArrayHelper.AddValue();
+				InnerProp->CopySingleValueToScriptVM(ArrayHelper.GetRawPtr(LastIndex), MapHelper.GetKeyPtr(I));
+				--Size;
+			}
 		}
 	}
 }
@@ -78,11 +84,17 @@ void UBlueprintMapLibrary::GenericMap_Values(const void* TargetMap, const FMapPr
 		FScriptArrayHelper ArrayHelper(ArrayProperty, TargetArray);
 		ArrayHelper.EmptyValues();
 
-		const FProperty* InnerProp = ArrayProperty->Inner;
-		for (FScriptMapHelper::FIterator It(MapHelper); It; ++It)
+		FProperty* InnerProp = ArrayProperty->Inner;
+		
+		int32 Size = MapHelper.Num();
+		for( int32 I = 0; Size; ++I )
 		{
-			const int32 LastIndex = ArrayHelper.AddValue();
-			InnerProp->CopySingleValueToScriptVM(ArrayHelper.GetRawPtr(LastIndex), MapHelper.GetValuePtr(It));
+			if(MapHelper.IsValidIndex(I))
+			{
+				int32 LastIndex = ArrayHelper.AddValue();
+				InnerProp->CopySingleValueToScriptVM(ArrayHelper.GetRawPtr(LastIndex), MapHelper.GetValuePtr(I));
+				--Size;
+			}
 		}
 	}
 }
diff --git a/Engine/Source/Runtime/Engine/Private/BlueprintSetLibrary.cpp b/Engine/Source/Runtime/Engine/Private/BlueprintSetLibrary.cpp
--- a/Engine/Source/Runtime/Engine/Private/BlueprintSetLibrary.cpp
+++ b/Engine/Source/Runtime/Engine/Private/BlueprintSetLibrary.cpp
@@ -64,9 +64,15 @@ void UBlueprintSetLibrary::GenericSet_ToArray(const void* TargetSet, const FSetP
 	if (TargetSet && TargetArray)
 	{
 		FScriptSetHelper SetHelper(SetProperty, TargetSet);
-		for (FScriptSetHelper::FIterator It(SetHelper); It; ++It)
+
+		int32 Size = SetHelper.Num();
+		for (int32 I = 0; Size; ++I)
 		{
-			UKismetArrayLibrary::GenericArray_Add(TargetArray, ArrayProperty, SetHelper.GetElementPtr(It));
+			if(SetHelper.IsValidIndex(I))
+			{
+				UKismetArrayLibrary::GenericArray_Add(TargetArray, ArrayProperty, SetHelper.GetElementPtr(I));
+				--Size;
+			}
 		}
 	}
 }
@@ -136,12 +142,18 @@ void UBlueprintSetLibrary::GenericSet_Intersect(const void* SetA, const FSetProp
 		FScriptSetHelper SetHelperResult(SetPropertyResult, SetResult);
 		
 		SetHelperResult.EmptyElements();
-		for (FScriptSetHelper::FIterator It(SetHelperA); It; ++It)
+
+		int32 Size = SetHelperA.Num();
+		for (int32 I = 0; Size; ++I)
 		{
-			const void* EntryInA = SetHelperA.GetElementPtr(It);
-			if (SetHelperB.FindElementIndexFromHash(EntryInA) != INDEX_NONE)
+			if(SetHelperA.IsValidIndex(I))
 			{
-				SetHelperResult.AddElement(EntryInA);
+				const void* EntryInA = SetHelperA.GetElementPtr(I);
+				if (SetHelperB.FindElementIndexFromHash(EntryInA) != INDEX_NONE)
+				{
+					SetHelperResult.AddElement(EntryInA);
+				}
+				--Size;
 			}
 		}
 	}
@@ -157,14 +169,24 @@ void UBlueprintSetLibrary::GenericSet_Union(const void* SetA, const FSetProperty
 		
 		SetHelperResult.EmptyElements();
 
-		for (FScriptSetHelper::FIterator It(SetHelperA); It; ++It)
+		int32 SizeA = SetHelperA.Num();
+		for (int32 I = 0; SizeA; ++I)
 		{
-			SetHelperResult.AddElement(SetHelperA.GetElementPtr(It));
+			if(SetHelperA.IsValidIndex(I))
+			{
+				SetHelperResult.AddElement(SetHelperA.GetElementPtr(I));
+				--SizeA;
+			}
 		}
 
-		for (FScriptSetHelper::FIterator It(SetHelperB); It; ++It)
+		int32 SizeB = SetHelperB.Num();
+		for (int32 I = 0; SizeB; ++I)
 		{
-			SetHelperResult.AddElement(SetHelperB.GetElementPtr(It));
+			if(SetHelperB.IsValidIndex(I))
+			{
+				SetHelperResult.AddElement(SetHelperB.GetElementPtr(I));
+				--SizeB;
+			}
 		}
 	}
 }
@@ -179,12 +201,17 @@ void UBlueprintSetLibrary::GenericSet_Difference(const void* SetA, const FSetPro
 
 		SetHelperResult.EmptyElements();
 
-		for (FScriptSetHelper::FIterator It(SetHelperA); It; ++It)
+		int32 Size = SetHelperA.Num();
+		for (int32 I = 0; Size; ++I)
 		{
-			const void* EntryInA = SetHelperA.GetElementPtr(It);
-			if (SetHelperB.FindElementIndexFromHash(EntryInA) == INDEX_NONE)
+			if(SetHelperA.IsValidIndex(I))
 			{
-				SetHelperResult.AddElement(EntryInA);
+				const void* EntryInA = SetHelperA.GetElementPtr(I);
+				if (SetHelperB.FindElementIndexFromHash(EntryInA) == INDEX_NONE)
+				{
+					SetHelperResult.AddElement(EntryInA);
+				}
+				--Size;
 			}
 		}
 	}
diff --git a/Engine/Source/Runtime/Engine/Private/DataTableJSON.cpp b/Engine/Source/Runtime/Engine/Private/DataTableJSON.cpp
--- a/Engine/Source/Runtime/Engine/Private/DataTableJSON.cpp
+++ b/Engine/Source/Runtime/Engine/Private/DataTableJSON.cpp
@@ -265,10 +265,13 @@ bool TDataTableExporterJSON<CharType>::WriteStructEntry(const void* InRowData, c
 		JsonWriter->WriteArrayStart(Identifier);
 
 		FScriptSetHelper SetHelper(SetProp, InPropertyData);
-		for (FScriptSetHelper::FIterator It(SetHelper); It; ++It)
+		for (int32 SetSparseIndex = 0; SetSparseIndex < SetHelper.GetMaxIndex(); ++SetSparseIndex)
 		{
-			const uint8* SetEntryData = SetHelper.GetElementPtr(It);
-			WriteContainerEntry(SetHelper.GetElementProperty(), SetEntryData);
+			if (SetHelper.IsValidIndex(SetSparseIndex))
+			{
+				const uint8* SetEntryData = SetHelper.GetElementPtr(SetSparseIndex);
+				WriteContainerEntry(SetHelper.GetElementProperty(), SetEntryData);
+			}
 		}
 
 		JsonWriter->WriteArrayEnd();
@@ -278,14 +281,17 @@ bool TDataTableExporterJSON<CharType>::WriteStructEntry(const void* InRowData, c
 		JsonWriter->WriteObjectStart(Identifier);
 
 		FScriptMapHelper MapHelper(MapProp, InPropertyData);
-		for (FScriptMapHelper::FIterator It(MapHelper); It; ++It)
+		for (int32 MapSparseIndex = 0; MapSparseIndex < MapHelper.GetMaxIndex(); ++MapSparseIndex)
 		{
-			const uint8* MapKeyData = MapHelper.GetKeyPtr(It);
-			const uint8* MapValueData = MapHelper.GetValuePtr(It);
+			if (MapHelper.IsValidIndex(MapSparseIndex))
+			{
+				const uint8* MapKeyData = MapHelper.GetKeyPtr(MapSparseIndex);
+				const uint8* MapValueData = MapHelper.GetValuePtr(MapSparseIndex);
 
-			// JSON object keys must always be strings
-			const FString KeyValue = DataTableUtils::GetPropertyValueAsStringDirect(MapHelper.GetKeyProperty(), (uint8*)MapKeyData, DTExportFlags);
-			WriteContainerEntry(MapHelper.GetValueProperty(), MapValueData, &KeyValue);
+				// JSON object keys must always be strings
+				const FString KeyValue = DataTableUtils::GetPropertyValueAsStringDirect(MapHelper.GetKeyProperty(), (uint8*)MapKeyData, DTExportFlags);
+				WriteContainerEntry(MapHelper.GetValueProperty(), MapValueData, &KeyValue);
+			}
 		}
 
 		JsonWriter->WriteObjectEnd();
diff --git a/Engine/Source/Runtime/Engine/Private/DataTableUtils.cpp b/Engine/Source/Runtime/Engine/Private/DataTableUtils.cpp
--- a/Engine/Source/Runtime/Engine/Private/DataTableUtils.cpp
+++ b/Engine/Source/Runtime/Engine/Private/DataTableUtils.cpp
@@ -142,17 +142,21 @@ void GetPropertyValueAsStringDirect(const FProperty* InProp, const uint8* InData
 
 			OutString.AppendChar('(');
 
+			int32 NumWrittenSetEntries = 0;
 			FScriptSetHelper SetHelper(SetProp, InData);
-			for (FScriptSetHelper::FIterator It(SetHelper); It; ++It)
+			for (int32 SetSparseIndex = 0; SetSparseIndex < SetHelper.GetMaxIndex(); ++SetSparseIndex)
 			{
-					if (It.GetLogicalIndex() > 0)
+				if (SetHelper.IsValidIndex(SetSparseIndex))
+				{
+					if (NumWrittenSetEntries++ > 0)
 					{
 						OutString.AppendChar(',');
 						OutString.AppendChar(' ');
 					}
 
-					const uint8* SetEntryData = SetHelper.GetElementPtr(It);
+					const uint8* SetEntryData = SetHelper.GetElementPtr(SetSparseIndex);
 					OutString.Append(ExportStructAsJson(StructInner->Struct, SetEntryData));
+				}
 			}
 
 			OutString.AppendChar(')');
@@ -162,32 +166,36 @@ void GetPropertyValueAsStringDirect(const FProperty* InProp, const uint8* InData
 		{
 			OutString.AppendChar('(');
 
+			int32 NumWrittenMapEntries = 0;
 			FScriptMapHelper MapHelper(MapProp, InData);
-			for (FScriptMapHelper::FIterator It(MapHelper); It; ++It)
+			for (int32 MapSparseIndex = 0; MapSparseIndex < MapHelper.GetMaxIndex(); ++MapSparseIndex)
 			{
-				if (It.GetLogicalIndex() > 0)
+				if (MapHelper.IsValidIndex(MapSparseIndex))
 				{
-					OutString.AppendChar(',');
-					OutString.AppendChar(' ');
-				}
+					if (NumWrittenMapEntries++ > 0)
+					{
+						OutString.AppendChar(',');
+						OutString.AppendChar(' ');
+					}
 
-				const uint8* MapKeyData = MapHelper.GetKeyPtr(It);
-				const uint8* MapValueData = MapHelper.GetValuePtr(It);
+					const uint8* MapKeyData = MapHelper.GetKeyPtr(MapSparseIndex);
+					const uint8* MapValueData = MapHelper.GetValuePtr(MapSparseIndex);
 
-				OutString.AppendChar('"');
-				GetPropertyValueAsStringDirect(MapHelper.GetKeyProperty(), MapKeyData, InPortFlags, InDTExportFlags, OutString);
-				OutString.AppendChar('"');
+					OutString.AppendChar('"');
+					GetPropertyValueAsStringDirect(MapHelper.GetKeyProperty(), MapKeyData, InPortFlags, InDTExportFlags, OutString);
+					OutString.AppendChar('"');
 
-				OutString.Append(TEXT(" = "));
+					OutString.Append(TEXT(" = "));
 
-				if (MapHelper.GetValueProperty()->IsA<FStructProperty>() && EnumHasAnyFlags(InDTExportFlags, EDataTableExportFlags::UseJsonObjectsForStructs))
-				{
-					const FStructProperty* StructMapValue = CastFieldChecked<const FStructProperty>(MapHelper.GetValueProperty());
-					OutString.Append(ExportStructAsJson(StructMapValue->Struct, MapValueData));
-				}
-				else
-				{
-					GetPropertyValueAsStringDirect(MapHelper.GetValueProperty(), MapValueData, InPortFlags, InDTExportFlags, OutString);
+					if (MapHelper.GetValueProperty()->IsA<FStructProperty>() && EnumHasAnyFlags(InDTExportFlags, EDataTableExportFlags::UseJsonObjectsForStructs))
+					{
+						const FStructProperty* StructMapValue = CastFieldChecked<const FStructProperty>(MapHelper.GetValueProperty());
+						OutString.Append(ExportStructAsJson(StructMapValue->Struct, MapValueData));
+					}
+					else
+					{
+						GetPropertyValueAsStringDirect(MapHelper.GetValueProperty(), MapValueData, InPortFlags, InDTExportFlags, OutString);
+					}
 				}
 			}
 
diff --git a/Engine/Source/Runtime/Engine/Private/InstancedReferenceSubobjectHelper.cpp b/Engine/Source/Runtime/Engine/Private/InstancedReferenceSubobjectHelper.cpp
--- a/Engine/Source/Runtime/Engine/Private/InstancedReferenceSubobjectHelper.cpp
+++ b/Engine/Source/Runtime/Engine/Private/InstancedReferenceSubobjectHelper.cpp
@@ -160,20 +160,26 @@ void FFindInstancedReferenceSubobjectHelper::ForEachInstancedSubObject(FInstance
 			return;
 		}
 
+		int32 LogicalIndex = 0;
 		FScriptMapHelper MapHelper(MapProperty, ContainerAddress);
-		for (FScriptMapHelper::FIterator It(MapHelper); It; ++It)
+		for (int32 ElementIndex = 0; ElementIndex < MapHelper.GetMaxIndex(); ++ElementIndex)
 		{
-			T KeyAddress = MapHelper.GetKeyPtr(It);
-			T ValueAddress = MapHelper.GetValuePtr(It);
+			if (MapHelper.IsValidIndex(ElementIndex))
+			{
+				T KeyAddress = MapHelper.GetKeyPtr(ElementIndex);
+				T ValueAddress = MapHelper.GetValuePtr(ElementIndex);
 
-			// Note: Keep these as the logical (Nth) index in case the map changes internally after we construct the path or in case we resolve using a different object.
-			PropertyPath.Push(MapProperty->KeyProp, It.GetLogicalIndex());
-			ForEachInstancedSubObject(PropertyPath, KeyAddress, ObjRefFunc);
-			PropertyPath.Pop();
+				// Note: Keep these as the logical (Nth) index in case the map changes internally after we construct the path or in case we resolve using a different object.
+				PropertyPath.Push(MapProperty->KeyProp, LogicalIndex);
+				ForEachInstancedSubObject(PropertyPath, KeyAddress, ObjRefFunc);
+				PropertyPath.Pop();
 
-			PropertyPath.Push(MapProperty->ValueProp, It.GetLogicalIndex(), true);
-			ForEachInstancedSubObject(PropertyPath, ValueAddress, ObjRefFunc);
-			PropertyPath.Pop();
+				PropertyPath.Push(MapProperty->ValueProp, LogicalIndex, true);
+				ForEachInstancedSubObject(PropertyPath, ValueAddress, ObjRefFunc);
+				PropertyPath.Pop();
+
+				++LogicalIndex;
+			}
 		}
 	}
 	else if (const FSetProperty* SetProperty = CastField<const FSetProperty>(TargetProp))
@@ -184,15 +190,21 @@ void FFindInstancedReferenceSubobjectHelper::ForEachInstancedSubObject(FInstance
 			return;
 		}
 
+		int32 LogicalIndex = 0;
 		FScriptSetHelper SetHelper(SetProperty, ContainerAddress);
-		for (FScriptSetHelper::FIterator It(SetHelper); It; ++It)
+		for (int32 ElementIndex = 0; ElementIndex < SetHelper.GetMaxIndex(); ++ElementIndex)
 		{
-			T ValueAddress = SetHelper.GetElementPtr(It);
+			if (SetHelper.IsValidIndex(ElementIndex))
+			{
+				T ValueAddress = SetHelper.GetElementPtr(ElementIndex);
 
-			// Note: Keep this as the logical (Nth) index in case the set changes internally after we construct the path or in case we resolve using a different object.
-			PropertyPath.Push(SetProperty->ElementProp, It.GetLogicalIndex());
-			ForEachInstancedSubObject(PropertyPath, ValueAddress, ObjRefFunc);
-			PropertyPath.Pop();
+				// Note: Keep this as the logical (Nth) index in case the set changes internally after we construct the path or in case we resolve using a different object.
+				PropertyPath.Push(SetProperty->ElementProp, LogicalIndex);
+				ForEachInstancedSubObject(PropertyPath, ValueAddress, ObjRefFunc);
+				PropertyPath.Pop();
+
+				++LogicalIndex;
+			}
 		}
 	}
 	else if (const FOptionalProperty* OptionalProperty = CastField<FOptionalProperty>(TargetProp))
diff --git a/Engine/Source/Runtime/Engine/Private/ObjectEditorUtils.cpp b/Engine/Source/Runtime/Engine/Private/ObjectEditorUtils.cpp
--- a/Engine/Source/Runtime/Engine/Private/ObjectEditorUtils.cpp
+++ b/Engine/Source/Runtime/Engine/Private/ObjectEditorUtils.cpp
@@ -202,19 +202,25 @@ namespace FObjectEditorUtils
 
 			TargetMapHelper.EmptyValues();
 
-			for (FScriptMapHelper::FIterator It(SourceMapHelper); It; ++It)
+			int32 Num = SourceMapHelper.Num();
+			for ( int32 Index = 0; Num; Index++ )
 			{
-				uint8* SrcPairPtr = SourceMapHelper.GetPairPtr(It);
+				if ( SourceMapHelper.IsValidIndex(Index) )
+				{
+					uint8* SrcPairPtr = SourceMapHelper.GetPairPtr(Index);
+
+					int32 NewIndex = TargetMapHelper.AddDefaultValue_Invalid_NeedsRehash();
+					TargetMapHelper.Rehash();
 
-				int32 NewIndex = TargetMapHelper.AddDefaultValue_Invalid_NeedsRehash();
-				TargetMapHelper.Rehash();
+					uint8* PairPtr = TargetMapHelper.GetPairPtr(NewIndex);
 
-				uint8* PairPtr = TargetMapHelper.GetPairPtr(NewIndex);
+					CopySinglePropertyRecursive(SourceObject, SrcPairPtr, SrcMapProperty->KeyProp, PairPtr, InDestinationObject, DestMapProperty->KeyProp);
+					CopySinglePropertyRecursive(SourceObject, SrcPairPtr, SrcMapProperty->ValueProp, PairPtr, InDestinationObject, DestMapProperty->ValueProp);
 
-				CopySinglePropertyRecursive(SourceObject, SrcPairPtr, SrcMapProperty->KeyProp, PairPtr, InDestinationObject, DestMapProperty->KeyProp);
-				CopySinglePropertyRecursive(SourceObject, SrcPairPtr, SrcMapProperty->ValueProp, PairPtr, InDestinationObject, DestMapProperty->ValueProp);
+					TargetMapHelper.Rehash();
 
-				TargetMapHelper.Rehash();
+					--Num;
+				}
 			}
 
 			bNeedsShallowCopy = false;
@@ -229,17 +235,23 @@ namespace FObjectEditorUtils
 
 			TargetSetHelper.EmptyElements();
 
-			for (FScriptSetHelper::FIterator It(SourceSetHelper); It; ++It)
+			int32 Num = SourceSetHelper.Num();
+			for ( int32 Index = 0; Num; Index++ )
 			{
-				uint8* SrcPtr = SourceSetHelper.GetElementPtr(It);
+				if ( SourceSetHelper.IsValidIndex(Index) )
+				{
+					uint8* SrcPtr = SourceSetHelper.GetElementPtr(Index);
+
+					int32 NewIndex = TargetSetHelper.AddDefaultValue_Invalid_NeedsRehash();
+					TargetSetHelper.Rehash();
 
-				int32 NewIndex = TargetSetHelper.AddDefaultValue_Invalid_NeedsRehash();
-				TargetSetHelper.Rehash();
+					uint8* TargetPtr = TargetSetHelper.GetElementPtr(NewIndex);
+					CopySinglePropertyRecursive(SourceObject, SrcPtr, SrcSetProperty->ElementProp, TargetPtr, InDestinationObject, DestSetProperty->ElementProp);
 
-				uint8* TargetPtr = TargetSetHelper.GetElementPtr(NewIndex);
-				CopySinglePropertyRecursive(SourceObject, SrcPtr, SrcSetProperty->ElementProp, TargetPtr, InDestinationObject, DestSetProperty->ElementProp);
+					TargetSetHelper.Rehash();
 
-				TargetSetHelper.Rehash();
+					--Num;
+				}
 			}
 
 			bNeedsShallowCopy = false;
diff --git a/Engine/Source/Runtime/Serialization/Private/StructSerializer.cpp b/Engine/Source/Runtime/Serialization/Private/StructSerializer.cpp
--- a/Engine/Source/Runtime/Serialization/Private/StructSerializer.cpp
+++ b/Engine/Source/Runtime/Serialization/Private/StructSerializer.cpp
@@ -22,18 +22,18 @@ namespace StructSerializer
 	template <>
 	struct TScriptHelper_InContainer<FMapProperty> { using Type = FScriptMapHelper_InContainer; };
 
-	FStructSerializerState CreateItemState(FScriptArrayHelper& InHelper, FArrayProperty* InArrayProperty, int32 InValidatedInternalElementIndex, const FStructSerializerPolicies& Policies, EStructSerializerStateFlags InFlags)
+	FStructSerializerState CreateItemState(FScriptArrayHelper& InHelper, FArrayProperty* InArrayProperty, int32 InElementIndex, const FStructSerializerPolicies& Policies, EStructSerializerStateFlags InFlags)
 	{
-		return FStructSerializerState(InHelper.GetRawPtr(InValidatedInternalElementIndex), InArrayProperty->Inner, InFlags);
+		return FStructSerializerState(InHelper.GetRawPtr(InElementIndex), InArrayProperty->Inner, InFlags);
 	};
 
-	FStructSerializerState CreateItemState(FScriptMapHelper& InHelper, FMapProperty* InMapProperty, int32 InValidatedInternalElementIndex, const FStructSerializerPolicies& Policies, EStructSerializerStateFlags InFlags)
+	FStructSerializerState CreateItemState(FScriptMapHelper& InHelper, FMapProperty* InMapProperty, int32 InElementIndex, const FStructSerializerPolicies& Policies, EStructSerializerStateFlags InFlags)
 	{
-		FStructSerializerState State = FStructSerializerState(InHelper.GetPairPtr(InValidatedInternalElementIndex), InMapProperty->ValueProp, InFlags);
+		FStructSerializerState State = FStructSerializerState(InHelper.GetPairPtr(InElementIndex), InMapProperty->ValueProp, InFlags);
 
 		if (Policies.MapSerialization == EStructSerializerMapPolicies::KeyValuePair)
 		{
-			State.KeyData = InHelper.GetPairPtr(InValidatedInternalElementIndex);
+			State.KeyData = InHelper.GetPairPtr(InElementIndex);
 			State.ValueData = State.KeyData;
 			State.KeyProperty = InMapProperty->KeyProp;
 		}
@@ -41,9 +41,9 @@ namespace StructSerializer
 		return State;
 	}
 
-	FStructSerializerState CreateItemState(FScriptSetHelper& InHelper, FSetProperty* InSetProperty, int32 InValidatedInternalElementIndex, const FStructSerializerPolicies& Policies, EStructSerializerStateFlags InFlags)
+	FStructSerializerState CreateItemState(FScriptSetHelper& InHelper, FSetProperty* InSetProperty, int32 InElementIndex, const FStructSerializerPolicies& Policies, EStructSerializerStateFlags InFlags)
 	{
-		return FStructSerializerState(InHelper.GetElementPtr(InValidatedInternalElementIndex), InSetProperty->ElementProp, InFlags);
+		return FStructSerializerState(InHelper.GetElementPtr(InElementIndex), InSetProperty->ElementProp, InFlags);
 	}
 
 	bool IsArrayLike(FProperty* Property, const FStructSerializerPolicies& Policies)
