UE-188423 - Ensure that objects are always preloaded prior to postloading.

Currently, we can only make this guarantee when the Zen loader is used.

diff --git a/Engine/Source/Runtime/CoreUObject/Private/UObject/Obj.cpp b/Engine/Source/Runtime/CoreUObject/Private/UObject/Obj.cpp
--- a/Engine/Source/Runtime/CoreUObject/Private/UObject/Obj.cpp
+++ b/Engine/Source/Runtime/CoreUObject/Private/UObject/Obj.cpp
@@ -1238,12 +1238,10 @@ void UObject::ConditionalPostLoad()
 {
 	LLM_SCOPE(ELLMTag::UObject);
 
-	check(!GEventDrivenLoaderEnabled || !HasAnyFlags(RF_NeedLoad)); //@todoio Added this as "nicks rule"
-									  // PostLoad only if the object needs it and has already been serialized
-	//@todoio note this logic should be unchanged compared to main
+	ensureAlways((GetLoaderType() != ELoaderType::ZenLoader) || !HasAnyFlags(RF_NeedLoad));
+
 	if (HasAnyFlags(RF_NeedPostLoad))
 	{
-
 		check(IsInGameThread() || HasAnyFlags(RF_ClassDefaultObject|RF_ArchetypeObject) || IsPostLoadThreadSafe() || IsA(UClass::StaticClass()))
 		UE_TRACK_REFERENCING_PACKAGE_SCOPED(this, PackageAccessTrackingOps::NAME_PostLoad);
 
