Add tagged property serialization control extensions to UObject
Add property tag extensions

Both of those exentions can be created and register as an enum flag which is now serialized as an attribute of UObject during TPS. Registered extensions can them be toggled on/off for specific objects with the intent to modify how TPS operates. i.e. change how delta serialization operates, load to property bags, etc.

Needed for the new overridable serialization mode.

diff --git a/Engine/Plugins/FX/Niagara/Source/Niagara/Private/NiagaraScript.cpp b/Engine/Plugins/FX/Niagara/Source/Niagara/Private/NiagaraScript.cpp
--- a/Engine/Plugins/FX/Niagara/Source/Niagara/Private/NiagaraScript.cpp
+++ b/Engine/Plugins/FX/Niagara/Source/Niagara/Private/NiagaraScript.cpp
@@ -1013,7 +1013,7 @@ FNiagaraVMExecutableDataId& UNiagaraScript::GetLastGeneratedVMId(const FGuid& Ve
 //	-graph specific information (node guids) used in serialized errors/warnings for backtracking to the source in the UI
 FString UNiagaraScript::BuildNiagaraDDCKeyString(const FNiagaraVMExecutableDataId& CompileId, const FString& ScriptPath)
 {
-	enum { UE_NIAGARA_COMPILATION_DERIVEDDATA_VER = 2 };
+	enum { UE_NIAGARA_COMPILATION_DERIVEDDATA_VER = 3 };
 
 	FString KeyString;
 	KeyString.Reserve(1024);
diff --git a/Engine/Source/Runtime/Core/Public/UObject/ObjectVersion.h b/Engine/Source/Runtime/Core/Public/UObject/ObjectVersion.h
--- a/Engine/Source/Runtime/Core/Public/UObject/ObjectVersion.h
+++ b/Engine/Source/Runtime/Core/Public/UObject/ObjectVersion.h
@@ -76,6 +76,11 @@ enum class EUnrealEngineObjectUE5Version : uint32
 	// Added script property serialization offset to export table entries for saved, versioned packages
 	SCRIPT_SERIALIZATION_OFFSET,
 
+	// Adding property tag extension,
+	// Support for overridable serialization on UObject,
+	// Support for overridable logic in containers
+	PROPERTY_TAG_EXTENSION_AND_OVERRIDABLE_SERIALIZATION,
+
 	// -----<new versions can be added before this line>-------------------------------------------------
 	// - this needs to be the last line (see note below)
 	AUTOMATIC_VERSION_PLUS_ONE,
diff --git a/Engine/Source/Runtime/CoreUObject/Private/UObject/Class.cpp b/Engine/Source/Runtime/CoreUObject/Private/UObject/Class.cpp
--- a/Engine/Source/Runtime/CoreUObject/Private/UObject/Class.cpp
+++ b/Engine/Source/Runtime/CoreUObject/Private/UObject/Class.cpp
@@ -1305,14 +1305,45 @@ const FBlake3Hash& UStruct::GetSchemaHash(bool bSkipEditorOnly) const
 }
 #endif
 
+
+/**
+ *  Enum flags that indicate that additional data may be serialized prior to actual tagged property serialization
+ *	Those extensions are used to store additional function to control how TPS will resolved. i.e use overridable serialization for example
+ *	Registered flag should be serialized in ascending order
+ *  @Note: do not use lightly
+ */
+enum class EClassSerializationControlExtension : uint8
+{
+	NoExtension					= 0x00,
+	ReserveForFutureUse			= 0x01, // Can be use to add a next group of extension
+
+	////////////////////////////////////////////////
+	// First extension group
+	OverridableSerialization	= 0x02,
+
+	//
+	// Add more extension for the first group here
+	//
+};
+ENUM_CLASS_FLAGS(EClassSerializationControlExtension);
+
 void UStruct::SerializeVersionedTaggedProperties(FStructuredArchive::FSlot Slot, uint8* Data, UStruct* DefaultsStruct, uint8* Defaults, const UObject* BreakRecursionIfFullyLoad) const
 {
 	using namespace UE;
+	checkf(Data, TEXT("Expecting a non null data ptr"));
 
 	FArchive& UnderlyingArchive = Slot.GetUnderlyingArchive();
 	FUObjectSerializeContext* LoadContext = UnderlyingArchive.GetSerializeContext();
 	//SCOPED_LOADTIMER(SerializeTaggedPropertiesTime);
 
+	// Setup tagged property serialization control data extensions, this is serialized only on root i.e. UObject and not structs
+	const bool bIsUClass = IsA<UClass>();
+	EClassSerializationControlExtension SerializationControl = EClassSerializationControlExtension::NoExtension;
+	if (bIsUClass && UnderlyingArchive.UEVer() >= EUnrealEngineObjectUE5Version::PROPERTY_TAG_EXTENSION_AND_OVERRIDABLE_SERIALIZATION)
+	{
+		Slot << SA_ATTRIBUTE(TEXT("SerializationControlExtensions"), SerializationControl);
+	}
+
 	// Determine if this struct supports optional property guid's (UBlueprintGeneratedClasses Only)
 	const bool bArePropertyGuidsAvailable = (UnderlyingArchive.UEVer() >= VER_UE4_PROPERTY_GUID_IN_PROPERTY_TAG) && (!FPlatformProperties::RequiresCookedData() || UnderlyingArchive.IsSaveGame()) && ArePropertyGuidsAvailable();
 	const bool bUseRedirects = (!FPlatformProperties::RequiresCookedData() || UnderlyingArchive.IsSaveGame()) && !UnderlyingArchive.IsLoadingFromCookedPackage();
diff --git a/Engine/Source/Runtime/CoreUObject/Private/UObject/PropertyTag.cpp b/Engine/Source/Runtime/CoreUObject/Private/UObject/PropertyTag.cpp
--- a/Engine/Source/Runtime/CoreUObject/Private/UObject/PropertyTag.cpp
+++ b/Engine/Source/Runtime/CoreUObject/Private/UObject/PropertyTag.cpp
@@ -225,6 +225,13 @@ void operator<<(FStructuredArchive::FSlot Slot, FPropertyTag& Tag)
 			}
 		}
 	}
+
+	// Serialize tag extensions
+	if (UnderlyingArchive.UEVer() >= EUnrealEngineObjectUE5Version::PROPERTY_TAG_EXTENSION_AND_OVERRIDABLE_SERIALIZATION)
+	{
+		EPropertyTagExtension PropertyTagExtensions = EPropertyTagExtension::NoExtension;
+		Slot << SA_ATTRIBUTE(TEXT("PropertyExtensions"), PropertyTagExtensions);
+	}
 }
 
 // Property serializer.
@@ -272,4 +279,4 @@ void FPropertyTag::SerializeTaggedProperty(FStructuredArchive::FSlot Slot, FProp
 		UnderlyingArchive.Seek(StartOfProperty + Size);
 		Property->ClearValue(Value);
 	}
-}
\ No newline at end of file
+}
diff --git a/Engine/Source/Runtime/CoreUObject/Public/UObject/PropertyTag.h b/Engine/Source/Runtime/CoreUObject/Public/UObject/PropertyTag.h
--- a/Engine/Source/Runtime/CoreUObject/Public/UObject/PropertyTag.h
+++ b/Engine/Source/Runtime/CoreUObject/Public/UObject/PropertyTag.h
@@ -14,6 +14,25 @@
 class FArchive;
 class FProperty;
 
+/**
+ *  Enum flags that indicate that additional data was serialized for that property tag
+ *	Registered flag should be serialized in ascending order 
+ */
+enum class EPropertyTagExtension : uint8
+{
+	NoExtension					= 0x00,
+	ReserveForFutureUse			= 0x01, // Can be use to add a next group of extension
+
+	////////////////////////////////////////////////
+	// First extension group
+	OverridableOperation		= 0x02,
+
+	//
+	// Add more extension for the first group here
+	//
+};
+ENUM_CLASS_FLAGS(EPropertyTagExtension);
+
 /**
  *  A tag describing a class property, to aid in serialization.
  */
