Added saving of FProperty to FPropertyTag

Added loading for property types that do not wrap another property.

diff --git a/Engine/Source/Runtime/CoreUObject/Private/UObject/EnumProperty.cpp b/Engine/Source/Runtime/CoreUObject/Private/UObject/EnumProperty.cpp
--- a/Engine/Source/Runtime/CoreUObject/Private/UObject/EnumProperty.cpp
+++ b/Engine/Source/Runtime/CoreUObject/Private/UObject/EnumProperty.cpp
@@ -1,6 +1,7 @@
 // Copyright Epic Games, Inc. All Rights Reserved.
 
 #include "UObject/EnumProperty.h"
+#include "UObject/Package.h"
 #include "UObject/PropertyPortFlags.h"
 #include "UObject/UObjectThreadContext.h"
 #include "UObject/PropertyTag.h"
@@ -547,3 +548,33 @@ uint64 FEnumProperty::GetMaxNetSerializeBits() const
 	
 	return FMath::Min(DesiredBits, MaxBits);
 }
+
+bool FEnumProperty::LoadFromTag(const FPropertyTag& Tag)
+{
+	if (!Super::LoadFromTag(Tag))
+	{
+		return false;
+	}
+
+	// Update FByteProperty when making changes here.
+	return false;
+}
+
+void FEnumProperty::SaveToTag(FPropertyTag& Tag)
+{
+	Super::SaveToTag(Tag);
+
+	if (const UEnum* LocalEnum = Enum)
+	{
+		// RobM: Ugly hack so that we can avoid content changes in most of the packages
+		// Update FByteProperty when making changes here.
+		if (LocalEnum->GetPackage()->HasAnyPackageFlags(PKG_CompiledIn))
+		{				
+			Tag.EnumName = LocalEnum->GetFName();
+		}
+		else
+		{
+			Tag.EnumName = FName(*LocalEnum->GetPathName());
+		}
+	}
+}
diff --git a/Engine/Source/Runtime/CoreUObject/Private/UObject/Property.cpp b/Engine/Source/Runtime/CoreUObject/Private/UObject/Property.cpp
--- a/Engine/Source/Runtime/CoreUObject/Private/UObject/Property.cpp
+++ b/Engine/Source/Runtime/CoreUObject/Private/UObject/Property.cpp
@@ -2103,6 +2103,19 @@ UPropertyWrapper* FProperty::GetUPropertyWrapper()
 }
 #endif //  WITH_EDITORONLY_DATA
 
+bool FProperty::LoadFromTag(const FPropertyTag& Tag)
+{
+	checkf(GetID() == Tag.Type, TEXT("Failed to load property '%s' of type '%s' from tag of type '%s'"),
+		*WriteToString<64>(GetFName()), *WriteToString<64>(GetID()), *WriteToString<64>(Tag.Type));
+	return true;
+}
+
+void FProperty::SaveToTag(FPropertyTag& Tag)
+{
+	Tag.Prop = this;
+	Tag.Type = GetID();
+	Tag.Name = GetFName();
+}
 
 FProperty* UStruct::FindPropertyByName(FName InName) const
 {
diff --git a/Engine/Source/Runtime/CoreUObject/Private/UObject/PropertyArray.cpp b/Engine/Source/Runtime/CoreUObject/Private/UObject/PropertyArray.cpp
--- a/Engine/Source/Runtime/CoreUObject/Private/UObject/PropertyArray.cpp
+++ b/Engine/Source/Runtime/CoreUObject/Private/UObject/PropertyArray.cpp
@@ -916,3 +916,22 @@ void* FArrayProperty::GetValueAddressAtIndex_Direct(const FProperty* InInner, vo
 		return nullptr;
 	}
 }
+
+bool FArrayProperty::LoadFromTag(const FPropertyTag& Tag)
+{
+	if (!Super::LoadFromTag(Tag))
+	{
+		return false;
+	}
+
+	return false;
+}
+
+void FArrayProperty::SaveToTag(FPropertyTag& Tag)
+{
+	Super::SaveToTag(Tag);
+
+	const FProperty* LocalInner = Inner;
+	check(LocalInner);
+	Tag.InnerType = LocalInner->GetID();
+}
diff --git a/Engine/Source/Runtime/CoreUObject/Private/UObject/PropertyByte.cpp b/Engine/Source/Runtime/CoreUObject/Private/UObject/PropertyByte.cpp
--- a/Engine/Source/Runtime/CoreUObject/Private/UObject/PropertyByte.cpp
+++ b/Engine/Source/Runtime/CoreUObject/Private/UObject/PropertyByte.cpp
@@ -3,6 +3,7 @@
 #include "CoreMinimal.h"
 #include "UObject/ObjectMacros.h"
 #include "UObject/Class.h"
+#include "UObject/Package.h"
 #include "UObject/PropertyPortFlags.h"
 #include "UObject/UnrealType.h"
 #include "UObject/UnrealTypePrivate.h"
@@ -510,3 +511,33 @@ uint64 FByteProperty::GetMaxNetSerializeBits() const
 
 	return FMath::Min(DesiredBits, MaxBits);
 }
+
+bool FByteProperty::LoadFromTag(const FPropertyTag& Tag)
+{
+	if (!Super::LoadFromTag(Tag))
+	{
+		return false;
+	}
+
+	// Update FEnumProperty when making changes here.
+	return Tag.EnumName.IsNone();
+}
+
+void FByteProperty::SaveToTag(FPropertyTag& Tag)
+{
+	Super::SaveToTag(Tag);
+
+	if (const UEnum* LocalEnum = Enum)
+	{
+		// RobM: Ugly hack so that we can avoid content changes in most of the packages
+		// Update FEnumProperty when making changes here.
+		if (LocalEnum->GetPackage()->HasAnyPackageFlags(PKG_CompiledIn))
+		{				
+			Tag.EnumName = LocalEnum->GetFName();
+		}
+		else
+		{
+			Tag.EnumName = FName(*LocalEnum->GetPathName());
+		}
+	}
+}
diff --git a/Engine/Source/Runtime/CoreUObject/Private/UObject/PropertyMap.cpp b/Engine/Source/Runtime/CoreUObject/Private/UObject/PropertyMap.cpp
--- a/Engine/Source/Runtime/CoreUObject/Private/UObject/PropertyMap.cpp
+++ b/Engine/Source/Runtime/CoreUObject/Private/UObject/PropertyMap.cpp
@@ -1310,3 +1310,25 @@ void* FMapProperty::GetValueAddressAtIndex_Direct(const FProperty* Inner, void*
 	}
 	return nullptr;
 }
+
+bool FMapProperty::LoadFromTag(const FPropertyTag& Tag)
+{
+	if (!Super::LoadFromTag(Tag))
+	{
+		return false;
+	}
+
+	return false;
+}
+
+void FMapProperty::SaveToTag(FPropertyTag& Tag)
+{
+	Super::SaveToTag(Tag);
+
+	const FProperty* LocalKeyProp = KeyProp;
+	const FProperty* LocalValueProp = ValueProp;
+	check(LocalKeyProp);
+	check(LocalValueProp);
+	Tag.InnerType = LocalKeyProp->GetID();
+	Tag.ValueType = LocalValueProp->GetID();
+}
diff --git a/Engine/Source/Runtime/CoreUObject/Private/UObject/PropertyMulticastDelegate.cpp b/Engine/Source/Runtime/CoreUObject/Private/UObject/PropertyMulticastDelegate.cpp
--- a/Engine/Source/Runtime/CoreUObject/Private/UObject/PropertyMulticastDelegate.cpp
+++ b/Engine/Source/Runtime/CoreUObject/Private/UObject/PropertyMulticastDelegate.cpp
@@ -690,4 +690,22 @@ void FMulticastSparseDelegateProperty::ClearDelegate(UObject* Parent, void* Prop
 	SparseDelegate.__Internal_Clear(Parent, SparseDelegateFunc->DelegateName);
 }
 
+bool FMulticastSparseDelegateProperty::LoadFromTag(const FPropertyTag& Tag)
+{
+	if (!Super::LoadFromTag(Tag))
+	{
+		return false;
+	}
+
+	// This cannot be used without its SignatureFunction and the tag lacks the information needed to load it.
+	return false;
+}
+
+void FMulticastSparseDelegateProperty::SaveToTag(FPropertyTag& Tag)
+{
+	Super::SaveToTag(Tag);
+
+	// This needs a way to save SignatureFunction to the tag.
+}
+
 IMPLEMENT_FIELD(FMulticastSparseDelegateProperty)
diff --git a/Engine/Source/Runtime/CoreUObject/Private/UObject/PropertyOptional.cpp b/Engine/Source/Runtime/CoreUObject/Private/UObject/PropertyOptional.cpp
--- a/Engine/Source/Runtime/CoreUObject/Private/UObject/PropertyOptional.cpp
+++ b/Engine/Source/Runtime/CoreUObject/Private/UObject/PropertyOptional.cpp
@@ -525,3 +525,22 @@ uint32 FOptionalProperty::GetValueTypeHashInternal(const void* Src) const
 		return 0;
 	}
 }
+
+bool FOptionalProperty::LoadFromTag(const FPropertyTag& Tag)
+{
+	if (!Super::LoadFromTag(Tag))
+	{
+		return false;
+	}
+
+	return false;
+}
+
+void FOptionalProperty::SaveToTag(FPropertyTag& Tag)
+{
+	Super::SaveToTag(Tag);
+
+	const FProperty* LocalValueProperty = ValueProperty;
+	check(LocalValueProperty);
+	Tag.InnerType = LocalValueProperty->GetID();
+}
diff --git a/Engine/Source/Runtime/CoreUObject/Private/UObject/PropertySet.cpp b/Engine/Source/Runtime/CoreUObject/Private/UObject/PropertySet.cpp
--- a/Engine/Source/Runtime/CoreUObject/Private/UObject/PropertySet.cpp
+++ b/Engine/Source/Runtime/CoreUObject/Private/UObject/PropertySet.cpp
@@ -1082,3 +1082,22 @@ void* FSetProperty::GetValueAddressAtIndex_Direct(const FProperty* Inner, void*
 
 	return nullptr;
 }
+
+bool FSetProperty::LoadFromTag(const FPropertyTag& Tag)
+{
+	if (!Super::LoadFromTag(Tag))
+	{
+		return false;
+	}
+
+	return false;
+}
+
+void FSetProperty::SaveToTag(FPropertyTag& Tag)
+{
+	Super::SaveToTag(Tag);
+
+	const FProperty* LocalElementProp = ElementProp;
+	check(LocalElementProp);
+	Tag.InnerType = LocalElementProp->GetID();
+}
diff --git a/Engine/Source/Runtime/CoreUObject/Private/UObject/PropertyStruct.cpp b/Engine/Source/Runtime/CoreUObject/Private/UObject/PropertyStruct.cpp
--- a/Engine/Source/Runtime/CoreUObject/Private/UObject/PropertyStruct.cpp
+++ b/Engine/Source/Runtime/CoreUObject/Private/UObject/PropertyStruct.cpp
@@ -465,3 +465,23 @@ void FStructProperty::AppendSchemaHash(FBlake3& Builder, bool bSkipEditorOnly) c
 	}
 }
 #endif
+
+bool FStructProperty::LoadFromTag(const FPropertyTag& Tag)
+{
+	if (!Super::LoadFromTag(Tag))
+	{
+		return false;
+	}
+
+	return false;
+}
+
+void FStructProperty::SaveToTag(FPropertyTag& Tag)
+{
+	Super::SaveToTag(Tag);
+
+	const UScriptStruct* LocalStruct = Struct;
+	check(LocalStruct);
+	Tag.StructName = LocalStruct->GetFName();
+	Tag.StructGuid = LocalStruct->GetCustomGuid();
+}
diff --git a/Engine/Source/Runtime/CoreUObject/Private/UObject/PropertyTag.cpp b/Engine/Source/Runtime/CoreUObject/Private/UObject/PropertyTag.cpp
--- a/Engine/Source/Runtime/CoreUObject/Private/UObject/PropertyTag.cpp
+++ b/Engine/Source/Runtime/CoreUObject/Private/UObject/PropertyTag.cpp
@@ -6,77 +6,20 @@
 #include "Serialization/ArchiveUObjectFromStructuredArchive.h"
 #include "UObject/Package.h"
 #include "UObject/UnrealType.h"
-#include "UObject/EnumProperty.h"
 #include "UObject/BlueprintsObjectVersion.h"
-#include "UObject/PropertyOptional.h"
 
 /*-----------------------------------------------------------------------------
 FPropertyTag
 -----------------------------------------------------------------------------*/
 
 FPropertyTag::FPropertyTag( FArchive& InSaveAr, FProperty* Property, int32 InIndex, uint8* Value, const uint8* Defaults )
-	: Prop      (Property)
-	, Type      (Property->GetID())
-	, Name      (Property->GetFName())
-	, ArrayIndex(InIndex)
+	: ArrayIndex(InIndex)
 {
 	check(!InSaveAr.GetArchiveState().UseUnversionedPropertySerialization());
-	if (Property)
+	Property->SaveToTag(*this);
+	if (FBoolProperty* Bool = CastField<FBoolProperty>(Property))
 	{
-		// RobM: Ugly hack so that we can avoid content changes in most of the packages
-		auto GetEnumName = [](UEnum* InEnum)
-		{
-			if (InEnum->GetPackage()->HasAnyPackageFlags(PKG_CompiledIn))
-			{				
-				return InEnum->GetFName();
-			}
-			else
-			{
-				return FName(*InEnum->GetPathName());
-			}
-		};
-
-		// Handle structs.
-		if (FStructProperty* StructProperty = CastField<FStructProperty>(Property))
-		{
-			StructName = StructProperty->Struct->GetFName();
-			StructGuid = StructProperty->Struct->GetCustomGuid();
-		}
-		else if (FEnumProperty* EnumProp = CastField<FEnumProperty>(Property))
-		{
-			if (UEnum* Enum = EnumProp->GetEnum())
-			{
-				EnumName = GetEnumName(Enum);
-			}
-		}
-		else if (FByteProperty* ByteProp = CastField<FByteProperty>(Property))
-		{
-			if (ByteProp->Enum != nullptr)
-			{
-				EnumName = GetEnumName(ByteProp->Enum);
-			}
-		}
-		else if (FArrayProperty* ArrayProp = CastField<FArrayProperty>(Property))
-		{
-			InnerType = ArrayProp->Inner->GetID();
-		}
-		else if (FSetProperty* SetProp = CastField<FSetProperty>(Property))
-		{
-			InnerType = SetProp->ElementProp->GetID();
-		}
-		else if (FMapProperty* MapProp = CastField<FMapProperty>(Property))
-		{
-			InnerType = MapProp->KeyProp->GetID();
-			ValueType = MapProp->ValueProp->GetID();
-		}
-		else if (FBoolProperty* Bool = CastField<FBoolProperty>(Property))
-		{
-			BoolVal = Bool->GetPropertyValue(Value);
-		}
-		else if (FOptionalProperty* OptionalProp = CastField<FOptionalProperty>(Property))
-		{
-			InnerType = OptionalProp->GetValueProperty()->GetID();
-		}
+		BoolVal = Bool->GetPropertyValue(Value);
 	}
 }
 
diff --git a/Engine/Source/Runtime/CoreUObject/Public/UObject/EnumProperty.h b/Engine/Source/Runtime/CoreUObject/Public/UObject/EnumProperty.h
--- a/Engine/Source/Runtime/CoreUObject/Public/UObject/EnumProperty.h
+++ b/Engine/Source/Runtime/CoreUObject/Public/UObject/EnumProperty.h
@@ -78,6 +78,8 @@ public:
 #if WITH_EDITORONLY_DATA
 	COREUOBJECT_API virtual void AppendSchemaHash(FBlake3& Builder, bool bSkipEditorOnly) const override;
 #endif
+	COREUOBJECT_API virtual bool LoadFromTag(const FPropertyTag& Tag) override;
+	COREUOBJECT_API virtual void SaveToTag(FPropertyTag& Tag) override;
 	// End of FProperty interface
 
 	/**
diff --git a/Engine/Source/Runtime/CoreUObject/Public/UObject/PropertyOptional.h b/Engine/Source/Runtime/CoreUObject/Public/UObject/PropertyOptional.h
--- a/Engine/Source/Runtime/CoreUObject/Public/UObject/PropertyOptional.h
+++ b/Engine/Source/Runtime/CoreUObject/Public/UObject/PropertyOptional.h
@@ -228,5 +228,7 @@ public:
 	virtual EConvertFromTypeResult ConvertFromType(const FPropertyTag& Tag, FStructuredArchive::FSlot Slot, uint8* Data, UStruct* DefaultsStruct, const uint8* Defaults) override;
 	virtual uint32 GetValueTypeHashInternal(const void* Src) const override;
 	virtual void AddReferencedObjects(FReferenceCollector& Collector) override;
+	virtual bool LoadFromTag(const FPropertyTag& Tag) override;
+	virtual void SaveToTag(FPropertyTag& Tag) override;
 	// End of FProperty interface
 };
diff --git a/Engine/Source/Runtime/CoreUObject/Public/UObject/UnrealType.h b/Engine/Source/Runtime/CoreUObject/Public/UObject/UnrealType.h
--- a/Engine/Source/Runtime/CoreUObject/Public/UObject/UnrealType.h
+++ b/Engine/Source/Runtime/CoreUObject/Public/UObject/UnrealType.h
@@ -1047,6 +1047,18 @@ public:
 		return ShouldPort() && GetOwnerClass() != UObject::StaticClass();
 	}
 
+	/**
+	 * Loads types and other details of the property from the tag. Must be a compatible type.
+	 *
+	 * @return true if this property is usable after loading from the tag.
+	 */
+	COREUOBJECT_API virtual bool LoadFromTag(const FPropertyTag& Tag);
+
+	/**
+	 * Saves types and other details of the property to the tag.
+	 */
+	COREUOBJECT_API virtual void SaveToTag(FPropertyTag& Tag);
+
 	/**
 	 * Returns the first FProperty in this property's Outer chain that does not have a FProperty for an Outer
 	 */
@@ -1963,6 +1975,8 @@ public:
 #if WITH_EDITORONLY_DATA
 	virtual void AppendSchemaHash(FBlake3& Builder, bool bSkipEditorOnly) const override;
 #endif
+	virtual bool LoadFromTag(const FPropertyTag& Tag) override;
+	virtual void SaveToTag(FPropertyTag& Tag) override;
 	// End of FProperty interface
 
 	// FNumericProperty interface.
@@ -3669,6 +3683,8 @@ public:
 	}
 
 	virtual void* GetValueAddressAtIndex_Direct(const FProperty* Inner, void* InValueAddress, int32 Index) const override;
+	virtual bool LoadFromTag(const FPropertyTag& Tag) override;
+	virtual void SaveToTag(FPropertyTag& Tag) override;
 	// End of FProperty interface
 
 	FString GetCPPTypeCustom(FString* ExtendedTypeText, uint32 CPPExportFlags, const FString& InnerTypeText, const FString& InInnerExtendedTypeText) const;
@@ -3800,6 +3816,8 @@ public:
 	virtual bool SameType(const FProperty* Other) const override;
 	virtual EConvertFromTypeResult ConvertFromType(const FPropertyTag& Tag, FStructuredArchive::FSlot Slot, uint8* Data, UStruct* DefaultsStruct, const uint8* Defaults) override;
 	virtual void* GetValueAddressAtIndex_Direct(const FProperty* Inner, void* InValueAddress, int32 LogicalIndex) const override;
+	virtual bool LoadFromTag(const FPropertyTag& Tag) override;
+	virtual void SaveToTag(FPropertyTag& Tag) override;
 	// End of FProperty interface
 
 	FString GetCPPTypeCustom(FString* ExtendedTypeText, uint32 CPPExportFlags, const FString& KeyTypeText, const FString& InKeyExtendedTypeText, const FString& ValueTypeText, const FString& InValueExtendedTypeText) const;
@@ -3926,6 +3944,8 @@ public:
 	virtual bool SameType(const FProperty* Other) const override;
 	virtual EConvertFromTypeResult ConvertFromType(const FPropertyTag& Tag, FStructuredArchive::FSlot Slot, uint8* Data, UStruct* DefaultsStruct, const uint8* Defaults) override;
 	virtual void* GetValueAddressAtIndex_Direct(const FProperty* Inner, void* InValueAddress, int32 LogicalIndex) const override;
+	virtual bool LoadFromTag(const FPropertyTag& Tag) override;
+	virtual void SaveToTag(FPropertyTag& Tag) override;
 	// End of FProperty interface
 
 	FString GetCPPTypeCustom(FString* ExtendedTypeText, uint32 CPPExportFlags, const FString& ElementTypeText, const FString& InElementExtendedTypeText) const;
@@ -6040,6 +6060,8 @@ public:
 #if WITH_EDITORONLY_DATA
 	virtual void AppendSchemaHash(FBlake3& Builder, bool bSkipEditorOnly) const override;
 #endif
+	virtual bool LoadFromTag(const FPropertyTag& Tag) override;
+	virtual void SaveToTag(FPropertyTag& Tag) override;
 	// End of FProperty interface
 
 	bool UseBinaryOrNativeSerialization(const FArchive& Ar) const;
@@ -6358,6 +6380,8 @@ PRAGMA_ENABLE_DEPRECATION_WARNINGS
 	// FProperty interface
 	virtual void SerializeItem(FStructuredArchive::FSlot Slot, void* Value, void const* Defaults) const override;
 	virtual const TCHAR* ImportText_Internal(const TCHAR* Buffer, void* ContainerOrPropertyPtr, EPropertyPointerType PropertyPointerType, UObject* OwnerObject, int32 PortFlags, FOutputDevice* ErrorText) const override;
+	virtual bool LoadFromTag(const FPropertyTag& Tag) override;
+	virtual void SaveToTag(FPropertyTag& Tag) override;
 	// End of FProperty interface
 
 	// FMulticastDelegateProperty interface
