ZenLoader: Fix crash in PkgInfoCommandlet -all when using zenloader.
#rnx

diff --git a/Engine/Source/Runtime/CoreUObject/Private/Serialization/AsyncLoading.cpp b/Engine/Source/Runtime/CoreUObject/Private/Serialization/AsyncLoading.cpp
--- a/Engine/Source/Runtime/CoreUObject/Private/Serialization/AsyncLoading.cpp
+++ b/Engine/Source/Runtime/CoreUObject/Private/Serialization/AsyncLoading.cpp
@@ -7723,6 +7723,9 @@ FAsyncArchive::FAsyncArchive(const FPackagePath& InPackagePath, FLinkerLoad* InO
 	, SummaryReadyCallback(Forward<TFunction<void()>>(InSummaryReadyCallback))
 	, OwnerLinker(InOwner)
 {
+	// Set the Archive flags for code that uses this as an FArchive
+	SetIsLoading(true);
+
 	LogItem(TEXT("Open"));
 	FOpenAsyncPackageResult OpenResult = IPackageResourceManager::Get().OpenAsyncReadPackage(PackagePath, EPackageSegment::Header);
 	Handle = OpenResult.Handle.Release();
