Extended FPropertyBag to allow values at non-leaf paths

diff --git a/Engine/Source/Runtime/CoreUObject/Private/UObject/PropertyBag.cpp b/Engine/Source/Runtime/CoreUObject/Private/UObject/PropertyBag.cpp
--- a/Engine/Source/Runtime/CoreUObject/Private/UObject/PropertyBag.cpp
+++ b/Engine/Source/Runtime/CoreUObject/Private/UObject/PropertyBag.cpp
@@ -21,15 +21,15 @@ void FPropertyBag::Empty()
 
 void FPropertyBag::Add(const FPropertyPathName& Path, FProperty* Property, void* Data, int32 ArrayIndex)
 {
-	// TODO: NullAr is a workaround to FPropertyTag requiring an archive to assert on versioned property serialization.
-	FNullArchive NullAr;
-
 	FValue& Value = FindOrCreateValue(Path);
 
 	const bool bPropertyChanged = Value.Tag.Prop != Property;
 
 	if (bPropertyChanged)
 	{
+		// TODO: NullAr is a workaround to FPropertyTag requiring an archive to assert on versioned property serialization.
+		FNullArchive NullAr;
+
 		Value.Destroy();
 		Value.Tag = FPropertyTag(NullAr, Property, INDEX_NONE, (uint8*)Data, nullptr);
 	}
@@ -55,9 +55,9 @@ void FPropertyBag::Remove(const FPropertyPathName& Path)
 		const FPropertyPathNameSegment Segment = Path.GetSegment(SegmentIndex);
 		if (TUniquePtr<FNode>* Node = Nodes[SegmentIndex]->Find(Segment.PackNameWithIndex()))
 		{
-			if (FNodeMap* NodeMap = (*Node)->ValueOrNodes.TryGet<FNodeMap>())
+			if (FNodeMap& NodeMap = (*Node)->Nodes; !NodeMap.IsEmpty())
 			{
-				Nodes.Add(NodeMap);
+				Nodes.Add(&NodeMap);
 				continue;
 			}
 		}
@@ -100,16 +100,18 @@ void FPropertyBag::LoadPropertyByTag(const FPropertyPathName& Path, const FPrope
 		Value.Tag.ArrayIndex = INDEX_NONE;
 	}
 
-	Value.AllocateAndInitializeValue();
-
+	// Serialize the value using the existing property from the tag.
 	if (FProperty* Property = Value.Tag.Prop)
 	{
+		Value.AllocateAndInitializeValue();
 		Tag.SerializeTaggedProperty(ValueSlot, Property, (uint8*)Value.Data, (const uint8*)Defaults);
+		return;
 	}
-	else
-	{
-		UnderlyingArchive.Serialize(Value.Data, Value.Tag.Size);
-	}
+
+	// Fall back to loading the serialized value.
+	// Persisting this serialized value will require capturing version information from the archive.
+	Value.AllocateAndInitializeValue();
+	UnderlyingArchive.Serialize(Value.Data, Value.Tag.Size);
 }
 
 FPropertyBag::FValue& FPropertyBag::FindOrCreateValue(const FPropertyPathName& Path)
@@ -125,12 +127,8 @@ FPropertyBag::FValue& FPropertyBag::FindOrCreateValue(const FPropertyPathName& P
 		{
 			Node = MakeUnique<FNode>();
 		}
-		if (!Node->ValueOrNodes.IsType<FNodeMap>())
-		{
-			Node->ValueOrNodes.Emplace<FNodeMap>();
-		}
 		Node->Type = Segment.Type;
-		NodeMap = &Node->ValueOrNodes.Get<FNodeMap>();
+		NodeMap = &Node->Nodes;
 	}
 
 	FPropertyPathNameSegment LastSegment = Path.GetSegment(SegmentCount - 1);
@@ -139,12 +137,8 @@ FPropertyBag::FValue& FPropertyBag::FindOrCreateValue(const FPropertyPathName& P
 	{
 		Node = MakeUnique<FNode>();
 	}
-	if (!Node->ValueOrNodes.IsType<FValue>())
-	{
-		Node->ValueOrNodes.Emplace<FValue>();
-	}
 	Node->Type = LastSegment.Type;
-	return Node->ValueOrNodes.Get<FValue>();
+	return Node->Value;
 }
 
 ///////////////////////////////////////////////////////////////////////////////////////////////////////////////////////
@@ -174,17 +168,25 @@ void FPropertyBag::FValue::AllocateAndInitializeValue()
 
 void FPropertyBag::FValue::Destroy()
 {
-	if (Data)
+	if (const FProperty* Property = Tag.Prop)
 	{
-		if (const FProperty* Property = Tag.Prop)
+		if (Data)
 		{
 			// TODO: Need to destroy and free only one element for arrays.
 			Property->DestroyAndFreeValue(Data);
+			Data = nullptr;
 		}
-		else
+
+		if (bOwnsProperty)
 		{
-			FMemory::Free(Data);
+			bOwnsProperty = false;
+			delete Property;
+			Tag.Prop = nullptr;
 		}
+	}
+	else
+	{
+		FMemory::Free(Data);
 		Data = nullptr;
 	}
 }
@@ -205,20 +207,7 @@ int32 FPropertyBag::FValue::GetSize() const
 
 ///////////////////////////////////////////////////////////////////////////////////////////////////////////////////////
 
-FPropertyBag::FConstIterator::FConstIterator(const FNodeIterator& NodeIt)
-{
-	NodeIterators.Push(NodeIt);
-	EnterNode();
-}
-
-FPropertyBag::FConstIterator& FPropertyBag::FConstIterator::operator++()
-{
-	++NodeIterators.Last();
-	EnterNode();
-	return *this;
-}
-
-inline void FPropertyBag::FConstIterator::EnterNode()
+void FPropertyBag::FConstIterator::EnterNode()
 {
 	for (;;)
 	{
@@ -229,27 +218,24 @@ inline void FPropertyBag::FConstIterator::EnterNode()
 
 		if (FNodeIterator& NodeIt = NodeIterators.Last())
 		{
-			const TUniquePtr<FNode>& Node = NodeIt.Value();
+			FNode* Node = NodeIt.Value().Get();
 			CurrentPath.Push(FPropertyPathNameSegment().SetNameWithIndex(NodeIt.Key()).SetType(Node->Type));
+			++NodeIt;
 
-			if (FValue* Value = Node->ValueOrNodes.TryGet<FValue>())
-			{
-				CurrentValue = Value;
-				return;
-			}
-			else if (const FNodeMap* NodeMap = Node->ValueOrNodes.TryGet<FNodeMap>())
+			if (!Node->Nodes.IsEmpty())
 			{
-				NodeIterators.Push(NodeMap->CreateConstIterator());
+				NodeIterators.Push(Node->Nodes.CreateConstIterator());
 			}
-			else
+
+			if (Node->Value)
 			{
-				checkNoEntry();
+				CurrentValue = &Node->Value;
+				return;
 			}
 		}
 		else if (NodeIterators.Num() > 1)
 		{
 			NodeIterators.Pop();
-			++NodeIterators.Last();
 		}
 		else
 		{
diff --git a/Engine/Source/Runtime/CoreUObject/Public/UObject/PropertyBag.h b/Engine/Source/Runtime/CoreUObject/Public/UObject/PropertyBag.h
--- a/Engine/Source/Runtime/CoreUObject/Public/UObject/PropertyBag.h
+++ b/Engine/Source/Runtime/CoreUObject/Public/UObject/PropertyBag.h
@@ -31,8 +31,11 @@ class FPropertyBag
 
 		UE_API int32 GetSize() const;
 
+		inline explicit operator bool() const { return !!Data; }
+
 		FPropertyTag Tag;
 		void* Data = nullptr;
+		bool bOwnsProperty = false;
 	};
 
 	struct FNode;
@@ -41,7 +44,8 @@ class FPropertyBag
 	struct FNode
 	{
 		FName Type;
-		TVariant<TYPE_OF_NULLPTR, FValue, FNodeMap> ValueOrNodes;
+		FValue Value;
+		FNodeMap Nodes;
 	};
 
 	FValue& FindOrCreateValue(const FPropertyPathName& Path);
@@ -89,14 +93,18 @@ public:
 	public:
 		using FNodeIterator = typename FNodeMap::TConstIterator;
 
-		UE_API explicit FConstIterator(const FNodeIterator& NodeIt);
+		inline explicit FConstIterator(const FNodeIterator& NodeIt)
+		{
+			NodeIterators.Push(NodeIt);
+			EnterNode();
+		}
 
-		UE_API FConstIterator& operator++();
+		inline FConstIterator& operator++() { EnterNode(); return *this; }
 
-		inline explicit operator bool() const { return !!NodeIterators.Last(); }
+		inline explicit operator bool() const { return !!CurrentValue; }
 
-		inline bool operator==(const FConstIterator& Rhs) const { return NodeIterators.Last() == Rhs.NodeIterators.Last(); }
-		inline bool operator!=(const FConstIterator& Rhs) const { return NodeIterators.Last() != Rhs.NodeIterators.Last(); }
+		inline bool operator==(const FConstIterator& Rhs) const { return CurrentValue == Rhs.CurrentValue; }
+		inline bool operator!=(const FConstIterator& Rhs) const { return CurrentValue != Rhs.CurrentValue; }
 
 		inline const FPropertyPathName& GetPath() const { return CurrentPath; }
 		inline const FProperty* GetProperty() const { return CurrentValue->Tag.Prop; }
@@ -104,7 +112,7 @@ public:
 		inline int32 GetValueSize() const { return CurrentValue->GetSize(); }
 
 	private:
-		void EnterNode();
+		UE_API void EnterNode();
 
 		TArray<FNodeIterator, TInlineAllocator<8>> NodeIterators;
 		FPropertyPathName CurrentPath;
