LinkerLoad: Fix spurious reporting of native CDO subobjects as import dependencies, and the spurious warning it causes: Using non-native XObject %s!!!
#rnx

diff --git a/Engine/Source/Runtime/CoreUObject/Private/UObject/LinkerLoad.cpp b/Engine/Source/Runtime/CoreUObject/Private/UObject/LinkerLoad.cpp
--- a/Engine/Source/Runtime/CoreUObject/Private/UObject/LinkerLoad.cpp
+++ b/Engine/Source/Runtime/CoreUObject/Private/UObject/LinkerLoad.cpp
@@ -3072,15 +3072,40 @@ void FLinkerLoad::GatherImportDependencies(int32 ImportIndex, TSet<FDependencyRe
 		// turn off the flag
 		bIsGatheringDependencies = false;
 
-		bool bIsValidImport =
-			(Import.XObject != NULL && !Import.XObject->IsNative() && (!Import.XObject->HasAnyFlags(RF_ClassDefaultObject) || !(Import.XObject->GetClass()->HasAllFlags(EObjectFlags(RF_Public | RF_Transient)) && Import.XObject->GetClass()->IsNative()))) ||
-			(Import.SourceLinker != NULL && Import.SourceIndex != INDEX_NONE);
+		bool bIsValidImport = Import.SourceLinker != nullptr && Import.SourceIndex != INDEX_NONE;
+		if (!bIsValidImport &&
+			Import.XObject && // Found the XObject, so potentially report it as an import anyway
+			!Import.XObject->IsNative() // Imports of native classes are not reported as dependencies
+			)
+		{
+			// Imports that found their XObject are reported as dependencies, unless they are suppressed
+			// XObject-found Imports are suppressed if they are native classes, native class CDOs, or subobjects of native class CDOs
+			// XObject-found Imports are suppressed if they are transient non-native CDOs (or subobjects thereof)
+			UObject* RootObject = Import.XObject;
+			while (RootObject && RootObject->HasAllFlags(RF_DefaultSubObject))
+			{
+				RootObject = RootObject->GetOuter();
+			}
+			if (RootObject)
+			{
+				if (!RootObject->HasAnyFlags(RF_ClassDefaultObject))
+				{
+					// Not a CDO, so a valid import dependency
+					bIsValidImport = true;
+				}
+				else if (!RootObject->GetClass()->IsNative() && !RootObject->HasAllFlags(EObjectFlags(RF_Transient)))
+				{
+					// A non-native, non-transient CDO is a valid import dependency
+					bIsValidImport = true;
+				}
+			}
+		}
 
 		// make sure it succeeded
 		if (!bIsValidImport)
 		{
-			// don't print out for intrinsic native classes
-			if (!Import.XObject || !(Import.XObject->GetClass()->HasAnyClassFlags(CLASS_Intrinsic)))
+			// don't warn about the suppressed Import.XObject dependencies
+			if (!Import.XObject)
 			{
 				UE_ASSET_LOG(LogLinker, Warning, PackagePath, TEXT("VerifyImportInner failed [(%x, %d), (%x, %d)] for %s"), 
 					Import.XObject, Import.XObject ? (Import.XObject->IsNative() ? 1 : 0) : 0, 
