Support of new overridable serialization
New Array and Map overridable serialization logic to better support conatiner of subobjects in BP
New property tag scope for access inside property serialization methods

diff --git a/Engine/Source/Programs/Shared/EpicGames.Core/UnrealEngineTypes.cs b/Engine/Source/Programs/Shared/EpicGames.Core/UnrealEngineTypes.cs
--- a/Engine/Source/Programs/Shared/EpicGames.Core/UnrealEngineTypes.cs
+++ b/Engine/Source/Programs/Shared/EpicGames.Core/UnrealEngineTypes.cs
@@ -1204,6 +1204,11 @@ namespace EpicGames.Core
 		/// </summary>
 		TObjectPtrWrapper = TObjectPtr | UObjectWrapper,
 
+		/// <summary>
+		/// ****Experimental*** Property will use different logic to serialize knowing what changes are done against its default use the overridable information provided by the overidable manager on the object
+		/// </summary>
+		ExperimentalOverridableLogic = 0x0200000000000000,
+
 		/// <summary>
 		/// All Native Access Specifier flags
 		/// </summary>
diff --git a/Engine/Source/Programs/Shared/EpicGames.UHT/Specifiers/UhtPropertyMemberSpecifiers.cs b/Engine/Source/Programs/Shared/EpicGames.UHT/Specifiers/UhtPropertyMemberSpecifiers.cs
--- a/Engine/Source/Programs/Shared/EpicGames.UHT/Specifiers/UhtPropertyMemberSpecifiers.cs
+++ b/Engine/Source/Programs/Shared/EpicGames.UHT/Specifiers/UhtPropertyMemberSpecifiers.cs
@@ -382,6 +382,14 @@ namespace EpicGames.UHT.Parsers
 			context.PropertySettings.PropertyFlags |= EPropertyFlags.SkipSerialization;
 		}
 
+		[UhtSpecifier(Extends = UhtTableNames.PropertyMember, ValueType = UhtSpecifierValueType.Legacy)]
+		private static void ExperimentalOverridableLogicSpecifier(UhtSpecifierContext specifierContext)
+		{
+			UhtPropertySpecifierContext context = (UhtPropertySpecifierContext)specifierContext;
+			context.PropertySettings.PropertyFlags |= EPropertyFlags.ExperimentalOverridableLogic;
+		}
+
+
 		[UhtSpecifier(Extends = UhtTableNames.PropertyMember, ValueType = UhtSpecifierValueType.OptionalString)]
 		private static void GetterSpecifier(UhtSpecifierContext specifierContext, StringView? value)
 		{
diff --git a/Engine/Source/Runtime/Core/Public/Serialization/ArchiveSerializedPropertyChain.h b/Engine/Source/Runtime/Core/Public/Serialization/ArchiveSerializedPropertyChain.h
--- a/Engine/Source/Runtime/Core/Public/Serialization/ArchiveSerializedPropertyChain.h
+++ b/Engine/Source/Runtime/Core/Public/Serialization/ArchiveSerializedPropertyChain.h
@@ -107,6 +107,10 @@ public:
 #endif
 	}
 
+	FORCEINLINE TArray<class FProperty*, TInlineAllocator<8>>::TConstIterator GetRootIterator() const
+	{
+		return SerializedPropertyChain.CreateConstIterator();
+	}
 private:
 	void IncrementUpdateCount()
 	{
@@ -123,4 +127,4 @@ private:
 	/** Non-zero if on the current stack of serialized properties there's an editor-only property */
 	int32 EditorOnlyPropertyStack;
 #endif
-};
+};
\ No newline at end of file
diff --git a/Engine/Source/Runtime/CoreUObject/Private/InstancedReferenceSubobjectHelper.cpp b/Engine/Source/Runtime/CoreUObject/Private/InstancedReferenceSubobjectHelper.cpp
new file mode 100644
--- /dev/null
+++ b/Engine/Source/Runtime/CoreUObject/Private/InstancedReferenceSubobjectHelper.cpp
@@ -0,0 +1,350 @@
+// Copyright Epic Games, Inc. All Rights Reserved.
+
+#include "InstancedReferenceSubobjectHelper.h"
+#include "UObject/PropertyOptional.h"
+#include "UObject/Package.h"
+#include "UObject/Package.h"
+	
+UObject* FInstancedPropertyPath::Resolve(const UObject* Container) const
+{
+	UStruct* CurrentContainerType = Container->GetClass();
+
+	const TArray<FPropertyLink>& PropChainRef = PropertyChain;
+	auto GetProperty = [&CurrentContainerType, &PropChainRef](int32 ChainIndex)->FProperty*
+	{
+		const FProperty* SrcProperty = PropChainRef[ChainIndex].PropertyPtr;
+		return FindFProperty<FProperty>(CurrentContainerType, SrcProperty->GetFName());
+	};
+
+	auto GetArrayIndex = [&PropChainRef](int32 ChainIndex)->int32
+	{
+		return PropChainRef[ChainIndex].ArrayIndex == INDEX_NONE ? 0 : PropChainRef[ChainIndex].ArrayIndex;
+	};
+
+	const FProperty* CurrentProp = GetProperty(0);
+	const uint8* ValuePtr = (CurrentProp) ? CurrentProp->ContainerPtrToValuePtr<uint8>(Container, GetArrayIndex(0)) : nullptr;
+
+	for (int32 ChainIndex = 1; CurrentProp && ChainIndex < PropertyChain.Num(); ++ChainIndex)
+	{
+		const FPropertyLink& PropertyLink = PropertyChain[ChainIndex];
+
+		if (const FArrayProperty* ArrayProperty = CastField<FArrayProperty>(CurrentProp))
+		{
+			if (!PropertyLink.PropertyPtr->SameType(ArrayProperty->Inner))
+			{
+				CurrentProp = nullptr;
+				break;
+			}
+
+			const int32 TargetIndex = PropertyLink.ArrayIndex;
+			check(TargetIndex != INDEX_NONE);
+
+			FScriptArrayHelper ArrayHelper(ArrayProperty, ValuePtr);
+			if (TargetIndex >= ArrayHelper.Num())
+			{
+				CurrentProp = nullptr;
+				break;
+			}
+
+			CurrentProp = ArrayProperty->Inner;
+			ValuePtr    = ArrayHelper.GetRawPtr(TargetIndex);
+		}
+		else if (const FSetProperty* SetProperty = CastField<FSetProperty>(CurrentProp))
+		{
+			if (!PropertyLink.PropertyPtr->SameType(SetProperty->ElementProp))
+			{
+				CurrentProp = nullptr;
+				break;
+			}
+
+			FScriptSetHelper SetHelper(SetProperty, ValuePtr);
+
+			// Convert the logical index (recorded in the path) to the actual index used internally by the set.
+			const int32 TargetIndex = SetHelper.FindInternalIndex(PropertyLink.ArrayIndex);
+			if (SetHelper.IsValidIndex(TargetIndex))
+			{
+				CurrentProp = SetProperty->ElementProp;
+				ValuePtr = SetHelper.GetElementPtr(TargetIndex);
+			}
+			else
+			{
+				CurrentProp = nullptr;
+			}
+		}
+		else if (const FMapProperty* MapProperty = CastField<FMapProperty>(CurrentProp))
+		{
+			FScriptMapHelper MapHelper(MapProperty, ValuePtr);
+
+			// Convert the logical index (recorded in the path) to the actual index used internally by the map.
+			const int32 TargetIndex = MapHelper.FindInternalIndex(PropertyLink.ArrayIndex);
+			if (!MapHelper.IsValidIndex(TargetIndex))
+			{
+				CurrentProp = nullptr;
+			}
+			else if (!PropertyLink.bIsMapValue && PropertyLink.PropertyPtr->SameType(MapProperty->KeyProp))
+			{
+				ValuePtr = MapHelper.GetKeyPtr(TargetIndex);
+				CurrentProp = MapProperty->KeyProp;
+			}
+			else if (PropertyLink.bIsMapValue && PropertyLink.PropertyPtr->SameType(MapProperty->ValueProp))
+			{
+				ValuePtr = MapHelper.GetValuePtr(TargetIndex);
+				CurrentProp = MapProperty->ValueProp;
+			}
+			else
+			{
+				CurrentProp = nullptr;
+			}
+		}
+		else if(const FOptionalProperty* OptionalProperty = CastField<FOptionalProperty>(CurrentProp))
+		{
+			if(const uint8* InnerValuePtr = static_cast<const uint8*>(OptionalProperty->GetValuePointerForReadIfSet(ValuePtr)))
+			{
+				ValuePtr = InnerValuePtr;
+				CurrentProp = OptionalProperty->GetValueProperty();
+			}
+			else
+			{
+				CurrentProp = nullptr;
+			}
+		}
+		else
+		{
+			if (const FStructProperty* StructProperty = CastField<FStructProperty>(CurrentProp))
+			{
+				CurrentContainerType = StructProperty->Struct;
+			}
+
+			CurrentProp = GetProperty(ChainIndex);
+			ValuePtr = (CurrentProp) ? CurrentProp->ContainerPtrToValuePtr<uint8>(ValuePtr, GetArrayIndex(ChainIndex)) : nullptr;
+		}
+	}
+
+	const FObjectProperty* TargetProperty = CastField<FObjectProperty>(CurrentProp);
+	if (TargetProperty && TargetProperty->HasAnyPropertyFlags(CPF_InstancedReference))
+	{
+		return TargetProperty->GetObjectPropertyValue(ValuePtr);
+	}
+	return nullptr;
+}
+
+template<typename T>
+void FFindInstancedReferenceSubobjectHelper::ForEachInstancedSubObject(FInstancedPropertyPath& PropertyPath, T ContainerAddress, TFunctionRef<void(const FInstancedSubObjRef&, T)> ObjRefFunc)
+{
+	check(ContainerAddress);
+	const FProperty* TargetProp = PropertyPath.Head();
+
+	if (const FArrayProperty* ArrayProperty = CastField<const FArrayProperty>(TargetProp))
+	{
+		// Exit now if the array doesn't contain any instanced references.
+		if (!ArrayProperty->HasAnyPropertyFlags(CPF_ContainsInstancedReference))
+		{
+			return;
+		}
+
+		FScriptArrayHelper ArrayHelper(ArrayProperty, ContainerAddress);
+		for (int32 ElementIndex = 0; ElementIndex < ArrayHelper.Num(); ++ElementIndex)
+		{
+			T ValueAddress = ArrayHelper.GetRawPtr(ElementIndex);
+
+			PropertyPath.Push(ArrayProperty->Inner, ElementIndex);
+			ForEachInstancedSubObject(PropertyPath, ValueAddress, ObjRefFunc);
+			PropertyPath.Pop();
+		}
+	}
+	else if (const FMapProperty* MapProperty = CastField<const FMapProperty>(TargetProp))
+	{
+		// Exit now if the map doesn't contain any instanced references.
+		if (!MapProperty->HasAnyPropertyFlags(CPF_ContainsInstancedReference))
+		{
+			return;
+		}
+
+		FScriptMapHelper MapHelper(MapProperty, ContainerAddress);
+		for (FScriptMapHelper::FIterator It(MapHelper); It; ++It)
+		{
+			T KeyAddress = MapHelper.GetKeyPtr(It);
+			T ValueAddress = MapHelper.GetValuePtr(It);
+
+			// Note: Keep these as the logical (Nth) index in case the map changes internally after we construct the path or in case we resolve using a different object.
+			PropertyPath.Push(MapProperty->KeyProp, It.GetLogicalIndex());
+			ForEachInstancedSubObject(PropertyPath, KeyAddress, ObjRefFunc);
+			PropertyPath.Pop();
+
+			PropertyPath.Push(MapProperty->ValueProp, It.GetLogicalIndex(), true);
+			ForEachInstancedSubObject(PropertyPath, ValueAddress, ObjRefFunc);
+			PropertyPath.Pop();
+		}
+	}
+	else if (const FSetProperty* SetProperty = CastField<const FSetProperty>(TargetProp))
+	{
+		// Exit now if the set doesn't contain any instanced references.
+		if (!SetProperty->HasAnyPropertyFlags(CPF_ContainsInstancedReference))
+		{
+			return;
+		}
+
+		FScriptSetHelper SetHelper(SetProperty, ContainerAddress);
+		for (FScriptSetHelper::FIterator It(SetHelper); It; ++It)
+		{
+			T ValueAddress = SetHelper.GetElementPtr(It);
+
+			// Note: Keep this as the logical (Nth) index in case the set changes internally after we construct the path or in case we resolve using a different object.
+			PropertyPath.Push(SetProperty->ElementProp, It.GetLogicalIndex());
+			ForEachInstancedSubObject(PropertyPath, ValueAddress, ObjRefFunc);
+			PropertyPath.Pop();
+		}
+	}
+	else if (const FOptionalProperty* OptionalProperty = CastField<FOptionalProperty>(TargetProp))
+	{
+		// Exit early if the optional does not contain any instanced references.
+		if (!OptionalProperty->HasAnyPropertyFlags(CPF_ContainsInstancedReference))
+		{
+			return;
+		}
+		
+		if (T ValueAddress = static_cast<T>(OptionalProperty->GetValuePointerForReadOrReplaceIfSet(ContainerAddress)))
+		{
+			PropertyPath.Push(OptionalProperty->GetValueProperty());
+			ForEachInstancedSubObject(PropertyPath, ValueAddress, ObjRefFunc);
+			PropertyPath.Pop();
+		}
+	}
+	else if (const FStructProperty* StructProperty = CastField<const FStructProperty>(TargetProp))
+	{
+		// Exit early if the struct does not contain any instanced references or if the struct is invalid.
+		if (!StructProperty->HasAnyPropertyFlags(CPF_ContainsInstancedReference) || !StructProperty->Struct)
+		{
+			return;
+		}
+
+		for (FProperty* StructProp = StructProperty->Struct->RefLink; StructProp; StructProp = StructProp->NextRef)
+		{
+			for (int32 ArrayIdx = 0; ArrayIdx < StructProp->ArrayDim; ++ArrayIdx)
+			{
+				T ValueAddress = StructProp->ContainerPtrToValuePtr<uint8>(ContainerAddress, ArrayIdx);
+
+				PropertyPath.Push(StructProp, ArrayIdx);
+				ForEachInstancedSubObject(PropertyPath, ValueAddress, ObjRefFunc);
+				PropertyPath.Pop();
+			}
+		}
+	}
+	else if (TargetProp->HasAllPropertyFlags(CPF_PersistentInstance))
+	{
+		ensure(TargetProp->HasAllPropertyFlags(CPF_InstancedReference));
+		if (const FObjectProperty* ObjectProperty = CastField<const FObjectProperty>(TargetProp))
+		{
+			if (UObject* ObjectValue = ObjectProperty->GetObjectPropertyValue(ContainerAddress))
+			{
+				// don't need to push to PropertyPath, since this property is already at its head
+				ObjRefFunc(FInstancedSubObjRef(ObjectValue, PropertyPath), ContainerAddress);
+			}
+		}
+		else
+		{
+			// @todo figure out how to support verse property types
+			//UE_LOG(LogBlueprint, Error, TEXT("%s: Skipping unknow property type(%s)"), ANSI_TO_TCHAR(__FUNCTION__), *TargetProp->GetName());
+		}
+	}
+}
+
+template COREUOBJECT_API void FFindInstancedReferenceSubobjectHelper::ForEachInstancedSubObject<void*>(FInstancedPropertyPath& PropertyPath, void* ContainerAddress, TFunctionRef<void(const FInstancedSubObjRef&, void*)> ObjRefFunc);
+template COREUOBJECT_API void FFindInstancedReferenceSubobjectHelper::ForEachInstancedSubObject<const void*>(FInstancedPropertyPath& PropertyPath, const void* ContainerAddress, TFunctionRef<void(const FInstancedSubObjRef&, const void*)> ObjRefFunc);
+
+void FFindInstancedReferenceSubobjectHelper::Duplicate(
+	UObject* OldObject, 
+	UObject* NewObject, 
+	TMap<UObject*, UObject*>& ReferenceReplacementMap, 
+	TArray<UObject*>& DuplicatedObjects, 
+	TMap<UObject*, UObject*>* OptionalMappings)
+{
+	if (OldObject->GetClass()->HasAnyClassFlags(CLASS_HasInstancedReference) &&
+		NewObject->GetClass()->HasAnyClassFlags(CLASS_HasInstancedReference))
+	{
+		TArray<FInstancedSubObjRef> OldInstancedSubObjects;
+		GetInstancedSubObjects(OldObject, OldInstancedSubObjects);
+		if (OldInstancedSubObjects.Num() > 0)
+		{
+			TArray<FInstancedSubObjRef> NewInstancedSubObjects;
+			GetInstancedSubObjects(NewObject, NewInstancedSubObjects);
+			for (const FInstancedSubObjRef& Obj : NewInstancedSubObjects)
+			{
+				const bool bNewObjectHasOldOuter = (Obj->GetOuter() == OldObject);
+				if (bNewObjectHasOldOuter)
+				{
+					const bool bKeptByOld = OldInstancedSubObjects.Contains(Obj);
+					const bool bNotHandledYet = !ReferenceReplacementMap.Contains(Obj) &&
+						(!OptionalMappings || !OptionalMappings->Contains(Obj));
+					if (bKeptByOld)
+					{
+						if (bNotHandledYet)
+						{
+							// This name may have been taken by an instanced subobject, since we're copying from
+							// an object we want to recreate it here with a duplicate of the old object. I'm 
+							// using a rename call here to reserve the name for ourself. These objects come from
+							// FObjectInitializer::InstanceSubobjects()
+							UObject* ExistingObject = StaticFindObjectFast(UObject::StaticClass(), NewObject, Obj->GetFName());
+							if (ExistingObject)
+							{
+								ExistingObject->Rename(
+									nullptr,
+									GetTransientPackage(),
+									REN_DoNotDirty | REN_DontCreateRedirectors | REN_ForceNoResetLoaders | REN_NonTransactional);
+							}
+
+							UObject* NewEditInlineSubobject = StaticDuplicateObject(Obj, NewObject, Obj->GetFName());
+							ReferenceReplacementMap.Add(Obj, NewEditInlineSubobject);
+							if(OptionalMappings)
+							{
+								OptionalMappings->Add(Obj, NewEditInlineSubobject);
+							}
+
+							// NOTE: we cannot patch OldObject's linker table here, since we don't 
+							//       know the relation between the two objects (one could be of a 
+							//       super class, and the other a child)
+
+							// We also need to make sure to fixup any properties here
+							DuplicatedObjects.Add(NewEditInlineSubobject);
+						}
+						else
+						{
+							// make sure the object is outered correctly:
+							check(bNewObjectHasOldOuter);
+							UObject* RealNewSubobject = nullptr;
+							if (OptionalMappings)
+							{
+								UObject* const* Entry = OptionalMappings->Find(Obj);
+								if (Entry)
+								{
+									RealNewSubobject = *Entry;
+								}
+							}
+
+							if (!RealNewSubobject)
+							{
+								RealNewSubobject = ReferenceReplacementMap.FindChecked(Obj);
+							}
+							if (RealNewSubobject->GetOuter() != NewObject)
+							{
+								UObject* ExistingObject = StaticFindObjectFast(UObject::StaticClass(), NewObject, RealNewSubobject->GetFName());
+								if (ExistingObject)
+								{
+									ExistingObject->Rename(
+										nullptr,
+										GetTransientPackage(),
+										REN_DoNotDirty | REN_DontCreateRedirectors | REN_ForceNoResetLoaders | REN_NonTransactional);
+									ReferenceReplacementMap.Add(ExistingObject, RealNewSubobject);
+								}
+								RealNewSubobject->Rename(
+									*(RealNewSubobject->GetName()),
+									NewObject,
+									REN_DoNotDirty | REN_DontCreateRedirectors | REN_ForceNoResetLoaders | REN_NonTransactional);
+							}
+						}
+					}
+				}
+			}
+		}
+	}
+}
diff --git a/Engine/Source/Runtime/CoreUObject/Private/UObject/Class.cpp b/Engine/Source/Runtime/CoreUObject/Private/UObject/Class.cpp
--- a/Engine/Source/Runtime/CoreUObject/Private/UObject/Class.cpp
+++ b/Engine/Source/Runtime/CoreUObject/Private/UObject/Class.cpp
@@ -67,6 +67,8 @@
 #include "String/ParseTokens.h"
 #include "AutoRTFM/AutoRTFM.h"
 #include "Serialization/TestUndeclaredScriptStructObjectReferences.h"
+#include "UObject/OverridableManager.h"
+#include "UObject/OverriddenPropertySet.h"
 
 // This flag enables some expensive class tree validation that is meant to catch mutations of 
 // the class tree outside of SetSuperStruct. It has been disabled because loading blueprints 
@@ -1319,7 +1321,7 @@ enum class EClassSerializationControlExtension : uint8
 
 	////////////////////////////////////////////////
 	// First extension group
-	OverridableSerialization	= 0x02,
+	OverridableSerializationInformation	= 0x02,
 
 	//
 	// Add more extension for the first group here
@@ -1327,6 +1329,28 @@ enum class EClassSerializationControlExtension : uint8
 };
 ENUM_CLASS_FLAGS(EClassSerializationControlExtension);
 
+struct FSerializationControlExtensionContext
+{
+	uint8* Data = nullptr;
+	bool bEnableOverridableSerialization = false;
+	FOverriddenPropertySet* OverriddenProperties = nullptr;
+
+	EClassSerializationControlExtension InitializeSerializationControlExtensions()
+	{
+		EClassSerializationControlExtension SerializationExtension = EClassSerializationControlExtension::NoExtension;
+	
+		// Overridable serialization information initialization
+		if (FOverriddenPropertySet* ObjectOverriddenProperties = FOverridableManager::Get().GetOverriddenProperties(*(UObject*)Data))
+		{
+			SerializationExtension |= EClassSerializationControlExtension::OverridableSerializationInformation;
+			bEnableOverridableSerialization = true;
+			OverriddenProperties = ObjectOverriddenProperties;
+		}
+	
+		return SerializationExtension;
+	}
+};
+
 void UStruct::SerializeVersionedTaggedProperties(FStructuredArchive::FSlot Slot, uint8* Data, UStruct* DefaultsStruct, uint8* Defaults, const UObject* BreakRecursionIfFullyLoad) const
 {
 	using namespace UE;
@@ -1336,14 +1360,37 @@ void UStruct::SerializeVersionedTaggedProperties(FStructuredArchive::FSlot Slot,
 	FUObjectSerializeContext* LoadContext = UnderlyingArchive.GetSerializeContext();
 	//SCOPED_LOADTIMER(SerializeTaggedPropertiesTime);
 
-	// Setup tagged property serialization control data extensions, this is serialized only on root i.e. UObject and not structs
+	// Setup serialization control data extensions, this is serialized only on root i.e. UObject and not structs!
 	const bool bIsUClass = IsA<UClass>();
-	EClassSerializationControlExtension SerializationControl = EClassSerializationControlExtension::NoExtension;
+	FSerializationControlExtensionContext ControlContext{ Data };
 	if (bIsUClass && UnderlyingArchive.UEVer() >= EUnrealEngineObjectUE5Version::PROPERTY_TAG_EXTENSION_AND_OVERRIDABLE_SERIALIZATION)
 	{
+		EClassSerializationControlExtension SerializationControl = EClassSerializationControlExtension::NoExtension;
+		if (UnderlyingArchive.IsSaving())
+		{
+			SerializationControl = ControlContext.InitializeSerializationControlExtensions();
+		}
+
 		Slot << SA_ATTRIBUTE(TEXT("SerializationControlExtensions"), SerializationControl);
+
+		// Overridable serialization information serialization
+		if (EnumHasAnyFlags(SerializationControl, EClassSerializationControlExtension::OverridableSerializationInformation))
+		{
+			checkf(!UnderlyingArchive.ArUseCustomPropertyList, TEXT("Overridable serialization does not support custom property list"))
+			EOverriddenPropertyOperation Operation = UnderlyingArchive.IsSaving() ? ControlContext.OverriddenProperties->GetOverriddenPropertyOperation(/*CurrentPropertyChain*/nullptr, /*Property*/nullptr) : EOverriddenPropertyOperation::None;
+			Slot << SA_ATTRIBUTE(TEXT("OverridableOperation"), Operation);
+
+			if (UnderlyingArchive.IsLoading())
+			{
+				ControlContext.bEnableOverridableSerialization = true;
+				ControlContext.OverriddenProperties = &FOverridableManager::Get().SetOverriddenProperties(*(UObject*)Data, Operation);
+			}
+		}
 	}
 
+	// Scope that enables the overridable serialization for this object
+	FEnableOverridableSerializationScope OverridableSerializationScope(ControlContext.bEnableOverridableSerialization, ControlContext.OverriddenProperties);
+
 	// Determine if this struct supports optional property guid's (UBlueprintGeneratedClasses Only)
 	const bool bArePropertyGuidsAvailable = (UnderlyingArchive.UEVer() >= VER_UE4_PROPERTY_GUID_IN_PROPERTY_TAG) && (!FPlatformProperties::RequiresCookedData() || UnderlyingArchive.IsSaveGame()) && ArePropertyGuidsAvailable();
 	const bool bUseRedirects = (!FPlatformProperties::RequiresCookedData() || UnderlyingArchive.IsSaveGame()) && !UnderlyingArchive.IsLoadingFromCookedPackage();
@@ -1555,6 +1602,12 @@ void UStruct::SerializeVersionedTaggedProperties(FStructuredArchive::FSlot Slot,
 					{
 						FStructuredArchive::FSlot ValueSlot = PropertyRecord.EnterField(TEXT("Value"));
 
+						// The operation was set part of the tag, now that we know the associated property, restore the overridden operation on the object
+						if (FOverriddenPropertySet* OverriddenProperties = FOverridableSerializationLogic::GetOverriddenProperties())
+						{
+							OverriddenProperties->SetOverriddenPropertyOperation(Tag.OverrideOperation, UnderlyingArchive.GetSerializedPropertyChain(), Property);
+						}
+
 						switch (Property->ConvertFromType(Tag, ValueSlot, Data, DefaultsStruct, Defaults))
 						{
 							case EConvertFromTypeResult::Converted:
@@ -1648,6 +1701,7 @@ void UStruct::SerializeVersionedTaggedProperties(FStructuredArchive::FSlot Slot,
 		}
 
 		// Save tagged properties.
+		const bool bDoDeltaSerialization = UnderlyingArchive.DoDelta() && !UnderlyingArchive.IsTransacting() && (Defaults || bIsUClass);
 
 		// Iterate over properties in the order they were linked and serialize them.
 		const FCustomPropertyListNode* CustomPropertyNode = UnderlyingArchive.ArUseCustomPropertyList ? UnderlyingArchive.ArCustomPropertyList : nullptr;
@@ -1671,7 +1725,9 @@ void UStruct::SerializeVersionedTaggedProperties(FStructuredArchive::FSlot Slot,
 				{
 					uint8* DataPtr      = Property->ContainerPtrToValuePtr           <uint8>(Data, Idx);
 					uint8* DefaultValue = Property->ContainerPtrToValuePtrForDefaults<uint8>(DefaultsStruct, Defaults, Idx);
-					if (StaticArrayContainer.IsSet() || CustomPropertyNode || !UnderlyingArchive.DoDelta() || UnderlyingArchive.IsTransacting() || (!Defaults && !dynamic_cast<const UClass*>(this)) || !Property->Identical(DataPtr, DefaultValue, UnderlyingArchive.GetPortFlags()))
+					if (StaticArrayContainer.IsSet() || CustomPropertyNode || !bDoDeltaSerialization ||
+						(FOverridableSerializationLogic::IsEnabled() && FOverridableSerializationLogic::GetOverriddenPropertyOperation(UnderlyingArchive, Property) != EOverriddenPropertyOperation::None) ||
+						(!FOverridableSerializationLogic::IsEnabled() && !Property->Identical(DataPtr, DefaultValue, UnderlyingArchive.GetPortFlags())))
 					{
 						if (bUseAtomicSerialization)
 						{
diff --git a/Engine/Source/Runtime/CoreUObject/Private/UObject/Obj.cpp b/Engine/Source/Runtime/CoreUObject/Private/UObject/Obj.cpp
--- a/Engine/Source/Runtime/CoreUObject/Private/UObject/Obj.cpp
+++ b/Engine/Source/Runtime/CoreUObject/Private/UObject/Obj.cpp
@@ -39,6 +39,7 @@
 #include "UObject/ObjectSaveContext.h"
 #include "UObject/UObjectAnnotation.h"
 #include "UObject/ReferenceChainSearch.h"
+#include "UObject/OverridableManager.h"
 #include "Serialization/ArchiveCountMem.h"
 #include "Serialization/ArchiveShowReferences.h"
 #include "Serialization/ArchiveFindCulprit.h"
@@ -471,6 +472,8 @@ void UObject::PreEditChange( FEditPropertyChain& PropertyAboutToChange )
 		GetArchetypeInstances(Objects);
 		PropagatePreEditChange(Objects, PropertyAboutToChange);
 	}
+
+	FOverridableManager::Get().PreOverrideProperty(*this, PropertyAboutToChange);
 }
 
 
@@ -542,6 +545,8 @@ void UObject::PostEditChangeChainProperty(FPropertyChangedChainEvent& PropertyCh
 	}
 
 	PostEditChangeProperty(PropertyEvent);
+
+	FOverridableManager::Get().PostOverrideProperty(*this, PropertyChangedEvent, PropertyChangedEvent.PropertyChain);
 }
 
 bool UObject::CanEditChange( const FProperty* InProperty ) const
@@ -557,6 +562,13 @@ bool UObject::CanEditChange(const FEditPropertyChain& PropertyChain) const
 
 void UObject::PropagatePreEditChange( TArray<UObject*>& AffectedObjects, FEditPropertyChain& PropertyAboutToChange )
 {
+	// This feature is not essential for overridable serialization
+	// Disable it until we fix the fact that do it will create overrides on all the instances
+	if (FOverridableManager::Get().IsEnabled(*this))
+	{
+		return;
+	}
+
 	TArray<UObject*> Instances;
 
 	for ( int32 i = 0; i < AffectedObjects.Num(); i++ )
@@ -642,6 +654,13 @@ void UObject::PropagatePreEditChange( TArray<UObject*>& AffectedObjects, FEditPr
 
 void UObject::PropagatePostEditChange( TArray<UObject*>& AffectedObjects, FPropertyChangedChainEvent& PropertyChangedEvent )
 {
+	// This feature is not essential for overridable serialization
+	// Disable it until we fix the fact that do it will create overrides on all the instances
+	if (FOverridableManager::Get().IsEnabled(*this))
+	{
+		return;
+	}
+
 	TArray<UObject*> Instances;
 
 	for ( int32 i = 0; i < AffectedObjects.Num(); i++ )
@@ -1625,6 +1644,12 @@ void UObject::Serialize(FStructuredArchive::FRecord Record)
 		// Handle derived UClass objects (exact UClass objects are native only and shouldn't be touched)
 		if (ObjClass != UClass::StaticClass())
 		{
+			if (UnderlyingArchive.IsTransacting())
+			{
+				// Serialize the overriden property inside the transaction buffer
+				FOverridableManager::Get().SerializeOverriddenProperties(*this, Record);
+			}
+
 			SerializeScriptProperties(Record.EnterField(TEXT("Properties")));
 		}
 
@@ -3865,6 +3890,7 @@ COREUOBJECT_API TArray<const TCHAR*> ParsePropertyFlags(EPropertyFlags InFlags)
 		TEXT("CPF_NativeAccessSpecifierProtected"),
 		TEXT("CPF_NativeAccessSpecifierPrivate"),
 		TEXT("CPF_SkipSerialization"),
+		TEXT("CPF_ExperimentalOverridableLogic"),
 	};
 
 	uint64 Flags = InFlags;
diff --git a/Engine/Source/Runtime/CoreUObject/Private/UObject/OverridableManager.cpp b/Engine/Source/Runtime/CoreUObject/Private/UObject/OverridableManager.cpp
new file mode 100644
--- /dev/null
+++ b/Engine/Source/Runtime/CoreUObject/Private/UObject/OverridableManager.cpp
@@ -0,0 +1,183 @@
+// Copyright Epic Games, Inc. All Rights Reserved.
+
+#include "UObject/OverridableManager.h"
+#include "InstancedReferenceSubobjectHelper.h"
+
+/*
+ *************************************************************************************
+ * Overridable serialization is experimental, not supported and use at your own risk *
+ *************************************************************************************
+ */
+ 
+FOverridableManager& FOverridableManager::Get()
+{
+	static FOverridableManager OverridableManager;
+	return OverridableManager;
+}
+
+bool FOverridableManager::IsEnabled(const UObject& Object)
+{
+	return OverriddenObjectAnnotations.IsEnabled(Object);
+}
+
+void FOverridableManager::Enable(UObject& Object)
+{
+	OverriddenObjectAnnotations.FindOrAdd(Object);
+}
+
+void FOverridableManager::Disable(UObject& Object)
+{
+	OverriddenObjectAnnotations.RemoveAnnotation(&Object);
+}
+
+void FOverridableManager::InheritEnabledFrom(UObject& Object, const UObject* DefaultData)
+{
+	if (!OverriddenObjectAnnotations.IsEnabled(Object))
+	{
+		const UObject* Outer = Object.GetOuter();
+		if ((Outer && IsEnabled(*Outer)) || (DefaultData && IsEnabled(*DefaultData)))
+		{
+			Enable(Object);
+		}
+	}
+}
+
+FOverriddenPropertySet* FOverridableManager::GetOverriddenProperties(UObject& Object)
+{
+	return OverriddenObjectAnnotations.Find(Object);
+}
+
+FOverriddenPropertySet& FOverridableManager::SetOverriddenProperties(UObject& Object, EOverriddenPropertyOperation Operation)
+{
+	FOverriddenPropertySet& ObjectOverriddenProperties = OverriddenObjectAnnotations.FindOrAdd(Object);
+	ObjectOverriddenProperties.Reset();
+	ObjectOverriddenProperties.SetOverriddenPropertyOperation(Operation, /*CurrentPropertyChain*/nullptr, /*Property*/nullptr);
+	return ObjectOverriddenProperties;
+}
+
+void FOverridableManager::CopyOverriddenProperties(UObject& DestObject, const UObject& SourceObject)
+{
+	if (const FOverriddenPropertySet* SourceObjectOverriddenProperties = OverriddenObjectAnnotations.Find(SourceObject))
+	{
+		FOverriddenPropertySet& DestObjectOverriddenProperties = OverriddenObjectAnnotations.FindOrAdd(DestObject);
+		DestObjectOverriddenProperties = *SourceObjectOverriddenProperties;
+	}
+	else
+	{
+		OverriddenObjectAnnotations.RemoveAnnotation(&DestObject);
+	}
+}
+
+void FOverridableManager::OverrideObject(UObject& Object)
+{
+	if (FOverriddenPropertySet* ThisObjectOverriddenProperties = OverriddenObjectAnnotations.Find(Object))
+	{
+		ThisObjectOverriddenProperties->NotifyPropertyChange(EPropertyNotificationType::PreEdit, FPropertyChangedEvent(nullptr), nullptr, &Object);
+		ThisObjectOverriddenProperties->NotifyPropertyChange(EPropertyNotificationType::PostEdit, FPropertyChangedEvent(nullptr), nullptr, &Object);
+	}
+}
+
+void FOverridableManager::OverrideInstancedSubObject(UObject& Object, UObject& InstancedSubObject)
+{
+	if (InstancedSubObject.IsIn(&Object))
+	{
+		OverrideObject(InstancedSubObject);
+	}
+}
+
+void FOverridableManager::PropagateOverrideToInstancedSubObjects(UObject& Object)
+{
+	TSet<UObject*> InstancedSubObjects;
+	FFindInstancedReferenceSubobjectHelper::GetInstancedSubObjects(&Object, InstancedSubObjects);
+	for (UObject* InstancedSubObject : InstancedSubObjects)
+	{
+		checkf(InstancedSubObject, TEXT("Expecting non null SubObjects"));
+		OverrideInstancedSubObject(Object, *InstancedSubObject);
+	}
+}
+
+void FOverridableManager::OverrideProperty(UObject& Object, const FPropertyChangedEvent& PropertyEvent, const FEditPropertyChain& PropertyChain)
+{
+	if (FOverriddenPropertySet* ThisObjectOverriddenProperties = OverriddenObjectAnnotations.Find(Object) )
+	{
+		ThisObjectOverriddenProperties->NotifyPropertyChange(EPropertyNotificationType::PreEdit, FPropertyChangedEvent(nullptr), PropertyChain.GetHead(), &Object);
+		ThisObjectOverriddenProperties->NotifyPropertyChange(EPropertyNotificationType::PostEdit, PropertyEvent, PropertyChain.GetHead(), &Object);
+	}
+}
+
+void FOverridableManager::PreOverrideProperty(UObject& Object, const FEditPropertyChain& PropertyChain)
+{
+	if (FOverriddenPropertySet* ThisObjectOverriddenProperties = OverriddenObjectAnnotations.Find(Object))
+	{
+		ThisObjectOverriddenProperties->NotifyPropertyChange(EPropertyNotificationType::PreEdit, FPropertyChangedEvent(nullptr), PropertyChain.GetHead(), &Object);
+	}
+}
+
+void FOverridableManager::PostOverrideProperty(UObject& Object, const FPropertyChangedEvent& PropertyEvent, const FEditPropertyChain& PropertyChain)
+{
+	if (FOverriddenPropertySet* ThisObjectOverriddenProperties = OverriddenObjectAnnotations.Find(Object))
+	{
+		ThisObjectOverriddenProperties->NotifyPropertyChange(EPropertyNotificationType::PostEdit, PropertyEvent, PropertyChain.GetHead(), &Object);
+	}
+}
+
+void FOverridableManager::NotifyPropertyChange(const EPropertyNotificationType Notification, UObject& Object, const FPropertyChangedEvent& PropertyEvent, const FEditPropertyChain::TDoubleLinkedListNode* PropertyNode)
+{
+	OverriddenObjectAnnotations.FindChecked(Object).NotifyPropertyChange(Notification, PropertyEvent, PropertyNode, &Object);
+}
+
+void FOverridableManager::ClearOverrides(UObject& Object)
+{
+	if(FOverriddenPropertySet* ThisObjectOverriddenProperties = OverriddenObjectAnnotations.Find(Object))
+	{
+		ThisObjectOverriddenProperties->Reset();
+		PropagateClearOverridesToInstancedSubObjects(Object);
+	}
+}
+
+void FOverridableManager::ClearInstancedSubObjectOverrides(UObject& Object, UObject& InstancedSubObject)
+{
+	if (InstancedSubObject.IsIn(&Object))
+	{
+		ClearOverrides(InstancedSubObject);
+	}
+}
+
+void FOverridableManager::PropagateClearOverridesToInstancedSubObjects(UObject& Object)
+{
+	TSet<UObject*> InstancedSubObjects;
+	FFindInstancedReferenceSubobjectHelper::GetInstancedSubObjects(&Object, InstancedSubObjects);
+	for (UObject* InstancedSubObject : InstancedSubObjects)
+	{
+		checkf(InstancedSubObject, TEXT("Expecting non null SubObjects"));
+
+		// There are some cases where the property has information about that should be an instanced subobject, but it is not owned by us.
+		ClearInstancedSubObjectOverrides(Object, *InstancedSubObject);
+	}
+}
+
+void FOverridableManager::SerializeOverriddenProperties(UObject& Object, FStructuredArchive::FRecord ObjectRecord)
+{
+	const FArchiveState& ArchiveState = ObjectRecord.GetArchiveState();
+	FOverriddenPropertySet* OverriddenProperties = ArchiveState.IsSaving() ? GetOverriddenProperties(Object) : nullptr;
+	TOptional<FStructuredArchiveSlot> OverridenPropertiesSlot = ObjectRecord.TryEnterField(TEXT("OverridenProperties"), OverriddenProperties != nullptr);
+	if (OverridenPropertiesSlot.IsSet())
+	{
+		EOverriddenPropertyOperation Operation = OverriddenProperties ? OverriddenProperties->GetOverriddenPropertyOperation(nullptr, nullptr) : EOverriddenPropertyOperation::None;
+		*OverridenPropertiesSlot << SA_ATTRIBUTE( TEXT("OverriddenOperation"), Operation);
+
+		if (ArchiveState.IsLoading())
+		{
+			OverriddenProperties = &SetOverriddenProperties(Object, Operation);
+		}
+
+		if (Operation != EOverriddenPropertyOperation::None)
+		{
+			FOverriddenPropertySet::StaticStruct()->SerializeItem(*OverridenPropertiesSlot, OverriddenProperties, /* Defaults */ nullptr);
+		}
+	}
+	else if (ArchiveState.IsLoading())
+	{
+		Disable(Object);
+	}
+}
\ No newline at end of file
diff --git a/Engine/Source/Runtime/CoreUObject/Private/UObject/OverriddenPropertySet.cpp b/Engine/Source/Runtime/CoreUObject/Private/UObject/OverriddenPropertySet.cpp
new file mode 100644
--- /dev/null
+++ b/Engine/Source/Runtime/CoreUObject/Private/UObject/OverriddenPropertySet.cpp
@@ -0,0 +1,760 @@
+// Copyright Epic Games, Inc. All Rights Reserved.
+
+#include "UObject/OverriddenPropertySet.h"
+
+#include "UObject/OverridableManager.h"
+#include "Serialization/ArchiveSerializedPropertyChain.h"
+#include "UObject/PropertyOptional.h"
+#include "Misc/ScopeExit.h"
+#include "Serialization/StructuredArchiveNameHelpers.h"
+
+/*
+ *************************************************************************************
+ * Overridable serialization is experimental, not supported and use at your own risk *
+ *************************************************************************************
+ */
+
+DEFINE_LOG_CATEGORY(LogOverridableObject);
+
+//----------------------------------------------------------------------//
+// FOverridableSerializationLogic
+//----------------------------------------------------------------------//
+thread_local bool FOverridableSerializationLogic::bUseOverridableSerialization = false;
+thread_local FOverriddenPropertySet* FOverridableSerializationLogic::OverriddenProperties = nullptr;
+
+EOverriddenPropertyOperation FOverridableSerializationLogic::GetOverriddenPropertyOperation(const FArchive& Ar, FProperty* Property /*= nullptr*/)
+{
+	checkf(bUseOverridableSerialization, TEXT("Nobody should use this method if it is not setup to use overridable serialization"));
+
+	const FArchiveSerializedPropertyChain* CurrentPropertyChain = Ar.GetSerializedPropertyChain();
+	const EOverriddenPropertyOperation OverriddenOperation = OverriddenProperties ? OverriddenProperties->GetOverriddenPropertyOperation(CurrentPropertyChain, Property) : EOverriddenPropertyOperation::None;
+	if (OverriddenOperation != EOverriddenPropertyOperation::None)
+	{
+		return OverriddenOperation;
+	}
+
+	// It does not mean that if we have no record of an overriden operation that a subobject might have one, need to traverse all possible subobjects. 
+	if (const FProperty* CurrentProperty = Property ? Property : (CurrentPropertyChain ? CurrentPropertyChain->GetPropertyFromStack(0) : nullptr) )
+	{
+		// Here we should just use CurrentProperty->ContainsInstancedObjectProperty() but somehow this CPF_InstanceReference seems to always be there on verse classes
+		// This should probably be fix in verse at some point.
+		if (CurrentProperty->HasAnyPropertyFlags(CPF_PersistentInstance))
+		{
+			return EOverriddenPropertyOperation::Modified;
+		}
+		else if (const FArrayProperty* CurrentArrayProperty = CastField<FArrayProperty>(CurrentProperty))
+		{
+			return CurrentArrayProperty->Inner->HasAnyPropertyFlags(CPF_PersistentInstance) ? EOverriddenPropertyOperation::Modified : EOverriddenPropertyOperation::None;
+		}
+		else if(const FMapProperty* CurrentMapProperty = CastField<FMapProperty>(CurrentProperty))
+		{
+			return CurrentMapProperty->ValueProp->HasAnyPropertyFlags(CPF_PersistentInstance) ? EOverriddenPropertyOperation::Modified : EOverriddenPropertyOperation::None;
+		}
+	}
+	
+	return EOverriddenPropertyOperation::None;
+}
+
+//----------------------------------------------------------------------//
+// FOverridableSerializationScope
+//----------------------------------------------------------------------//
+FEnableOverridableSerializationScope::FEnableOverridableSerializationScope(bool bEnableOverridableSerialization, FOverriddenPropertySet* OverriddenProperties)
+{
+	if (bEnableOverridableSerialization)
+	{
+		FOverridableSerializationLogic::Enable(OverriddenProperties);
+		bOverridableSerializationEnabled = true;
+	}
+}
+
+FEnableOverridableSerializationScope::~FEnableOverridableSerializationScope()
+{
+	if (bOverridableSerializationEnabled)
+	{
+		FOverridableSerializationLogic::Disable();
+	}
+}
+
+
+//----------------------------------------------------------------------//
+// FOverriddenPropertySet
+//----------------------------------------------------------------------//
+FOverriddenPropertyNode& FOverriddenPropertySet::FindOrAddNode(FOverriddenPropertyNode& ParentNode, FName NodeID)
+{
+	FName& SubNodeKey = ParentNode.SubPropertyNodeKeys.FindOrAdd(NodeID, NAME_None);
+	if (!SubNodeKey.IsNone())
+	{
+		FOverriddenPropertyNode* FoundNode = OverriddenPropertyNodes.FindByHash(GetTypeHash(SubNodeKey), SubNodeKey);
+		checkf(FoundNode, TEXT("Expecting a node"));
+		return *FoundNode;
+	}
+
+	// We can safely assume that the parent node is at least modified from now on
+	if (ParentNode.Operation == EOverriddenPropertyOperation::None)
+	{
+		ParentNode.Operation = EOverriddenPropertyOperation::Modified;
+	}
+
+	// Not found add the node
+	FStringBuilderBase SubPropertyKeyBuilder;
+	SubPropertyKeyBuilder = *ParentNode.PropertyPath.ToString();
+	SubPropertyKeyBuilder.Append(TEXT("."));
+	SubPropertyKeyBuilder.Append(*NodeID.ToString());
+	SubNodeKey = FName(SubPropertyKeyBuilder.ToString());
+	const FSetElementId NewID = OverriddenPropertyNodes.Emplace(SubNodeKey);
+	return OverriddenPropertyNodes.Get(NewID);
+}
+
+void FOverriddenPropertySet::NotifyPropertyChange(FOverriddenPropertyNode& ParentPropertyNode, const EPropertyNotificationType Notification, const FPropertyChangedEvent& PropertyEvent, const FEditPropertyChain::TDoubleLinkedListNode* PropertyNode, const void* Data)
+{
+	checkf(IsValid(Owner), TEXT("Expecting a valid overridable owner"));
+
+	FOverridableManager& OverridableManager = FOverridableManager::Get();
+	if (!PropertyNode)
+	{
+		if (Notification == EPropertyNotificationType::PostEdit)
+		{
+			// Replacing this entire property
+			ParentPropertyNode.Operation = EOverriddenPropertyOperation::Replace;
+
+			// Sub-property overrides are not needed from now on, so clear them
+			RemoveOverriddenSubProperties(ParentPropertyNode);
+
+			// If we are overriding the root node, need to propagate the overrides to all instanced sub object
+			const FOverriddenPropertyNode* RootNode = OverriddenPropertyNodes.FindByHash(GetTypeHash(RootNodeName), RootNodeName);
+			checkf(RootNode, TEXT("Expecting to always have a "));
+			if (RootNode == &ParentPropertyNode)
+			{
+				OverridableManager.PropagateOverrideToInstancedSubObjects(*Owner);
+			}
+
+		}
+		return;
+	}
+
+	const FProperty* Property = PropertyNode->GetValue();
+	checkf(Property, TEXT("Expecting a valid property"));
+
+	const FName PropID = Property->GetFName();
+	const void* SubValuePtr = Property->ContainerPtrToValuePtr<void>(Data, 0); //@todo support static arrays
+
+	FOverriddenPropertyNode& SubPropertyNode = FindOrAddNode(ParentPropertyNode, PropID);
+	if (SubPropertyNode.Operation == EOverriddenPropertyOperation::Replace)
+	{
+		// No need to continue further, everything is replaced for here down anyway
+		return;
+	}
+
+	if (const FArrayProperty* ArrayProperty = CastField<FArrayProperty>(Property))
+	{
+		// Only special case is instanced subobjects, otherwise we fallback to full array override
+		if (FObjectProperty* InnerObjectProperty = CastField<FObjectProperty>(ArrayProperty->Inner))
+		{
+			if (InnerObjectProperty->HasAnyPropertyFlags(CPF_PersistentInstance))
+			{
+				FScriptArrayHelper ArrayHelper(ArrayProperty, SubValuePtr);
+				int32 ArrayIndex = PropertyEvent.GetArrayIndex(Property->GetName());
+				if (!PropertyNode->GetNextNode())
+				{
+					checkf(ArrayProperty->Inner, TEXT("Expecting an inner type for Arrays"));
+
+					static TArray<uint8> SavedPreEditSubObjects;
+					FScriptArrayHelper PreEditSubObjectsArrayHelper(ArrayProperty, &SavedPreEditSubObjects);
+					if (Notification == EPropertyNotificationType::PreEdit)
+					{
+						PreEditSubObjectsArrayHelper.EmptyAndAddValues(ArrayHelper.Num());
+						for(int32 i = 0; i < ArrayHelper.Num(); i++)
+						{
+							InnerObjectProperty->SetObjectPropertyValue(PreEditSubObjectsArrayHelper.GetElementPtr(i), InnerObjectProperty->GetObjectPropertyValue(ArrayHelper.GetElementPtr(i)));
+						}
+						return;
+					}
+
+					auto ArrayReplace = [&]
+					{
+						// Overriding all entry in the array
+						SubPropertyNode.Operation = EOverriddenPropertyOperation::Replace;
+
+						// This is a case of the entire array is overridden
+						// Need to loop through every sub object and setup them up as overridden
+						for (int i = 0; i < ArrayHelper.Num(); ++i)
+						{
+							if (UObject* SubObject = InnerObjectProperty->GetObjectPropertyValue(ArrayHelper.GetElementPtr(i)))
+							{
+								const FName SubObjectID = SubObject->GetFName();
+								FOverriddenPropertyNode& SubObjectNode = FindOrAddNode(SubPropertyNode, SubObjectID);
+								SubObjectNode.Operation = EOverriddenPropertyOperation::Replace;
+
+								OverridableManager.OverrideInstancedSubObject(*Owner, *SubObject);
+							}
+						}
+					};
+
+					auto ArrayAddImpl = [&]()
+					{
+						checkf(ArrayHelper.IsValidIndex(ArrayIndex), TEXT("ArrayAdd change type expected to have an valid index"));
+						if (UObject* AddedSubObject = InnerObjectProperty->GetObjectPropertyValue(ArrayHelper.GetElementPtr(ArrayIndex)))
+						{
+							const FName AddedSubObjectID = AddedSubObject->GetFName();
+							FOverriddenPropertyNode& AddedSubObjectNode = FindOrAddNode(SubPropertyNode, AddedSubObjectID);
+							if (AddedSubObjectNode.Operation == EOverriddenPropertyOperation::Remove)
+							{
+								// This is a readd, make sure all the subobjects values are overridden
+								OverridableManager.OverrideInstancedSubObject(*Owner, *AddedSubObject);
+							}
+							AddedSubObjectNode.Operation = EOverriddenPropertyOperation::Add;
+						}
+					};
+
+					auto ArrayRemoveImpl = [&]()
+					{
+						checkf(PreEditSubObjectsArrayHelper.IsValidIndex(ArrayIndex), TEXT("ArrayRemove change type expected to have an valid index"));
+						if (UObject* RemovedSubObject = InnerObjectProperty->GetObjectPropertyValue(PreEditSubObjectsArrayHelper.GetElementPtr(ArrayIndex)))
+						{
+							const FName RemovedSubObjectID = RemovedSubObject->GetFName();
+							FOverriddenPropertyNode& RemovedSubObjectNode = FindOrAddNode(SubPropertyNode, RemovedSubObjectID);
+							OverridableManager.ClearInstancedSubObjectOverrides(*Owner, *RemovedSubObject);
+							if (RemovedSubObjectNode.Operation == EOverriddenPropertyOperation::Add)
+							{
+								// An add then a remove becomes no opt
+								SubPropertyNode.SubPropertyNodeKeys.Remove(RemovedSubObjectID);
+							}
+							else
+							{
+								RemovedSubObjectNode.Operation = EOverriddenPropertyOperation::Remove;
+							}
+						}
+					};
+
+					// Only arrays flagged overridable logic can record deltas, for now just override entire array
+					if (!ArrayProperty->HasAnyPropertyFlags(CPF_ExperimentalOverridableLogic))
+					{
+						if(PropertyEvent.ChangeType == EPropertyChangeType::Unspecified && ArrayIndex == INDEX_NONE)
+						{
+							// Overriding all entry in the array + override instanced sub obejects
+							ArrayReplace();
+						}
+						else
+						{
+							// Overriding all entry in the array
+							SubPropertyNode.Operation = EOverriddenPropertyOperation::Replace;
+						}
+						return;
+					}
+
+					switch(PropertyEvent.ChangeType)
+					{
+					case EPropertyChangeType::Unspecified:
+						{
+							if (ArrayIndex != INDEX_NONE)
+							{
+								// Overriding a single entry in the array
+								ArrayRemoveImpl();
+								ArrayAddImpl();
+							}
+							else
+							{
+								ArrayReplace();
+							}
+							return;
+						}
+					case EPropertyChangeType::ArrayAdd:
+						{
+							ArrayAddImpl();
+							return;
+						}
+					case EPropertyChangeType::ArrayRemove:
+						{
+							ArrayRemoveImpl();
+							return;
+						}
+					case EPropertyChangeType::ArrayClear:
+						{
+							checkf(ArrayIndex == INDEX_NONE, TEXT("ArrayClear change type should not have associated indexes"));
+
+							for (int i = 0; i < PreEditSubObjectsArrayHelper.Num(); ++i)
+							{
+								ArrayIndex = i;
+								ArrayRemoveImpl();
+							}
+							return;
+						}
+					case EPropertyChangeType::ArrayMove:
+						{
+							UE_LOG(LogOverridableObject, Warning, TEXT("ArrayMove change type is not going to change anything as ordering of object isn't supported yet"));
+							return;
+						}
+					default:
+						{
+							UE_LOG(LogOverridableObject, Warning, TEXT("Property change type is not supported will default to full array override"));
+							break;
+						}
+					}
+				}
+				else if (Notification == EPropertyNotificationType::PostEdit)
+				{
+					checkf(ArrayHelper.IsValidIndex(ArrayIndex), TEXT("Any sub operation is expected to have a valid index"));
+					if (UObject* SubObject = InnerObjectProperty->GetObjectPropertyValue(ArrayHelper.GetElementPtr(ArrayIndex)))
+					{
+						// This should not be needed in the property grid, as it should already been called on the subobject itself.
+						OverridableManager.NotifyPropertyChange(Notification, *SubObject, PropertyEvent, PropertyNode->GetNextNode());
+						return;
+					}
+				}
+			}
+		}
+	}
+	// @todo support set in the overridable serialization
+	//else if (const FSetProperty* SetProperty = CastField<FSetProperty>(Property))
+	//{
+	//	
+	//}
+	else if (const FMapProperty* MapProperty = CastField<FMapProperty>(Property))
+	{
+		// Special handling of instanced subobjects
+		checkf(MapProperty->KeyProp, TEXT("Expecting an key type for Maps"));
+		FObjectProperty* KeyObjectProperty = CastField<FObjectProperty>(MapProperty->KeyProp);
+
+		// SubObjects
+		checkf(!KeyObjectProperty || !MapProperty->KeyProp->HasAnyPropertyFlags(CPF_PersistentInstance) || CastField<FClassProperty>(MapProperty->KeyProp), TEXT("Keys as a instanced subobject is not supported yet"));
+
+		checkf(MapProperty->ValueProp, TEXT("Expecting an value type for Maps"));
+		FObjectProperty* ValueInstancedObjectProperty = MapProperty->ValueProp->HasAnyPropertyFlags(CPF_PersistentInstance) ? CastField<FObjectProperty>(MapProperty->ValueProp) : nullptr;
+
+		FScriptMapHelper MapHelper(MapProperty, SubValuePtr);
+		int32 LogicalMapIndex = PropertyEvent.GetArrayIndex(Property->GetName());
+		int32 InternalMapIndex = LogicalMapIndex != INDEX_NONE ? MapHelper.FindInternalIndex(LogicalMapIndex) : INDEX_NONE;
+		if (!PropertyNode->GetNextNode())
+		{
+			static const FProperty* SavedProp = nullptr;
+			static uint8* SavedPreEditMap = nullptr;
+
+			auto FreePreEditMap = []()
+			{
+				if (SavedPreEditMap)
+				{
+					checkf(SavedProp, TEXT("Expecting a matching property to the allocated memory"));
+					SavedProp->DestroyValue(SavedPreEditMap);
+					FMemory::Free(SavedPreEditMap);
+					SavedPreEditMap = nullptr;
+					SavedProp = nullptr;
+				}
+			};
+
+			if (Notification == EPropertyNotificationType::PreEdit)
+			{
+				FreePreEditMap();
+
+				SavedPreEditMap = (uint8*)FMemory::Malloc(MapProperty->GetSize(), MapProperty->GetMinAlignment());
+				MapProperty->InitializeValue(SavedPreEditMap);
+				SavedProp = MapProperty;
+
+				FScriptMapHelper PreEditMapHelper(MapProperty, SavedPreEditMap);
+				PreEditMapHelper.EmptyValues();
+				for (FScriptMapHelper::FIterator It(MapHelper); It; ++It)
+				{
+					PreEditMapHelper.AddPair(MapHelper.GetKeyPtr(It.GetInternalIndex()), MapHelper.GetValuePtr(It.GetInternalIndex()));
+				}
+				return;
+			}
+
+			checkf(SavedProp == MapProperty, TEXT("Expecting the same property as the pre edit flow"));
+			FScriptMapHelper PreEditMapHelper(MapProperty, SavedPreEditMap);
+			// The logical should map directly to the pre edit map internal index as we skipped all of the invalid entries
+			int32 InternalPreEditMapIndex = LogicalMapIndex;
+
+			ON_SCOPE_EXIT
+			{
+				FreePreEditMap();
+			};
+
+			auto MapReplace = [&]()
+			{
+				// Overriding a all entries in the map
+				SubPropertyNode.Operation = EOverriddenPropertyOperation::Replace;
+
+				// This is a case of the entire array is overridden
+				// Need to loop through every sub object and setup them up as overridden
+				for (FScriptMapHelper::FIterator It(MapHelper); It; ++It)
+				{
+					FName OverriddenKeyID = NAME_None;
+					if (const UObject* OverriddenKeyObject = KeyObjectProperty ? KeyObjectProperty->GetObjectPropertyValue(MapHelper.GetKeyPtr(It.GetInternalIndex())) : nullptr)
+					{
+						OverriddenKeyID = FName(OverriddenKeyObject->GetPathName());;
+					}
+					else
+					{
+						FString AddedKey;
+						MapProperty->KeyProp->ExportTextItem_Direct(AddedKey, MapHelper.GetKeyPtr(It.GetInternalIndex()), nullptr, nullptr, PPF_None);
+						OverriddenKeyID = FName(AddedKey);
+					}
+					FOverriddenPropertyNode& OverriddenKeyNode = FindOrAddNode(SubPropertyNode, OverriddenKeyID);
+					OverriddenKeyNode.Operation = EOverriddenPropertyOperation::Replace;
+
+					// @todo support instanced object as a key in maps
+					//if (UObject* KeySubObject = KeyInstancedObjectProperty ? KeyInstancedObjectProperty->GetObjectPropertyValue(MapHelper.GetKeyPtr(*It)) : nullptr)
+					//{
+					//	OverridableManager.OverrideInstancedSubObject(*Owner, *KeySubObject);
+					//}
+					if (UObject* ValueSubObject = ValueInstancedObjectProperty ? ValueInstancedObjectProperty->GetObjectPropertyValue(MapHelper.GetValuePtr(It.GetInternalIndex())) : nullptr)
+					{
+						OverridableManager.OverrideInstancedSubObject(*Owner, *ValueSubObject);
+					}
+				}
+			};
+
+			auto MapAddImpl = [&]()
+			{
+				checkf(MapHelper.IsValidIndex(InternalMapIndex), TEXT("ArrayAdd change type expected to have an valid index"));
+				FName AddedKeyID = NAME_None;
+				if (const UObject* AddedKeyObject = KeyObjectProperty ? KeyObjectProperty->GetObjectPropertyValue(MapHelper.GetKeyPtr(InternalMapIndex)) : nullptr)
+				{
+					AddedKeyID = FName(AddedKeyObject->GetPathName());
+				}
+				else
+				{
+					FString AddedKey;
+					MapProperty->KeyProp->ExportTextItem_Direct(AddedKey, MapHelper.GetKeyPtr(InternalMapIndex), nullptr, nullptr, PPF_None);
+					AddedKeyID = FName(AddedKey);
+				}
+
+				FOverriddenPropertyNode& AddedKeyNode = FindOrAddNode(SubPropertyNode, AddedKeyID);
+				// @todo support instanced object as a key in maps
+				//if (UObject* AddedKeySubObject = KeyInstancedObjectProperty ? KeyInstancedObjectProperty->GetObjectPropertyValue(MapHelper.GetKeyPtr(InternalMapIndex)) : nullptr)
+				//{
+				//	// This is a readd, make sure all the subobjects values are overridden
+				//	if (AddedKeyNode.Operation == EOverriddenPropertyOperation::Remove)
+				//	{
+				//		OverridableManager.OverrideInstancedSubObject(*Owner, *AddedKeySubObject);
+				//	}
+				//}
+				if (UObject* AddedValueSubObject = ValueInstancedObjectProperty ? ValueInstancedObjectProperty->GetObjectPropertyValue(MapHelper.GetValuePtr(InternalMapIndex)) : nullptr)
+				{
+					// This is a readd, make sure all the subobjects values are overridden
+					if (AddedKeyNode.Operation == EOverriddenPropertyOperation::Remove)
+					{
+						OverridableManager.OverrideInstancedSubObject(*Owner, *AddedValueSubObject);
+					}
+				}
+				AddedKeyNode.Operation = EOverriddenPropertyOperation::Add;
+			};
+
+			auto MapRemoveImpl = [&]()
+			{
+				checkf(PreEditMapHelper.IsValidIndex(InternalPreEditMapIndex), TEXT("ArrayRemove change type expected to have an valid index"));
+				FName RemovedKeyID = NAME_None;
+				if (const UObject* RemovedKeyObject = KeyObjectProperty ? KeyObjectProperty->GetObjectPropertyValue(PreEditMapHelper.GetKeyPtr(InternalPreEditMapIndex)) : nullptr)
+				{
+					RemovedKeyID = FName(RemovedKeyObject->GetPathName());;
+				}
+				else
+				{
+					FString RemovedKey;
+					MapProperty->KeyProp->ExportTextItem_Direct(RemovedKey, PreEditMapHelper.GetKeyPtr(InternalPreEditMapIndex), nullptr, nullptr, PPF_None);
+					RemovedKeyID = FName(RemovedKey);
+				}
+				FOverriddenPropertyNode& RemovedKeyNode = FindOrAddNode(SubPropertyNode, RemovedKeyID);
+				if (RemovedKeyNode.Operation == EOverriddenPropertyOperation::Add)
+				{
+					// @Todo support remove/add/remove
+					SubPropertyNode.SubPropertyNodeKeys.Remove(RemovedKeyID);
+				}
+				else
+				{
+					RemovedKeyNode.Operation = EOverriddenPropertyOperation::Remove;
+				}
+
+				// @todo support instanced object as a key in maps
+				//if (UObject* RemovedKeySubObject = KeyInstancedObjectProperty ? KeyInstancedObjectProperty->GetObjectPropertyValue(PreEditMapHelper.GetKeyPtr(InternalPreEditMapIndex)) : nullptr)
+				//{
+				//	OverridableManager.ClearInstancedSubObjectOverrides(*Owner, *RemovedKeySubObject);
+				//}
+				if (UObject* RemovedValueSubObject = ValueInstancedObjectProperty ? ValueInstancedObjectProperty->GetObjectPropertyValue(PreEditMapHelper.GetValuePtr(InternalPreEditMapIndex)) : nullptr)
+				{
+					OverridableManager.ClearInstancedSubObjectOverrides(*Owner, *RemovedValueSubObject);
+				}
+			};
+
+			// Only maps flagged overridable logic can be handled here
+			if (!MapProperty->HasAnyPropertyFlags(CPF_ExperimentalOverridableLogic))
+			{
+				if (PropertyEvent.ChangeType == EPropertyChangeType::Unspecified && InternalMapIndex == INDEX_NONE)
+				{
+					// Overriding all entry in the array + override instanced sub obejects
+					MapReplace();
+				}
+				else
+				{
+					// Overriding all entry in the array
+					SubPropertyNode.Operation = EOverriddenPropertyOperation::Replace;
+				}
+				return;
+			}
+
+			switch (PropertyEvent.ChangeType)
+			{
+			case EPropertyChangeType::Unspecified:
+				{
+					if(LogicalMapIndex != INDEX_NONE)
+					{
+						// Overriding a single entry in the map
+						MapRemoveImpl();
+						MapAddImpl();
+					}
+					else
+					{
+						MapReplace();
+					}
+					return;
+				}
+			case EPropertyChangeType::ArrayAdd:
+				{
+					MapAddImpl();
+					return;
+				}
+			case EPropertyChangeType::ArrayRemove:
+				{
+					MapRemoveImpl();
+					return;
+				}
+			case EPropertyChangeType::ArrayClear:
+				{
+					checkf(InternalPreEditMapIndex == INDEX_NONE, TEXT("ArrayClear change type should not have associated indexes"));
+
+					for (FScriptMapHelper::FIterator It(PreEditMapHelper); It; ++It)
+					{
+						InternalPreEditMapIndex = It.GetInternalIndex();
+						MapRemoveImpl();
+					}
+					return;
+				}
+			case EPropertyChangeType::ArrayMove:
+				{
+					UE_LOG(LogOverridableObject, Warning, TEXT("ArrayMove change type is not going to change anything as ordering of object isn't supported yet"));
+					return;
+				}
+			default:
+				{
+					UE_LOG(LogOverridableObject, Warning, TEXT("Property change type is not supported will default to full array override"));
+					break;
+				}
+			}
+		}
+		else if (Notification == EPropertyNotificationType::PostEdit)
+		{
+			checkf(MapHelper.IsValidIndex(InternalMapIndex), TEXT("Any sub operation is expected to not have a valid index"));
+
+			// @todo support instanced object as a key in maps
+			//if (UObject* SubObject = KeyInstancedObjectProperty ? KeyInstancedObjectProperty->GetObjectPropertyValue(MapHelper.GetValuePtr(InternalMapIndex)) : nullptr)
+			//{
+			//	// This should not be needed in the property grid, as it should already been called on the subobject.
+			//	OverridableManager.NotifyPropertyChange(Notification, *SubObject, PropertyEvent, PropertyNode->GetNextNode());
+			//	return;
+			//}
+
+			if (UObject* SubObject = ValueInstancedObjectProperty ? ValueInstancedObjectProperty->GetObjectPropertyValue(MapHelper.GetValuePtr(InternalMapIndex)) : nullptr)
+			{
+				// This should not be needed in the property grid, as it should already been called on the subobject.
+				OverridableManager.NotifyPropertyChange(Notification, *SubObject, PropertyEvent, PropertyNode->GetNextNode());
+				return;
+			}
+		}
+	}
+	else if (Property->IsA<FStructProperty>())
+	{
+		NotifyPropertyChange(SubPropertyNode, Notification, PropertyEvent, PropertyNode->GetNextNode(), SubValuePtr);
+		return;
+	}
+	else if (const FObjectProperty* ObjectProperty = CastField<FObjectProperty>(Property))
+	{
+		if(Notification == EPropertyNotificationType::PostEdit && !PropertyNode->GetNextNode())
+		{
+			SubPropertyNode.Operation = EOverriddenPropertyOperation::Replace;
+		}
+
+		// @TODO Need to handle non instanced object here.
+		if (UObject* SubObject = ObjectProperty->GetObjectPropertyValue(SubValuePtr))
+		{
+			// This should not be needed in the property grid, as it should already been called on the subobject.
+			OverridableManager.NotifyPropertyChange(Notification, *SubObject, PropertyEvent, PropertyNode->GetNextNode());
+			return;
+		}
+	}
+	else if (const FOptionalProperty* OptionalProperty = CastField<FOptionalProperty>(Property))
+	{
+		if (OptionalProperty->IsSet(Data))
+		{
+			NotifyPropertyChange(SubPropertyNode, Notification, PropertyEvent, PropertyNode->GetNextNode(), OptionalProperty->GetValuePointerForRead(SubValuePtr));
+			return;
+		}
+	}
+
+	UE_CLOG(PropertyNode->GetNextNode(), LogOverridableObject, Warning, TEXT("Unsupported property type(%s), fallback to overriding entire property"), *Property->GetName());
+	if (Notification == EPropertyNotificationType::PostEdit)
+	{
+		NotifyPropertyChange(SubPropertyNode, Notification, PropertyEvent, nullptr, SubValuePtr);
+	}
+}
+
+void FOverriddenPropertySet::RemoveOverriddenSubProperties(FOverriddenPropertyNode& PropertyNode)
+{
+	for (const auto& Pair : PropertyNode.SubPropertyNodeKeys)
+	{
+		FOverriddenPropertyNode* RemovedPropertyNode = OverriddenPropertyNodes.FindByHash(GetTypeHash(Pair.Value), Pair.Value);
+		checkf(RemovedPropertyNode, TEXT("Expecting a node"));
+		RemoveOverriddenSubProperties(*RemovedPropertyNode);
+		OverriddenPropertyNodes.RemoveByHash(GetTypeHash(Pair.Value), Pair.Value);
+	}
+	PropertyNode.SubPropertyNodeKeys.Empty();
+}
+
+void FOverriddenPropertySet::NotifyPropertyChange(const EPropertyNotificationType Notification, const FPropertyChangedEvent& PropertyEvent, const FEditPropertyChain::TDoubleLinkedListNode* PropertyNode, const void* Data)
+{
+	NotifyPropertyChange(OverriddenPropertyNodes.FindOrAddByHash(GetTypeHash(RootNodeName), RootNodeName), Notification, PropertyEvent, PropertyNode, Data);
+}
+
+EOverriddenPropertyOperation FOverriddenPropertySet::GetOverriddenPropertyOperation(const FArchiveSerializedPropertyChain* CurrentPropertyChain, FProperty* Property) const
+{
+	if (const FOverriddenPropertyNode* RootNode = OverriddenPropertyNodes.FindByHash(GetTypeHash(RootNodeName), RootNodeName))
+	{
+		return GetOverriddenPropertyOperation(*RootNode, CurrentPropertyChain, Property);
+	}
+	return EOverriddenPropertyOperation::None;
+}
+
+FOverriddenPropertyNode* FOverriddenPropertySet::SetOverriddenPropertyOperation(EOverriddenPropertyOperation Operation, const FArchiveSerializedPropertyChain* CurrentPropertyChain, FProperty* Property)
+{
+	return SetOverriddenPropertyOperation(Operation, OverriddenPropertyNodes.FindOrAddByHash(GetTypeHash(RootNodeName), RootNodeName), CurrentPropertyChain, Property);
+}
+
+const FOverriddenPropertyNode* FOverriddenPropertySet::GetOverriddenPropertyNode(const FArchiveSerializedPropertyChain* CurrentPropertyChain) const
+{
+	if (const FOverriddenPropertyNode* RootNode = OverriddenPropertyNodes.FindByHash(GetTypeHash(RootNodeName), RootNodeName))
+	{
+		return GetOverriddenPropertyNode(*RootNode, CurrentPropertyChain);
+	}
+	return nullptr;
+}
+
+EOverriddenPropertyOperation FOverriddenPropertySet::GetOverriddenPropertyOperation(const FOverriddenPropertyNode& ParentPropertyNode, const FArchiveSerializedPropertyChain* CurrentPropertyChain, FProperty* Property) const
+{
+	// No need to look further
+	// if it is the entire property is replaced or
+	// if it is the FOverriddenPropertySet struct which is always Overridden
+	if (ParentPropertyNode.Operation == EOverriddenPropertyOperation::Replace)
+	{
+		return EOverriddenPropertyOperation::Replace;
+	}
+
+	// @Todo optimize find a way to not have to copy the property chain here.
+	FArchiveSerializedPropertyChain PropertyChain(CurrentPropertyChain ? *CurrentPropertyChain : FArchiveSerializedPropertyChain());
+	if(Property)
+	{
+		PropertyChain.PushProperty(Property, Property->IsEditorOnlyProperty());
+	}
+
+	TArray<class FProperty*, TInlineAllocator<8>>::TConstIterator PropertyIterator = PropertyChain.GetRootIterator();
+	const FOverriddenPropertyNode* OverriddenPropertyNode = &ParentPropertyNode;
+	while (PropertyIterator && OverriddenPropertyNode && OverriddenPropertyNode->Operation != EOverriddenPropertyOperation::Replace)
+	{
+		const FProperty* CurrentProperty = (*PropertyIterator);
+		const FName CurrentPropID = CurrentProperty->GetFName();
+		if (const FName* CurrentPropKey = OverriddenPropertyNode->SubPropertyNodeKeys.Find(CurrentPropID))
+		{
+			OverriddenPropertyNode = OverriddenPropertyNodes.FindByHash(GetTypeHash(*CurrentPropKey), *CurrentPropKey);
+			checkf(OverriddenPropertyNode, TEXT("Expecting a node"));
+		}
+		else
+		{
+			OverriddenPropertyNode = nullptr;
+			break;
+		}
+		++PropertyIterator;
+	}
+
+	return OverriddenPropertyNode ? OverriddenPropertyNode->Operation : EOverriddenPropertyOperation::None;
+}
+
+FOverriddenPropertyNode* FOverriddenPropertySet::SetOverriddenPropertyOperation(EOverriddenPropertyOperation Operation, FOverriddenPropertyNode& ParentPropertyNode, const FArchiveSerializedPropertyChain* CurrentPropertyChain, FProperty* Property)
+{
+	// No need to look further
+	// if it is the entire property is replaced or
+	// if it is the FOverriddenPropertySet struct which is always Overridden
+	if (ParentPropertyNode.Operation == EOverriddenPropertyOperation::Replace)
+	{
+		return nullptr;
+	}
+
+	// @Todo optimize find a way to not have to copy the property chain here.
+	FArchiveSerializedPropertyChain PropertyChain(CurrentPropertyChain ? *CurrentPropertyChain : FArchiveSerializedPropertyChain());
+	if (Property)
+	{
+		PropertyChain.PushProperty(Property, Property->IsEditorOnlyProperty());
+	}
+
+	TArray<class FProperty*, TInlineAllocator<8>>::TConstIterator PropertyIterator = PropertyChain.GetRootIterator();
+	FOverriddenPropertyNode* OverriddenPropertyNode = &ParentPropertyNode;
+	while (PropertyIterator && OverriddenPropertyNode->Operation != EOverriddenPropertyOperation::Replace)
+	{
+		const FProperty* CurrentProperty = (*PropertyIterator);
+		const FName CurrentPropID = CurrentProperty->GetFName();
+		OverriddenPropertyNode = &FindOrAddNode(*OverriddenPropertyNode, CurrentPropID);
+		++PropertyIterator;
+	}
+
+	// Might have stop before as one of the parent property was completely replaced.
+	if (!PropertyIterator)
+	{
+		OverriddenPropertyNode->Operation = Operation;
+		return OverriddenPropertyNode;
+	}
+
+	return nullptr;
+}
+
+EOverriddenPropertyOperation FOverriddenPropertySet::GetSubPropertyOperation(FName PropertyKey) const
+{
+	const FOverriddenPropertyNode* OverriddenPropertyNode = OverriddenPropertyNodes.FindByHash(GetTypeHash(PropertyKey), PropertyKey);
+	return OverriddenPropertyNode ? OverriddenPropertyNode->Operation : EOverriddenPropertyOperation::None;
+}
+
+FOverriddenPropertyNode* FOverriddenPropertySet::SetSubPropertyOperation(EOverriddenPropertyOperation Operation, FOverriddenPropertyNode& Node, FName PropID)
+{
+	FOverriddenPropertyNode& OverriddenPropertyNode = FindOrAddNode(Node, PropID);
+	OverriddenPropertyNode.Operation = Operation;
+	return &OverriddenPropertyNode;
+}
+
+void FOverriddenPropertySet::Reset()
+{
+	OverriddenPropertyNodes.Reset();
+}
+
+const FOverriddenPropertyNode* FOverriddenPropertySet::GetOverriddenPropertyNode(const FOverriddenPropertyNode& ParentPropertyNode, const FArchiveSerializedPropertyChain* CurrentPropertyChain) const
+{
+	if (!CurrentPropertyChain)
+	{
+		return &ParentPropertyNode;
+	}
+
+	TArray<class FProperty*, TInlineAllocator<8>>::TConstIterator PropertyIterator = CurrentPropertyChain->GetRootIterator();
+	const FOverriddenPropertyNode* OverriddenPropertyNode = &ParentPropertyNode;
+	while (PropertyIterator && OverriddenPropertyNode)
+	{
+		const FProperty* CurrentProperty = (*PropertyIterator);
+		const FName CurrentPropID = CurrentProperty->GetFName();
+		if (const FName* CurrentPropKey = OverriddenPropertyNode->SubPropertyNodeKeys.Find(CurrentPropID))
+		{
+			OverriddenPropertyNode = OverriddenPropertyNodes.FindByHash(GetTypeHash(*CurrentPropKey), *CurrentPropKey);
+			checkf(OverriddenPropertyNode, TEXT("Expecting a node"));
+		}
+		else
+		{
+			OverriddenPropertyNode = nullptr;
+			break;
+		}
+		++PropertyIterator;
+	}
+
+	return OverriddenPropertyNode;
+}
diff --git a/Engine/Source/Runtime/CoreUObject/Private/UObject/PropertyArray.cpp b/Engine/Source/Runtime/CoreUObject/Private/UObject/PropertyArray.cpp
--- a/Engine/Source/Runtime/CoreUObject/Private/UObject/PropertyArray.cpp
+++ b/Engine/Source/Runtime/CoreUObject/Private/UObject/PropertyArray.cpp
@@ -11,6 +11,7 @@
 #include "UObject/LinkerLoad.h"
 #include "UObject/PropertyHelper.h"
 #include "UObject/UObjectThreadContext.h"
+#include "UObject/OverriddenPropertySet.h"
 
 /*-----------------------------------------------------------------------------
 	FArrayProperty.
@@ -123,7 +124,8 @@ void FArrayProperty::SerializeItem(FStructuredArchive::FSlot Slot, void* Value,
 	check(Inner);
 	FArchive& UnderlyingArchive = Slot.GetUnderlyingArchive();
 	const bool bIsTextFormat = UnderlyingArchive.IsTextFormat();
-	const bool bUPS = Slot.GetArchiveState().UseUnversionedPropertySerialization();
+	const bool bUPS = UnderlyingArchive.UseUnversionedPropertySerialization();
+	bool bExperimentalOverridableLogic = HasAnyPropertyFlags(CPF_ExperimentalOverridableLogic);
 	TOptional<FPropertyTag> MaybeInnerTag;
 
 	// Ensure that the Inner itself has been loaded before calling SerializeItem() on it
@@ -133,7 +135,7 @@ void FArrayProperty::SerializeItem(FStructuredArchive::FSlot Slot, void* Value,
 	int32		n		= ArrayHelper.Num();
 
 	// Custom branch for UPS to try and take advantage of bulk serialization
-	if (bUPS)
+	if (bUPS && !bExperimentalOverridableLogic)
 	{
 		checkf(!UnderlyingArchive.ArUseCustomPropertyList, TEXT("Custom property lists are not supported with UPS"));
 		checkf(!bIsTextFormat, TEXT("Text-based archives are not supported with UPS"));
@@ -185,31 +187,6 @@ void FArrayProperty::SerializeItem(FStructuredArchive::FSlot Slot, void* Value,
 		Slot << SA_OPTIONAL_ATTRIBUTE(TEXT("InnerStructGuid"), MaybeInnerTag.GetValue().StructGuid, FGuid());
 	}
 
-	FStructuredArchiveArray Array = Slot.EnterArray(n);
-
-	if( UnderlyingArchive.IsLoading() )
-	{
-		// If using a custom property list, don't empty the array on load. Not all indices may have been serialized, so we need to preserve existing values at those slots.
-		if (UnderlyingArchive.ArUseCustomPropertyList)
-		{
-			const int32 OldNum = ArrayHelper.Num();
-			if (n > OldNum)
-			{
-				ArrayHelper.AddValues(n - OldNum);
-			}
-			else if (n < OldNum)
-			{
-				ArrayHelper.RemoveValues(n, OldNum - n);
-			}
-		}
-		else
-		{
-			ArrayHelper.EmptyAndAddValues(n);
-		}
-	}
-	ArrayHelper.CountBytes( UnderlyingArchive );
-
-
 	TOptional<FPropertyTag> SerializeFromMismatchedTag;
 
 	// TODO: Should work for maps + sets too.
@@ -237,6 +214,335 @@ void FArrayProperty::SerializeItem(FStructuredArchive::FSlot Slot, void* Value,
 		Inner->SerializeItem(Slot, Item);
 	};
 	
+	// Make sure the container is reloading accordingly to the value set in the property tag if any
+	if (!bUPS && UnderlyingArchive.IsLoading() && FPropertyTagScope::GetCurrentPropertyTag())
+	{
+		bExperimentalOverridableLogic = FPropertyTagScope::GetCurrentPropertyTag()->bExperimentalOverridableLogic;
+	}
+
+	// *** Experimental *** Special serialization path for array with overridable serialization
+	if(bExperimentalOverridableLogic)
+	{
+		checkf(!UnderlyingArchive.ArUseCustomPropertyList, TEXT("Using custom property list is not supported by overridable serialization"));
+
+		FStructuredArchive::FRecord Record = Slot.EnterRecord();
+		if (UnderlyingArchive.IsLoading())
+		{
+			int32 NumReplaced = 0;
+			FStructuredArchive::FArray ReplacedArray = Record.EnterArray(TEXT("Replaced"), NumReplaced);
+			if (NumReplaced != INDEX_NONE)
+			{
+				ArrayHelper.EmptyAndAddValues(NumReplaced);
+				for (int32 i = 0; i < NumReplaced; i++)
+				{
+					SerializeContainerItem(ReplacedArray.EnterElement(), ArrayHelper.GetRawPtr(i));
+				}
+			}
+			else
+			{
+				// Only Array of Instanced subobject are handled here as sort of a set where the matching key is done using the archetype
+				const FObjectProperty* InnerObjectProperty = CastFieldChecked<FObjectProperty>(Inner);
+				checkf(InnerObjectProperty->HasAnyPropertyFlags(CPF_PersistentInstance), TEXT("Only supported code path here is the instanced subobjects"));
+
+				FOverriddenPropertySet* OverriddenProperties = FOverridableSerializationLogic::GetOverriddenProperties();
+
+				checkf(Defaults, TEXT("Expecting overridable serialization to have defaults to compare to"));
+
+				FScriptArrayHelper DefaultsArrayHelper(this, Defaults);
+
+				auto FindObject = [InnerObjectProperty](UObject* Object, UObject* Object2, FScriptArrayHelper& ArrayHelper) -> int32
+				{
+					if (Object)
+					{
+						const int32 ArrayNum = ArrayHelper.Num();
+						for (int i = 0; i < ArrayNum; ++i)
+						{
+							UObject* CurrentObject = InnerObjectProperty->GetObjectPropertyValue(ArrayHelper.GetElementPtr(i));
+							if (CurrentObject == Object || (Object2 && CurrentObject == Object2))
+							{
+								return i;
+							}
+						}
+					}
+					return INDEX_NONE;
+				};
+
+				uint8* TempValueStorage = nullptr;
+				ON_SCOPE_EXIT
+				{
+					if (TempValueStorage)
+					{
+						InnerObjectProperty->DestroyValue(TempValueStorage);
+						FMemory::Free(TempValueStorage);
+					}
+				};
+
+				int32 NumRemoved = 0;
+				FStructuredArchive::FArray RemovedArray = Record.EnterArray(TEXT("Removed"), NumRemoved);
+				if(NumRemoved != 0)
+				{
+					TArray<int32> IndicesToRemove;
+					TempValueStorage = (uint8*)FMemory::Malloc(InnerObjectProperty->ElementSize);
+					InnerObjectProperty->InitializeValue(TempValueStorage);
+
+					for (int32 i = 0; i < NumRemoved; ++i)
+					{
+						{
+							FSerializedPropertyScope SerializedProperty(UnderlyingArchive, Inner, this);
+							SerializeContainerItem(RemovedArray.EnterElement(), TempValueStorage);
+						}
+
+						if (UObject* RemovedSubObject = InnerObjectProperty->GetObjectPropertyValue(TempValueStorage) )
+						{
+							int32 Index = FindObject(RemovedSubObject, nullptr, DefaultsArrayHelper);
+							if (Index != INDEX_NONE)
+							{
+								IndicesToRemove.Add(Index);
+							}
+
+							// Need to fetch the ArrayOverriddenPropertyNode every loop as the previous iteration might have reallocated the node.
+							if (FOverriddenPropertyNode* ArrayOverriddenPropertyNode = OverriddenProperties ? OverriddenProperties->SetOverriddenPropertyOperation(EOverriddenPropertyOperation::Modified, UnderlyingArchive.GetSerializedPropertyChain(), /*Property*/nullptr) : nullptr)
+							{
+								// Rebuild the overridden info
+								const FName RemovedSubObjectID = RemovedSubObject->GetFName();
+								OverriddenProperties->SetSubPropertyOperation(EOverriddenPropertyOperation::Remove, *ArrayOverriddenPropertyNode, RemovedSubObjectID);
+							}
+						}
+					}
+
+					IndicesToRemove.Sort(TGreater<>());
+					for(int32 IndexToRemove : IndicesToRemove)
+					{
+						ArrayHelper.RemoveValues(IndexToRemove);
+					}
+				}
+
+				int32 NumModified = 0;
+				FStructuredArchive::FArray ModifiedArray = Record.EnterArray(TEXT("Modified"), NumModified);
+				if (NumModified != 0)
+				{
+					if (!TempValueStorage)
+					{
+						TempValueStorage = (uint8*)FMemory::Malloc(InnerObjectProperty->ElementSize);
+						InnerObjectProperty->InitializeValue(TempValueStorage);
+					}
+
+					for (int32 i = 0; i < NumModified; ++i)
+					{
+						{
+							FSerializedPropertyScope SerializedProperty(UnderlyingArchive, Inner, this);
+							SerializeContainerItem(ModifiedArray.EnterElement(), TempValueStorage);
+						}
+
+						if (UObject* ModifiedObject = InnerObjectProperty->GetObjectPropertyValue(TempValueStorage))
+						{
+							int32 Index = FindObject(ModifiedObject->GetArchetype(), ModifiedObject, ArrayHelper);
+							if (Index != INDEX_NONE)
+							{
+								InnerObjectProperty->SetObjectPropertyValue(ArrayHelper.GetRawPtr(Index), ModifiedObject);
+							}
+						}
+					}
+				}
+
+				int32 NumAdded = 0;
+				FStructuredArchive::FArray AddedArray = Record.EnterArray(TEXT("Added"), NumAdded);
+				if (NumAdded != 0)
+				{
+					if (!TempValueStorage)
+					{
+						TempValueStorage = (uint8*)FMemory::Malloc(InnerObjectProperty->ElementSize);
+						InnerObjectProperty->InitializeValue(TempValueStorage);
+					}
+
+					int32 AddIndex = ArrayHelper.Num();
+					for (int32 i = 0; i < NumAdded; ++i)
+					{
+						{
+							FSerializedPropertyScope SerializedProperty(UnderlyingArchive, Inner, this);
+							SerializeContainerItem(AddedArray.EnterElement(), TempValueStorage);
+						}
+
+						if (UObject* AddedSubObject = InnerObjectProperty->GetObjectPropertyValue(TempValueStorage))
+						{
+							int32 Index = FindObject(AddedSubObject->GetArchetype(), AddedSubObject, ArrayHelper);
+							if (Index == INDEX_NONE)
+							{
+								ArrayHelper.AddValue();
+								Index = AddIndex++;
+							}
+							InnerObjectProperty->SetObjectPropertyValue(ArrayHelper.GetRawPtr(Index), AddedSubObject);
+
+							// Need to fetch the ArrayOverriddenPropertyNode every loop as the previous iteration might have reallocated the node.
+							if (FOverriddenPropertyNode* ArrayOverriddenPropertyNode = OverriddenProperties ? OverriddenProperties->SetOverriddenPropertyOperation(EOverriddenPropertyOperation::Modified, UnderlyingArchive.GetSerializedPropertyChain(), /*Property*/nullptr) : nullptr)
+							{
+								// Rebuild the overridden info
+								const FName AddedSubObjectID = AddedSubObject->GetFName();
+								OverriddenProperties->SetSubPropertyOperation(EOverriddenPropertyOperation::Add, *ArrayOverriddenPropertyNode, AddedSubObjectID);
+							}
+						}
+					}
+				}
+			}
+		}
+		else
+		{
+			// Container for temporarily tracking some indices
+			TArray<int32> RemovedIndices;
+			TArray<int32> ModifiedIndices;
+			TArray<int32> AddedIndices;
+
+			bool bReplaceArray = false;
+			if (!Defaults)
+			{
+				bReplaceArray = true;
+			}
+			else
+			{
+				const FObjectProperty* InnerObjectProperty = CastField<FObjectProperty>(Inner);
+				EOverriddenPropertyOperation ArrayOverrideOp = EOverriddenPropertyOperation::None;
+				FOverriddenPropertySet* OverriddenProperties = FOverridableSerializationLogic::GetOverriddenProperties();
+				if (OverriddenProperties)
+				{
+					ArrayOverrideOp = OverriddenProperties->GetOverriddenPropertyOperation(UnderlyingArchive.GetSerializedPropertyChain(), /*Property*/nullptr);
+					bReplaceArray = ArrayOverrideOp == EOverriddenPropertyOperation::Replace;
+				}
+				else
+				{
+					bReplaceArray = !InnerObjectProperty || !InnerObjectProperty->HasAnyPropertyFlags(CPF_PersistentInstance);
+				}
+
+				if (!bReplaceArray)
+				{
+					// Only array of instanced subobjects are handled here as sort of a set where the matching key is done using the archetype
+					checkf(InnerObjectProperty&& InnerObjectProperty->HasAnyPropertyFlags(CPF_PersistentInstance), TEXT("Expecting only arrays of instanced subobjects"));
+
+					// We need to always serialize instanced subobjects to know if they have overridden values.
+					const int32 ArrayNum = ArrayHelper.Num();
+					for (int i = 0; i < ArrayNum; i++)
+					{
+						ModifiedIndices.Add(i);
+					}
+
+					if (OverriddenProperties && ArrayOverrideOp != EOverriddenPropertyOperation::None)
+					{
+
+						auto FindObject = [InnerObjectProperty](const FName ObjectToFind, FScriptArrayHelper& ArrayHelper) -> int32
+						{
+							const int32 ArrayNum = ArrayHelper.Num();
+							for (int i = 0; i < ArrayNum; ++i)
+							{
+								UObject* CurrentObject = InnerObjectProperty->GetObjectPropertyValue(ArrayHelper.GetElementPtr(i));
+								if (CurrentObject && CurrentObject->GetName() == ObjectToFind)
+								{
+									return i;
+								}
+							}
+							return INDEX_NONE;
+						};
+
+						FScriptArrayHelper DefaultsArrayHelper(this, Defaults);
+
+						if (const FOverriddenPropertyNode* ArrayOverriddenPropertyNode = OverriddenProperties->GetOverriddenPropertyNode(UnderlyingArchive.GetSerializedPropertyChain()))
+						{
+							for (const auto& Pair : ArrayOverriddenPropertyNode->SubPropertyNodeKeys)
+							{
+								const EOverriddenPropertyOperation OverrideOp = OverriddenProperties->GetSubPropertyOperation(Pair.Value);
+								switch (OverrideOp)
+								{
+								case EOverriddenPropertyOperation::Remove:
+									{
+										const int32 DefaultIndex = FindObject(Pair.Key, DefaultsArrayHelper);
+										if (DefaultIndex != INDEX_NONE)
+										{
+											RemovedIndices.Add(DefaultIndex);
+										}
+										break;
+									}
+								case EOverriddenPropertyOperation::Add:
+									{
+										const int32 Index = FindObject(Pair.Key, ArrayHelper);
+										if(Index != INDEX_NONE)
+										{
+											AddedIndices.Add(Index);
+											ModifiedIndices.Remove(Index);
+										}
+										break;
+									}
+								default:
+									checkf(false, TEXT("Unsupported operation type"));
+									break;
+								}
+							}
+						}
+					}
+				}
+			}
+
+			int32 NumReplaced = bReplaceArray ? ArrayHelper.Num() : INDEX_NONE;
+			FStructuredArchive::FArray ReplacedArray = Record.EnterArray(TEXT("Replaced"), NumReplaced);
+			if(bReplaceArray)
+			{
+				const int32 ArrayNum = ArrayHelper.Num();
+				for (int32 i =0; i < ArrayNum; ++i)
+				{
+					SerializeContainerItem(ReplacedArray.EnterElement(), ArrayHelper.GetRawPtr(i));
+				}
+			}
+			else
+			{
+				checkf(Defaults, TEXT("Expecting overridable serialization to have defaults to compare to"));
+				FScriptArrayHelper DefaultsArrayHelper(this, Defaults);
+
+				int32 NumRemoved = RemovedIndices.Num();
+				FStructuredArchive::FArray RemovedArray = Record.EnterArray(TEXT("Removed"), NumRemoved);
+				for (int32 i : RemovedIndices)
+				{
+					SerializeContainerItem(RemovedArray.EnterElement(), DefaultsArrayHelper.GetRawPtr(i));
+				}
+
+				int32 NumModified = ModifiedIndices.Num();
+				FStructuredArchive::FArray ModifiedArray = Record.EnterArray(TEXT("Modified"), NumModified);
+				for (int32 i : ModifiedIndices)
+				{
+					SerializeContainerItem(ModifiedArray.EnterElement(), ArrayHelper.GetRawPtr(i));
+				}
+
+				int32 NumAdded = AddedIndices.Num();
+				FStructuredArchive::FArray AddedArray = Record.EnterArray(TEXT("Added"), NumAdded);
+				for (int32 i : AddedIndices)
+				{
+					SerializeContainerItem(AddedArray.EnterElement(), ArrayHelper.GetRawPtr(i));
+				}
+			}
+		}
+		return;
+	}
+
+	FStructuredArchiveArray Array = Slot.EnterArray(n);
+
+	if( UnderlyingArchive.IsLoading() )
+	{
+		// If using a custom property list, don't empty the array on load. Not all indices may have been serialized, so we need to preserve existing values at those slots.
+		if (UnderlyingArchive.ArUseCustomPropertyList)
+		{
+			const int32 OldNum = ArrayHelper.Num();
+			if (n > OldNum)
+			{
+				ArrayHelper.AddValues(n - OldNum);
+			}
+			else if (n < OldNum)
+			{
+				ArrayHelper.RemoveValues(n, OldNum - n);
+			}
+		}
+		else
+		{
+			ArrayHelper.EmptyAndAddValues(n);
+		}
+	}
+	ArrayHelper.CountBytes( UnderlyingArchive );
+	
 	// Serialize a PropertyTag for the inner property of this array, allows us to validate the inner struct to see if it has changed
 	if (UnderlyingArchive.UEVer() >= VER_UE4_INNER_ARRAY_TAG_INFO && Inner->IsA<FStructProperty>())
 	{
diff --git a/Engine/Source/Runtime/CoreUObject/Private/UObject/PropertyMap.cpp b/Engine/Source/Runtime/CoreUObject/Private/UObject/PropertyMap.cpp
--- a/Engine/Source/Runtime/CoreUObject/Private/UObject/PropertyMap.cpp
+++ b/Engine/Source/Runtime/CoreUObject/Private/UObject/PropertyMap.cpp
@@ -13,6 +13,7 @@
 #include "Misc/ScopeExit.h"
 #include "Serialization/ArchiveUObjectFromStructuredArchive.h"
 #include "UObject/UObjectThreadContext.h"
+#include "UObject/OverriddenPropertySet.h"
 
 namespace UEMapProperty_Private
 {
@@ -281,6 +282,8 @@ void FMapProperty::GetPreloadDependencies(TArray<UObject*>& OutDeps)
 void FMapProperty::SerializeItem(FStructuredArchive::FSlot Slot, void* Value, const void* Defaults) const
 {
 	FArchive& UnderlyingArchive = Slot.GetUnderlyingArchive();
+	const bool bUPS = UnderlyingArchive.UseUnversionedPropertySerialization();
+	bool bExperimentalOverridableLogic = HasAnyPropertyFlags(CPF_ExperimentalOverridableLogic);
 	FStructuredArchive::FRecord Record = Slot.EnterRecord();
 
 	// Map containers must be serialized as a "whole" value, which means that we need to serialize every field for struct-typed entries.
@@ -304,22 +307,388 @@ void FMapProperty::SerializeItem(FStructuredArchive::FSlot Slot, void* Value, co
 
 	FScriptMapHelper MapHelper(this, Value);
 
-	if (UnderlyingArchive.IsLoading())
+	// Make sure the container is reloading accordingly to the value set in the property tag if any
+	if (!bUPS && UnderlyingArchive.IsLoading() && FPropertyTagScope::GetCurrentPropertyTag())
 	{
-		FUObjectSerializeContext* Context = UnderlyingArchive.GetSerializeContext();
+		bExperimentalOverridableLogic = FPropertyTagScope::GetCurrentPropertyTag()->bExperimentalOverridableLogic;
+	}
 
-		if (Defaults)
+	// *** Experimental *** Special serialization path for map with overridable serialization logic
+	if (bExperimentalOverridableLogic)
+	{
+		checkf(!UnderlyingArchive.ArUseCustomPropertyList, TEXT("Using custom property list is not supported by overridable serialization"));
+
+		auto GetIDFromKey = [&](uint8* KeyData) -> FName
 		{
-			CopyValuesInternal(Value, Defaults, 1);
+			if (FObjectProperty* KeyObjectProperty = CastField<FObjectProperty>(KeyProp))
+			{
+				if (const UObject* Object = KeyObjectProperty->GetObjectPropertyValue(KeyData))
+				{
+					return FName(Object->GetPathName());
+				}
+			}
+			else
+			{
+				FString KeyString;
+				KeyProp->ExportTextItem_Direct(KeyString, KeyData, nullptr, nullptr, PPF_None);
+				return FName(KeyString);
+			}
+	
+			checkf(false, TEXT("This case is not handled"))
+			return NAME_None;
+		};
+
+		if (UnderlyingArchive.IsLoading())
+		{
+			int32 NumReplaced = 0;
+			FStructuredArchive::FArray ReplacedArray = Record.EnterArray(TEXT("Replaced"), NumReplaced);
+			if (NumReplaced != INDEX_NONE)
+			{
+				MapHelper.EmptyValues(NumReplaced);
+				for (int32 i = 0; i < NumReplaced; i++)
+				{
+					FStructuredArchive::FRecord EntryRecord = ReplacedArray.EnterElement().EnterRecord();
+					int32 Index = MapHelper.AddDefaultValue_Invalid_NeedsRehash();
+					uint8* PairPtr = MapHelper.GetPairPtr(Index);
+
+					{
+						FSerializedPropertyScope SerializedProperty(UnderlyingArchive, KeyProp, this);
+						KeyProp->SerializeItem(EntryRecord.EnterField(TEXT("Key")), PairPtr);
+					}
+					{
+						FSerializedPropertyScope SerializedProperty(UnderlyingArchive, ValueProp, this);
+						ValueProp->SerializeItem(EntryRecord.EnterField(TEXT("Value")), PairPtr + MapLayout.ValueOffset);
+					}
+				}
+				MapHelper.Rehash();
+			}
+			else
+			{
+				FOverriddenPropertySet* OverriddenProperties = FOverridableSerializationLogic::GetOverriddenProperties();
+
+				// This is not fully implemented yet and not a priority right now, so just trying to prevent it as the result could be random
+				checkf(!KeyProp->HasAnyPropertyFlags(CPF_PersistentInstance) || CastField<FClassProperty>(KeyProp) || !CastField<FObjectProperty>(KeyProp), TEXT("The key as an instanced sub object is NYI"));
+
+				uint8* TempKeyValueStorage = nullptr;
+				ON_SCOPE_EXIT
+				{
+					if (TempKeyValueStorage)
+					{
+						KeyProp->DestroyValue(TempKeyValueStorage);
+						ValueProp->DestroyValue(TempKeyValueStorage + MapLayout.ValueOffset);
+						FMemory::Free(TempKeyValueStorage);
+					}
+				};
+
+				int32 NumRemoved = 0;
+				FStructuredArchive::FArray RemovedArray = Record.EnterArray(TEXT("Removed"), NumRemoved);
+				if (NumRemoved != 0)
+				{
+					TempKeyValueStorage = (uint8*)FMemory::Malloc(MapLayout.SetLayout.Size);
+					KeyProp->InitializeValue(TempKeyValueStorage);
+					ValueProp->InitializeValue(TempKeyValueStorage + MapLayout.ValueOffset);
+
+					for (int32 i = 0; i < NumRemoved; ++i)
+					{
+						{
+							FSerializedPropertyScope SerializedProperty(UnderlyingArchive, KeyProp, this);
+							KeyProp->SerializeItem(RemovedArray.EnterElement().EnterRecord().EnterField(TEXT("Key")), TempKeyValueStorage);
+						}
+
+						MapHelper.RemovePair(TempKeyValueStorage);
+
+						// Need to fetch the MapOverriddenPropertyNode every loop as the previous might have reallocated the node.
+						if (FOverriddenPropertyNode* MapOverriddenPropertyNode = OverriddenProperties ? OverriddenProperties->SetOverriddenPropertyOperation(EOverriddenPropertyOperation::Modified, UnderlyingArchive.GetSerializedPropertyChain(), /*Property*/nullptr) : nullptr)
+						{
+							// Rebuild the overridden info
+							FName RemovedKeyID = GetIDFromKey(TempKeyValueStorage);
+							OverriddenProperties->SetSubPropertyOperation(EOverriddenPropertyOperation::Remove, *MapOverriddenPropertyNode, RemovedKeyID);
+						}
+					}
+				}
+
+				int32 NumModified = 0;
+				FStructuredArchive::FArray ModifiedArray = Record.EnterArray(TEXT("Modified"), NumModified);
+				if (NumModified != 0)
+				{
+					if (!TempKeyValueStorage)
+					{
+						TempKeyValueStorage = (uint8*)FMemory::Malloc(MapLayout.SetLayout.Size);
+						KeyProp->InitializeValue(TempKeyValueStorage);
+						ValueProp->InitializeValue(TempKeyValueStorage + MapLayout.ValueOffset);
+					}
+					for (int32 i = 0; i < NumModified; ++i)
+					{
+						FStructuredArchive::FRecord EntryRecord = ModifiedArray.EnterElement().EnterRecord();
+
+						// Read key into temporary storage
+						{
+							FSerializedPropertyScope SerializedProperty(UnderlyingArchive, KeyProp, this);
+							KeyProp->SerializeItem(EntryRecord.EnterField(TEXT("Key")), TempKeyValueStorage);
+						}
+
+						const int32 Index = MapHelper.FindMapPairIndexFromHash(TempKeyValueStorage);
+						uint8* ValuePtr = Index != INDEX_NONE ? MapHelper.GetValuePtr(Index) : TempKeyValueStorage + MapLayout.ValueOffset;
+
+						// Deserialize value into hash map-owned memory
+						{
+							FSerializedPropertyScope SerializedProperty(UnderlyingArchive, ValueProp, this);
+							ValueProp->SerializeItem(EntryRecord.EnterField(TEXT("Value")), ValuePtr);
+						}
+					}
+				}
+
+				int32 NumAdded = 0;
+				FStructuredArchive::FArray AddedArray = Record.EnterArray(TEXT("Added"), NumAdded);
+				if (NumAdded != 0)
+				{
+					if (!TempKeyValueStorage)
+					{
+						TempKeyValueStorage = (uint8*)FMemory::Malloc(MapLayout.SetLayout.Size);
+						KeyProp->InitializeValue(TempKeyValueStorage);
+						ValueProp->InitializeValue(TempKeyValueStorage + MapLayout.ValueOffset);
+					}
+
+					for (int32 i = 0; i < NumAdded; ++i)
+					{
+						FStructuredArchive::FRecord EntryRecord = AddedArray.EnterElement().EnterRecord();
+
+						// Read key into temporary storage
+						{
+							FSerializedPropertyScope SerializedProperty(UnderlyingArchive, KeyProp, this);
+							KeyProp->SerializeItem(EntryRecord.EnterField(TEXT("Key")), TempKeyValueStorage);
+						}
+
+						void* ValuePtr = MapHelper.FindOrAdd(TempKeyValueStorage);
+
+						// Deserialize value into hash map-owned memory
+						{
+							FSerializedPropertyScope SerializedProperty(UnderlyingArchive, ValueProp, this);
+							ValueProp->SerializeItem(EntryRecord.EnterField(TEXT("Value")), ValuePtr);
+						}
+
+						// Need to fetch the MapOverriddenPropertyNode every loop as the previous might have reallocated the node.
+						if (FOverriddenPropertyNode* MapOverriddenPropertyNode = OverriddenProperties ? OverriddenProperties->SetOverriddenPropertyOperation(EOverriddenPropertyOperation::Modified, UnderlyingArchive.GetSerializedPropertyChain(), /*Property*/nullptr) : nullptr)
+						{
+							// Rebuild the overridden info
+							FName AddedKeyID = GetIDFromKey(TempKeyValueStorage);
+							OverriddenProperties->SetSubPropertyOperation(EOverriddenPropertyOperation::Add, *MapOverriddenPropertyNode, AddedKeyID);
+						}
+					}
+				}
+			}
 		}
+		else
+		{
+			auto FindKeyInternalIndex = [this](const FName KeyIDToFind, FScriptMapHelper& MapHelper) -> int32
+			{
+				if (const FObjectProperty* KeyObjectProperty = CastField<FObjectProperty>(KeyProp))
+				{
+					for (FScriptMapHelper::FIterator It(MapHelper); It; ++It) 
+					{
+						UObject* CurrentObject = KeyObjectProperty->GetObjectPropertyValue(MapHelper.GetKeyPtr(It.GetInternalIndex()));
+						if (FName(CurrentObject->GetPathName()) == KeyIDToFind)
+						{
+							return It.GetInternalIndex();
+						}
+					}
+				}
+				else
+				{
+					void* TempKeyValueStorage = FMemory::Malloc(MapLayout.SetLayout.Size);
+					KeyProp->InitializeValue(TempKeyValueStorage);
+
+					FString KeyToFind(KeyIDToFind.ToString());
+					KeyProp->ImportText_Direct(*KeyToFind, TempKeyValueStorage, nullptr, PPF_None);
+
+					const int32 InternalIndex = MapHelper.FindMapPairIndexFromHash(TempKeyValueStorage);
+
+					KeyProp->DestroyValue(TempKeyValueStorage);
+					FMemory::Free(TempKeyValueStorage);
+
+					return InternalIndex;
+				}
+				return INDEX_NONE;
+			};
+
+			// Container for temporarily tracking some indices
+			TArray<int32> RemovedIndices;
+			TArray<int32> AddedIndices;
+			TSet<int32> ModifiedIndices;
+
+			bool bReplaceMap = false;
+			if (!Defaults)
+			{
+				bReplaceMap = true;
+			}
+			else 
+			{
+				EOverriddenPropertyOperation MapOverrideOp = EOverriddenPropertyOperation::None;
+				FOverriddenPropertySet* OverriddenProperties = FOverridableSerializationLogic::GetOverriddenProperties();
+				if (OverriddenProperties)
+				{
+					MapOverrideOp = OverriddenProperties->GetOverriddenPropertyOperation(UnderlyingArchive.GetSerializedPropertyChain(), /*Property*/nullptr);
+					bReplaceMap = MapOverrideOp == EOverriddenPropertyOperation::Replace;
+				}
+				else
+				{
+					// Only instanced subobjects keys are not supported
+					bReplaceMap = CastField<FObjectProperty>(KeyProp) && KeyProp->HasAnyPropertyFlags(CPF_PersistentInstance) && !CastField<FClassProperty>(KeyProp); // Class property should never be instanced
+				}
+
+				if(!bReplaceMap)
+				{
+					checkf(!KeyProp->HasAnyPropertyFlags(CPF_PersistentInstance) || CastField<FClassProperty>(KeyProp) || !CastField<FObjectProperty>(KeyProp), TEXT("The key as an instanced sub object is NYI"));
+
+					// For instanced subobject, the overridable is handled per object base, not only here, so we need to serialize the object ptr no matter what.
+					const bool bAreValuesInstancedSubObjects = ValueProp->HasAnyPropertyFlags(CPF_PersistentInstance) && CastField<FObjectProperty>(ValueProp);
+					if(bAreValuesInstancedSubObjects)
+					{
+						for (FScriptMapHelper::FIterator It(MapHelper); It; ++It)
+						{
+							ModifiedIndices.Add(It.GetInternalIndex());
+						}
+					}
+
+					if (OverriddenProperties && MapOverrideOp != EOverriddenPropertyOperation::None)
+					{
+						checkf(Defaults, TEXT("Expecting overridable serialization to have defaults to compare to"));
+						FScriptMapHelper DefaultsMapHelper(this, Defaults);
+
+						if (const FOverriddenPropertyNode* MapOverriddenPropertyNode = OverriddenProperties->GetOverriddenPropertyNode(UnderlyingArchive.GetSerializedPropertyChain()))
+						{
+							// Figure out the modifications of the map
+							for (const auto& Pair : MapOverriddenPropertyNode->SubPropertyNodeKeys)
+							{
+								const EOverriddenPropertyOperation OverrideOp = OverriddenProperties->GetSubPropertyOperation(Pair.Value);
+								switch (OverrideOp)
+								{
+								case EOverriddenPropertyOperation::Remove:
+									{
+										const int32 InternalIndex = FindKeyInternalIndex(Pair.Key, DefaultsMapHelper);
+										if (InternalIndex != INDEX_NONE)
+										{
+											RemovedIndices.Add(InternalIndex);
+										}
+										break;
+									}
+								case EOverriddenPropertyOperation::Add:
+									{
+										const int32 InternalIndex = FindKeyInternalIndex(Pair.Key, MapHelper);
+										if (InternalIndex != INDEX_NONE)
+										{
+											AddedIndices.Add(InternalIndex);
+											ModifiedIndices.Remove(InternalIndex);
+										}
+										break;
+									}
+								case EOverriddenPropertyOperation::Modified:
+									if (!bAreValuesInstancedSubObjects)
+									{
+										const int32 InternalIndex = FindKeyInternalIndex(Pair.Key, MapHelper);
+										if (InternalIndex != INDEX_NONE)
+										{
+											ModifiedIndices.Add(InternalIndex);
+										}
+									}
+									break;
+								default:
+									checkf(false, TEXT("Unsupported map operation"));
+									break;
+								}
+							}
+						}
+					}
+				}
+			}
+
+			int32 NumReplaced = bReplaceMap ? MapHelper.Num() : INDEX_NONE;
+			FStructuredArchive::FArray ReplacedArray = Record.EnterArray(TEXT("Replaced"), NumReplaced);
+			if (bReplaceMap)
+			{
+				for (FScriptMapHelper::FIterator It(MapHelper); It; ++It)
+				{
+					FStructuredArchive::FRecord EntryRecord = ReplacedArray.EnterElement().EnterRecord();
+					uint8* PairPtr = MapHelper.GetPairPtr(It.GetInternalIndex());
+					{
+						FSerializedPropertyScope SerializedProperty(UnderlyingArchive, KeyProp, this);
+						KeyProp->SerializeItem(EntryRecord.EnterField(TEXT("Key")), PairPtr);
+					}
+					{
+						FSerializedPropertyScope SerializedProperty(UnderlyingArchive, ValueProp, this);
+						ValueProp->SerializeItem(EntryRecord.EnterField(TEXT("Value")), PairPtr + MapLayout.ValueOffset);
+					}
+				}
+			}
+			else
+			{
+				checkf(Defaults, TEXT("Expecting overridable serialization to have defaults to compare to"));
+				FScriptMapHelper DefaultsMapHelper(this, Defaults);
+
+				int32 NumRemoved = bReplaceMap ? INDEX_NONE : RemovedIndices.Num();
+				FStructuredArchive::FArray RemovedArray = Record.EnterArray(TEXT("Removed"), NumRemoved);
+				for (int32 i : RemovedIndices)
+				{
+					FStructuredArchive::FRecord EntryRecord = RemovedArray.EnterElement().EnterRecord();
+					FSerializedPropertyScope SerializedProperty(UnderlyingArchive, KeyProp, this);
+					KeyProp->SerializeItem(EntryRecord.EnterField(TEXT("Key")), DefaultsMapHelper.GetKeyPtr(i));
+				}
+
+				int32 NumModified = ModifiedIndices.Num();
+				FStructuredArchive::FArray ModifiedArray = Record.EnterArray(TEXT("Modified"), NumModified);
+				for (int32 i : ModifiedIndices)
+				{
+					FStructuredArchive::FRecord EntryRecord = ModifiedArray.EnterElement().EnterRecord();
+					uint8* PairPtr = MapHelper.GetPairPtr(i);
+					{
+						FSerializedPropertyScope SerializedProperty(UnderlyingArchive, KeyProp, this);
+						KeyProp->SerializeItem(EntryRecord.EnterField(TEXT("Key")), PairPtr);
+					}
+					{
+						FSerializedPropertyScope SerializedProperty(UnderlyingArchive, ValueProp, this);
+						ValueProp->SerializeItem(EntryRecord.EnterField(TEXT("Value")), PairPtr + MapLayout.ValueOffset);
+					}
+				}
+
+				// Added keys
+				int32 NumAdded = AddedIndices.Num();
+				FStructuredArchive::FArray AddedArray = Record.EnterArray(TEXT("Added"), NumAdded);
+				for (int32 i : AddedIndices)
+				{
+					FStructuredArchive::FRecord EntryRecord = AddedArray.EnterElement().EnterRecord();
+					uint8* PairPtr = MapHelper.GetPairPtr(i);
+					{
+						FSerializedPropertyScope SerializedProperty(UnderlyingArchive, KeyProp, this);
+						KeyProp->SerializeItem(EntryRecord.EnterField(TEXT("Key")), PairPtr);
+					}
+					{
+						FSerializedPropertyScope SerializedProperty(UnderlyingArchive, ValueProp, this);
+						ValueProp->SerializeItem(EntryRecord.EnterField(TEXT("Value")), PairPtr + MapLayout.ValueOffset);
+					}
+				}
+			}
+		}
+
+		return;
+	}
+
+	if (UnderlyingArchive.IsLoading())
+	{
+		FUObjectSerializeContext* Context = UnderlyingArchive.GetSerializeContext();
 
 		// Delete any explicitly-removed elements
 		int32 NumKeysToRemove = 0;
 		FStructuredArchive::FArray KeysToRemoveArray = Record.EnterArray(TEXT("KeysToRemove"), NumKeysToRemove);
+		const bool bReplaceMap = NumKeysToRemove == INDEX_NONE;
 
-		if (!Defaults || MapHelper.Num() == 0) // Faster loading path when loading into an empty map
+		if (Defaults && !bReplaceMap)
 		{
-			if (NumKeysToRemove)
+			CopyValuesInternal(Value, Defaults, 1);
+		}
+
+		if (!Defaults || MapHelper.Num() == 0 || bReplaceMap) // Faster loading path when loading into an empty map or replacing the entire map
+		{
+			if (NumKeysToRemove && !bReplaceMap)
 			{
 				TOptional<TGuardValue<bool>> SerializeUnknownProperty;
 				if (Context)
@@ -458,11 +827,14 @@ void FMapProperty::SerializeItem(FStructuredArchive::FSlot Slot, void* Value, co
 	{
 		FScriptMapHelper DefaultsHelper(this, Defaults);
 
+		// Override logic should only supports replacing the entire array
+		const bool bReplaceMap = FOverridableSerializationLogic::GetOverriddenProperties() != nullptr;
+
 		// Container for temporarily tracking some indices
 		TSet<int32> Indices;
 
 		// Determine how many keys are missing from the object
-		if (Defaults)
+		if (Defaults && !bReplaceMap)
 		{
 			for (FScriptMapHelper::FIterator Iterator(DefaultsHelper); Iterator; ++Iterator)
 			{
@@ -475,7 +847,7 @@ void FMapProperty::SerializeItem(FStructuredArchive::FSlot Slot, void* Value, co
 		}
 
 		// Write out the missing keys
-		int32 MissingKeysNum = Indices.Num();
+		int32 MissingKeysNum = bReplaceMap ? INDEX_NONE : Indices.Num();
 		FStructuredArchive::FArray KeysToRemoveArray = Record.EnterArray(TEXT("KeysToRemove"), MissingKeysNum);
 		{
 			FSerializedPropertyScope SerializedProperty(UnderlyingArchive, KeyProp, this);
@@ -486,7 +858,7 @@ void FMapProperty::SerializeItem(FStructuredArchive::FSlot Slot, void* Value, co
 		}
 
 		// Write out differences from defaults
-		if (Defaults)
+		if (Defaults && !bReplaceMap)
 		{
 			Indices.Empty(Indices.Num());
 			for (FScriptMapHelper::FIterator Iterator(MapHelper); Iterator; ++Iterator)
diff --git a/Engine/Source/Runtime/CoreUObject/Private/UObject/PropertyTag.cpp b/Engine/Source/Runtime/CoreUObject/Private/UObject/PropertyTag.cpp
--- a/Engine/Source/Runtime/CoreUObject/Private/UObject/PropertyTag.cpp
+++ b/Engine/Source/Runtime/CoreUObject/Private/UObject/PropertyTag.cpp
@@ -7,20 +7,27 @@
 #include "UObject/Package.h"
 #include "UObject/UnrealType.h"
 #include "UObject/BlueprintsObjectVersion.h"
+#include "UObject/OverriddenPropertySet.h"
+
+thread_local const FPropertyTag* FPropertyTagScope::CurrentPropertyTag = nullptr;
 
 /*-----------------------------------------------------------------------------
 FPropertyTag
 -----------------------------------------------------------------------------*/
+FPropertyTag::FPropertyTag()
+: OverrideOperation(EOverriddenPropertyOperation::None)
+{}
 
 FPropertyTag::FPropertyTag( FArchive& InSaveAr, FProperty* Property, int32 InIndex, uint8* Value, const uint8* Defaults )
 	: ArrayIndex(InIndex)
+	, OverrideOperation(EOverriddenPropertyOperation::None)
 {
 	check(!InSaveAr.GetArchiveState().UseUnversionedPropertySerialization());
 	Property->SaveToTag(*this);
 	if (FBoolProperty* Bool = CastField<FBoolProperty>(Property))
-	{
-		BoolVal = Bool->GetPropertyValue(Value);
-	}
+		{
+			BoolVal = Bool->GetPropertyValue(Value);
+		}
 }
 
 // Set optional property guid
@@ -169,11 +176,31 @@ void operator<<(FStructuredArchive::FSlot Slot, FPropertyTag& Tag)
 		}
 	}
 
-	// Serialize tag extensions
+	// Serialize tag extensions, consider doing a init function and context as `EClassSerializationControlExtension` if we add more extensions
 	if (UnderlyingArchive.UEVer() >= EUnrealEngineObjectUE5Version::PROPERTY_TAG_EXTENSION_AND_OVERRIDABLE_SERIALIZATION)
 	{
 		EPropertyTagExtension PropertyTagExtensions = EPropertyTagExtension::NoExtension;
+
+		if (UnderlyingArchive.IsSaving())
+		{
+			// Overridable information extension
+			const FOverriddenPropertySet* OverriddenProperties = FOverridableSerializationLogic::GetOverriddenProperties();
+			Tag.bExperimentalOverridableLogic = Tag.Prop->HasAnyPropertyFlags(CPF_ExperimentalOverridableLogic);
+			if (OverriddenProperties || Tag.bExperimentalOverridableLogic)
+			{
+				Tag.OverrideOperation = OverriddenProperties ? OverriddenProperties->GetOverriddenPropertyOperation(UnderlyingArchive.GetSerializedPropertyChain(), Tag.Prop) : EOverriddenPropertyOperation::None;
+				PropertyTagExtensions |= EPropertyTagExtension::OverridableInformation;
+			}
+		}
+
 		Slot << SA_ATTRIBUTE(TEXT("PropertyExtensions"), PropertyTagExtensions);
+
+		// Overridable information extension
+		if (EnumHasAnyFlags(PropertyTagExtensions,EPropertyTagExtension::OverridableInformation))
+		{
+			Slot << SA_ATTRIBUTE(TEXT("OverriddenPropertyOperation"), Tag.OverrideOperation);
+			Slot << SA_ATTRIBUTE(TEXT("ExperimentalOverridableLogic"), Tag.bExperimentalOverridableLogic);
+		}
 	}
 }
 
@@ -210,6 +237,7 @@ void FPropertyTag::SerializeTaggedProperty(FStructuredArchive::FSlot Slot, FProp
 		FArchive::FScopeAddDebugData A(UnderlyingArchive, Property->GetFName());
 #endif
 		FSerializedPropertyScope SerializedProperty(UnderlyingArchive, Property);
+		FPropertyTagScope CurrentPropertyTagScope(this);
 
 		Property->SerializeItem(Slot, Value, Defaults);
 	}
diff --git a/Engine/Source/Runtime/CoreUObject/Private/UObject/UObjectGlobals.cpp b/Engine/Source/Runtime/CoreUObject/Private/UObject/UObjectGlobals.cpp
--- a/Engine/Source/Runtime/CoreUObject/Private/UObject/UObjectGlobals.cpp
+++ b/Engine/Source/Runtime/CoreUObject/Private/UObject/UObjectGlobals.cpp
@@ -32,6 +32,7 @@
 #include "UObject/Class.h"
 #include "UObject/CoreRedirects.h"
 #include "UObject/FastReferenceCollector.h"
+#include "UObject/OverridableManager.h"
 #include "UObject/UObjectIterator.h"
 #include "UObject/Package.h"
 #include "Templates/Casts.h"
@@ -3614,6 +3615,8 @@ void UObject::PostInitProperties()
 #if USE_UBER_GRAPH_PERSISTENT_FRAME
 	GetClass()->CreatePersistentUberGraphFrame(this, true);
 #endif
+
+	FOverridableManager::Get().ClearOverrides(*this);
 }
 
 UObject::UObject()
@@ -4101,6 +4104,8 @@ void FObjectInitializer::InitProperties(UObject* Obj, UClass* DefaultsClass, UOb
 
 	check(DefaultsClass && Obj);
 
+	FOverridableManager::Get().InheritEnabledFrom(*Obj, DefaultData);
+
 	UClass* Class = Obj->GetClass();
 
 	// bool to indicate that we need to initialize any non-native properties (native ones were done when the native constructor was called by the code that created and passed in a FObjectInitializer object)
diff --git a/Engine/Source/Runtime/CoreUObject/Public/InstancedReferenceSubobjectHelper.h b/Engine/Source/Runtime/CoreUObject/Public/InstancedReferenceSubobjectHelper.h
new file mode 100644
--- /dev/null
+++ b/Engine/Source/Runtime/CoreUObject/Public/InstancedReferenceSubobjectHelper.h
@@ -0,0 +1,148 @@
+// Copyright Epic Games, Inc. All Rights Reserved.
+
+#pragma once
+
+#include "CoreMinimal.h"
+#include "UObject/ObjectMacros.h"
+#include "UObject/UnrealType.h"
+
+/** 
+ * Meant to represent a specific object property that is setup to reference a 
+ * instanced sub-object. Tracks the property hierarchy used to reach the 
+ * property, so that we can easily retrieve instanced sub-objects from a 
+ * container object.
+ */
+struct FInstancedPropertyPath
+{
+private:
+	struct FPropertyLink
+	{
+		FPropertyLink(const FProperty* Property, int32 ArrayIndexIn, bool bIsMapValueIn)
+			: PropertyPtr(Property), ArrayIndex(ArrayIndexIn), bIsMapValue(bIsMapValueIn)
+		{}
+
+		const FProperty* PropertyPtr;
+		int32            ArrayIndex;
+		bool             bIsMapValue;
+	};
+
+public:
+	//--------------------------------------------------------------------------
+	FInstancedPropertyPath(FProperty* RootProperty, int32 ArrayIndex = INDEX_NONE, bool bIsMapValue = false)
+	{
+		Push(RootProperty, ArrayIndex, bIsMapValue);
+	}
+
+	//--------------------------------------------------------------------------
+	void Push(const FProperty* Property, int32 ArrayIndex = INDEX_NONE, bool bIsMapValue = false)
+	{
+		PropertyChain.Add(FPropertyLink(Property, ArrayIndex, bIsMapValue));		
+	}
+
+	//--------------------------------------------------------------------------
+	void Pop()
+	{
+ 		PropertyChain.RemoveAt(PropertyChain.Num() - 1);
+	}
+
+	//--------------------------------------------------------------------------
+	const FProperty* Head() const
+	{
+		return PropertyChain.Last().PropertyPtr;
+	}
+
+	//--------------------------------------------------------------------------
+	COREUOBJECT_API UObject* Resolve(const UObject* Container) const;
+
+private:
+	TArray<FPropertyLink> PropertyChain;
+};
+
+/** 
+ * Can be used as a raw sub-object pointer, but also contains a 
+ * FInstancedPropertyPath to identify the property that this sub-object is 
+ * referenced by. Paired together for ease of use (so API users don't have to manage a map).
+ */
+struct FInstancedSubObjRef
+{
+	FInstancedSubObjRef(UObject* SubObj, const FInstancedPropertyPath& PropertyPathIn)
+		: SubObjInstance(SubObj)
+		, PropertyPath(PropertyPathIn)
+	{}
+
+	//--------------------------------------------------------------------------
+	operator UObject*() const
+	{
+		return SubObjInstance;
+	}
+
+	//--------------------------------------------------------------------------
+	UObject* operator->() const
+	{
+		return SubObjInstance;
+	}
+
+	//--------------------------------------------------------------------------
+	friend uint32 GetTypeHash(const FInstancedSubObjRef& SubObjRef)
+	{
+		return GetTypeHash((UObject*)SubObjRef);
+	}
+
+	UObject* SubObjInstance;
+	FInstancedPropertyPath PropertyPath;
+};
+ 
+/** 
+ * Contains a set of utility functions useful for searching out and identifying
+ * instanced sub-objects contained within a specific outer object.
+ */
+class FFindInstancedReferenceSubobjectHelper
+{
+public:
+	template<typename T>
+	static void GetInstancedSubObjects(const UObject* Container, T& OutObjects)
+	{
+		const UClass* ContainerClass = Container->GetClass();
+		for (FProperty* Prop = ContainerClass->RefLink; Prop; Prop = Prop->NextRef)
+		{
+			for (int32 ArrayIdx = 0; ArrayIdx < Prop->ArrayDim; ++ArrayIdx)
+			{
+				FInstancedPropertyPath RootPropertyPath(Prop, ArrayIdx);
+				const uint8* ValuePtr = Prop->ContainerPtrToValuePtr<uint8>(Container, ArrayIdx);
+				ForEachInstancedSubObject<const void*>(RootPropertyPath, ValuePtr, [&OutObjects](const FInstancedSubObjRef& Ref, const void*){ OutObjects.Add(Ref); });
+			}
+		}
+	}
+
+	template<typename T>
+	static void GetInstancedSubObjectsRecursive(const UObject* Container, T& OutObjects)
+	{
+		const UClass* ContainerClass = Container->GetClass();
+		for (FProperty* Prop = ContainerClass->RefLink; Prop; Prop = Prop->NextRef)
+		{
+			for (int32 ArrayIdx = 0; ArrayIdx < Prop->ArrayDim; ++ArrayIdx)
+			{
+				FInstancedPropertyPath RootPropertyPath(Prop, ArrayIdx);
+				const uint8* ValuePtr = Prop->ContainerPtrToValuePtr<uint8>(Container, ArrayIdx);
+				ForEachInstancedSubObject<const void*>(RootPropertyPath, ValuePtr, [&OutObjects](const FInstancedSubObjRef& Ref, const void*)
+				{
+					if (!OutObjects.Contains(Ref))
+					{
+						OutObjects.Add(Ref);
+						GetInstancedSubObjectsRecursive(Ref.SubObjInstance, OutObjects);
+					}
+				});
+			}
+		}
+	}
+
+	static COREUOBJECT_API void Duplicate(
+		UObject* OldObject, 
+		UObject* NewObject, 
+		TMap<UObject*, UObject*>& ReferenceReplacementMap, 
+		TArray<UObject*>& DuplicatedObjects, 
+		TMap<UObject*, UObject*>* OptionalMappings);
+
+	template<typename T>
+	static void ForEachInstancedSubObject(FInstancedPropertyPath& PropertyPath, T ContainerAddress, TFunctionRef<void(const FInstancedSubObjRef& Ref, T PropertyValueAddress)> ObjRefFunc);
+};
diff --git a/Engine/Source/Runtime/CoreUObject/Public/UObject/ObjectMacros.h b/Engine/Source/Runtime/CoreUObject/Public/UObject/ObjectMacros.h
--- a/Engine/Source/Runtime/CoreUObject/Public/UObject/ObjectMacros.h
+++ b/Engine/Source/Runtime/CoreUObject/Public/UObject/ObjectMacros.h
@@ -452,6 +452,7 @@ enum EPropertyFlags : uint64
 	CPF_NativeAccessSpecifierPrivate	= 0x0040000000000000,	///< Private native access specifier
 	CPF_SkipSerialization				= 0x0080000000000000,	///< Property shouldn't be serialized, can still be exported to text
 	CPF_TObjectPtr						= 0x0100000000000000,	///< Property is a TObjectPtr<T> instead of a USomething*. Need to differentiate between TObjectclassOf and TObjectPtr
+	CPF_ExperimentalOverridableLogic	= 0x0200000000000000,	///< ****Experimental*** Property will use different logic to serialize knowing what changes are done against its default use the overridable information provided by the overidable manager on the object
 };
 
 /** All Native Access Specifier flags */
diff --git a/Engine/Source/Runtime/CoreUObject/Public/UObject/OverridableManager.h b/Engine/Source/Runtime/CoreUObject/Public/UObject/OverridableManager.h
new file mode 100644
--- /dev/null
+++ b/Engine/Source/Runtime/CoreUObject/Public/UObject/OverridableManager.h
@@ -0,0 +1,193 @@
+// Copyright Epic Games, Inc. All Rights Reserved.
+
+#pragma once
+
+#include "UObject/Object.h"
+#include "UObject/UObjectAnnotation.h"
+#include "UObject/OverriddenPropertySet.h"
+
+/*
+ *************************************************************************************
+ * Overridable serialization is experimental, not supported and use at your own risk *
+ *************************************************************************************
+ */
+
+/**
+ * Struct holding the shared ptr of the overridden properties
+ */
+struct FOverriddenPropertyAnnotation
+{
+	bool IsDefault() const
+	{
+		return !OverriddenProperties.IsValid();
+	}
+
+	TSharedPtr<FOverriddenPropertySet> OverriddenProperties;
+};
+
+/**
+ * Global container of overriden object annotations
+ */
+class FOverriddenPropertyAnnotations : public FUObjectAnnotationChunked<FOverriddenPropertyAnnotation, true/*bAutoRemove*/>
+{
+public:
+	bool IsEnabled(const UObject& Object)
+	{
+		return !GetAnnotation(&Object).IsDefault();
+	}
+
+	FOverriddenPropertySet* Find(const UObject& Object)
+	{
+		FOverriddenPropertyAnnotation Annotation = GetAnnotation(&Object);
+		return Annotation.OverriddenProperties.Get();
+	}
+
+	FOverriddenPropertySet& FindChecked(const UObject& Object)
+	{
+		FOverriddenPropertyAnnotation Annotation = GetAnnotation(&Object);
+		checkf(!Annotation.IsDefault(), TEXT("Caller is expecting the object to have overridable serialization enabled"));
+		return *Annotation.OverriddenProperties.Get();
+	}
+
+	FOverriddenPropertySet& FindOrAdd(UObject& Object)
+	{
+		FOverriddenPropertyAnnotation Annotation = GetAnnotation(&Object);
+		if (Annotation.IsDefault())
+		{
+			Annotation.OverriddenProperties = MakeShared<FOverriddenPropertySet>(Object);
+			AddAnnotation(&Object, Annotation);
+		}
+
+		return *Annotation.OverriddenProperties;
+	}
+};
+
+/*
+ *************************************************************************************
+ * Overridable serialization is experimental, not supported and use at your own risk *
+ *************************************************************************************
+ */
+class FOverridableManager
+{
+public:
+	/**
+	 * @return the static instance managing the overridability */
+	COREUOBJECT_API static FOverridableManager& Get();
+
+	/**
+	 * Lookup if the for the specified object has overridable serialization enabled
+	 * @param Object to check against
+	 * @return true if it uses the overridable serialization */
+	 COREUOBJECT_API bool IsEnabled(const UObject& Object);
+
+	/**
+	 * Sets on the specified object to use overridable serialization
+	 * @param Object to enable */
+	COREUOBJECT_API void Enable(UObject& Object);
+
+	/**
+	 * Sets on the specified object to use overridable serialization
+	 * @param Object to disable */
+	COREUOBJECT_API void Disable(UObject& Object);
+
+	/**
+	 * Inherit if the specified object should enable overridable serialization. It inherits it from either its default object or its outer.
+	 * @param Object to be inherited on
+	 * @param DefaultData its default object */
+	COREUOBJECT_API void InheritEnabledFrom(UObject& Object, const UObject* DefaultData);
+
+	/**
+	 * Retrieve the overridden properties for the specified object
+	 * @param Object to fetch the overridden properties
+	 * @return the overridden properties if the object have overridable serialization enabled */
+	COREUOBJECT_API FOverriddenPropertySet* GetOverriddenProperties(UObject& Object);
+
+	/**
+	 * Set the override operation on this object and it will enable it if it wasn't already enabled
+	 * @param Object to set the override operation on
+	 * @param Operation the override operation to set on the object
+	 * @return the overridden properties of the object */
+	COREUOBJECT_API FOverriddenPropertySet& SetOverriddenProperties(UObject& Object, EOverriddenPropertyOperation Operation);
+
+	/**
+	 * Copies the overridden properties from one object to the other
+	 * @param DestObject copying the overridden properties to
+	 * @param SourceObject copying the overridden properties from */
+	COREUOBJECT_API void CopyOverriddenProperties(UObject& DestObject, const UObject& SourceObject);
+
+	/**
+	 * Override the entire object properties and all its instanced subobjects
+	 * @param Object to override all it properties */
+	COREUOBJECT_API void OverrideObject(UObject& Object);
+
+	/**
+	 * Override all the properties of the specified instanced subobject if it is really owned by the referencer
+	 * @param Object referencing this subobject
+	 * @param InstancedSubObject the sub object to override all its properties */
+	COREUOBJECT_API void OverrideInstancedSubObject(UObject& Object, UObject& InstancedSubObject);
+
+	/**
+	 * Propagate override to all instanced sub object of the specified object
+	 * @param Object to propagate from */
+	COREUOBJECT_API void PropagateOverrideToInstancedSubObjects(UObject& Object);
+
+	/**
+	 * Clears all the overrides on the specified object
+	 * @param Object to clear them on */
+	COREUOBJECT_API void ClearOverrides(UObject& Object);
+
+	/**
+	 * Clear all overrides of the specified instanced subobject if it is really owned by the referencer
+	 * @param Object referencing the subobject
+	 * @param InstancedSubObject to clear overrides on */
+	COREUOBJECT_API void ClearInstancedSubObjectOverrides(UObject& Object, UObject& InstancedSubObject);
+
+	/**
+	 * Propagate the clear overrides to all instanced suboject of the specified object
+	 * @param Object*/
+	COREUOBJECT_API void PropagateClearOverridesToInstancedSubObjects(UObject& Object);
+
+	/**
+	 * Override a specific property of an object (Helper methods to call Pre/PostOverride)
+	 * Note: Supports object that does not have overridable serialization enabled
+	 * @param Object owning the property
+	 * @param PropertyEvent information about the type of change
+	 * @param PropertyChain leading to the property that is about to be overridden */
+	COREUOBJECT_API void OverrideProperty(UObject& Object, const FPropertyChangedEvent& PropertyEvent, const FEditPropertyChain& PropertyChain);
+
+	/**
+	 * To be called prior to override a property of the specified object
+	 * Note: Supports object that does not have overridable serialization enabled
+	 * @param Object owning the property
+	 * @param PropertyChain leading to the property about to be overriden */
+	COREUOBJECT_API void PreOverrideProperty(UObject& Object, const FEditPropertyChain& PropertyChain);
+
+	/**
+	 * To be called after the property was overridden of the specified object
+	 * Note: Supports object that does not have overridable serialization enabled
+	 * @param Object owning the property
+	 * @param PropertyEvent information about the type of change
+	 * @param PropertyChain leading to the property that was overridden */
+	COREUOBJECT_API void PostOverrideProperty(UObject& Object, const FPropertyChangedEvent& PropertyEvent, const FEditPropertyChain& PropertyChain);
+
+	/**
+	 * Notifies that a properties is about to change or has been changed for the specified object.
+	 * Note: Will ensure if the specified object does not have overridable serialization enabled
+	 * @param Notification the type of notification (pre or post
+	 * @param Object owning the property
+	 * @param PropertyEvent information about the type of change
+	 * @param PropertyNode leading to the property that is changing
+	 */
+	COREUOBJECT_API void NotifyPropertyChange(const EPropertyNotificationType Notification, UObject& Object, const FPropertyChangedEvent& PropertyEvent, const FEditPropertyChain::TDoubleLinkedListNode* PropertyNode);
+
+	/**
+	 * Serializes the overriden properties of the specified object into the record
+	 * @param Object to serialize the overriden property
+	 * @param ObjectRecord the record to use for serialization
+	 */
+	COREUOBJECT_API void SerializeOverriddenProperties(UObject& Object, FStructuredArchive::FRecord ObjectRecord);
+
+protected:
+	FOverriddenPropertyAnnotations OverriddenObjectAnnotations;
+};
+
diff --git a/Engine/Source/Runtime/CoreUObject/Public/UObject/OverriddenPropertySet.h b/Engine/Source/Runtime/CoreUObject/Public/UObject/OverriddenPropertySet.h
new file mode 100644
--- /dev/null
+++ b/Engine/Source/Runtime/CoreUObject/Public/UObject/OverriddenPropertySet.h
@@ -0,0 +1,235 @@
+// Copyright Epic Games, Inc. All Rights Reserved.
+
+#pragma once
+
+#include "UObject/Object.h"
+#include "Serialization/ArchiveSerializedPropertyChain.h"
+
+#include "OverriddenPropertySet.generated.h"
+
+COREUOBJECT_API DECLARE_LOG_CATEGORY_EXTERN(LogOverridableObject, Warning, All);
+
+struct FPropertyChangedChainEvent;
+
+/*
+ *************************************************************************************
+ * Overridable serialization is experimental, not supported and use at your own risk *
+ *************************************************************************************
+ * Scope responsible to control overridable serialization logic.
+ */
+struct FOverridableSerializationLogic
+{
+public:
+
+	/**
+	 * Call to enable overridable serialization and to set the overridden properties of the current serialized object
+	 * Note this is not re-entrant and it stores information in a thread local storage
+	 * @param InOverriddenProperties of the current serializing object */
+	FORCEINLINE static void Enable(FOverriddenPropertySet* InOverriddenProperties)
+	{
+		checkf(!bUseOverridableSerialization, TEXT("Nobody should use this method if overridable serialization is already enabled"));
+		bUseOverridableSerialization = true;
+		OverriddenProperties = InOverriddenProperties;
+	}
+
+	/**
+	 * Call to disable overridable serialization
+	 * Note this is not re-entrant and it stores information in a thread local storage */
+	FORCEINLINE static void Disable()
+	{
+		checkf(bUseOverridableSerialization, TEXT("Expecting overridable serialization to be already enabled"));
+		OverriddenProperties = nullptr;
+		bUseOverridableSerialization = false;
+	}
+
+	/**
+	 * Called during the serialization of an object to know to know if it should do overridden serialization logic
+	 * @return true if the overridable serialization is enabled on the current serializing object */
+	FORCEINLINE static bool IsEnabled()
+	{
+		return bUseOverridableSerialization;
+	}
+
+	/**
+	 * Call during the serialization of an object to get its overriden properties
+	 * Note: Expects the current serialized object to use overridable serialization
+	 * Note this is not re-entrant and it stores information in a thread local storage
+	 * @return the overridden properties of the current object being serialized */
+	FORCEINLINE static FOverriddenPropertySet* GetOverriddenProperties()
+	{
+		return OverriddenProperties;
+	}
+
+	/**
+	 * Retrieve from the Archive and the current property the overridden property operation to know if it has to be serialized or not
+	 * @param Ar currently being used to serialize the current object (will be used to retrieve the current property serialized path)
+	 * @param Property the property about to be serialized, can be null
+	 * @return the overridden property operation */
+	static EOverriddenPropertyOperation GetOverriddenPropertyOperation(const FArchive& Ar, FProperty* Property = nullptr);
+
+private:
+	static thread_local bool bUseOverridableSerialization;
+	static thread_local FOverriddenPropertySet* OverriddenProperties;
+};
+
+/*
+ *************************************************************************************
+ * Overridable serialization is experimental, not supported and use at your own risk *
+ *************************************************************************************
+ * Scope responsible for enabling/disabling the overridable serialization from the parameters
+*/
+struct FEnableOverridableSerializationScope
+{
+	FEnableOverridableSerializationScope(bool bEnableOverridableSerialization, FOverriddenPropertySet* OverriddenProperties);
+	~FEnableOverridableSerializationScope();
+
+protected:
+	bool bOverridableSerializationEnabled = false;
+};
+
+/*
+ * Override operation type for each property node
+ */
+UENUM()
+enum class EOverriddenPropertyOperation : uint8
+{
+	None =	0,	/* no overridden operation was recorded on this property  */
+	Modified,	/* some sub property has recorded overridden operation */
+	Replace,	/* everything has been overridden from this property down to every sub property/sub object*/
+	Add,		/* this element was added in the container */
+	Remove,		/* this element was removed from the container */
+};
+
+/*
+ *************************************************************************************
+ * Overridable serialization is experimental, not supported and use at your own risk *
+ *************************************************************************************
+ * Overridden property information node, there will be one per overriden property/subojects
+ *
+ */
+USTRUCT()
+struct FOverriddenPropertyNode
+{
+	GENERATED_BODY()
+
+	FOverriddenPropertyNode(FName InPropertyPath = NAME_None)
+		: PropertyPath(InPropertyPath)
+	{}
+
+	UPROPERTY()
+	FName PropertyPath;
+
+	UPROPERTY()
+	EOverriddenPropertyOperation Operation = EOverriddenPropertyOperation::None;
+
+	UPROPERTY()
+	TMap<FName,FName> SubPropertyNodeKeys;
+
+	bool operator==(const FOverriddenPropertyNode& Other) const
+	{
+		return PropertyPath == Other.PropertyPath;
+	}
+
+	friend uint32 GetTypeHash(const FOverriddenPropertyNode& Node)
+	{
+		return GetTypeHash(Node.PropertyPath);
+	}
+};
+
+
+
+/*
+ * Property change notification type mapping the Pre/PostEditChange callbacks
+ */
+enum class EPropertyNotificationType : uint8
+{
+	PreEdit,
+	PostEdit
+};
+
+/*
+ *************************************************************************************
+ * Overridable serialization is experimental, not supported and use at your own risk *
+ *************************************************************************************
+ * Structure holding and tracking overridden properties of an UObject
+ */
+USTRUCT()
+struct COREUOBJECT_API FOverriddenPropertySet
+{
+	GENERATED_BODY()
+
+public:
+
+	FOverriddenPropertySet() = default;
+	FOverriddenPropertySet(UObject& InOwner)
+		: Owner(&InOwner)
+	{}
+
+	/**
+	 * Handling and storing modification on a property of an object
+	 * @param Notification type either pre/post property overridden
+	 * @param PropertyEvent information about the type of change
+	 * @param PropertyNode leading to the property that is changing
+	 * @param Data memory of the current property */
+	void NotifyPropertyChange(const EPropertyNotificationType Notification, const FPropertyChangedEvent& PropertyEvent, const FEditPropertyChain::TDoubleLinkedListNode* PropertyNode, const void* Data);
+
+	/**
+	 * Retrieve the overridable operation from the specified the serialized property chain and the specified property
+	 * @param CurrentPropertyChain leading to the property being serialized if any
+	 * @param Property being serialized if any, otherwise will fallback on the last property of the chain
+	 * @return the current type of override operation on the property */
+	EOverriddenPropertyOperation GetOverriddenPropertyOperation(const FArchiveSerializedPropertyChain* CurrentPropertyChain, FProperty* Property) const;
+
+	/**
+	 * Setup the overridable operation of the current property from the serialized property chain and the specified property
+	 * @param Operation to set on the proto set for this property
+	 * @param CurrentPropertyChain leading to the property being serialized if any
+	 * @param Property being serialized if any, otherwise will fallback on the last property of the chain
+	 * @return the node containing the information of the overridden property */
+	FOverriddenPropertyNode* SetOverriddenPropertyOperation(EOverriddenPropertyOperation Operation, const FArchiveSerializedPropertyChain* CurrentPropertyChain, FProperty* Property);
+
+	/**
+	 * Retrieve the overridden property node from the serialized property chain
+	 * @param CurrentPropertyChain leading to the property being serialized.
+	 * @return the node containing the information of the overridden property */
+	const FOverriddenPropertyNode* GetOverriddenPropertyNode(const FArchiveSerializedPropertyChain* CurrentPropertyChain) const;
+
+	/**
+	 * Retrieve the overridable operation given the property key
+	 * @param PropertyKey that uniquely identify the property within the object 
+	 * @return the current type of override operation on the property */
+	EOverriddenPropertyOperation GetSubPropertyOperation(FName PropertyKey) const;
+
+	/**
+	 * Set the overridable operation of a sub property of the specified node.
+	 * @param Operation to set for this property
+	 * @param Node from where the sub property is owned by
+	 * @param PropID the ID of the sub property 
+	 * @return the node to the sub property */
+	FOverriddenPropertyNode* SetSubPropertyOperation(EOverriddenPropertyOperation Operation, FOverriddenPropertyNode& Node, FName PropID);
+
+	/**
+	 * Resets all overrides of the object */
+	void Reset();
+
+protected:
+
+	void NotifyPropertyChange(FOverriddenPropertyNode& ParentPropertyNode, const EPropertyNotificationType Notification, const FPropertyChangedEvent& PropertyEvent, const FEditPropertyChain::TDoubleLinkedListNode* PropertyNode, const void* Data);
+	EOverriddenPropertyOperation GetOverriddenPropertyOperation(const FOverriddenPropertyNode& ParentPropertyNode, const FArchiveSerializedPropertyChain* CurrentPropertyChain, FProperty* Property) const;
+	FOverriddenPropertyNode* SetOverriddenPropertyOperation(EOverriddenPropertyOperation Operation, FOverriddenPropertyNode& ParentPropertyNode, const FArchiveSerializedPropertyChain* CurrentPropertyChain, FProperty* Property);
+
+	const FOverriddenPropertyNode* GetOverriddenPropertyNode(const FOverriddenPropertyNode& ParentPropertyNode, const FArchiveSerializedPropertyChain* CurrentPropertyChain) const;
+
+	void RemoveOverriddenSubProperties(FOverriddenPropertyNode& PropertyNode);
+
+	FOverriddenPropertyNode& FindOrAddNode(FOverriddenPropertyNode& ParentPropertyNode, FName PropID);
+
+private:
+	UPROPERTY()
+	TObjectPtr<UObject> Owner = nullptr;
+
+	UPROPERTY()
+	TSet<FOverriddenPropertyNode> OverriddenPropertyNodes;
+
+	static inline FName RootNodeName = TEXT("root");
+};
diff --git a/Engine/Source/Runtime/CoreUObject/Public/UObject/PropertyTag.h b/Engine/Source/Runtime/CoreUObject/Public/UObject/PropertyTag.h
--- a/Engine/Source/Runtime/CoreUObject/Public/UObject/PropertyTag.h
+++ b/Engine/Source/Runtime/CoreUObject/Public/UObject/PropertyTag.h
@@ -11,6 +11,7 @@
 #include "Serialization/StructuredArchive.h"
 #include "UObject/NameTypes.h"
 
+enum class EOverriddenPropertyOperation : uint8;
 class FArchive;
 class FProperty;
 
@@ -25,7 +26,7 @@ enum class EPropertyTagExtension : uint8
 
 	////////////////////////////////////////////////
 	// First extension group
-	OverridableOperation		= 0x02,
+	OverridableInformation		= 0x02,
 
 	//
 	// Add more extension for the first group here
@@ -55,9 +56,11 @@ struct FPropertyTag
 	FGuid	StructGuid;
 	uint8	HasPropertyGuid = 0;
 	FGuid	PropertyGuid;
+	EOverriddenPropertyOperation OverrideOperation; // Overridable serialization state reconstruction 
+	bool	bExperimentalOverridableLogic = false; // Remember if property had CPF_ExperimentalOverridableLogic when saved
 
 	// Constructors.
-	FPropertyTag() {}
+	FPropertyTag();
 	FPropertyTag( FArchive& InSaveAr, FProperty* Property, int32 InIndex, uint8* Value, const uint8* Defaults );
 
 	// Set optional property guid
@@ -72,6 +75,29 @@ struct FPropertyTag
 	void SerializeTaggedProperty(FStructuredArchive::FSlot Slot, FProperty* Property, uint8* Value, const uint8* Defaults) const;
 };
 
+struct FPropertyTagScope
+{
+	FPropertyTagScope(const FPropertyTag* InCurrentPropertyTag)
+	: PropertyTagToRestore(CurrentPropertyTag)
+	{
+		CurrentPropertyTag = InCurrentPropertyTag;
+	}
+
+	~FPropertyTagScope()
+	{
+		CurrentPropertyTag = PropertyTagToRestore;
+	}
+
+	static FORCEINLINE const FPropertyTag* GetCurrentPropertyTag()
+	{
+		return CurrentPropertyTag;
+	}
+private:
+	const FPropertyTag* PropertyTagToRestore;
+
+	static thread_local const FPropertyTag* CurrentPropertyTag;
+};
+
 #if UE_ENABLE_INCLUDE_ORDER_DEPRECATED_IN_5_2
 #include "CoreMinimal.h"
 #endif
diff --git a/Engine/Source/Runtime/Engine/Private/InstancedReferenceSubobjectHelper.cpp b/Engine/Source/Runtime/Engine/Private/InstancedReferenceSubobjectHelper.cpp
deleted file mode 100644
diff --git a/Engine/Source/Runtime/Engine/Public/InstancedReferenceSubobjectHelper.h b/Engine/Source/Runtime/Engine/Public/InstancedReferenceSubobjectHelper.h
deleted file mode 100644
