Dynamically create properties in FPropertyBag when loading by tag

diff --git a/Engine/Source/Runtime/CoreUObject/Private/UObject/PropertyBag.cpp b/Engine/Source/Runtime/CoreUObject/Private/UObject/PropertyBag.cpp
--- a/Engine/Source/Runtime/CoreUObject/Private/UObject/PropertyBag.cpp
+++ b/Engine/Source/Runtime/CoreUObject/Private/UObject/PropertyBag.cpp
@@ -108,6 +108,18 @@ void FPropertyBag::LoadPropertyByTag(const FPropertyPathName& Path, const FPrope
 		return;
 	}
 
+	// Construct a property from the tag and try to use it to serialize the value.
+	FField* Field = FField::Construct(Tag.Type, {}, Tag.Name, RF_NoFlags);
+	if (FProperty* Property = CastField<FProperty>(Field); Property && Property->LoadFromTag(Value.Tag))
+	{
+		Value.bOwnsProperty = true;
+		Value.Tag.Prop = Property;
+		Value.AllocateAndInitializeValue();
+		Tag.SerializeTaggedProperty(ValueSlot, Property, (uint8*)Value.Data, nullptr);
+		return;
+	}
+	delete Field;
+
 	// Fall back to loading the serialized value.
 	// Persisting this serialized value will require capturing version information from the archive.
 	Value.AllocateAndInitializeValue();
