Deprecate GetCurrentTag/InvalidateTag and remove all calls as the feature was unreliable and is no longer used by TSoftObjectPtr as of 5.3

diff --git a/Engine/Source/Developer/AssetTools/Private/AssetRenameManager.cpp b/Engine/Source/Developer/AssetTools/Private/AssetRenameManager.cpp
--- a/Engine/Source/Developer/AssetTools/Private/AssetRenameManager.cpp
+++ b/Engine/Source/Developer/AssetTools/Private/AssetRenameManager.cpp
@@ -1244,9 +1244,6 @@ void FAssetRenameManager::RenameReferencingSoftObjectPaths(const TArray<UPackage
 			}
 		}
 	}
-
-	// Invalidate the soft object tag as we have created new valid paths
-	FSoftObjectPath::InvalidateTag();
 }
 
 void FAssetRenameManager::OnMarkPackageDirty(UPackage* Pkg, bool bWasDirty)
diff --git a/Engine/Source/Editor/GraphEditor/Private/KismetPins/SGraphPinDataTableRowName.cpp b/Engine/Source/Editor/GraphEditor/Private/KismetPins/SGraphPinDataTableRowName.cpp
--- a/Engine/Source/Editor/GraphEditor/Private/KismetPins/SGraphPinDataTableRowName.cpp
+++ b/Engine/Source/Editor/GraphEditor/Private/KismetPins/SGraphPinDataTableRowName.cpp
@@ -31,7 +31,6 @@ void SGraphPinDataTableRowName::PreChange(const UDataTable* Changed, FDataTableE
 
 void SGraphPinDataTableRowName::PostChange(const UDataTable* Changed, FDataTableEditorUtils::EDataTableChangeInfo Info)
 {
-	//FSoftObjectPath::InvalidateTag(); // Should be removed after UE-5615 is fixed
 	if (Changed && (Changed == DataTable.Get()) && (FDataTableEditorUtils::EDataTableChangeInfo::RowList == Info))
 	{
 		RefreshNameList();
diff --git a/Engine/Source/Runtime/CoreUObject/Private/Serialization/AsyncLoading.cpp b/Engine/Source/Runtime/CoreUObject/Private/Serialization/AsyncLoading.cpp
--- a/Engine/Source/Runtime/CoreUObject/Private/Serialization/AsyncLoading.cpp
+++ b/Engine/Source/Runtime/CoreUObject/Private/Serialization/AsyncLoading.cpp
@@ -7162,9 +7162,6 @@ EAsyncPackageState::Type FAsyncPackage::PostLoadDeferredObjects(double InTickSta
 			}
 			::IsTimeLimitExceeded(InTickStartTime, bInUseTimeLimit, InOutTimeLimit, LastTypeOfWorkPerformed, LastObjectWorkWasPerformedOn);
 		}
-
-		FSoftObjectPath::InvalidateTag();
-		FUniqueObjectGuid::InvalidateTag();
 	}
 
 	--ReentryCount;
diff --git a/Engine/Source/Runtime/CoreUObject/Private/Serialization/AsyncLoading2.cpp b/Engine/Source/Runtime/CoreUObject/Private/Serialization/AsyncLoading2.cpp
--- a/Engine/Source/Runtime/CoreUObject/Private/Serialization/AsyncLoading2.cpp
+++ b/Engine/Source/Runtime/CoreUObject/Private/Serialization/AsyncLoading2.cpp
@@ -7320,12 +7320,7 @@ EEventLoadNodeExecutionResult FAsyncPackage2::Event_DeferredPostLoadExportBundle
 	check(InExportBundleIndex < Package->Data.TotalExportBundleCount);
 	EEventLoadNodeExecutionResult LoadingState = EEventLoadNodeExecutionResult::Complete;
 
-	if (Package->bLoadHasFailed)
-	{
-		FSoftObjectPath::InvalidateTag();
-		FUniqueObjectGuid::InvalidateTag();
-	}
-	else
+	if (!Package->bLoadHasFailed)
 	{
 		SCOPED_LOADTIMER(PostLoadDeferredObjectsTime);
 		TRACE_LOADTIME_POSTLOAD_SCOPE;
@@ -7741,9 +7736,6 @@ EAsyncPackageState::Type FAsyncLoadingThread2::ProcessLoadedPackagesFromGameThre
 				}
 			}
 
-			FSoftObjectPath::InvalidateTag();
-			FUniqueObjectGuid::InvalidateTag();
-
 			// push stats so that we don't overflow number of tags per thread during blocking loading
 			LLM_PUSH_STATS_FOR_ASSET_TAGS();
 
diff --git a/Engine/Source/Runtime/CoreUObject/Private/UObject/LazyObjectPtr.cpp b/Engine/Source/Runtime/CoreUObject/Private/UObject/LazyObjectPtr.cpp
--- a/Engine/Source/Runtime/CoreUObject/Private/UObject/LazyObjectPtr.cpp
+++ b/Engine/Source/Runtime/CoreUObject/Private/UObject/LazyObjectPtr.cpp
@@ -178,7 +178,6 @@ void FLazyObjectPtr::PossiblySerializeObjectGuid(UObject *Object, FStructuredArc
 					{
 						GuidAnnotation.AddAnnotation(Object, Guid);
 					}
-					FUniqueObjectGuid::InvalidateTag();
 				}
 			}
 		}
diff --git a/Engine/Source/Runtime/CoreUObject/Private/UObject/Obj.cpp b/Engine/Source/Runtime/CoreUObject/Private/UObject/Obj.cpp
--- a/Engine/Source/Runtime/CoreUObject/Private/UObject/Obj.cpp
+++ b/Engine/Source/Runtime/CoreUObject/Private/UObject/Obj.cpp
@@ -1678,12 +1678,6 @@ void UObject::Serialize(FStructuredArchive::FRecord Record)
 		// Serialize a GUID if this object has one mapped to it
 		FLazyObjectPtr::PossiblySerializeObjectGuid(this, Record);
 
-		// Invalidate asset pointer caches when loading a new object
-		if (UnderlyingArchive.IsLoading())
-		{
-			FSoftObjectPath::InvalidateTag();
-		}
-
 		// Keep track of sparse class data for undo/redo
 		if (UnderlyingArchive.IsTransacting() && HasAnyFlags(RF_ClassDefaultObject))
 		{
diff --git a/Engine/Source/Runtime/CoreUObject/Private/UObject/UObjectGlobals.cpp b/Engine/Source/Runtime/CoreUObject/Private/UObject/UObjectGlobals.cpp
--- a/Engine/Source/Runtime/CoreUObject/Private/UObject/UObjectGlobals.cpp
+++ b/Engine/Source/Runtime/CoreUObject/Private/UObject/UObjectGlobals.cpp
@@ -2388,9 +2388,6 @@ void EndLoad(FUObjectSerializeContext* LoadContext, TArray<UPackage*>* OutLoaded
 		check(GGameThreadLoadCounter >= 0);
 	}
 
-	// Loaded new objects, so allow reaccessing asset ptrs
-	FSoftObjectPath::InvalidateTag();
-
 #if WITH_EDITOR
 	// Now call asset loaded callbacks for anything that was loaded. We do this at the very end so any nested objects will load properly
 	// Useful for updating UI such as ContentBrowser's loaded status.
diff --git a/Engine/Source/Runtime/CoreUObject/Public/UObject/LazyObjectPtr.h b/Engine/Source/Runtime/CoreUObject/Public/UObject/LazyObjectPtr.h
--- a/Engine/Source/Runtime/CoreUObject/Public/UObject/LazyObjectPtr.h
+++ b/Engine/Source/Runtime/CoreUObject/Public/UObject/LazyObjectPtr.h
@@ -103,11 +103,12 @@ struct FUniqueObjectGuid
 		Slot << ObjectGuid.Guid;
 	}
 
-	/** These will be deprecated in a followup change as they are no longer in use */
+	UE_DEPRECATED(5.4, "The current object tag is no longer used by TSoftObjectPtr, you can remove all calls")
 	static int32 GetCurrentTag()
 	{
 		return 0;
 	}
+	UE_DEPRECATED(5.4, "The current object tag is no longer used by TSoftObjectPtr, you can remove all calls")
 	static int32 InvalidateTag()
 	{
 		return 0;
diff --git a/Engine/Source/Runtime/CoreUObject/Public/UObject/SoftObjectPath.h b/Engine/Source/Runtime/CoreUObject/Public/UObject/SoftObjectPath.h
--- a/Engine/Source/Runtime/CoreUObject/Public/UObject/SoftObjectPath.h
+++ b/Engine/Source/Runtime/CoreUObject/Public/UObject/SoftObjectPath.h
@@ -318,11 +318,12 @@ PRAGMA_ENABLE_DEPRECATION_WARNINGS;
 		return Hash;
 	}
 
-	/** These will be deprecated in a followup change as they are no longer used by FSoftObjectPtr */
+	UE_DEPRECATED(5.4, "The current object tag is no longer used by TSoftObjectPtr, you can remove all calls")
 	static int32 GetCurrentTag()
 	{
 		return 0;
 	}
+	UE_DEPRECATED(5.4, "The current object tag is no longer used by TSoftObjectPtr, you can remove all calls")
 	static int32 InvalidateTag()
 	{
 		return 0;
