Workaround for MSVC 17.8 bug.

diff --git a/Engine/Source/Runtime/CoreUObject/Public/UObject/ObjectPtr.h b/Engine/Source/Runtime/CoreUObject/Public/UObject/ObjectPtr.h
--- a/Engine/Source/Runtime/CoreUObject/Public/UObject/ObjectPtr.h
+++ b/Engine/Source/Runtime/CoreUObject/Public/UObject/ObjectPtr.h
@@ -367,10 +367,11 @@ namespace ObjectPtr_Private
 	/** Perform shallow equality check between a TObjectPtr and another (non TObjectPtr) type that we can coerce to a pointer. */
 	template <
 		typename T,
-		typename U
+		typename U,
+		decltype(CoerceToPointer<T>(std::declval<U>()) == std::declval<const T*>())* = nullptr
 		UE_REQUIRES(!TIsTObjectPtr_V<U>)
 	>
-	auto IsObjectPtrEqual(const TObjectPtr<T>& Ptr, const U& Other) -> decltype(CoerceToPointer<T>(Other) == std::declval<const T*>())
+	bool IsObjectPtrEqual(const TObjectPtr<T>& Ptr, const U& Other)
 	{
 		// This function deliberately avoids the tracking code path as we are only doing
 		// a shallow pointer comparison.
