Remove old hack that would ensure if a raw memory address was reused when creating temporary maps in the editor

diff --git a/Engine/Source/Editor/UnrealEd/Private/EditorServer.cpp b/Engine/Source/Editor/UnrealEd/Private/EditorServer.cpp
--- a/Engine/Source/Editor/UnrealEd/Private/EditorServer.cpp
+++ b/Engine/Source/Editor/UnrealEd/Private/EditorServer.cpp
@@ -2615,9 +2615,6 @@ bool UEditorEngine::Map_Load(const TCHAR* Str, FOutputDevice& Ar)
 				Context.SetCurrentWorld(World);
 				GWorld = World;
 
-				// UE-21181 - Tracking where the loaded editor level's package gets flagged as a PIE object
-				UPackage::EditorPackage = WorldPackage;
-
 				World->WorldType = EWorldType::Editor;
 
 				Context.World()->PersistentLevel->HandleLegacyMapBuildData();
diff --git a/Engine/Source/Runtime/CoreUObject/Private/UObject/Package.cpp b/Engine/Source/Runtime/CoreUObject/Private/UObject/Package.cpp
--- a/Engine/Source/Runtime/CoreUObject/Private/UObject/Package.cpp
+++ b/Engine/Source/Runtime/CoreUObject/Private/UObject/Package.cpp
@@ -381,16 +381,6 @@ bool UPackage::IsPostLoadThreadSafe() const
 	return true;
 }
 
-// UE-21181 - Tracking where the loaded editor level's package gets flagged as a PIE object
-#if WITH_EDITOR
-UPackage* UPackage::EditorPackage = nullptr;
-void UPackage::SetPackageFlagsTo( uint32 NewFlags )
-{
-	PackageFlagsPrivate = NewFlags;
-	ensure(((NewFlags & PKG_PlayInEditor) == 0) || (this != EditorPackage));
-}
-#endif
-
 #if WITH_EDITORONLY_DATA
 void FixupPackageEditorOnlyFlag(FName PackageThatGotEditorOnlyFlagCleared, bool bRecursive);
 
diff --git a/Engine/Source/Runtime/CoreUObject/Public/UObject/Package.h b/Engine/Source/Runtime/CoreUObject/Public/UObject/Package.h
--- a/Engine/Source/Runtime/CoreUObject/Public/UObject/Package.h
+++ b/Engine/Source/Runtime/CoreUObject/Public/UObject/Package.h
@@ -639,11 +639,6 @@ public:
 		return HasAnyPackageFlags(PKG_RequiresLocalizationGather);
 	}
 
-	// UE-21181 - trying to track when a flag gets set on a package due to PIE
-#if WITH_EDITOR
-	static UPackage* EditorPackage;
-	void SetPackageFlagsTo( uint32 NewFlags );
-#else
 	/**
 	* Sets all package flags to the specified values.
 	*
@@ -653,7 +648,6 @@ public:
 	{
 		PackageFlagsPrivate = NewFlags;
 	}
-#endif
 
 	/**
 	* Set the specified flags to true. Does not affect any other flags.
