Support for Incremental Gather Unreachable Objects

Split unreachable clusters destruction to DissolveUnreachableClusters (non-incremental as it's very fast)
Flipping MaybeUnreachable -> Unreachable flag is now a single Swap() function call (MaybeUnreachable / Unreachable flags are deprecated and replaced with interchangable Reachability flags)
GatherUnreachableObjects will now only gather unreachable objects for NotifyUnreachableObjects call and perf gains during unhashing and purging all objects.
Enable with gc.AllowIncrementalGather=1

diff --git a/Engine/Plugins/EnhancedInput/Source/InputEditor/Private/InputEditorModule.cpp b/Engine/Plugins/EnhancedInput/Source/InputEditor/Private/InputEditorModule.cpp
--- a/Engine/Plugins/EnhancedInput/Source/InputEditor/Private/InputEditorModule.cpp
+++ b/Engine/Plugins/EnhancedInput/Source/InputEditor/Private/InputEditorModule.cpp
@@ -507,7 +507,7 @@ void FInputEditorModule::AutoUpgradeDefaultInputClasses()
 const UPlayerMappableKeySettings* FInputEditorModule::FindMappingByName(const FName InName)
 {
 	// Instanced objects that have been deleted will have an outer of the transient package, and we don't want to include them
-	for (TObjectIterator<UPlayerMappableKeySettings> Itr(RF_ClassDefaultObject | RF_Transient, true, EInternalObjectFlags::Garbage | EInternalObjectFlags::Unreachable); Itr; ++Itr)
+	for (TObjectIterator<UPlayerMappableKeySettings> Itr(RF_ClassDefaultObject | RF_Transient, true, EInternalObjectFlags::Garbage | UE::GC::GUnreachableObjectFlag); Itr; ++Itr)
 	{
 		const UPlayerMappableKeySettings* Settings = *Itr;
 		if (Settings && Settings->GetOuter() != GetTransientPackage() && Settings->Name == InName)
diff --git a/Engine/Source/Editor/UnrealEd/Private/PackageTools.cpp b/Engine/Source/Editor/UnrealEd/Private/PackageTools.cpp
--- a/Engine/Source/Editor/UnrealEd/Private/PackageTools.cpp
+++ b/Engine/Source/Editor/UnrealEd/Private/PackageTools.cpp
@@ -156,7 +156,7 @@ UPackageTools::UPackageTools(const FObjectInitializer& ObjectInitializer)
 						Object->SetFlags(RF_Standalone);
 					}
 					return true;
-				}, true, RF_NoFlags, EInternalObjectFlags::Unreachable);
+				}, true, RF_NoFlags, UE::GC::GUnreachableObjectFlag);
 			}
 		}
 	}
diff --git a/Engine/Source/Runtime/CoreUObject/Private/Serialization/AsyncLoading.cpp b/Engine/Source/Runtime/CoreUObject/Private/Serialization/AsyncLoading.cpp
--- a/Engine/Source/Runtime/CoreUObject/Private/Serialization/AsyncLoading.cpp
+++ b/Engine/Source/Runtime/CoreUObject/Private/Serialization/AsyncLoading.cpp
@@ -3154,7 +3154,7 @@ UObject* FAsyncPackage::EventDrivenIndexToObject(FPackageIndex Index, bool bChec
 	}
 	if (Result)
 	{
-		UE_CLOG(Result->HasAnyInternalFlags(EInternalObjectFlags::Unreachable), LogStreaming, Fatal, TEXT("Returning an object  (%s) from EventDrivenIndexToObject that is unreachable."), *Result->GetFullName());
+		UE_CLOG(Result->IsUnreachable(), LogStreaming, Fatal, TEXT("Returning an object  (%s) from EventDrivenIndexToObject that is unreachable."), *Result->GetFullName());
 		checkSlow(ReferencedObjects.Contains(Result));
 	}
 	return Result;
@@ -5727,7 +5727,7 @@ void FAsyncPackage::AddObjectReference(UObject* InObject)
 				ReferencedObjects.Add(InObject);
 			}
 		}
-		UE_CLOG(InObject->HasAnyInternalFlags(EInternalObjectFlags::Unreachable), LogStreaming, Fatal, TEXT("Trying to add an unreachable object %s to FAsyncPackage %s referenced objects list."), *InObject->GetFullName(), *GetPackageName().ToString());
+		UE_CLOG(InObject->IsUnreachable(), LogStreaming, Fatal, TEXT("Trying to add an unreachable object %s to FAsyncPackage %s referenced objects list."), *InObject->GetFullName(), *GetPackageName().ToString());
 	}
 }
 
diff --git a/Engine/Source/Runtime/CoreUObject/Private/Serialization/AsyncLoading2.cpp b/Engine/Source/Runtime/CoreUObject/Private/Serialization/AsyncLoading2.cpp
--- a/Engine/Source/Runtime/CoreUObject/Private/Serialization/AsyncLoading2.cpp
+++ b/Engine/Source/Runtime/CoreUObject/Private/Serialization/AsyncLoading2.cpp
@@ -765,17 +765,29 @@ private:
 			return nullptr;
 		}
 
-		void PinForGC()
+		[[nodiscard]] bool PinForGC(TArray<int32>& OutUnreachableObjectIndices)
 		{
-			for (int32 ObjectIndex : GetValues())
+			OutUnreachableObjectIndices.Reset();
+			for (int32& ObjectIndex : GetValues())
 			{
 				if (ObjectIndex >= 0)
 				{
-					UObject* Object = static_cast<UObject*>(GUObjectArray.IndexToObject(ObjectIndex)->Object);
-					checkf(!Object->HasAnyInternalFlags(EInternalObjectFlags::LoaderImport), TEXT("%s"), *Object->GetFullName());
-					Object->SetInternalFlags(EInternalObjectFlags::LoaderImport);
+					FUObjectItem* ObjectItem = GUObjectArray.IndexToObject(ObjectIndex);
+					if (!ObjectItem->IsUnreachable())
+					{
+						UObject* Object = static_cast<UObject*>(ObjectItem->Object);
+						checkf(!ObjectItem->HasAnyFlags(EInternalObjectFlags::LoaderImport), TEXT("%s"), *Object->GetFullName());
+						ObjectItem->SetFlags(EInternalObjectFlags::LoaderImport);
+					}
+					else
+					{
+						OutUnreachableObjectIndices.Reserve(Count);
+						OutUnreachableObjectIndices.Add(ObjectIndex);
+						ObjectIndex = -1;
+					}
 				}
 			}
+			return OutUnreachableObjectIndices.Num() == 0;
 		}
 
 		void UnpinForGC()
@@ -887,7 +899,7 @@ public:
 		}
 	}
 
-	void RemoveCompletedRenamedPackage()
+	void RemoveUnreferencedObsoletePackage()
 	{
 		check(RefCount == 0);
 		*this = FLoadedPackageRef();
@@ -969,7 +981,7 @@ public:
 		return PublicExportMap.Find(ExportHash);
 	}
 
-	void PinPublicExportsForGC()
+	void PinPublicExportsForGC(TArray<int32>& OutUnreachableObjectIndices)
 	{
 		UPackage* Package = GetPackage();
 		UE_ASYNC_UPACKAGE_DEBUG(Package);
@@ -978,7 +990,7 @@ public:
 		{
 			return;
 		}
-		PublicExportMap.PinForGC();
+		bAreAllPublicExportsLoaded = PublicExportMap.PinForGC(OutUnreachableObjectIndices);
 		checkf(!Package->HasAnyInternalFlags(EInternalObjectFlags::LoaderImport), TEXT("%s"), *Package->GetFullName());
 		Package->SetInternalFlags(EInternalObjectFlags::LoaderImport);
 	}
@@ -1116,19 +1128,24 @@ public:
 #endif
 			if (UPackage* Package = PackageRef.GetPackage())
 			{
-				if (PackageRef.GetOriginalPackageName() == Package->GetFName())
+				if (Package->IsUnreachable() || PackageRef.GetOriginalPackageName() != Package->GetFName())
 				{
-					PackageRef.PinPublicExportsForGC();
-				}
-				else
-				{
-					UE_LOG(LogStreaming, Display,
+					UE_CLOG(!Package->IsUnreachable(), LogStreaming, Display,
 						TEXT("FGlobalImportStore:AddPackageRef: Dropping renamed package %s before reloading %s (0x%llX)"),
 						*Package->GetName(),
 						*PackageRef.GetOriginalPackageName().ToString(),
 						Package->GetPackageId().ValueForDebugging());
 
-					RemoveCompletedRenamedPackage(PackageRef);
+					RemoveUnreferencedObsoletePackage(PackageRef);
+				}
+				else
+				{
+					TArray<int32> UnreachableObjectIndices;
+					PackageRef.PinPublicExportsForGC(UnreachableObjectIndices);
+					for (int32 ObjectIndex : UnreachableObjectIndices)
+					{
+						ObjectIndexToPublicExport.Remove(ObjectIndex);
+					}
 				}
 			}
 		}
@@ -1187,7 +1204,7 @@ public:
 		}
 	}
 
-	void RemoveCompletedRenamedPackage(FLoadedPackageRef& PackageRef)
+	void RemoveUnreferencedObsoletePackage(FLoadedPackageRef& PackageRef)
 	{
 		UPackage* OldPackage = PackageRef.GetPackage();
 		UE_ASYNC_UPACKAGE_DEBUG(OldPackage);
@@ -1203,9 +1220,9 @@ public:
 				ObjectIndexToPublicExport.Remove(ObjectIndex);
 			}
 		}
-		PackageRef.RemoveCompletedRenamedPackage();
-		// Update PackageId to prevent us from removing our updated package ref from GC
-		OldPackage->SetPackageId(FPackageId::FromName(OldPackage->GetFName()));
+		PackageRef.RemoveUnreferencedObsoletePackage();
+		// Reset PackageId to prevent a double remove from GC NotifyUnreachableObjects
+		OldPackage->SetPackageId(FPackageId());
 	}
 
 	void ReplaceReferencedRenamedPackage(FLoadedPackageRef& PackageRef, UPackage* NewPackage)
@@ -5686,7 +5703,7 @@ bool FAsyncPackage2::CreateLinkerLoadExports(FAsyncLoadingThreadState2& ThreadSt
 		FExportObject& ExportObject = Data.Exports[ExportIndex];
 		if (UObject* Object = LinkerLoadState->Linker->CreateExport(ExportIndex))
 		{
-			checkf(!Object->HasAnyInternalFlags(EInternalObjectFlags::Unreachable), TEXT("Trying to store an unreachable object '%s' in the import store"), *Object->GetFullName());
+			checkf(!Object->IsUnreachable(), TEXT("Trying to store an unreachable object '%s' in the import store"), *Object->GetFullName());
 			ExportObject.Object = Object;
 			ExportObject.bWasFoundInMemory = true; // Make sure that the async flags are cleared in ClearConstructedObjects
 			EInternalObjectFlags FlagsToSet = EInternalObjectFlags::Async;
@@ -6249,7 +6266,7 @@ UObject* FAsyncPackage2::EventDrivenIndexToObject(const FAsyncPackageHeaderData&
 	}
 	if (Result)
 	{
-		UE_CLOG(Result->HasAnyInternalFlags(EInternalObjectFlags::Unreachable), LogStreaming, Fatal, TEXT("Returning an object  (%s) from EventDrivenIndexToObject that is unreachable."), *Result->GetFullName());
+		UE_CLOG(Result->IsUnreachable(), LogStreaming, Fatal, TEXT("Returning an object  (%s) from EventDrivenIndexToObject that is unreachable."), *Result->GetFullName());
 	}
 #endif
 	return Result;
diff --git a/Engine/Source/Runtime/CoreUObject/Private/UObject/GarbageCollection.cpp b/Engine/Source/Runtime/CoreUObject/Private/UObject/GarbageCollection.cpp
--- a/Engine/Source/Runtime/CoreUObject/Private/UObject/GarbageCollection.cpp
+++ b/Engine/Source/Runtime/CoreUObject/Private/UObject/GarbageCollection.cpp
@@ -291,6 +291,26 @@ static FAutoConsoleVariableRef CVarIncrementalReachabilityTimeLimit(
 	ECVF_Default
 );
 
+static int32 GAllowIncrementalGather = 0;
+static FAutoConsoleVariableRef CVarAllowIncrementalGather(
+	TEXT("gc.AllowIncrementalGather"),
+	GAllowIncrementalGather,
+	TEXT("Set to control incremental Gather Unreachable Objects (experimental)"),
+	ECVF_Default
+);
+
+#if UE_BUILD_SHIPPING
+static constexpr float GIncrementalGatherTimeLimit = 0.0f;
+#else
+static float GIncrementalGatherTimeLimit = 0.0f;
+static FAutoConsoleVariableRef CVarIncrementalGatherTimeLimit(
+	TEXT("gc.IncrementalGatherTimeLimit"),
+	GIncrementalGatherTimeLimit,
+	TEXT("Override Incremental Gather Time Limit (in seconds)."),
+	ECVF_Default
+);
+#endif
+
 #if VERIFY_DISREGARD_GC_ASSUMPTIONS
 static int32 GVerifyNoUnreachableObjects = 0;
 static FAutoConsoleVariableRef CVarVerifyNoUnreachableObjects(
@@ -524,17 +544,34 @@ namespace UE::GC
 	{
 		Obj->MarkAsReachable();
 	}
+
+	EGCOptions GetReferenceCollectorOptions(bool bPerformFullPurge);
+	EGatherOptions GetObjectGatherOptions(bool bPerformFullPurge);
 }
 
-namespace UE::GC::Private
+namespace UE::GC
 {
-	/** List of objects marked as reachable by GC barrier (see UObject::MarkAsReachable()) */
-	static TExpandingChunkedList<UObject*> GReachableObjects;
-	/** List of FUObjectItems representing cluster root objects marker as reachable by GC barrier (see UObject::MarkAsReachable()) */
-	static TExpandingChunkedList<FUObjectItem*> GReachableClusters;
-	bool GIsIncrementalReachabilityPending = false;
+	/** EInternalObjectFlag value representing an unreachable object */
+	EInternalObjectFlags GUnreachableObjectFlag = EInternalObjectFlags::ReachabilityFlag0;
+
+	/** EInternalObjectFlag value representing a maybe unreachable object */
+	EInternalObjectFlags GMaybeUnreachableObjectFlag = EInternalObjectFlags::ReachabilityFlag1;
+
+	namespace Private
+	{
+		/** List of objects marked as reachable by GC barrier (see UObject::MarkAsReachable()) */
+		static TExpandingChunkedList<UObject*> GReachableObjects;
+		/** List of FUObjectItems representing cluster root objects marker as reachable by GC barrier (see UObject::MarkAsReachable()) */
+		static TExpandingChunkedList<FUObjectItem*> GReachableClusters;
+		bool GIsIncrementalReachabilityPending = false;
+
+		typedef TThreadedGather<FUObjectItem*> FGatherUnreachableObjectsState;
+		static FGatherUnreachableObjectsState GGatherUnreachableObjectsState;
+	}
 }
 
+static bool GatherUnreachableObjects(UE::GC::EGatherOptions Options, double TimeLimit = 0.0);
+
 /**
  * Helper class for destroying UObjects on a worker thread
  */
@@ -2792,7 +2829,7 @@ constexpr FORCEINLINE EKillable MayKill(EOrigin Origin, bool bAllowKill)
 // Return whether flag was cleared. Only thread-safe for concurrent clear, not concurrent set+clear. Don't use during mark phase.
 FORCEINLINE static bool ClearMaybeUnreachableInterlocked(int32& Flags)
 {
-	static constexpr int32 FlagToClear = int32(EInternalObjectFlags::MaybeUnreachable);
+	const int32 FlagToClear = int32(UE::GC::GMaybeUnreachableObjectFlag);
 	if (FPlatformAtomics::AtomicRead_Relaxed(&Flags) & FlagToClear)
 	{
 		int32 Old = FPlatformAtomics::InterlockedAnd(&Flags, ~FlagToClear);
@@ -3866,7 +3903,7 @@ public:
 
 					// We can't collect garbage during an async load operation and by now all unreachable objects should've been purged.
 					checkf(	bIsRerun ||
-							!ObjectItem->HasAnyFlags(EInternalObjectFlags::Unreachable|EInternalObjectFlags::PendingConstruction),
+							!ObjectItem->HasAnyFlags(UE::GC::GUnreachableObjectFlag | EInternalObjectFlags::PendingConstruction),
 							TEXT("Object: '%s' with ObjectFlags=0x%08x and InternalObjectFlags=0x%08x. ")
 							TEXT("State: IsEngineExitRequested=%d, GIsCriticalError=%d, GExitPurge=%d, GObjPurgeIsRequired=%d, GObjIncrementalPurgeIsInProgress=%d, GObjFinishDestroyHasBeenRoutedToAllObjects=%d, GGCObjectsPendingDestructionCount=%d"),
 							*Object->GetFullName(),
@@ -4254,6 +4291,9 @@ static bool IncrementalDestroyGarbage(bool bUseTimeLimit, double TimeLimit);
  */
 void IncrementalPurgeGarbage(bool bUseTimeLimit, double TimeLimit)
 {
+	using namespace UE::GC;
+	using namespace UE::GC::Private;
+
 	if (GExitPurge)
 	{
 		GObjPurgeIsRequired = true;
@@ -4308,7 +4348,7 @@ void IncrementalPurgeGarbage(bool bUseTimeLimit, double TimeLimit)
 		GCStartTime = FPlatformTime::Seconds();
 		bool bTimeLimitReached = false;
 
-		if (GUnrechableObjectIndex < GUnreachableObjects.Num())
+		if (IsIncrementalUnhashPending())
 		{
 			bTimeLimitReached = UnhashUnreachableObjects(bUseTimeLimit, TimeLimit);
 
@@ -4698,89 +4738,171 @@ static FAutoConsoleVariableRef CVarFlushStreamingOnGC(
 	ECVF_Default
 	);
 
-void GatherUnreachableObjects(bool bForceSingleThreaded)
+void DissolveUnreachableClusters(UE::GC::EGatherOptions Options)
 {
-	TRACE_CPUPROFILER_EVENT_SCOPE(GatherUnreachableObjects);
-	DECLARE_SCOPE_CYCLE_COUNTER(TEXT("CollectGarbageInternal.GatherUnreachableObjects"), STAT_CollectGarbageInternal_GatherUnreachableObjects, STATGROUP_GC);
+	const int32 InitialNumClusters = GUObjectClusters.GetNumAllocatedClusters();
+	if (!InitialNumClusters)
+	{
+		// Early out if clustering is disabled or there's no clusters
+		return;
+	}
 
-	const double StartTime = FPlatformTime::Seconds();
+	TRACE_CPUPROFILER_EVENT_SCOPE(DissolveUnreachableClusters);
+	DECLARE_SCOPE_CYCLE_COUNTER(TEXT("CollectGarbageInternal.DissolveUnreachableClusters"), STAT_CollectGarbageInternal_DissolveUnreachableClusters, STATGROUP_GC);
 
-	GUnreachableObjects.Reset();
-	GUnrechableObjectIndex = 0;
+	const double StartTime = FPlatformTime::Seconds();
 
-	int32 MaxNumberOfObjects = GUObjectArray.GetObjectArrayNum() - (GExitPurge ? 0 : GUObjectArray.GetFirstGCIndex());
-	int32 NumThreads = FMath::Max(1, FTaskGraphInterface::Get().GetNumWorkerThreads());
-	int32 NumberOfObjectsPerThread = (MaxNumberOfObjects / NumThreads) + 1;
+	using namespace UE::GC;
+	using namespace UE::GC::Private;
+	typedef TThreadedGather<int32> FGatherClustersState;
 
-	TArray<FUObjectItem*> ClusterItemsToDestroy;
-	int32 ClusterObjects = 0;
+	std::atomic<int32> TotalClusteredObjects = 0;
+	FGatherClustersState GatherClustersState;
+	TArray<FUObjectCluster>& ClusterArray = GUObjectClusters.GetClustersUnsafe();
+	
+	// StartGathering calculates the number of threads based on the number of objects but here the objects are actually clusters
+	// that contain many more objects than the number of clusters so we want to be able to process at least two clusters per thread	
+	const int32 NumThreads = !!(Options & EGatherOptions::Parallel) ? FMath::Min(GetNumCollectReferenceWorkers(), (ClusterArray.Num() + 1) / 2) : 1;
+	GatherClustersState.Start(Options, ClusterArray.Num(), /* FirstIndex = */ 0, NumThreads);
+	FGatherClustersState::FThreadIterators& ThreadIterators = GatherClustersState.GetThreadIterators();
 
-	// Iterate over all objects. Note that we iterate over the UObjectArray and usually check only internal flags which
-	// are part of the array so we don't suffer from cache misses as much as we would if we were to check ObjectFlags.
-	ParallelFor( TEXT("GC.GatherUnreachable"),NumThreads,1, [&ClusterItemsToDestroy, NumberOfObjectsPerThread, NumThreads, MaxNumberOfObjects](int32 ThreadIndex)
+	ParallelFor(TEXT("GC.DissolveUnreachableClusters"), GatherClustersState.NumWorkerThreads(), 1, [&ThreadIterators, &ClusterArray, &TotalClusteredObjects](int32 ThreadIndex)
 	{
-		TRACE_CPUPROFILER_EVENT_SCOPE(GatherUnreachableObjectsTask);
-		int32 FirstObjectIndex = ThreadIndex * NumberOfObjectsPerThread + (GExitPurge ? 0 : GUObjectArray.GetFirstGCIndex());
-		int32 NumObjects = (ThreadIndex < (NumThreads - 1)) ? NumberOfObjectsPerThread : (MaxNumberOfObjects - (NumThreads - 1) * NumberOfObjectsPerThread);
-		int32 LastObjectIndex = FMath::Min(GUObjectArray.GetObjectArrayNum() - 1, FirstObjectIndex + NumObjects - 1);
-		TArray<FUObjectItem*> ThisThreadUnreachableObjects;
-		TArray<FUObjectItem*> ThisThreadClusterItemsToDestroy;
+		TRACE_CPUPROFILER_EVENT_SCOPE(DissolveUnreachableClustersTask);
+		FGatherClustersState::FIterator& ThreadState = ThreadIterators[ThreadIndex];
+		int32 ThisThreadClusteredObjects = 0;
 
-		for (int32 ObjectIndex = FirstObjectIndex; ObjectIndex <= LastObjectIndex; ++ObjectIndex)
+		while (ThreadState.Index <= ThreadState.LastIndex)
 		{
-			FUObjectItem* ObjectItem = &GUObjectArray.GetObjectItemArrayUnsafe()[ObjectIndex];
-			if (ObjectItem->HasAnyFlags(EInternalObjectFlags::MaybeUnreachable | EInternalObjectFlags::Unreachable))
+			int32 ClusterIndex = ThreadState.Index++;
+			FUObjectCluster& Cluster = ClusterArray[ClusterIndex];
+			if (Cluster.RootIndex >= 0)
 			{
-				ObjectItem->ClearMaybeUnreachable();
-				ObjectItem->SetUnreachable();
-				ThisThreadUnreachableObjects.Add(ObjectItem);
-				if (ObjectItem->HasAnyFlags(EInternalObjectFlags::ClusterRoot))
+				FUObjectItem* RootItem = &GUObjectArray.GetObjectItemArrayUnsafe()[Cluster.RootIndex];
+				if (RootItem->IsUnreachable())
 				{
-					// We can't mark cluster objects as unreachable here as they may be currently being processed on another thread
-					ThisThreadClusterItemsToDestroy.Add(ObjectItem);
+					for (int32 ObjectIndex : Cluster.Objects)
+					{
+						FUObjectItem* ObjectItem = &GUObjectArray.GetObjectItemArrayUnsafe()[ObjectIndex];
+						if (!ObjectItem->HasAnyFlags(EInternalObjectFlags::ReachableInCluster))
+						{
+							ObjectItem->SetUnreachable();
+							ThisThreadClusteredObjects++;
+						}
+					}
+					ThreadState.Payload.Add(ClusterIndex);
 				}
 			}
 		}
-		if (ThisThreadUnreachableObjects.Num())
-		{
-			FScopeLock UnreachableObjectsLock(&GUnreachableObjectsCritical);
-			GUnreachableObjects.Append(ThisThreadUnreachableObjects);
-			ClusterItemsToDestroy.Append(ThisThreadClusterItemsToDestroy);
-		}
-	}, bForceSingleThreaded ? EParallelForFlags::ForceSingleThread : EParallelForFlags::None);
+		TotalClusteredObjects += ThisThreadClusteredObjects;
+	}, (GatherClustersState.NumWorkerThreads() == 1) ? EParallelForFlags::ForceSingleThread : EParallelForFlags::None);
+	
+	TArray<int32> ClustersToDestroy;
+	GatherClustersState.Finish(ClustersToDestroy);
 
+	// @todo: if GUObjectClusters.FreeCluster() was thread safe we could do this in parallel too
+	for (int32 ClusterIndex : ClustersToDestroy)
 	{
-		// @todo: if GUObjectClusters.FreeCluster() was thread safe we could do this in parallel too
-		for (FUObjectItem* ClusterRootItem : ClusterItemsToDestroy)
-		{
 #if UE_GCCLUSTER_VERBOSE_LOGGING
-			UE_LOG(LogGarbage, Log, TEXT("Destroying cluster (%d) %s"), ClusterRootItem->GetClusterIndex(), *static_cast<UObject*>(ClusterRootItem->Object)->GetFullName());
+		UE_LOG(LogGarbage, Log, TEXT("Destroying cluster (%d) %s"), ClusterRootItem->GetClusterIndex(), *static_cast<UObject*>(ClusterRootItem->Object)->GetFullName());
 #endif
-			ClusterRootItem->ClearFlags(EInternalObjectFlags::ClusterRoot);
+		GUObjectClusters.FreeCluster(ClusterIndex);
+	}
 
-			const int32 ClusterIndex = ClusterRootItem->GetClusterIndex();
-			FUObjectCluster& Cluster = GUObjectClusters[ClusterIndex];
-			for (int32 ClusterObjectIndex : Cluster.Objects)
+	UE_LOG(LogGarbage, Log, TEXT("%f ms for Dissolve Unreachable Clusters (%d/%d clusters dissolved containing %d cluster objects)"),
+		(FPlatformTime::Seconds() - StartTime) * 1000,
+		ClustersToDestroy.Num(),
+		InitialNumClusters,
+		TotalClusteredObjects.load(std::memory_order_relaxed));
+}
+
+bool GatherUnreachableObjects(UE::GC::EGatherOptions Options, double TimeLimit /*= 0.0*/)
+{
+	TRACE_CPUPROFILER_EVENT_SCOPE(GatherUnreachableObjects);
+	DECLARE_SCOPE_CYCLE_COUNTER(TEXT("CollectGarbageInternal.GatherUnreachableObjects"), STAT_CollectGarbageInternal_GatherUnreachableObjects, STATGROUP_GC);
+
+	using namespace UE::GC;
+	using namespace UE::GC::Private;
+
+	const double StartTime = FPlatformTime::Seconds();
+
+	GUnreachableObjects.Reset();
+	GUnrechableObjectIndex = 0;
+
+	FGatherUnreachableObjectsState::FThreadIterators& ThreadIterators = GGatherUnreachableObjectsState.GetThreadIterators();
+	if (!GGatherUnreachableObjectsState.IsSuspended())
+	{
+		const int32 FirstIndex = GExitPurge ? 0 : GUObjectArray.GetFirstGCIndex();
+		GGatherUnreachableObjectsState.Start(Options, GUObjectArray.GetObjectArrayNum(), FirstIndex);
+	}
+
+	std::atomic<int32> TimeLimitExceededFlag = 0;
+
+	// Iterate over all objects. Note that we iterate over the UObjectArray and usually check only internal flags which
+	// are part of the array so we don't suffer from cache misses as much as we would if we were to check ObjectFlags.
+	ParallelFor( TEXT("GC.GatherUnreachable"), GGatherUnreachableObjectsState.NumWorkerThreads(), 1, [&ThreadIterators, &StartTime, &TimeLimit, &TimeLimitExceededFlag](int32 ThreadIndex)
+	{
+		TRACE_CPUPROFILER_EVENT_SCOPE(GatherUnreachableObjectsTask);
+		constexpr int32 TimeLimitPollInterval = 10;
+		FTimeSlicer Timer(TimeLimitPollInterval, ThreadIndex * TimeLimitPollInterval, StartTime, TimeLimit, TimeLimitExceededFlag);
+		FGatherUnreachableObjectsState::FIterator& Iterator = ThreadIterators[ThreadIndex];
+		
+		while (Iterator.Index <= Iterator.LastIndex)
+		{
+			FUObjectItem* ObjectItem = &GUObjectArray.GetObjectItemArrayUnsafe()[Iterator.Index++];
+			if (ObjectItem->IsUnreachable())
 			{
-				FUObjectItem* ClusterObjectItem = GUObjectArray.IndexToObjectUnsafeForGC(ClusterObjectIndex);
-				ClusterObjectItem->SetOwnerIndex(0);
+				checkf(!ObjectItem->HasAnyFlags(EInternalObjectFlags::ClusterRoot), TEXT("Unreachable cluster root found. Unreachable clusters should have been dissolved in DissolveUnreachableClusters!"));
+				Iterator.Payload.Add(ObjectItem);
+			}
 
-				if (!ClusterObjectItem->HasAnyFlags(EInternalObjectFlags::ReachableInCluster))
-				{
-					ClusterObjectItem->SetUnreachable();
-					ClusterObjects++;
-					GUnreachableObjects.Add(ClusterObjectItem);
-				}
+			if (Timer.IsTimeLimitExceeded())
+			{
+				return;
 			}
-			GUObjectClusters.FreeCluster(ClusterIndex);
 		}
+	}, (GGatherUnreachableObjectsState.NumWorkerThreads() == 1) ? EParallelForFlags::ForceSingleThread : EParallelForFlags::None);
+
+	// Need to grab the stats now as Finish() will reset the global state but we still want to include Finish() execution time when logging
+	const int32 NumGathered = GGatherUnreachableObjectsState.NumGathered();
+	const int32 NumScanned = GGatherUnreachableObjectsState.NumScanned();
+	const int32 NumThreads = GGatherUnreachableObjectsState.NumWorkerThreads();
+	const bool bCompleted = !TimeLimitExceededFlag.load();
+	if (bCompleted)
+	{
+		GGatherUnreachableObjectsState.Finish(GUnreachableObjects);
 	}
 
-	UE_LOG(LogGarbage, Log, TEXT("%f ms for Gather Unreachable Objects (%d objects collected including %d cluster objects from %d clusters)"),
+	UE_LOG(LogGarbage, Log, TEXT("%f ms for Gather Unreachable Objects (%d objects collected / %d scanned with %d thread(s))"),
 		(FPlatformTime::Seconds() - StartTime) * 1000,
-		GUnreachableObjects.Num(),
-		ClusterObjects,
-		ClusterItemsToDestroy.Num());
+		NumGathered,
+		NumScanned,
+		NumThreads);
+
+	if (bCompleted)
+	{
+		// NotifyUnreachableObjects needs to be called from within a GC lock so that Async Loading Thread does not attempt to create new exports  
+		// before it gets a chance to process all unreachable objects
+		const bool bNeedsGCLock = !FGCCSyncObject::Get().IsGCLocked();
+		if (bNeedsGCLock)
+		{
+			AcquireGCLock();
+		}
+		NotifyUnreachableObjects(GUnreachableObjects);
+		if (bNeedsGCLock)
+		{
+			ReleaseGCLock();
+		}
+	}
+
+	// !bCompleted = bTimeLimitReached (to make this function return value behave like UnhashUnreachableObjects)
+	return !bCompleted;
+}
+
+void GatherUnreachableObjects(bool bForceSingleThreaded)
+{
+	using namespace UE::GC;
+	GatherUnreachableObjects(bForceSingleThreaded ? EGatherOptions::None : EGatherOptions::Parallel);
 }
 
 namespace UE::GC
@@ -4977,6 +5099,11 @@ EGCOptions GetReferenceCollectorOptions(bool bPerformFullPurge)
 		((GAllowIncrementalReachability && !bPerformFullPurge) ? EGCOptions::IncrementalReachability : EGCOptions::None);
 }
 
+EGatherOptions GetObjectGatherOptions()
+{
+	return (ShouldForceSingleThreadedGC() ? EGatherOptions::None : EGatherOptions::Parallel);
+}
+
 #if WITH_VERSE_VM || defined(__INTELLISENSE__)
 static bool UpdateFrankenGCMode()
 {
@@ -5183,7 +5310,8 @@ void CollectGarbageImpl(EObjectFlags KeepFlags)
 template<bool bPerformFullPurge>
 void PostCollectGarbageImpl(EObjectFlags KeepFlags)
 {
-	using namespace UE::GC::Private;
+	using namespace UE::GC;
+	using namespace UE::GC::Private;	
 
 	if (!GIsIncrementalReachabilityPending)
 	{
@@ -5202,9 +5330,12 @@ void PostCollectGarbageImpl(EObjectFlags KeepFlags)
 
 		DumpGarbageReferencers(AllContexts);
 
-		GatherUnreachableObjects(!(GetReferenceCollectorOptions(bPerformFullPurge) & EGCOptions::Parallel));
+		Swap(GUnreachableObjectFlag, GMaybeUnreachableObjectFlag);
 
-		// This needs to happen after GatherUnreachableObjects since it can mark more objects as unreachable
+		const EGatherOptions GatherOptions = GetObjectGatherOptions();
+		DissolveUnreachableClusters(GatherOptions);
+
+		// This needs to happen after DissolveUnreachableClusters since it can mark more objects as unreachable
 		if (GReachabilityState.GetNumIterations() > 1)
 		{
 			ClearWeakReferences<true>(AllContexts);
@@ -5220,6 +5351,13 @@ void PostCollectGarbageImpl(EObjectFlags KeepFlags)
 		{
 			ContextPool.Cleanup();
 		}
+
+		GGatherUnreachableObjectsState.Init();
+
+		if (bPerformFullPurge || !GAllowIncrementalGather)
+		{
+			GatherUnreachableObjects(GatherOptions, /*TimeLimit =*/ 0.0);
+		}
 	}
 
 	GIsGarbageCollectingAndLockingUObjectHashTables = false;
@@ -5227,11 +5365,6 @@ void PostCollectGarbageImpl(EObjectFlags KeepFlags)
 
 	GIsGarbageCollecting = false;
 
-	if (!GIsIncrementalReachabilityPending)
-	{
-		NotifyUnreachableObjects(GUnreachableObjects);
-	}
-
 	// The hash tables lock was released when reachability analysis was done.
 	// BeginDestroy, FinishDestroy, destructors and callbacks are allowed to call functions like StaticAllocateObject and StaticFindObject.
 	// Now release the GC lock to allow async loading and other threads to perform UObject operations under the FGCScopeGuard.
@@ -5536,11 +5669,35 @@ double GetLastGCDuration()
 
 bool IsIncrementalUnhashPending()
 {
-	return GUnrechableObjectIndex < GUnreachableObjects.Num();
+	return GUnrechableObjectIndex < GUnreachableObjects.Num() || UE::GC::Private::GGatherUnreachableObjectsState.IsPending();
 }
 
 bool UnhashUnreachableObjects(bool bUseTimeLimit, double TimeLimit)
 {
+	using namespace UE::GC;
+	using namespace UE::GC::Private;
+
+	bool bTimeLimitReached = false;
+
+	if (GGatherUnreachableObjectsState.IsPending())
+	{
+		// Incremental Gather needs to be called from UnhashUnreachableObjects to match changes in IsIncrementalUnhashPending() (and not introduce IsIncrementalGatherPending())
+		const EGatherOptions GatherOptions = GetObjectGatherOptions();
+		const double GatherTimeLimit = GIncrementalGatherTimeLimit > 0.0f ? GIncrementalGatherTimeLimit : TimeLimit;
+		bTimeLimitReached = GatherUnreachableObjects(GatherOptions, bUseTimeLimit ? GatherTimeLimit : 0.0);
+		if (!bTimeLimitReached)
+		{
+			if (bUseTimeLimit)
+			{
+				TimeLimit -= FMath::Min(TimeLimit, FPlatformTime::Seconds() - GCStartTime);
+			}
+		}
+		else
+		{
+			return bTimeLimitReached;
+		}
+	}
+
 	TRACE_CPUPROFILER_EVENT_SCOPE(UnhashUnreachableObjects);
 	DECLARE_SCOPE_CYCLE_COUNTER(TEXT("UnhashUnreachableObjects"), STAT_UnhashUnreachableObjects, STATGROUP_GC);
 
@@ -5554,7 +5711,6 @@ bool UnhashUnreachableObjects(bool bUseTimeLimit, double TimeLimit)
 	// Unhash all unreachable objects.
 	const double StartTime = FPlatformTime::Seconds();
 	const int32 TimeLimitEnforcementGranularityForBeginDestroy = 10;
-	int32 Items = 0;
 	int32 TimePollCounter = 0;
 	const bool bFirstIteration = (GUnrechableObjectIndex == 0);
 
@@ -5572,8 +5728,6 @@ bool UnhashUnreachableObjects(bool bUseTimeLimit, double TimeLimit)
 				Object->ConditionalBeginDestroy();
 			}
 
-			Items++;
-
 			const bool bPollTimeLimit = ((TimePollCounter++) % TimeLimitEnforcementGranularityForBeginDestroy == 0);
 			if (bUseTimeLimit && bPollTimeLimit && ((FPlatformTime::Seconds() - StartTime) > TimeLimit))
 			{
@@ -5582,7 +5736,7 @@ bool UnhashUnreachableObjects(bool bUseTimeLimit, double TimeLimit)
 		}
 	}
 
-	const bool bTimeLimitReached = (GUnrechableObjectIndex < GUnreachableObjects.Num());
+	bTimeLimitReached = (GUnrechableObjectIndex < GUnreachableObjects.Num());
 
 	if (!bUseTimeLimit)
 	{
diff --git a/Engine/Source/Runtime/CoreUObject/Private/UObject/GarbageCollectionVerification.cpp b/Engine/Source/Runtime/CoreUObject/Private/UObject/GarbageCollectionVerification.cpp
--- a/Engine/Source/Runtime/CoreUObject/Private/UObject/GarbageCollectionVerification.cpp
+++ b/Engine/Source/Runtime/CoreUObject/Private/UObject/GarbageCollectionVerification.cpp
@@ -455,7 +455,7 @@ public:
 		if (Object)
 		{
 			FUObjectItem* ObjectItem = GUObjectArray.ObjectToObjectItem(Object);
-			if (ObjectItem->HasAnyFlags(EInternalObjectFlags::MaybeUnreachable | EInternalObjectFlags::Unreachable))
+			if (ObjectItem->HasAnyFlags(UE::GC::GMaybeUnreachableObjectFlag | UE::GC::GUnreachableObjectFlag))
 			{
 				if (ReferencingObject)
 				{
@@ -522,7 +522,7 @@ void VerifyNoUnreachableObjects()
 		for (int32 ObjectIndex = 0; ObjectIndex < NumObjects && (FirstObjectIndex + ObjectIndex) < GUObjectArray.GetObjectArrayNum(); ++ObjectIndex)
 		{
 			FUObjectItem& ObjectItem = GUObjectArray.GetObjectItemArrayUnsafe()[FirstObjectIndex + ObjectIndex];
-			if (ObjectItem.Object && !ObjectItem.HasAnyFlags(EInternalObjectFlags::MaybeUnreachable | EInternalObjectFlags::Unreachable))
+			if (ObjectItem.Object && !ObjectItem.HasAnyFlags(UE::GC::GMaybeUnreachableObjectFlag | UE::GC::GUnreachableObjectFlag))
 			{
 				ObjectsToSerialize.Add(static_cast<UObject*>(ObjectItem.Object));
 			}
diff --git a/Engine/Source/Runtime/CoreUObject/Private/UObject/LinkerLoad.cpp b/Engine/Source/Runtime/CoreUObject/Private/UObject/LinkerLoad.cpp
--- a/Engine/Source/Runtime/CoreUObject/Private/UObject/LinkerLoad.cpp
+++ b/Engine/Source/Runtime/CoreUObject/Private/UObject/LinkerLoad.cpp
@@ -5560,7 +5560,7 @@ UObject* FLinkerLoad::CreateImport( int32 Index )
 	DeferPotentialCircularImport(Index); 
 #endif // USE_CIRCULAR_DEPENDENCY_LOAD_DEFERRING
 
-	if (Import.XObject != nullptr && Import.XObject->HasAnyInternalFlags(EInternalObjectFlags::Unreachable))
+	if (Import.XObject != nullptr && Import.XObject->IsUnreachable())
 	{
 		// This is just a safeguard to catch potential bugs that should have been fixed by calling UnhashUnreachableObjects in Async Loading code
 		UE_ASSET_LOG(LogLinker, Warning, PackagePath, TEXT("Unreachable object found when creating import %s"), *Import.XObject->GetFullName());
diff --git a/Engine/Source/Runtime/CoreUObject/Private/UObject/ReachabilityAnalysisState.h b/Engine/Source/Runtime/CoreUObject/Private/UObject/ReachabilityAnalysisState.h
--- a/Engine/Source/Runtime/CoreUObject/Private/UObject/ReachabilityAnalysisState.h
+++ b/Engine/Source/Runtime/CoreUObject/Private/UObject/ReachabilityAnalysisState.h
@@ -128,4 +128,184 @@ public:
 	void FinishIteration();
 };
 
+
+/** Settings for GatherUnreachableObjects function */
+enum class EGatherOptions : uint32
+{
+	None = 0,
+	Parallel = 1,	 // Use threading to increase performance
+};
+ENUM_CLASS_FLAGS(EGatherOptions);
+
+inline int32 GetNumThreadsForGather(const EGatherOptions Options, const int32 NumObjects)
+{
+	if (!!(Options & EGatherOptions::Parallel) && NumObjects > 0)
+	{
+		constexpr int32 MinNumObjectsPerThread = 100;
+		const int32 MaxNumThreads = FTaskGraphInterface::Get().GetNumWorkerThreads();
+		const int32 NumThreadsForNumObjects = (NumObjects + MinNumObjectsPerThread - 1) / MinNumObjectsPerThread;
+		return FMath::Min(NumThreadsForNumObjects, MaxNumThreads);
+	}
+	return 1;
+}
+
+/** Structure that holds the current state of Thread Iteration */
+template <typename PayloadType>
+struct TGatherIterator
+{
+	/** Inital start index */
+	int32 StartIndex = 0;
+	/** Current index of an object */
+	int32 Index = 0;
+	/** Estimated number of objects to process */
+	int32 Num = 0;
+	/** Last index of object to process */
+	int32 LastIndex = 0;
+	/** Data gathered by the iteration */
+	PayloadType Payload;
+};
+
+/** Helper class that holds the current state of Object Gathering phases */
+template <typename ObjectType>
+class TThreadedGather
+{
+public:
+
+	typedef TGatherIterator<TArray<ObjectType>> FIterator;
+	typedef TArray<FIterator, TInlineAllocator<32>> FThreadIterators;
+
+	FORCEINLINE FThreadIterators& GetThreadIterators()
+	{
+		return ThreadIterators;
+	}
+
+	FORCEINLINE bool IsSuspended() const
+	{
+		return ThreadIterators.Num() > 0;
+	}
+
+	FORCEINLINE bool IsPending() const
+	{
+		return bIsPending;
+	}
+
+	void Init()
+	{
+		bIsPending = true;
+	}
+
+	FORCENOINLINE void Start(const EGatherOptions Options, const int32 TotalNumObjects, const int32 FirstIndex = 0, const int32 DesiredNumThreads = 0)
+	{
+		int32 NumObjectsToProcess = TotalNumObjects - FirstIndex;
+		int32 NumThreads = DesiredNumThreads > 0 ? DesiredNumThreads : GetNumThreadsForGather(Options, NumObjectsToProcess);
+		int32 NumberOfObjectsPerThread = (NumObjectsToProcess / NumThreads) + 1;
+
+		ThreadIterators.AddDefaulted(NumThreads);
+
+		for (int32 ThreadIndex = 0; ThreadIndex < NumThreads; ++ThreadIndex)
+		{
+			FIterator& Iterator = ThreadIterators[ThreadIndex];
+			Iterator.StartIndex = ThreadIndex * NumberOfObjectsPerThread + FirstIndex;
+			Iterator.Index = Iterator.StartIndex;
+			Iterator.Num = (ThreadIndex < (NumThreads - 1)) ? NumberOfObjectsPerThread : (TotalNumObjects - (NumThreads - 1) * NumberOfObjectsPerThread);
+			Iterator.LastIndex = FMath::Min(TotalNumObjects - 1, Iterator.Index + Iterator.Num - 1);
+		}
+	}
+
+	FORCENOINLINE void Finish(TArray<ObjectType>& OutGatheredObjects)
+	{
+		const int32 NumGatheredObjects = NumGathered();
+		if (NumGatheredObjects)
+		{
+			OutGatheredObjects.Reserve(NumGatheredObjects);
+			if (ThreadIterators.Num() == 1)
+			{
+				OutGatheredObjects = MoveTemp(ThreadIterators[0].Payload);
+			}
+			else
+			{
+				for (FIterator& Iterator : ThreadIterators)
+				{
+					OutGatheredObjects += Iterator.Payload;
+				}
+			}
+		}
+		ThreadIterators.Reset();
+		bIsPending = false;
+	}
+
+	FORCEINLINE int32 NumWorkerThreads() const
+	{
+		return ThreadIterators.Num();
+	}
+
+	int32 NumScanned() const
+	{
+		int32 NumScanned = 0;
+		for (const FIterator& ThreadState : ThreadIterators)
+		{
+			NumScanned += ThreadState.Index - ThreadState.StartIndex;
+		}
+		return NumScanned;
+	}
+
+	int32 NumGathered() const
+	{
+		int32 NumGathered = 0;
+		for (const FIterator& ThreadState : ThreadIterators)
+		{
+			NumGathered += ThreadState.Payload.Num();
+		}
+		return NumGathered;
+	}
+
+private:
+
+	FThreadIterators ThreadIterators;
+	bool bIsPending = false;
+};
+
+/** Timer that can be used with ParallelFor that distributes TimeLimit checks evenly across multiple threads */
+class FTimeSlicer
+{
+	const int32 TimeLimitPollInterval = 0;	
+	int32 TimeLimitTimePollCounter = 0;
+	const double StartTime = 0.0;
+	const double TimeLimit = 0.0;
+	std::atomic<int32>& TimeLimitFlag;
+public:
+
+	FTimeSlicer(int32 InPollInterval, int32 InitialCounter, double InStartTime, double InTimeLimit, std::atomic<int32>& InTimeLimitFlag)
+		: TimeLimitPollInterval(InPollInterval)		
+		, TimeLimitTimePollCounter(InitialCounter)
+		, StartTime(InStartTime)
+		, TimeLimit(InTimeLimit)
+		, TimeLimitFlag(InTimeLimitFlag)
+	{
+	}
+
+	FORCEINLINE bool IsTimeLimitExceeded()
+	{
+		if (TimeLimit > 0.0)
+		{
+			const bool bPollTimeLimit = ((TimeLimitTimePollCounter++) % TimeLimitPollInterval == 0);
+			if (bPollTimeLimit)
+			{
+				// It's this thread's time to check if time limit has been exceeded
+				if ((FPlatformTime::Seconds() - StartTime) >= TimeLimit)
+				{
+					TimeLimitFlag.store(1, std::memory_order_relaxed);
+					return true;
+				}
+			}
+			else
+			{
+				// Check if any other thread signaled that time limit has been exceeded
+				return !!TimeLimitFlag.load(std::memory_order_relaxed);
+			}
+		}
+		return false;
+	}
+};
+
 } // namespace UE::GC
\ No newline at end of file
diff --git a/Engine/Source/Runtime/CoreUObject/Private/UObject/ReferenceChainSearch.cpp b/Engine/Source/Runtime/CoreUObject/Private/UObject/ReferenceChainSearch.cpp
--- a/Engine/Source/Runtime/CoreUObject/Private/UObject/ReferenceChainSearch.cpp
+++ b/Engine/Source/Runtime/CoreUObject/Private/UObject/ReferenceChainSearch.cpp
@@ -1319,7 +1319,7 @@ FString FReferenceChainSearch::GetObjectFlags(FGCObjectInfo* InObject)
 		Flags += TEXT("(ClusterRoot) ");
 	}
 
-	if (InObject->HasAnyInternalFlags(EInternalObjectFlags::MaybeUnreachable))
+	if (InObject->HasAnyInternalFlags(UE::GC::GMaybeUnreachableObjectFlag))
 	{
 		Flags += TEXT("(MaybeUnreachable) ");
 	}
diff --git a/Engine/Source/Runtime/CoreUObject/Private/UObject/UObjectGlobals.cpp b/Engine/Source/Runtime/CoreUObject/Private/UObject/UObjectGlobals.cpp
--- a/Engine/Source/Runtime/CoreUObject/Private/UObject/UObjectGlobals.cpp
+++ b/Engine/Source/Runtime/CoreUObject/Private/UObject/UObjectGlobals.cpp
@@ -5208,7 +5208,7 @@ public:
 				}
 				else
 				{
-					Object->SetInternalFlags(EInternalObjectFlags::Unreachable);
+					Object->SetInternalFlags(UE::GC::GUnreachableObjectFlag);
 				}
 			}
 		}
diff --git a/Engine/Source/Runtime/CoreUObject/Private/UObject/UObjectHash.cpp b/Engine/Source/Runtime/CoreUObject/Private/UObject/UObjectHash.cpp
--- a/Engine/Source/Runtime/CoreUObject/Private/UObject/UObjectHash.cpp
+++ b/Engine/Source/Runtime/CoreUObject/Private/UObject/UObjectHash.cpp
@@ -6,6 +6,7 @@
 
 #include "UObject/UObjectHash.h"
 #include "UObject/Class.h"
+#include "UObject/GarbageCollectionGlobals.h"
 #include "UObject/Package.h"
 #include "Misc/AsciiSet.h"
 #include "Misc/PackageName.h"
@@ -562,7 +563,7 @@ static int32 GetObjectOuterHash(FName ObjName, PTRINT Outer)
 
 UObject* StaticFindObjectFastExplicitThreadSafe(FUObjectHashTables& ThreadHash, const UClass* ObjectClass, FName ObjectName, const FString& ObjectPathName, bool bExactClass, EObjectFlags ExcludeFlags/*=0*/)
 {
-	const EInternalObjectFlags ExclusiveInternalFlags = EInternalObjectFlags::Unreachable;
+	const EInternalObjectFlags ExclusiveInternalFlags = UE::GC::GUnreachableObjectFlag;
 
 	// Find an object with the specified name and (optional) class, in any package; if bAnyPackage is false, only matches top-level packages
 	int32 Hash = GetObjectHash(ObjectName);
@@ -692,7 +693,7 @@ struct FObjectSearchPath
 
 UObject* StaticFindObjectInPackageInternal(FUObjectHashTables& ThreadHash, const UClass* ObjectClass, const UPackage* ObjectPackage, FName ObjectName, bool bExactClass, EObjectFlags ExcludeFlags, EInternalObjectFlags ExclusiveInternalFlags)
 {
-	ExclusiveInternalFlags |= EInternalObjectFlags::Unreachable;
+	ExclusiveInternalFlags |= UE::GC::GUnreachableObjectFlag;
 	UObject* Result = nullptr;
 	if (FHashBucket* Inners = ThreadHash.PackageToObjectListMap.Find(ObjectPackage))
 	{
@@ -726,7 +727,7 @@ UObject* StaticFindObjectInPackageInternal(FUObjectHashTables& ThreadHash, const
 
 UObject* StaticFindObjectFastInternalThreadSafe(FUObjectHashTables& ThreadHash, const UClass* ObjectClass, const UObject* ObjectPackage, FName ObjectName, bool bExactClass, bool bAnyPackage, EObjectFlags ExcludeFlags, EInternalObjectFlags ExclusiveInternalFlags)
 {
-	ExclusiveInternalFlags |= EInternalObjectFlags::Unreachable;
+	ExclusiveInternalFlags |= UE::GC::GUnreachableObjectFlag;
 
 	// If they specified an outer use that during the hashing
 	UObject* Result = nullptr;
@@ -874,7 +875,7 @@ bool StaticFindAllObjectsFastInternal(TArray<UObject*>& OutFoundObjects, const U
 {
 	INC_DWORD_STAT(STAT_FindObjectFast);
 
-	ExclusiveInternalFlags |= EInternalObjectFlags::Unreachable;
+	ExclusiveInternalFlags |= UE::GC::GUnreachableObjectFlag;
 
 	FObjectSearchPath SearchPath(ObjectName);
 	const int32 Hash = GetObjectHash(SearchPath.Inner);
@@ -918,7 +919,7 @@ UObject* StaticFindFirstObjectFastInternal(const UClass* ObjectClass, FName Obje
 {
 	INC_DWORD_STAT(STAT_FindObjectFast);
 
-	ExclusiveInternalFlags |= EInternalObjectFlags::Unreachable;
+	ExclusiveInternalFlags |= UE::GC::GUnreachableObjectFlag;
 
 	UObject* Result = nullptr;
 	FObjectSearchPath SearchPath(ObjectName);
@@ -1146,7 +1147,7 @@ void GetObjectsWithOuter(const class UObjectBase* Outer, TArray<UObject *>& Resu
 #endif
 
 	// We don't want to return any objects that are currently being background loaded unless we're using the object iterator during async loading.
-	ExclusionInternalFlags |= EInternalObjectFlags::Unreachable;
+	ExclusionInternalFlags |= UE::GC::GUnreachableObjectFlag;
 	if (!IsInAsyncLoadingThread())
 	{
 		ExclusionInternalFlags |= EInternalObjectFlags::AsyncLoading;
@@ -1215,7 +1216,7 @@ void ForEachObjectWithOuterBreakable(const class UObjectBase* Outer, TFunctionRe
 #endif
 
 	// We don't want to return any objects that are currently being background loaded unless we're using the object iterator during async loading.
-	ExclusionInternalFlags |= EInternalObjectFlags::Unreachable;
+	ExclusionInternalFlags |= UE::GC::GUnreachableObjectFlag;
 	if (!IsInAsyncLoadingThread())
 	{
 		ExclusionInternalFlags |= EInternalObjectFlags::AsyncLoading;
@@ -1265,7 +1266,7 @@ UObjectBase* FindObjectWithOuter(const class UObjectBase* Outer, const class UCl
 	UObject* Result = nullptr;
 	check( Outer );
 	// We don't want to return any objects that are currently being background loaded unless we're using the object iterator during async loading.
-	EInternalObjectFlags ExclusionInternalFlags = EInternalObjectFlags::Unreachable;
+	EInternalObjectFlags ExclusionInternalFlags = UE::GC::GUnreachableObjectFlag;
 	if (!IsInAsyncLoadingThread())
 	{
 		ExclusionInternalFlags = EInternalObjectFlags::AsyncLoading;
@@ -1319,7 +1320,7 @@ void ForEachObjectWithPackage(const class UPackage* Package, TFunctionRef<bool(U
 	check(Package != nullptr);
 
 	// We don't want to return any objects that are currently being background loaded unless we're using the object iterator during async loading.
-	ExclusionInternalFlags |= EInternalObjectFlags::Unreachable;
+	ExclusionInternalFlags |= UE::GC::GUnreachableObjectFlag;
 	if (!IsInAsyncLoadingThread())
 	{
 		ExclusionInternalFlags |= EInternalObjectFlags::AsyncLoading;
@@ -1431,7 +1432,7 @@ FORCEINLINE void ForEachObjectOfClasses_Implementation(FUObjectHashTables& Threa
 	TRACE_CPUPROFILER_EVENT_SCOPE(ForEachObjectOfClasses_Implementation);
 
 	// We don't want to return any objects that are currently being background loaded unless we're using the object iterator during async loading.
-	ExclusionInternalFlags |= EInternalObjectFlags::Unreachable;
+	ExclusionInternalFlags |= UE::GC::GUnreachableObjectFlag;
 	if (!IsInAsyncLoadingThread())
 	{
 		ExclusionInternalFlags |= EInternalObjectFlags::AsyncLoading;
diff --git a/Engine/Source/Runtime/CoreUObject/Private/UObject/UObjectMarks.cpp b/Engine/Source/Runtime/CoreUObject/Private/UObject/UObjectMarks.cpp
--- a/Engine/Source/Runtime/CoreUObject/Private/UObject/UObjectMarks.cpp
+++ b/Engine/Source/Runtime/CoreUObject/Private/UObject/UObjectMarks.cpp
@@ -7,7 +7,7 @@
 #include "UObject/UObjectMarks.h"
 #include "UObject/ObjectMacros.h"
 #include "UObject/UObjectIterator.h"
-
+#include "UObject/GarbageCollectionGlobals.h"
 
 struct FObjectMark
 {
@@ -213,7 +213,7 @@ EObjectMark ObjectGetAllMarks(const class UObjectBase* Object)
 void GetObjectsWithAllMarks(TArray<UObject *>& Results, EObjectMark Marks)
 {
 	// We don't want to return any objects that are currently being background loaded unless we're using the object iterator during async loading.
-	EInternalObjectFlags ExclusionFlags = EInternalObjectFlags::Unreachable;
+	EInternalObjectFlags ExclusionFlags = UE::GC::GUnreachableObjectFlag;
 	if (!IsInAsyncLoadingThread())
 	{
 		ExclusionFlags |= EInternalObjectFlags::AsyncLoading;
@@ -236,7 +236,7 @@ void GetObjectsWithAllMarks(TArray<UObject *>& Results, EObjectMark Marks)
 void GetObjectsWithAnyMarks(TArray<UObject *>& Results, EObjectMark Marks)
 {
 	// We don't want to return any objects that are currently being background loaded unless we're using the object iterator during async loading.
-	EInternalObjectFlags ExclusionFlags = EInternalObjectFlags::Unreachable;
+	EInternalObjectFlags ExclusionFlags = UE::GC::GUnreachableObjectFlag;
 	if (!IsInAsyncLoadingThread())
 	{
 		ExclusionFlags |= EInternalObjectFlags::AsyncLoading;
diff --git a/Engine/Source/Runtime/CoreUObject/Public/UObject/GarbageCollectionGlobals.h b/Engine/Source/Runtime/CoreUObject/Public/UObject/GarbageCollectionGlobals.h
new file mode 100644
--- /dev/null
+++ b/Engine/Source/Runtime/CoreUObject/Public/UObject/GarbageCollectionGlobals.h
@@ -0,0 +1,18 @@
+// Copyright Epic Games, Inc. All Rights Reserved.
+
+/*=============================================================================
+	GarbageCollectionGlobals.h: Garbage Collection Global State Vars
+=============================================================================*/
+
+#pragma once
+
+#include "UObject/ObjectMacros.h"
+
+namespace UE::GC
+{
+	/** Current EInternalObjectFlags value representing an unreachable object */
+	extern COREUOBJECT_API EInternalObjectFlags GUnreachableObjectFlag;
+
+	/** Current EInternalObjectFlags value representing a maybe unreachable object */
+	extern COREUOBJECT_API EInternalObjectFlags GMaybeUnreachableObjectFlag;
+}
diff --git a/Engine/Source/Runtime/CoreUObject/Public/UObject/ObjectMacros.h b/Engine/Source/Runtime/CoreUObject/Public/UObject/ObjectMacros.h
--- a/Engine/Source/Runtime/CoreUObject/Public/UObject/ObjectMacros.h
+++ b/Engine/Source/Runtime/CoreUObject/Public/UObject/ObjectMacros.h
@@ -607,7 +607,10 @@ enum class EInternalObjectFlags : int32
 {
 	None = 0,
 
-	MaybeUnreachable = 1 << 19, ///< Flag set on all non-root objects at the beginning of Reachability Analysis
+	ReachabilityFlag0 = 1 << 0, ///< One of the flags used by Garbage Collector to determine UObject's reachability state 
+	ReachabilityFlag1 = 1 << 1, ///< One of the flags used by Garbage Collector to determine UObject's reachability state 
+
+	MaybeUnreachable UE_DEPRECATED(5.4, "MaybeUnreachable flag should no longer be used. Use FUObjectItem::IsMaybeUnreachable() or UE::GC::GMaybeUnreachableObjectFlag flag instead.") = 1 << 19, ///< Flag set on all non-root objects at the beginning of Reachability Analysis
 	LoaderImport = 1 << 20, ///< Object is ready to be imported by another package during loading
 	Garbage = 1 << 21, ///< Garbage from logical point of view and should not be referenced. This flag is mirrored in EObjectFlags as RF_Garbage for performance
 	ReachableInCluster = 1 << 23, ///< External reference to object in cluster exists
@@ -615,7 +618,7 @@ enum class EInternalObjectFlags : int32
 	Native = 1 << 25, ///< Native (UClass only). 
 	Async = 1 << 26, ///< Object exists only on a different thread than the game thread.
 	AsyncLoading = 1 << 27, ///< Object is being asynchronously loaded.
-	Unreachable = 1 << 28, ///< Object is not reachable on the object graph.
+	Unreachable UE_DEPRECATED(5.4, "Unreachable flag should no longer be used. Use UObjectBaseUtility::IsUnreachable() or UE::GC::GUnreachableObjectFlag flag instead.") = 1 << 28, ///< Object is not reachable on the object graph.
 	PendingKill UE_DEPRECATED(5.0, "PendingKill flag should no longer be used. Use Garbage flag instead.") = 1 << 29, ///< Objects that are pending destruction (invalid for gameplay but valid objects). This flag is mirrored in EObjectFlags as RF_PendingKill for performance
 	RootSet = 1 << 30, ///< Object will not be garbage collected, even if unreferenced.
 	PendingConstruction = 1 << 31, ///< Object didn't have its class constructor called yet (only the UObjectBase one to initialize its most basic members)
@@ -625,7 +628,7 @@ enum class EInternalObjectFlags : int32
 	MirroredFlags = Garbage | PendingKill, /// Flags mirrored in EObjectFlags
 
 	//~ Make sure this is up to date!
-	AllFlags = MaybeUnreachable | LoaderImport | Garbage | ReachableInCluster | ClusterRoot | Native | Async | AsyncLoading | Unreachable | PendingKill | RootSet | PendingConstruction
+	AllFlags = ReachabilityFlag0 | ReachabilityFlag1 | MaybeUnreachable | LoaderImport | Garbage | ReachableInCluster | ClusterRoot | Native | Async | AsyncLoading | Unreachable | PendingKill | RootSet | PendingConstruction
 	PRAGMA_ENABLE_DEPRECATION_WARNINGS
 };
 ENUM_CLASS_FLAGS(EInternalObjectFlags);
diff --git a/Engine/Source/Runtime/CoreUObject/Public/UObject/UObjectArray.h b/Engine/Source/Runtime/CoreUObject/Public/UObject/UObjectArray.h
--- a/Engine/Source/Runtime/CoreUObject/Public/UObject/UObjectArray.h
+++ b/Engine/Source/Runtime/CoreUObject/Public/UObject/UObjectArray.h
@@ -8,7 +8,7 @@
 
 #include "HAL/ThreadSafeCounter.h"
 #include "Containers/LockFreeList.h"
-#include "UObject/ObjectMacros.h"
+#include "UObject/GarbageCollectionGlobals.h"
 #include "UObject/UObjectBase.h"
 
 #if WITH_VERSE_VM || defined(__INTELLISENSE__)
@@ -177,35 +177,35 @@ struct
 
 	FORCEINLINE void SetUnreachable()
 	{
-		ThisThreadAtomicallySetFlag(EInternalObjectFlags::Unreachable);
+		ThisThreadAtomicallySetFlag(UE::GC::GUnreachableObjectFlag);
 	}
 	FORCEINLINE void SetMaybeUnreachable()
 	{
-		ThisThreadAtomicallySetFlag(EInternalObjectFlags::MaybeUnreachable);
+		ThisThreadAtomicallySetFlag(UE::GC::GMaybeUnreachableObjectFlag);
 	}
 	FORCEINLINE void ClearUnreachable()
 	{
-		ThisThreadAtomicallyClearedFlag(EInternalObjectFlags::Unreachable);
+		ThisThreadAtomicallyClearedFlag(UE::GC::GUnreachableObjectFlag);
 	}
 	FORCEINLINE void ClearMaybeUnreachable()
 	{
-		ThisThreadAtomicallyClearedFlag(EInternalObjectFlags::MaybeUnreachable);
+		ThisThreadAtomicallyClearedFlag(UE::GC::GMaybeUnreachableObjectFlag);
 	}
 	FORCEINLINE bool IsUnreachable() const
 	{
-		return !!(GetFlagsInternal() & int32(EInternalObjectFlags::Unreachable));
+		return !!(GetFlagsInternal() & int32(UE::GC::GUnreachableObjectFlag));
 	}
 	FORCEINLINE bool IsMaybeUnreachable() const
 	{
-		return !!(GetFlagsInternal() & int32(EInternalObjectFlags::MaybeUnreachable));
+		return !!(GetFlagsInternal() & int32(UE::GC::GMaybeUnreachableObjectFlag));
 	}
 	FORCEINLINE bool ThisThreadAtomicallyClearedRFUnreachable()
 	{
-		return ThisThreadAtomicallyClearedFlag(EInternalObjectFlags::Unreachable);
+		return ThisThreadAtomicallyClearedFlag(UE::GC::GUnreachableObjectFlag);
 	}
 	FORCEINLINE bool ThisThreadAtomicallyClearedMaybeUnreachable()
 	{
-		return ThisThreadAtomicallyClearedFlag(EInternalObjectFlags::MaybeUnreachable);
+		return ThisThreadAtomicallyClearedFlag(UE::GC::GMaybeUnreachableObjectFlag);
 	}
 	FORCEINLINE void SetPendingKill()
 	{
@@ -805,7 +805,7 @@ public:
 		if (ObjectItem)
 		{
 			PRAGMA_DISABLE_DEPRECATION_WARNINGS
-			return bEvenIfPendingKill ? !ObjectItem->IsUnreachable() : !(ObjectItem->HasAnyFlags(EInternalObjectFlags::Unreachable | EInternalObjectFlags::PendingKill | EInternalObjectFlags::Garbage));
+			return bEvenIfPendingKill ? !ObjectItem->IsUnreachable() : !(ObjectItem->HasAnyFlags(UE::GC::GUnreachableObjectFlag | EInternalObjectFlags::PendingKill | EInternalObjectFlags::Garbage));
 			PRAGMA_ENABLE_DEPRECATION_WARNINGS
 		}
 		return false;
@@ -828,7 +828,7 @@ public:
 	{
 		// This method assumes ObjectItem is valid.
 		PRAGMA_DISABLE_DEPRECATION_WARNINGS
-		return bEvenIfPendingKill ? (ObjectItem->HasAnyFlags(EInternalObjectFlags::Unreachable | EInternalObjectFlags::PendingKill | EInternalObjectFlags::Garbage)) : (ObjectItem->IsUnreachable());
+		return bEvenIfPendingKill ? (ObjectItem->HasAnyFlags(UE::GC::GUnreachableObjectFlag | EInternalObjectFlags::PendingKill | EInternalObjectFlags::Garbage)) : (ObjectItem->IsUnreachable());
 		PRAGMA_ENABLE_DEPRECATION_WARNINGS
 	}
 
diff --git a/Engine/Source/Runtime/CoreUObject/Public/UObject/UObjectBaseUtility.h b/Engine/Source/Runtime/CoreUObject/Public/UObject/UObjectBaseUtility.h
--- a/Engine/Source/Runtime/CoreUObject/Public/UObject/UObjectBaseUtility.h
+++ b/Engine/Source/Runtime/CoreUObject/Public/UObject/UObjectBaseUtility.h
@@ -19,8 +19,8 @@
 #include "Trace/Detail/Channel.h"
 #include "Trace/Detail/Channel.inl"
 #include "Trace/Trace.h"
+#include "UObject/GarbageCollectionGlobals.h"
 #include "UObject/NameTypes.h"
-#include "UObject/ObjectMacros.h"
 #include "UObject/ObjectVersion.h"
 #include "UObject/UObjectArray.h"
 #include "UObject/UObjectBase.h"
@@ -313,7 +313,7 @@ public:
 	/** Checks if the object is unreachable. */
 	FORCEINLINE bool IsUnreachable() const
 	{
-		return GUObjectArray.IndexToObject(InternalIndex)->HasAnyFlags(EInternalObjectFlags::Unreachable);
+		return GUObjectArray.IndexToObject(InternalIndex)->IsUnreachable();
 	}
 
 	/** Checks if the object is pending kill or unreachable. INTERNAL USE ONLY! If you want to check if your object is valid use IsValid(Object)/IsValidObjectChecked(Object)/GetValid(Object) instead. */
@@ -321,7 +321,7 @@ public:
 	FORCEINLINE bool IsPendingKillOrUnreachable() const
 	{
 		PRAGMA_DISABLE_DEPRECATION_WARNINGS
-		return GUObjectArray.IndexToObject(InternalIndex)->HasAnyFlags(EInternalObjectFlags::PendingKill | EInternalObjectFlags::Garbage | EInternalObjectFlags::Unreachable);
+		return GUObjectArray.IndexToObject(InternalIndex)->HasAnyFlags(EInternalObjectFlags::PendingKill | EInternalObjectFlags::Garbage | UE::GC::GUnreachableObjectFlag);
 		PRAGMA_ENABLE_DEPRECATION_WARNINGS
 	}
 
diff --git a/Engine/Source/Runtime/CoreUObject/Public/UObject/UObjectIterator.h b/Engine/Source/Runtime/CoreUObject/Public/UObject/UObjectIterator.h
--- a/Engine/Source/Runtime/CoreUObject/Public/UObject/UObjectIterator.h
+++ b/Engine/Source/Runtime/CoreUObject/Public/UObject/UObjectIterator.h
@@ -7,7 +7,7 @@
 #pragma once
 
 #include "CoreMinimal.h"
-#include "UObject/ObjectMacros.h"
+#include "UObject/GarbageCollectionGlobals.h"
 #include "UObject/UObjectHash.h"
 #include "UObject/UObjectArray.h"
 #include "UObject/Object.h"
@@ -59,7 +59,7 @@ inline EInternalObjectFlags GetObjectIteratorDefaultInternalExclusionFlags(EInte
 {
 	InternalExclusionFlags = UObjectBaseUtility::FixGarbageOrPendingKillInternalObjectFlags(InternalExclusionFlags);
 
-	InternalExclusionFlags |= EInternalObjectFlags::Unreachable | EInternalObjectFlags::PendingConstruction;
+	InternalExclusionFlags |= UE::GC::GUnreachableObjectFlag | EInternalObjectFlags::PendingConstruction;
 	if (!IsInAsyncLoadingThread())
 	{
 		// We don't want to return any objects that are currently being background loaded unless we're using the object iterator during async loading.
diff --git a/Engine/Source/Runtime/Engine/Private/Actor.cpp b/Engine/Source/Runtime/Engine/Private/Actor.cpp
--- a/Engine/Source/Runtime/Engine/Private/Actor.cpp
+++ b/Engine/Source/Runtime/Engine/Private/Actor.cpp
@@ -733,7 +733,7 @@ void AActor::BeginDestroy()
 {
 	UnregisterAllComponents();
 	ULevel* Level = GetLevel();
-	if (Level && !Level->HasAnyInternalFlags(EInternalObjectFlags::Unreachable))
+	if (Level && !Level->IsUnreachable())
 	{
 		Level->Actors.RemoveSingleSwap(this, false);
 	}
diff --git a/Engine/Source/Runtime/Engine/Private/Components/SceneComponent.cpp b/Engine/Source/Runtime/Engine/Private/Components/SceneComponent.cpp
--- a/Engine/Source/Runtime/Engine/Private/Components/SceneComponent.cpp
+++ b/Engine/Source/Runtime/Engine/Private/Components/SceneComponent.cpp
@@ -1213,7 +1213,7 @@ void USceneComponent::OnComponentDestroyed(bool bDestroyingHierarchy)
 
 		// Do not involve objects which will be destroyed in hierarchy fixups
 PRAGMA_DISABLE_DEPRECATION_WARNINGS
-		constexpr EInternalObjectFlags SkipFlags = EInternalObjectFlags::PendingKill | EInternalObjectFlags::Garbage | EInternalObjectFlags::Unreachable;
+		const EInternalObjectFlags SkipFlags = EInternalObjectFlags::PendingKill | EInternalObjectFlags::Garbage | UE::GC::GUnreachableObjectFlag;
 PRAGMA_ENABLE_DEPRECATION_WARNINGS
 
 		if (bDestroyingHierarchy)
diff --git a/Engine/Source/Runtime/Engine/Private/ObjectCacheContext.cpp b/Engine/Source/Runtime/Engine/Private/ObjectCacheContext.cpp
--- a/Engine/Source/Runtime/Engine/Private/ObjectCacheContext.cpp
+++ b/Engine/Source/Runtime/Engine/Private/ObjectCacheContext.cpp
@@ -26,7 +26,7 @@ namespace ObjectCacheContextImpl {
 EInternalObjectFlags GetObjectCacheInternalFlagsExclusion()
 {
 	// We never want to return objects that are invalid or still being worked on by other threads
-	return EInternalObjectFlags::Unreachable | EInternalObjectFlags::Garbage | EInternalObjectFlags::AsyncLoading | EInternalObjectFlags::Async;
+	return UE::GC::GUnreachableObjectFlag | EInternalObjectFlags::Garbage | EInternalObjectFlags::AsyncLoading | EInternalObjectFlags::Async;
 }
 
 } // namespace ObjectCacheContextImpl
diff --git a/Engine/Source/Runtime/Engine/Public/ExternalPackageHelper.h b/Engine/Source/Runtime/Engine/Public/ExternalPackageHelper.h
--- a/Engine/Source/Runtime/Engine/Public/ExternalPackageHelper.h
+++ b/Engine/Source/Runtime/Engine/Public/ExternalPackageHelper.h
@@ -4,6 +4,7 @@
 
 #if WITH_EDITOR
 
+#include "UObject/GarbageCollectionGlobals.h"
 #include "UObject/LinkerLoad.h"
 #include "UObject/Package.h"
 #include "UObject/UObjectThreadContext.h"
@@ -200,7 +201,7 @@ void FExternalPackageHelper::LoadObjectsFromExternalPackages(UObject* InOuter, T
 					return false;
 				}
 				return true;
-			}, true, RF_NoFlags, EInternalObjectFlags::Unreachable);
+			}, true, RF_NoFlags, UE::GC::GUnreachableObjectFlag);
 
 			if (ensure(LoadedObject))
 			{
