Fixing CIS

diff --git a/Engine/Source/Runtime/CoreUObject/Private/Serialization/AsyncLoading2.cpp b/Engine/Source/Runtime/CoreUObject/Private/Serialization/AsyncLoading2.cpp
--- a/Engine/Source/Runtime/CoreUObject/Private/Serialization/AsyncLoading2.cpp
+++ b/Engine/Source/Runtime/CoreUObject/Private/Serialization/AsyncLoading2.cpp
@@ -1098,7 +1098,7 @@ public:
 							);
 							
 							check(OldPackageRef->GetRefCount() == 0);
-							RemoveCompletedRenamedPackage(*OldPackageRef);
+							RemoveUnreferencedObsoletePackage(*OldPackageRef);
 							RemovePackage(OldPackageId);
 						}
 					}
