Fix property bag repository failing to create outer bags

diff --git a/Engine/Source/Runtime/CoreUObject/Private/UObject/PropertyBagRepository.cpp b/Engine/Source/Runtime/CoreUObject/Private/UObject/PropertyBagRepository.cpp
--- a/Engine/Source/Runtime/CoreUObject/Private/UObject/PropertyBagRepository.cpp
+++ b/Engine/Source/Runtime/CoreUObject/Private/UObject/PropertyBagRepository.cpp
@@ -89,7 +89,8 @@ FPropertyBag* FPropertyBagRepository::CreateOuterBag(const UObjectBase* Owner)
 	const FPropertyBagAssociationData* BagData = AssociatedData.Find(Owner);
 	if(!BagData)
 	{
-		FPropertyBagAssociationData NewBagData;
+		FPropertyBagAssociationData NewBagData(new FPropertyBag);
+		NewBagData.Bag = new FPropertyBag;
 		BagData = &AssociatedData.Emplace(Owner, NewBagData);
 	}
 	return BagData->Bag;
