Fix compile error
#rnx

diff --git a/Engine/Source/Runtime/CoreUObject/Private/UObject/PropertyBagRepository.cpp b/Engine/Source/Runtime/CoreUObject/Private/UObject/PropertyBagRepository.cpp
--- a/Engine/Source/Runtime/CoreUObject/Private/UObject/PropertyBagRepository.cpp
+++ b/Engine/Source/Runtime/CoreUObject/Private/UObject/PropertyBagRepository.cpp
@@ -89,7 +89,7 @@ FPropertyBag* FPropertyBagRepository::CreateOuterBag(const UObjectBase* Owner)
 	const FPropertyBagAssociationData* BagData = AssociatedData.Find(Owner);
 	if(!BagData)
 	{
-		FPropertyBagAssociationData NewBagData(new FPropertyBag);
+		FPropertyBagAssociationData NewBagData;
 		NewBagData.Bag = new FPropertyBag;
 		BagData = &AssociatedData.Emplace(Owner, NewBagData);
 	}
