OverridableManager: Add API to retreive overridden state

diff --git a/Engine/Source/Runtime/CoreUObject/Private/UObject/OverridableManager.cpp b/Engine/Source/Runtime/CoreUObject/Private/UObject/OverridableManager.cpp
--- a/Engine/Source/Runtime/CoreUObject/Private/UObject/OverridableManager.cpp
+++ b/Engine/Source/Runtime/CoreUObject/Private/UObject/OverridableManager.cpp
@@ -68,6 +68,33 @@ void FOverridableManager::CopyOverriddenProperties(UObject& DestObject, const UO
 	}
 }
 
+EOverriddenState FOverridableManager::GetOverriddenState(UObject& Object)
+{
+	if(FOverriddenPropertySet* OverriddenProperties = GetOverriddenProperties(Object))
+	{
+		const EOverriddenPropertyOperation Operation = OverriddenProperties->GetOverriddenPropertyOperation(nullptr, nullptr);
+		if (Operation != EOverriddenPropertyOperation::None)
+		{
+			return Operation == EOverriddenPropertyOperation::Replace ? EOverriddenState::AllOverridden : EOverriddenState::HasOverrides;
+		}
+
+		// Need to check subobjects to 
+		TSet<UObject*> InstancedSubObjects;
+		FFindInstancedReferenceSubobjectHelper::GetInstancedSubObjects(&Object, InstancedSubObjects);
+		for (UObject* InstancedSubObject : InstancedSubObjects)
+		{
+			if (InstancedSubObject && InstancedSubObject->IsIn(&Object))
+			{
+				if (GetOverriddenState(*InstancedSubObject) != EOverriddenState::NoOverrides)
+				{
+					return EOverriddenState::HasOverrides;
+				}
+			}
+		}
+	}
+	return EOverriddenState::NoOverrides;
+}
+
 void FOverridableManager::OverrideObject(UObject& Object)
 {
 	if (FOverriddenPropertySet* ThisObjectOverriddenProperties = OverriddenObjectAnnotations.Find(Object))
diff --git a/Engine/Source/Runtime/CoreUObject/Public/UObject/OverridableManager.h b/Engine/Source/Runtime/CoreUObject/Public/UObject/OverridableManager.h
--- a/Engine/Source/Runtime/CoreUObject/Public/UObject/OverridableManager.h
+++ b/Engine/Source/Runtime/CoreUObject/Public/UObject/OverridableManager.h
@@ -62,6 +62,15 @@ public:
 	}
 };
 
+enum class EOverriddenState : uint8
+{
+	NoOverrides, // no on this object and any of its instanced subobjects
+	HasOverrides, // has overrides in the object properties
+	AllOverridden, // all properties are overridden for this object and its subobjects
+	SubObjectsHasOverrides, // at least one of its subobjects has overrides
+};
+
+
 /*
  *************************************************************************************
  * Overridable serialization is experimental, not supported and use at your own risk *
@@ -115,6 +124,12 @@ public:
 	 * @param SourceObject copying the overridden properties from */
 	COREUOBJECT_API void CopyOverriddenProperties(UObject& DestObject, const UObject& SourceObject);
 
+	/**
+	 * Retrieve the overridden state for the specified object
+	 * @param Object to fetch the overridden properties
+	 * @return the overridden state if the object has overridable serialization enabled */
+	COREUOBJECT_API EOverriddenState GetOverriddenState(UObject& Object);
+
 	/**
 	 * Override the entire object properties and all its instanced subobjects
 	 * @param Object to override all it properties */
