Clean up old blueprint compile and load code:
Remove nativization code that was deprecated in 5.0, which frees up RF_Dynamic
Remove USE_UBER_GRAPH_PERSISTENT_FRAME (the code depends on it now) and UberGraphFramePointerProperty_DEPRECATED (was never being read)
Cleaned up the other global BP defines that were close to compiling, they can now be overridden in target.cs if necessary
Delete some dead code and confusing old comments

diff --git a/Engine/Source/Editor/Kismet/Private/BlueprintCompilationManager.cpp b/Engine/Source/Editor/Kismet/Private/BlueprintCompilationManager.cpp
--- a/Engine/Source/Editor/Kismet/Private/BlueprintCompilationManager.cpp
+++ b/Engine/Source/Editor/Kismet/Private/BlueprintCompilationManager.cpp
@@ -745,9 +745,7 @@ void FBlueprintCompilationManagerImpl::FlushCompilationQueueImpl(bool bSuppressB
 					bDefaultComponentMustBeAdded = true;
 				}
 
-#if USE_UBER_GRAPH_PERSISTENT_FRAME
 				bHasPendingUberGraphFrame = BPGC->UberGraphFramePointerProperty || BPGC->UberGraphFunction;
-#endif//USE_UBER_GRAPH_PERSISTENT_FRAME
 			}
 
 			bool bSkipCompile = false;
@@ -3615,14 +3613,6 @@ void FBlueprintCompilationManager::CompileSynchronously(const FBPCompileRequest&
 	}
 }
 
-// @todo: BP2CPP_remove
-// [DEPRECATED] - No longer implemented or in use; will be removed later.
-PRAGMA_DISABLE_DEPRECATION_WARNINGS
-void FBlueprintCompilationManager::CompileSynchronouslyToCpp(UBlueprint* BP, TSharedPtr<FString> OutHeaderSource, TSharedPtr<FString> OutCppSource, const FCompilerNativizationOptions& NativizationOptions)
-PRAGMA_ENABLE_DEPRECATION_WARNINGS
-{
-}
-
 void FBlueprintCompilationManager::NotifyBlueprintLoaded(UBlueprint* BPLoaded)
 {
 	// Blueprints can be loaded before editor modules are on line:
diff --git a/Engine/Source/Editor/Kismet/Public/BlueprintCompilationManager.h b/Engine/Source/Editor/Kismet/Public/BlueprintCompilationManager.h
--- a/Engine/Source/Editor/Kismet/Public/BlueprintCompilationManager.h
+++ b/Engine/Source/Editor/Kismet/Public/BlueprintCompilationManager.h
@@ -64,15 +64,6 @@ struct KISMET_API FBlueprintCompilationManager
 	 */
 	static void CompileSynchronously(const FBPCompileRequest& Request);
 
-	// @todo: BP2CPP_remove
-	/** 
-	 * Write c++ represenetation of the blueprint to OutHeaderSource and OutCppSource
-	 */
-PRAGMA_DISABLE_DEPRECATION_WARNINGS
-	UE_DEPRECATED(5.0, "This API is no longer in use and will be removed.")
-	static void CompileSynchronouslyToCpp(UBlueprint* BP, TSharedPtr<FString> OutHeaderSource, TSharedPtr<FString> OutCppSource, const FCompilerNativizationOptions& NativizationOptions);
-PRAGMA_ENABLE_DEPRECATION_WARNINGS
-
 	/**
 	 * Adds a newly loaded blueprint to the compilation queue
 	 */
diff --git a/Engine/Source/Editor/KismetCompiler/Private/KismetCompilerMisc.cpp b/Engine/Source/Editor/KismetCompiler/Private/KismetCompilerMisc.cpp
--- a/Engine/Source/Editor/KismetCompiler/Private/KismetCompilerMisc.cpp
+++ b/Engine/Source/Editor/KismetCompiler/Private/KismetCompilerMisc.cpp
@@ -1459,9 +1459,6 @@ FProperty* FKismetCompilerUtilities::CreatePropertyOnScope(UStruct* Scope, const
 	// Check to see if there's already a object on this scope with the same name, and throw an internal compiler error if so
 	// If this happens, it breaks the property link, which causes stack corruption and hard-to-track errors, so better to fail at this point
 	{
-#if !USE_UBER_GRAPH_PERSISTENT_FRAME
-	#error "Without the uber graph frame we will intentionally create properties with conflicting names on the same scope - disable this error at your own risk"
-#else
 		FFieldVariant ExistingObject = CheckPropertyNameOnScope(Scope, PropertyName);
 		if (ExistingObject.IsValid())
 		{
@@ -1481,7 +1478,6 @@ FProperty* FKismetCompilerUtilities::CreatePropertyOnScope(UStruct* Scope, const
 
 			ValidatedPropertyName = TestName;
 		}
-#endif
 	}
 
 	FProperty* NewProperty = nullptr;
@@ -2631,8 +2627,6 @@ FString FNetNameMapping::MakeBaseName(const UObject* Net)
 //////////////////////////////////////////////////////////////////////////
 // FKismetFunctionContext
 
-// @todo: BP2CPP_remove - remove disable/enable deprecation warning once deprecated members are removed
-PRAGMA_DISABLE_DEPRECATION_WARNINGS
 FKismetFunctionContext::FKismetFunctionContext(FCompilerResultsLog& InMessageLog, const UEdGraphSchema_K2* InSchema, UBlueprintGeneratedClass* InNewClass, UBlueprint* InBlueprint)
 	: Blueprint(InBlueprint)
 	, SourceGraph(nullptr)
@@ -2664,10 +2658,7 @@ FKismetFunctionContext::FKismetFunctionContext(FCompilerResultsLog& InMessageLog
 		bCreateDebugData = false;
 	}
 }
-PRAGMA_ENABLE_DEPRECATION_WARNINGS
 
-// @todo: BP2CPP_remove - remove disable/enable deprecation warning once deprecated members are removed
-PRAGMA_DISABLE_DEPRECATION_WARNINGS
 FKismetFunctionContext::~FKismetFunctionContext()
 {
 	if (bAllocatedNetNameMap)
@@ -2681,7 +2672,6 @@ FKismetFunctionContext::~FKismetFunctionContext()
 		delete AllGeneratedStatements[i];
 	}
 }
-PRAGMA_ENABLE_DEPRECATION_WARNINGS
 
 void FKismetFunctionContext::SetExternalNetNameMap(FNetNameMapping* NewMap)
 {
diff --git a/Engine/Source/Editor/KismetCompiler/Public/KismetCompiledFunctionContext.h b/Engine/Source/Editor/KismetCompiler/Public/KismetCompiledFunctionContext.h
--- a/Engine/Source/Editor/KismetCompiler/Public/KismetCompiledFunctionContext.h
+++ b/Engine/Source/Editor/KismetCompiler/Public/KismetCompiledFunctionContext.h
@@ -105,15 +105,6 @@ public:
 public:
 	FKismetFunctionContext(FCompilerResultsLog& InMessageLog, const UEdGraphSchema_K2* InSchema, UBlueprintGeneratedClass* InNewClass, UBlueprint* InBlueprint);
 
-	// @todo: BP2CPP_remove
-	PRAGMA_DISABLE_DEPRECATION_WARNINGS
-	UE_DEPRECATED(5.0, "Please use the version that does not include the Cpp generation flag as a parameter.")
-	FKismetFunctionContext(FCompilerResultsLog& InMessageLog, const UEdGraphSchema_K2* InSchema, UBlueprintGeneratedClass* InNewClass, UBlueprint* InBlueprint, bool bInGeneratingCpp)
-		:FKismetFunctionContext(InMessageLog, InSchema, InNewClass, InBlueprint)
-	{
-	}
-	PRAGMA_ENABLE_DEPRECATION_WARNINGS
-
 	~FKismetFunctionContext();
 
 	void SetExternalNetNameMap(FNetNameMapping* NewMap);
@@ -317,10 +308,6 @@ public:
 		return (SourceStatementList != NULL) && (SourceStatementList->Num() > 0);
 	}
 
-	// @todo: BP2CPP_remove
-	UE_DEPRECATED(5.0, "This API is no longer in use and will be removed.")
-	bool MustUseSwitchState(const FBlueprintCompiledStatement* ExcludeThisOne) const { return false; }
-
 private:
 	// Optimize out any useless jumps (jump to the very next statement, where the control flow can just fall through)
 	void MergeAdjacentStates();
@@ -332,11 +319,10 @@ private:
 	void ResolveGotoFixups();
 
 public:
-	// @todo: BP2CPP_remove
-	UE_DEPRECATED(5.0, "This API is no longer in use and will be removed.")
-	static bool DoesStatementRequiresSwitch(const FBlueprintCompiledStatement* Statement) { return false; }
+	/** Returns true if this statement cannot be optimized to remove the flow stack */
 	static bool DoesStatementRequiresFlowStack(const FBlueprintCompiledStatement* Statement);
-	// The function links gotos, sorts statments, and merges adjacent ones. 
+	
+	/** This function links gotos, sorts statments, and merges adjacent ones */
 	void ResolveStatements();
 
 	/**
diff --git a/Engine/Source/Editor/KismetCompiler/Public/KismetCompilerModule.h b/Engine/Source/Editor/KismetCompiler/Public/KismetCompilerModule.h
--- a/Engine/Source/Editor/KismetCompiler/Public/KismetCompilerModule.h
+++ b/Engine/Source/Editor/KismetCompiler/Public/KismetCompilerModule.h
@@ -18,12 +18,6 @@ class UUserDefinedEnum;
 class UUserDefinedStruct;
 struct FKismetCompilerOptions;
 
-// @todo: BP2CPP_remove
-PRAGMA_DISABLE_DEPRECATION_WARNINGS
-struct FCompilerNativizationOptions;
-
-PRAGMA_ENABLE_DEPRECATION_WARNINGS
-
 #define KISMET_COMPILER_MODULENAME "KismetCompiler"
 
 //////////////////////////////////////////////////////////////////////////
@@ -109,17 +103,6 @@ public:
 	 */
 	virtual void GetBlueprintTypesForClass(UClass* ParentClass, UClass*& OutBlueprintClass, UClass*& OutBlueprintGeneratedClass) const = 0;
 	virtual void GetSubclassesWithDifferingBlueprintTypes(UClass* Class, TSet<const UClass*>& OutMismatchedSubclasses) const = 0;
-
-	// @todo: BP2CPP_remove
-PRAGMA_DISABLE_DEPRECATION_WARNINGS
-	UE_DEPRECATED(5.0, "This API is no longer required. Any overrides should be removed.")
-	virtual void GenerateCppCodeForEnum(UUserDefinedEnum* UDEnum, const FCompilerNativizationOptions& NativizationOptions, FString& OutHeaderCode, FString& OutCPPCode) {}
-	UE_DEPRECATED(5.0, "This API is no longer required. Any overrides should be removed.")
-	virtual void GenerateCppCodeForStruct(UUserDefinedStruct* UDStruct, const FCompilerNativizationOptions& NativizationOptions, FString& OutHeaderCode, FString& OutCPPCode) {}
-	// Generate a wrapper class, that helps accessing non-native properties and calling non-native functions
-	UE_DEPRECATED(5.0, "This API is no longer required. Any overrides should be removed.")
-	virtual FString GenerateCppWrapper(UBlueprintGeneratedClass* BPGC, const FCompilerNativizationOptions& NativizationOptions) { return FString(); }
-PRAGMA_ENABLE_DEPRECATION_WARNINGS
 };
 
 
diff --git a/Engine/Source/Editor/UnrealEd/Private/Kismet2/KismetDebugUtilities.cpp b/Engine/Source/Editor/UnrealEd/Private/Kismet2/KismetDebugUtilities.cpp
--- a/Engine/Source/Editor/UnrealEd/Private/Kismet2/KismetDebugUtilities.cpp
+++ b/Engine/Source/Editor/UnrealEd/Private/Kismet2/KismetDebugUtilities.cpp
@@ -1901,7 +1901,7 @@ FKismetDebugUtilities::EWatchTextResult FKismetDebugUtilities::FindDebuggingData
 					}
 				}
 			}
-#if USE_UBER_GRAPH_PERSISTENT_FRAME
+
 			// Try find the propertybase in the persistent ubergraph frame
 			UFunction* OuterFunction = Property->GetOwner<UFunction>();
 			if (!PropertyBase && OuterFunction)
@@ -1912,7 +1912,6 @@ FKismetDebugUtilities::EWatchTextResult FKismetDebugUtilities::FindDebuggingData
 					PropertyBase = BPGC->GetPersistentUberGraphFrame(ActiveObject, OuterFunction);
 				}
 			}
-#endif // USE_UBER_GRAPH_PERSISTENT_FRAME
 
 			// see if our WatchPin is on a animation node & if so try to get its property info
 			const UAnimBlueprintGeneratedClass* AnimBlueprintGeneratedClass = Cast<UAnimBlueprintGeneratedClass>(Blueprint->GeneratedClass);
diff --git a/Engine/Source/Runtime/Core/Public/Misc/Build.h b/Engine/Source/Runtime/Core/Public/Misc/Build.h
--- a/Engine/Source/Runtime/Core/Public/Misc/Build.h
+++ b/Engine/Source/Runtime/Core/Public/Misc/Build.h
@@ -108,6 +108,9 @@
 	#define WITH_PERFCOUNTERS		0
 #endif
 
+/** Enable perf counters on dedicated servers */
+#define USE_SERVER_PERF_COUNTERS ((UE_SERVER || UE_EDITOR) && WITH_PERFCOUNTERS)
+
 /** 
  * Whether we are compiling a PGO instrumented build.
  */
@@ -352,20 +355,25 @@
 #define USE_NETWORK_PROFILER !(UE_BUILD_SHIPPING || UE_BUILD_TEST)
 #endif
 
-/** Enable UberGraphPersistentFrame feature. It can speed up BP compilation (re-instancing) in editor, but introduce an unnecessary overhead in runtime. */
-#define USE_UBER_GRAPH_PERSISTENT_FRAME 1
-
 /** Enable validation of the Uber Graph's persistent frame's layout, this is useful to detect uber graph frame related corruption */
-#define VALIDATE_UBER_GRAPH_PERSISTENT_FRAME (!(UE_BUILD_SHIPPING || UE_BUILD_TEST)) && USE_UBER_GRAPH_PERSISTENT_FRAME
+#ifndef VALIDATE_USER_GRAPH_PERSISTENT_FRAME
+	#define VALIDATE_UBER_GRAPH_PERSISTENT_FRAME (!(UE_BUILD_SHIPPING || UE_BUILD_TEST))
+#endif
 
 /** Enable fast calls for event thunks into an event graph that have no parameters  */
-#define UE_BLUEPRINT_EVENTGRAPH_FASTCALLS 1
+#ifndef UE_BLUEPRINT_EVENTGRAPH_FASTCALLS
+	#define UE_BLUEPRINT_EVENTGRAPH_FASTCALLS 1
+#endif
 
-/** Enable perf counters on dedicated servers */
-#define USE_SERVER_PERF_COUNTERS ((UE_SERVER || UE_EDITOR) && WITH_PERFCOUNTERS)
+/** Enables code required for handling recursive dependencies during blueprint serialization */
+#ifndef USE_CIRCULAR_DEPENDENCY_LOAD_DEFERRING
+	#define USE_CIRCULAR_DEPENDENCY_LOAD_DEFERRING 1
+#endif
 
-#define USE_CIRCULAR_DEPENDENCY_LOAD_DEFERRING 1
-#define USE_DEFERRED_DEPENDENCY_CHECK_VERIFICATION_TESTS (USE_CIRCULAR_DEPENDENCY_LOAD_DEFERRING && 0)
+/** Enable validation of deferred dependencies loaded during blueprint serialization */
+#ifndef USE_DEFERRED_DEPENDENCY_CHECK_VERIFICATION_TESTS
+	#define USE_DEFERRED_DEPENDENCY_CHECK_VERIFICATION_TESTS (USE_CIRCULAR_DEPENDENCY_LOAD_DEFERRING && 0)
+#endif
 
 // 0 (default), set this to 1 to get draw events with "TOGGLEDRAWEVENTS" "r.ShowMaterialDrawEvents" and the "ProfileGPU" command working in test
 #ifndef ALLOW_PROFILEGPU_IN_TEST
diff --git a/Engine/Source/Runtime/CoreUObject/Private/Blueprint/BlueprintSupport.cpp b/Engine/Source/Runtime/CoreUObject/Private/Blueprint/BlueprintSupport.cpp
--- a/Engine/Source/Runtime/CoreUObject/Private/Blueprint/BlueprintSupport.cpp
+++ b/Engine/Source/Runtime/CoreUObject/Private/Blueprint/BlueprintSupport.cpp
@@ -2585,9 +2585,6 @@ UObject* FLinkerLoad::FindImportFast(UClass* ImportClass, UObject* ImportOuter,
  * UObject
  ******************************************************************************/
 
-/** 
- * Returns whether this object is contained in or part of a blueprint object
- */
 bool UObject::IsInBlueprint() const
 {
 	// Exclude blueprint classes as they may be regenerated at any time
@@ -2612,20 +2609,14 @@ bool UObject::IsInBlueprint() const
 	return false;
 }
 
-/** 
- *  Destroy properties that won't be destroyed by the native destructor
- */
 void UObject::DestroyNonNativeProperties()
 {
 	// Destroy properties that won't be destroyed by the native destructor
-#if USE_UBER_GRAPH_PERSISTENT_FRAME
 	GetClass()->DestroyPersistentUberGraphFrame(this);
-#endif
+
+	for (FProperty* P = GetClass()->DestructorLink; P; P = P->DestructorLinkNext)
 	{
-		for (FProperty* P = GetClass()->DestructorLink; P; P = P->DestructorLinkNext)
-		{
-			P->DestroyValue_InContainer(this);
-		}
+		P->DestroyValue_InContainer(this);
 	}
 }
 
@@ -2633,13 +2624,6 @@ void UObject::DestroyNonNativeProperties()
  * FObjectInitializer
  ******************************************************************************/
 
-/** 
- * Initializes a non-native property, according to the initialization rules. If the property is non-native
- * and does not have a zero constructor, it is initialized with the default value.
- * @param	Property			Property to be initialized
- * @param	Data				Default data
- * @return	Returns true if that property was a non-native one, otherwise false
- */
 bool FObjectInitializer::InitNonNativeProperty(FProperty* Property, UObject* Data)
 {
 	if (!Property->GetOwnerClass()->HasAnyClassFlags(CLASS_Native | CLASS_Intrinsic)) // if this property belongs to a native class, it was already initialized by the class constructor
@@ -3042,28 +3026,3 @@ void FDeferredObjInitializationHelper::ResolveDeferredInitsFromArchetype(UObject
 
 // don't want other files ending up with this internal define
 #undef DEFERRED_DEPENDENCY_CHECK
-
-// @todo: BP2CPP_remove
-// [DEPRECATED] - This type is no longer in use by the engine; remove later.
-FBlueprintDependencyObjectRef::FBlueprintDependencyObjectRef(const TCHAR* InPackageFolder
-	, const TCHAR* InShortPackageName
-	, const TCHAR* InObjectName
-	, const TCHAR* InClassPackageName
-	, const TCHAR* InClassName
-	, const TCHAR* InOuterName)
-	: PackageName(*(FString(InPackageFolder) + TEXT("/") + InShortPackageName))
-	, ObjectName(InObjectName)
-	, ClassPackageName(InClassPackageName)
-	, ClassName(InClassName)
-	, OuterName(InOuterName)
-{}
-
-// @todo: BP2CPP_remove
-// [DEPRECATED] - This type is no longer in use by the engine; remove later.
-PRAGMA_DISABLE_DEPRECATION_WARNINGS
-FConvertedBlueprintsDependencies& FConvertedBlueprintsDependencies::Get()
-{
-	static FConvertedBlueprintsDependencies ConvertedBlueprintsDependencies;
-	return ConvertedBlueprintsDependencies;
-}
-PRAGMA_ENABLE_DEPRECATION_WARNINGS
diff --git a/Engine/Source/Runtime/CoreUObject/Private/UObject/Class.cpp b/Engine/Source/Runtime/CoreUObject/Private/UObject/Class.cpp
--- a/Engine/Source/Runtime/CoreUObject/Private/UObject/Class.cpp
+++ b/Engine/Source/Runtime/CoreUObject/Private/UObject/Class.cpp
@@ -7402,25 +7402,6 @@ IMPLEMENT_CORE_INTRINSIC_CLASS(USparseDelegateFunction, UDelegateFunction,
 	}
 );
 
-// @todo: BP2CPP_remove
-// [DEPRECATED] - No longer in use; will be removed later.
-PRAGMA_DISABLE_DEPRECATION_WARNINGS
-const FString& UDynamicClass::GetTempPackagePrefix()
-{
-	static const FString PackagePrefix(TEXT("/Temp/__TEMP_BP__"));
-	return PackagePrefix;
-}
-PRAGMA_ENABLE_DEPRECATION_WARNINGS
-
-// @todo: BP2CPP_remove
-// [DEPRECATED] - No longer in use; will be removed later.
-PRAGMA_DISABLE_DEPRECATION_WARNINGS
-IMPLEMENT_CORE_INTRINSIC_CLASS(UDynamicClass, UClass,
-{
-	Class->CppClassStaticFunctions = UOBJECT_CPPCLASS_STATICFUNCTIONS_FORCLASS(UDynamicClass);
-}
-);
-PRAGMA_ENABLE_DEPRECATION_WARNINGS
 #if defined(_MSC_VER) && _MSC_VER == 1900
 	#ifdef PRAGMA_ENABLE_SHADOW_VARIABLE_WARNINGS
 		PRAGMA_ENABLE_SHADOW_VARIABLE_WARNINGS
diff --git a/Engine/Source/Runtime/CoreUObject/Private/UObject/LinkerLoad.cpp b/Engine/Source/Runtime/CoreUObject/Private/UObject/LinkerLoad.cpp
--- a/Engine/Source/Runtime/CoreUObject/Private/UObject/LinkerLoad.cpp
+++ b/Engine/Source/Runtime/CoreUObject/Private/UObject/LinkerLoad.cpp
@@ -972,7 +972,6 @@ FLinkerLoad::FLinkerLoad(UPackage* InParent, const FPackagePath& InPackagePath,
 : FLinker(ELinkerType::Load, InParent)
 , LoadFlags(InLoadFlags)
 , bHaveImportsBeenVerified(false)
-, bDynamicClassLinker(false)
 , TemplateForGetArchetypeFromLoader(nullptr)
 , bForceSimpleIndexToObject(false)
 , bLockoutLegacyOperations(false)
diff --git a/Engine/Source/Runtime/CoreUObject/Private/UObject/ScriptCore.cpp b/Engine/Source/Runtime/CoreUObject/Private/UObject/ScriptCore.cpp
--- a/Engine/Source/Runtime/CoreUObject/Private/UObject/ScriptCore.cpp
+++ b/Engine/Source/Runtime/CoreUObject/Private/UObject/ScriptCore.cpp
@@ -924,12 +924,11 @@ void ProcessScriptFunction(UObject* Context, UFunction* Function, FFrame& Stack,
 
 	// Allocate any temporary memory the script may need via AllocA. This AllocA dependency, along with
 	// the desire to inline calls to our Execution function are the reason for this template function:
-	uint8* FrameMemory = nullptr;
 	FFrame NewStack(Context, Function, nullptr, &Stack, Function->ChildProperties);
 	UE_VSTACK_MAKE_FRAME(ProcessScriptFunctionBookmark, NewStack.CachedThreadVirtualStackAllocator);
-#if USE_UBER_GRAPH_PERSISTENT_FRAME
-	FrameMemory = Function->GetOuterUClassUnchecked()->GetPersistentUberGraphFrame(Context, Function);
-#endif
+
+	uint8* FrameMemory = Function->GetOuterUClassUnchecked()->GetPersistentUberGraphFrame(Context, Function);
+
 	bool bUsePersistentFrame = (nullptr != FrameMemory);
 	if (!bUsePersistentFrame)
 	{
@@ -2054,16 +2053,15 @@ void UObject::ProcessEvent( UFunction* Function, void* Parms )
 
 	// Scope required for scoped script stats.
 	{
-		uint8* Frame = NULL;
-#if USE_UBER_GRAPH_PERSISTENT_FRAME
+		uint8* Frame = nullptr;
 		if (Function->HasAnyFunctionFlags(FUNC_UbergraphFunction))
 		{
 			Frame = Function->GetOuterUClassUnchecked()->GetPersistentUberGraphFrame(this, Function);
 		}
-#endif
+
 		FVirtualStackAllocator* VirtualStackAllocator = FBlueprintContext::GetThreadSingleton()->GetVirtualStackAllocator();
 		UE_VSTACK_MAKE_FRAME(ProcessEventBookmark, VirtualStackAllocator);
-		const bool bUsePersistentFrame = (NULL != Frame);
+		const bool bUsePersistentFrame = (nullptr != Frame);
 		if (!bUsePersistentFrame)
 		{
 			Frame = (uint8*)UE_VSTACK_ALLOC_ALIGNED(VirtualStackAllocator, Function->PropertiesSize, Function->GetMinAlignment());
@@ -2632,7 +2630,6 @@ IMPLEMENT_VM_FUNCTION( EX_PopExecutionFlowIfNot, execPopExecutionFlowIfNot );
 
 DEFINE_FUNCTION(UObject::execLetValueOnPersistentFrame)
 {
-#if USE_UBER_GRAPH_PERSISTENT_FRAME
 	Stack.MostRecentProperty = nullptr;
 	Stack.MostRecentPropertyAddress = nullptr;
 	Stack.MostRecentPropertyContainer = nullptr;
@@ -2646,9 +2643,6 @@ DEFINE_FUNCTION(UObject::execLetValueOnPersistentFrame)
 	uint8* DestAddress = DestProperty->ContainerPtrToValuePtr<uint8>(FrameBase);
 
 	Stack.Step(Stack.Object, DestAddress);
-#else
-	checkf(false, TEXT("execLetValueOnPersistentFrame: UberGraphPersistentFrame is not supported by current build!"));
-#endif
 }
 IMPLEMENT_VM_FUNCTION(EX_LetValueOnPersistentFrame, execLetValueOnPersistentFrame);
 
diff --git a/Engine/Source/Runtime/CoreUObject/Private/UObject/UObjectBase.cpp b/Engine/Source/Runtime/CoreUObject/Private/UObject/UObjectBase.cpp
--- a/Engine/Source/Runtime/CoreUObject/Private/UObject/UObjectBase.cpp
+++ b/Engine/Source/Runtime/CoreUObject/Private/UObject/UObjectBase.cpp
@@ -297,14 +297,11 @@ void UObjectBase::SetClass(UClass* NewClass)
 #endif
 
 	UnhashObject(this);
-#if USE_UBER_GRAPH_PERSISTENT_FRAME
+
 	UClass* OldClass = ClassPrivate;
 	ClassPrivate->DestroyPersistentUberGraphFrame((UObject*)this);
-#endif
 	ClassPrivate = NewClass;
-#if USE_UBER_GRAPH_PERSISTENT_FRAME
 	ClassPrivate->CreatePersistentUberGraphFrame((UObject*)this, /*bCreateOnlyIfEmpty =*/false, /*bSkipSuperClass =*/false, OldClass);
-#endif
 	HashObject(this);
 }
 #endif
@@ -593,18 +590,6 @@ class UEnum *GetStaticEnum(class UEnum *(*InRegister)(), UObject* EnumOuter, con
 	return Result;
 }
 
-// UClass deferred registration
-
-// @todo: BP2CPP_remove
-// [DEPRECATED] - No longer in use; will be removed later.
-PRAGMA_DISABLE_DEPRECATION_WARNINGS
-TMap<FName, FDynamicClassStaticData>& GetDynamicClassMap()
-{
-	static TMap<FName, FDynamicClassStaticData> DynamicClassMap;
-	return DynamicClassMap;
-}
-PRAGMA_ENABLE_DEPRECATION_WARNINGS
-
 FName UObjectBase::GetFNameForStatID() const
 {
 	return GetFName();
@@ -1171,48 +1156,3 @@ const TCHAR* DebugFullName(UObject* Object)
 		return TEXT("None");
 	}
 }
-
-// @todo: BP2CPP_remove
-// [DEPRECATED] - No longer implemented or in use; will be removed later.
-UScriptStruct* FindExistingStructIfHotReloadOrDynamic(UObject* Outer, const TCHAR* StructName, SIZE_T Size, uint32 Crc, bool bIsDynamic)
-{
-	return nullptr;
-}
-
-// @todo: BP2CPP_remove
-// [DEPRECATED] - No longer implemented or in use; will be removed later.
-UEnum* FindExistingEnumIfHotReloadOrDynamic(UObject* Outer, const TCHAR* EnumName, SIZE_T Size, uint32 Crc, bool bIsDynamic)
-{
-	return nullptr;
-}
-
-// @todo: BP2CPP_remove
-// [DEPRECATED] - No longer implemented or in use; will be removed later.
-PRAGMA_DISABLE_DEPRECATION_WARNINGS
-UObject* ConstructDynamicType(FName TypePathName, EConstructDynamicType ConstructionSpecifier)
-PRAGMA_ENABLE_DEPRECATION_WARNINGS
-{
-	return nullptr;
-}
-
-// @todo: BP2CPP_remove
-// [DEPRECATED] - No longer implemented or in use; will be removed later.
-FName GetDynamicTypeClassName(FName TypePathName)
-{
-	return NAME_None;
-}
-
-// @todo: BP2CPP_remove
-// [DEPRECATED] - No longer implemented or in use; will be removed later.
-UPackage* FindOrConstructDynamicTypePackage(const TCHAR* PackageName)
-{
-	return nullptr;
-}
-
-// @todo: BP2CPP_remove
-// [DEPRECATED] - No longer in use; will be removed later.
-TMap<FName, FName>& GetConvertedDynamicPackageNameToTypeName()
-{
-	static TMap<FName, FName> ConvertedDynamicPackageNameToTypeName;
-	return ConvertedDynamicPackageNameToTypeName;
-}
diff --git a/Engine/Source/Runtime/CoreUObject/Private/UObject/UObjectGlobals.cpp b/Engine/Source/Runtime/CoreUObject/Private/UObject/UObjectGlobals.cpp
--- a/Engine/Source/Runtime/CoreUObject/Private/UObject/UObjectGlobals.cpp
+++ b/Engine/Source/Runtime/CoreUObject/Private/UObject/UObjectGlobals.cpp
@@ -3609,10 +3609,8 @@ void UObject::PostInitProperties()
 #if !(UE_BUILD_SHIPPING || UE_BUILD_TEST)
 	FUObjectThreadContext::Get().PostInitPropertiesCheck.Push(this);
 #endif
-#if USE_UBER_GRAPH_PERSISTENT_FRAME
-	GetClass()->CreatePersistentUberGraphFrame(this, true);
-#endif
 
+	GetClass()->CreatePersistentUberGraphFrame(this, true);
 	FOverridableManager::Get().ClearOverrides(*this);
 }
 
@@ -3815,8 +3813,9 @@ FObjectInitializer::~FObjectInitializer()
 	// We're not popping this initializer from the stack in the same place where we decrement IsInConstructor
 	// because we still want to be able to access the current initializer from PostConstructInit or any of its callbacks
 	if (!bIsDeferredInitializer)
-	{
 #endif // USE_CIRCULAR_DEPENDENCY_LOAD_DEFERRING		
+	{
+
 		check(ThreadContext.TopInitializer() == this);
 		ThreadContext.PopInitializer();
 	}
diff --git a/Engine/Source/Runtime/CoreUObject/Public/Blueprint/BlueprintSupport.h b/Engine/Source/Runtime/CoreUObject/Public/Blueprint/BlueprintSupport.h
--- a/Engine/Source/Runtime/CoreUObject/Public/Blueprint/BlueprintSupport.h
+++ b/Engine/Source/Runtime/CoreUObject/Public/Blueprint/BlueprintSupport.h
@@ -18,12 +18,6 @@ class UPackage;
 class UScriptStruct;
 class UStruct;
 struct FUObjectSerializeContext;
-
-// @todo: BP2CPP_remove
-PRAGMA_DISABLE_DEPRECATION_WARNINGS
-class UDynamicClass;
-struct FCompilerNativizationOptions;
-PRAGMA_ENABLE_DEPRECATION_WARNINGS
 class ITargetPlatform;
 
 /**
@@ -196,82 +190,6 @@ private:
 	FScopedClassDependencyGather();
 };
 
-// @todo: BP2CPP_remove
-enum class UE_DEPRECATED(5.0, "This type is no longer in use and will be removed.") EReplacementResult
-{
-	/** Don't replace the provided package at all */
-	DontReplace,
-
-	/** Generate a stub file, but don't replace the package */
-	GenerateStub,
-
-	/** Completely replace the file with generated code */
-	ReplaceCompletely
-};
-
-// @todo: BP2CPP_remove
-/**
- * Interface needed by CoreUObject to the BlueprintNativeCodeGen logic. Used by cooker to convert assets 
- * to native code.
- */
-PRAGMA_DISABLE_DEPRECATION_WARNINGS
-struct UE_DEPRECATED(5.0, "This type is no longer in use and will be removed.") IBlueprintNativeCodeGenCore
-{
-	/** Returns the current IBlueprintNativeCodeGenCore, may return nullptr */
-	COREUOBJECT_API static const IBlueprintNativeCodeGenCore* Get();
-
-	/**
-	 * Registers the IBlueprintNativeCodeGenCore, just used to point us at an implementation.
-	 * By default, there is no IBlueprintNativeCodeGenCore, and thus no blueprints are
-	 * replaced at cook.
-	 */
-	COREUOBJECT_API static void Register(const IBlueprintNativeCodeGenCore* Coordinator);
-
-	/**
-	 * Determines whether the provided package needs to be replaced (in part or completely)
-	 * 
-	 * @param Package	The package in question
-	 * @return Whether the package should be converted
-	 */
-	virtual EReplacementResult IsTargetedForReplacement(const UPackage* Package, const FCompilerNativizationOptions& NativizationOptions) const = 0;
-	
-	/**
-	 * Determines whether the provided object needs to be replaced (in part or completely).
-	 * Some objects in a package may require conversion and some may not. If any object 
-	 * in a package wants to be converted then it is implied that all other objects will 
-	 * be converted with it (no support for partial package conversion, beyond stubs)
-	 *
-	 * @param Object	The package in question
-	 * @return Whether the object should be converted
-	 */
-	virtual EReplacementResult IsTargetedForReplacement(const UObject* Object, const FCompilerNativizationOptions& NativizationOptions) const = 0;
-
-	/** 
-	 * Function used to change the type of a class from, say, UBlueprintGeneratedClass to 
-	 * UDynamicClass. Cooking (and conversion in general) must be order independent so
-	 * The scope of this kind of type swap is limited.
-	 * 
-	 * @param Object whose class will be replaced
-	 * @return A replacement class ptr, null if none
-	 */
-	virtual UClass* FindReplacedClassForObject(const UObject* Object, const FCompilerNativizationOptions& NativizationOptions) const = 0;
-	
-	/** 
-	 * Function used to change the path of subobject from a nativized class.
-	 * 
-	 * @param Object Imported Object.
-	 * @param OutName Referenced to name, that will be saved in import table.
-	 * @return An Outer object that should be saved in import table.
-	 */
-	virtual UObject* FindReplacedNameAndOuter(UObject* Object, FName& OutName, const FCompilerNativizationOptions& NativizationOptions) const = 0;
-
-	/*
-	 * Return nativization options for given platform.
-	 */
-	virtual const FCompilerNativizationOptions& GetNativizationOptionsForPlatform(const ITargetPlatform* Platform) const = 0;
-};
-PRAGMA_ENABLE_DEPRECATION_WARNINGS
-
 #endif // WITH_EDITOR
 
 /** 
@@ -447,125 +365,3 @@ struct FDeferredObjInitializationHelper
 	 */
 	static void ResolveDeferredInitsFromArchetype(UObject* Archetype);
 };
-
-// @todo: BP2CPP_remove
-struct UE_DEPRECATED(5.0, "This type is no longer in use and will be removed.") FBlueprintDependencyType
-{
-	uint8 bSerializationBeforeSerializationDependency : 1;
-	uint8 bCreateBeforeSerializationDependency : 1;
-	uint8 bSerializationBeforeCreateDependency : 1;
-	uint8 bCreateBeforeCreateDependency : 1;
-
-	FBlueprintDependencyType()
-		: bSerializationBeforeSerializationDependency(0)
-		, bCreateBeforeSerializationDependency(0)
-		, bSerializationBeforeCreateDependency(0)
-		, bCreateBeforeCreateDependency(0) {}
-
-	FBlueprintDependencyType(bool bInSerializationBeforeSerializationDependency
-		, bool bInCreateBeforeSerializationDependency
-		, bool bInSerializationBeforeCreateDependency
-		, bool bInCreateBeforeCreateDependency)
-		: bSerializationBeforeSerializationDependency(bInSerializationBeforeSerializationDependency)
-		, bCreateBeforeSerializationDependency(bInCreateBeforeSerializationDependency)
-		, bSerializationBeforeCreateDependency(bInSerializationBeforeCreateDependency)
-		, bCreateBeforeCreateDependency(bInCreateBeforeCreateDependency)
-	{}
-};
-
-// @todo: BP2CPP_remove
-PRAGMA_DISABLE_DEPRECATION_WARNINGS
-struct COREUOBJECT_API UE_DEPRECATED(5.0, "This type is no longer in use and will be removed.") FCompactBlueprintDependencyData
-{
-	int16 ObjectRefIndex;
-	FBlueprintDependencyType StructDependency;
-	FBlueprintDependencyType CDODependency;
-
-	FCompactBlueprintDependencyData()
-		: ObjectRefIndex(-1)
-	{}
-
-	FCompactBlueprintDependencyData(int16 InObjectRefIndex
-		, FBlueprintDependencyType InStructDependency
-		, FBlueprintDependencyType InCDODependency = FBlueprintDependencyType())
-		: ObjectRefIndex(InObjectRefIndex)
-		, StructDependency(InStructDependency)
-		, CDODependency(InCDODependency)
-	{}
-};
-PRAGMA_ENABLE_DEPRECATION_WARNINGS
-
-// @todo: BP2CPP_remove
-struct COREUOBJECT_API UE_DEPRECATED(5.0, "This type is no longer in use and will be removed.") FBlueprintDependencyObjectRef
-{
-	FName PackageName;
-	FName ObjectName;
-	FName ClassPackageName;
-	FName ClassName;
-	FName OuterName;
-
-	FBlueprintDependencyObjectRef() {}
-
-	FORCENOINLINE FBlueprintDependencyObjectRef(const TCHAR* InPackageFolder
-		, const TCHAR* InShortPackageName
-		, const TCHAR* InObjectName
-		, const TCHAR* InClassPackageName
-		, const TCHAR* InClassName
-		, const TCHAR* InOuterName );
-};
-
-// @todo: BP2CPP_remove
-PRAGMA_DISABLE_DEPRECATION_WARNINGS
-struct COREUOBJECT_API UE_DEPRECATED(5.0, "This type is no longer in use and will be removed.") FBlueprintDependencyData
-{
-	FBlueprintDependencyObjectRef ObjectRef;
-	// 0 - dependency type for dynamic class or UDS
-	// 1 - dependency type for CD0
-	FBlueprintDependencyType DependencyTypes[2];
-
-	int16 ObjectRefIndex; // NativizationWithoutEDLBT
-
-	FBlueprintDependencyData(const FBlueprintDependencyObjectRef& InObjectRef
-		, const FCompactBlueprintDependencyData& InCompactDependencyData)
-		: ObjectRef(InObjectRef)
-		, ObjectRefIndex(InCompactDependencyData.ObjectRefIndex)
-	{
-		DependencyTypes[0] = InCompactDependencyData.StructDependency;
-		DependencyTypes[1] = InCompactDependencyData.CDODependency;
-	}
-
-	bool operator==(const FBlueprintDependencyData& Other) const
-	{
-		return Other.ObjectRefIndex == ObjectRefIndex;
-	}
-
-	static bool ContainsDependencyData(TArray<FBlueprintDependencyData>& Assets, int16 ObjectRefIndex) { return false; }
-	static void AppendUniquely(TArray<FBlueprintDependencyData>& Destination, const TArray<FBlueprintDependencyData>& AdditionalData) {}
-};
-PRAGMA_ENABLE_DEPRECATION_WARNINGS
-
-// @todo: BP2CPP_remove
-/**
- *	Stores info about dependencies of native classes converted from BPs
- */
-PRAGMA_DISABLE_DEPRECATION_WARNINGS
-struct COREUOBJECT_API UE_DEPRECATED(5.0, "This type is no longer in use and will be removed.") FConvertedBlueprintsDependencies
-{
-	typedef void(*GetDependenciesNamesFunc)(TArray<FBlueprintDependencyData>&);
-
-private:
-
-	TMap<FName, GetDependenciesNamesFunc> PackageNameToGetter;
-
-public:
-	static FConvertedBlueprintsDependencies& Get();
-
-	void RegisterConvertedClass(FName PackageName, GetDependenciesNamesFunc GetAssets) {}
-
-	/** Get all assets paths necessary for the class with the given class name and all converted classes that dependencies. */
-	void GetAssets(FName PackageName, TArray<FBlueprintDependencyData>& OutDependencies) const {}
-
-	static void FillUsedAssetsInDynamicClass(UDynamicClass* DynamicClass, GetDependenciesNamesFunc GetUsedAssets) {}
-	static UObject* LoadObjectForStructConstructor(UScriptStruct* ScriptStruct, const TCHAR* ObjectPath) { return nullptr; }
-};
-PRAGMA_ENABLE_DEPRECATION_WARNINGS
diff --git a/Engine/Source/Runtime/CoreUObject/Public/UObject/Class.h b/Engine/Source/Runtime/CoreUObject/Public/UObject/Class.h
--- a/Engine/Source/Runtime/CoreUObject/Public/UObject/Class.h
+++ b/Engine/Source/Runtime/CoreUObject/Public/UObject/Class.h
@@ -3536,57 +3536,6 @@ protected:
 	COREUOBJECT_API virtual UObject* CreateDefaultObject();
 };
 
-// @todo: BP2CPP_remove
-/**
-* Dynamic class (can be constructed after initial startup)
-*/
-class COREUOBJECT_API UE_DEPRECATED(5.0, "Dynamic class types are no longer supported.") UDynamicClass : public UClass
-{
-	DECLARE_CASTED_CLASS_INTRINSIC_NO_CTOR(UDynamicClass, UClass, 0, TEXT("/Script/CoreUObject"), CASTCLASS_None, NO_API)
-	DECLARE_WITHIN_UPACKAGE()
-
-public:
-
-	typedef void (*DynamicClassInitializerType)	(UDynamicClass*);
-
-	UDynamicClass(const FObjectInitializer& ObjectInitializer = FObjectInitializer::Get()) {}
-	explicit UDynamicClass(const FObjectInitializer& ObjectInitializer, UClass* InSuperClass) {}
-	UDynamicClass(EStaticConstructor, FName InName, uint32 InSize, uint32 InAlignment, EClassFlags InClassFlags, EClassCastFlags InClassCastFlags,
-		const TCHAR* InClassConfigName, EObjectFlags InFlags, ClassConstructorType InClassConstructor,
-		ClassVTableHelperCtorCallerType InClassVTableHelperCtorCaller,
-		FUObjectCppClassStaticFunctions&& InCppClassStaticFunctions,
-		DynamicClassInitializerType InDynamicClassInitializer)
-	{}
-
-	/** Find a struct property, called from generated code */
-	FStructProperty* FindStructPropertyChecked(const TCHAR* PropertyName) const { return nullptr; }
-
-	/** Misc objects owned by the class. */
-	TArray<UObject*> MiscConvertedSubobjects;
-
-	/** Additional converted fields, that are used by the class. */
-	TArray<UField*> ReferencedConvertedFields;
-
-	/** Outer assets used by the class */
-	TArray<UObject*> UsedAssets;
-
-	/** Specialized sub-object containers */
-	TArray<UObject*> DynamicBindingObjects;
-	TArray<UObject*> ComponentTemplates;
-	TArray<UObject*> Timelines;
-
-	/** Array of blueprint overrides of component classes in parent classes */
-	TArray<TPair<FName, const UClass*>> ComponentClassOverrides;
-
-	/** IAnimClassInterface (UAnimClassData) or null */
-	UObject* AnimClassImplementation;
-
-	DynamicClassInitializerType DynamicClassInitializer;
-
-	/** Prefix for the temporary package where the dynamic classes are stored when being generated */
-	static const FString& GetTempPackagePrefix();
-};
-
 /**
  * Helper template to call the default constructor for a class
  */
diff --git a/Engine/Source/Runtime/CoreUObject/Public/UObject/LinkerLoad.h b/Engine/Source/Runtime/CoreUObject/Public/UObject/LinkerLoad.h
--- a/Engine/Source/Runtime/CoreUObject/Public/UObject/LinkerLoad.h
+++ b/Engine/Source/Runtime/CoreUObject/Public/UObject/LinkerLoad.h
@@ -164,10 +164,6 @@ public:
 	uint32					LoadFlags;
 	/** Indicates whether the imports for this loader have been verified													*/
 	bool					bHaveImportsBeenVerified;
-	// @todo: BP2CPP_remove
-	/** Indicates that this linker was created for a dynamic class package and will not use Loader */
-	UE_DEPRECATED(5.0, "This member is no longer in use and will be removed.")
-	bool					bDynamicClassLinker;
 
 	UObject*				TemplateForGetArchetypeFromLoader;
 	bool					bForceSimpleIndexToObject;
diff --git a/Engine/Source/Runtime/CoreUObject/Public/UObject/ObjectMacros.h b/Engine/Source/Runtime/CoreUObject/Public/UObject/ObjectMacros.h
--- a/Engine/Source/Runtime/CoreUObject/Public/UObject/ObjectMacros.h
+++ b/Engine/Source/Runtime/CoreUObject/Public/UObject/ObjectMacros.h
@@ -564,7 +564,7 @@ enum EObjectFlags
 	RF_DuplicateTransient		=0x00800000,	///< Object should not be included in any type of duplication (copy/paste, binary duplication, etc.)
 	RF_StrongRefOnFrame			=0x01000000,	///< References to this object from persistent function frame are handled as strong ones.
 	RF_NonPIEDuplicateTransient	=0x02000000,	///< Object should not be included for duplication unless it's being duplicated for a PIE session
-	RF_Dynamic UE_DEPRECATED(5.0, "RF_Dynamic should no longer be used. It is no longer being set by engine code.") =0x04000000,	///< Field Only. Dynamic field - doesn't get constructed during static initialization, can be constructed multiple times  // @todo: BP2CPP_remove
+	// RF_Dynamic				=0x04000000,	///< Was removed along with bp nativization
 	RF_WillBeLoaded				=0x08000000,	///< This object was constructed during load and will be loaded shortly
 	RF_HasExternalPackage		=0x10000000,	///< This object has an external package assigned and should look it up when getting the outermost package
 
diff --git a/Engine/Source/Runtime/CoreUObject/Public/UObject/UObjectBase.h b/Engine/Source/Runtime/CoreUObject/Public/UObject/UObjectBase.h
--- a/Engine/Source/Runtime/CoreUObject/Public/UObject/UObjectBase.h
+++ b/Engine/Source/Runtime/CoreUObject/Public/UObject/UObjectBase.h
@@ -459,16 +459,6 @@ COREUOBJECT_API void RegisterCompiledInInfo(UPackage* (*InOuterRegister)(), cons
  */
 COREUOBJECT_API void RegisterCompiledInInfo(const TCHAR* PackageName, const FClassRegisterCompiledInInfo* ClassInfo, size_t NumClassInfo, const FStructRegisterCompiledInInfo* StructInfo, size_t NumStructInfo, const FEnumRegisterCompiledInInfo* EnumInfo, size_t NumEnumInfo);
 
-// @todo: BP2CPP_remove
-/** Called during HotReload to hook up an existing structure */
-UE_DEPRECATED(5.0, "This API is no longer in use and will be removed.")
-COREUOBJECT_API class UScriptStruct* FindExistingStructIfHotReloadOrDynamic(UObject* Outer, const TCHAR* StructName, SIZE_T Size, uint32 Crc, bool bIsDynamic);
-
-// @todo: BP2CPP_remove
-/** Called during HotReload to hook up an existing enum */
-UE_DEPRECATED(5.0, "This API is no longer in use and will be removed.")
-COREUOBJECT_API class UEnum* FindExistingEnumIfHotReloadOrDynamic(UObject* Outer, const TCHAR* EnumName, SIZE_T Size, uint32 Crc, bool bIsDynamic);
-
 /** Must be called after a module has been loaded that contains UObject classes */
 COREUOBJECT_API void ProcessNewlyLoadedUObjects(FName Package = NAME_None, bool bCanProcessNewlyLoadedObjects = true);
 
diff --git a/Engine/Source/Runtime/CoreUObject/Public/UObject/UObjectGlobals.h b/Engine/Source/Runtime/CoreUObject/Public/UObject/UObjectGlobals.h
--- a/Engine/Source/Runtime/CoreUObject/Public/UObject/UObjectGlobals.h
+++ b/Engine/Source/Runtime/CoreUObject/Public/UObject/UObjectGlobals.h
@@ -3201,56 +3201,10 @@ extern COREUOBJECT_API int32 GAssetClustreringEnabled;
 
 /** A struct used as stub for deleted ones. */
 COREUOBJECT_API UScriptStruct* GetFallbackStruct();
-// @todo: BP2CPP_remove
-enum class UE_DEPRECATED(5.0, "This type is no longer in use and will be removed.") EConstructDynamicType : uint8
-{
-	OnlyAllocateClassObject,
-	CallZConstructor
-};
-
-// @todo: BP2CPP_remove
-/** Constructs dynamic type of a given class. */
-PRAGMA_DISABLE_DEPRECATION_WARNINGS
-UE_DEPRECATED(5.0, "This method is no longer in use and will be removed.")
-COREUOBJECT_API UObject* ConstructDynamicType(FName TypePathName, EConstructDynamicType ConstructionSpecifier);
-PRAGMA_ENABLE_DEPRECATION_WARNINGS
-
-// @todo: BP2CPP_remove
-/** Given a dynamic type path name, returns that type's class name (can be either DynamicClass, ScriptStruct or Enum). */
-UE_DEPRECATED(5.0, "This method is no longer in use and will be removed.")
-COREUOBJECT_API FName GetDynamicTypeClassName(FName TypePathName);
-
-// @todo: BP2CPP_remove
-/** Finds or constructs a package for dynamic type. */
-UE_DEPRECATED(5.0, "This method is no longer in use and will be removed.")
-COREUOBJECT_API UPackage* FindOrConstructDynamicTypePackage(const TCHAR* PackageName);
-
-// @todo: BP2CPP_remove
-/** Get names of "virtual" packages, that contain Dynamic types  */
-UE_DEPRECATED(5.0, "This method is no longer in use and will be removed.")
-COREUOBJECT_API TMap<FName, FName>& GetConvertedDynamicPackageNameToTypeName();
 
 /** Utility accessor for whether we are running with component class overrides enabled */
 COREUOBJECT_API bool GetAllowNativeComponentClassOverrides();
 
-// @todo: BP2CPP_remove
-struct COREUOBJECT_API UE_DEPRECATED(5.0, "This type is no longer in use and will be removed.") FDynamicClassStaticData
-{
-	/** Autogenerated Z_Construct* function pointer */
-	UClass* (*ZConstructFn)();
-	/** StaticClass() function pointer */
-	UClass* (*StaticClassFn)();
-	/** Selected AssetRegistrySearchable values */
-	TMap<FName, FName> SelectedSearchableValues;
-};
-
-// @todo: BP2CPP_remove
-/** Returns map of all dynamic/nativized classes */
-PRAGMA_DISABLE_DEPRECATION_WARNINGS
-UE_DEPRECATED(5.0, "This method is no longer in use and will be removed.")
-COREUOBJECT_API TMap<FName, FDynamicClassStaticData>& GetDynamicClassMap();
-PRAGMA_ENABLE_DEPRECATION_WARNINGS
-
 namespace UE
 {
 class FAssetLog;
diff --git a/Engine/Source/Runtime/Engine/Classes/Animation/AnimClassData.h b/Engine/Source/Runtime/Engine/Classes/Animation/AnimClassData.h
deleted file mode 100644
diff --git a/Engine/Source/Runtime/Engine/Classes/Engine/Blueprint.h b/Engine/Source/Runtime/Engine/Classes/Engine/Blueprint.h
--- a/Engine/Source/Runtime/Engine/Classes/Engine/Blueprint.h
+++ b/Engine/Source/Runtime/Engine/Classes/Engine/Blueprint.h
@@ -88,8 +88,7 @@ namespace EKismetCompileType
 		Full,
 		StubAfterFailure, 
 		BytecodeOnly,
-		// @todo: BP2CPP_remove
-		Cpp UE_DEPRECATED(5.0, "Cpp is no longer a supported compile type and will be removed."),
+		// Cpp type was removed with BP nativization
 	};
 };
 
@@ -110,42 +109,6 @@ enum class EBlueprintCompileMode : uint8
 	FinalRelease UMETA(ToolTip="Always compile in final release mode.")
 };
 
-// @todo: BP2CPP_remove
-USTRUCT()
-struct UE_DEPRECATED(5.0, "This type is no longer in use and will be removed.") FCompilerNativizationOptions
-{
-	GENERATED_USTRUCT_BODY()
-
-	UPROPERTY()
-	FName PlatformName;
-
-	UPROPERTY()
-	bool ServerOnlyPlatform;
-
-	UPROPERTY()
-	bool ClientOnlyPlatform;
-
-	UPROPERTY()
-	bool bExcludeMonolithicHeaders;
-
-	UPROPERTY()
-	TArray<FName> ExcludedModules;
-
-	// Individually excluded assets
-	UPROPERTY()
-	TSet<FSoftObjectPath> ExcludedAssets;
-
-	// Excluded folders. It excludes only BPGCs, enums and structures are still converted.
-	UPROPERTY()
-	TArray<FString> ExcludedFolderPaths;
-
-	FCompilerNativizationOptions()
-		: ServerOnlyPlatform(false)
-		, ClientOnlyPlatform(false)
-		, bExcludeMonolithicHeaders(false)
-	{}
-};
-
 /** Cached 'cosmetic' information about a macro graph (this is transient and is computed at load) */
 USTRUCT()
 struct FBlueprintMacroCosmeticInfo
@@ -161,8 +124,7 @@ struct FBlueprintMacroCosmeticInfo
 	}
 };
 
-// @todo: BP2CPP_remove - remove disable/enable deprecation warning once NativizationOptions is removed
-PRAGMA_DISABLE_DEPRECATION_WARNINGS
+/** Options used for a specific invication of the blueprint compiler */
 struct FKismetCompilerOptions
 {
 public:
@@ -196,17 +158,6 @@ public:
 	TSharedPtr<FString> OutHeaderSourceCode;
 	TSharedPtr<FString> OutCppSourceCode;
 
-	// @todo: BP2CPP_remove
-	UE_DEPRECATED(5.0, "This member is no longer in use and will be removed.")
-	FCompilerNativizationOptions NativizationOptions;
-
-	// @todo: BP2CPP_remove
-	UE_DEPRECATED(5.0, "This API is no longer in use and will be removed.")
-	bool DoesRequireCppCodeGeneration() const
-	{
-		return false;
-	}
-
 	bool DoesRequireBytecodeGeneration() const
 	{
 		return (CompileType == EKismetCompileType::Full) 
@@ -224,7 +175,6 @@ public:
 	{
 	};
 };
-PRAGMA_ENABLE_DEPRECATION_WARNINGS
 
 /** One metadata entry for a variable */
 USTRUCT()
@@ -835,21 +785,6 @@ public:
 	UPROPERTY(transient, duplicatetransient)
 	TSet<TWeakObjectPtr<UStruct>> CachedUDSDependencies;
 
-	// @todo: BP2CPP_remove
-	enum class UE_DEPRECATED(5.1, "Blueprint Nativization has been removed as a supported feature. This type will eventually be removed.") EIsBPNonReducible : uint8
-	{
-		Unkown,
-		Yes,
-		No,
-	};
-
-	// @todo: BP2CPP_remove
-	// Cached information if the BP contains any non-reducible functions (that can benefit from nativization).
-	UE_DEPRECATED(5.1, "Blueprint Nativization has been removed as a supported feature. This member will eventually be removed.")
-	PRAGMA_DISABLE_DEPRECATION_WARNINGS
-	EIsBPNonReducible bHasAnyNonReducibleFunction;
-	PRAGMA_ENABLE_DEPRECATION_WARNINGS
-
 	// If this BP is just a duplicate created for a specific compilation, the reference to original GeneratedClass is needed
 	UPROPERTY(transient, duplicatetransient)
 	TObjectPtr<UClass> OriginalClass;
diff --git a/Engine/Source/Runtime/Engine/Classes/Engine/BlueprintGeneratedClass.h b/Engine/Source/Runtime/Engine/Classes/Engine/BlueprintGeneratedClass.h
--- a/Engine/Source/Runtime/Engine/Classes/Engine/BlueprintGeneratedClass.h
+++ b/Engine/Source/Runtime/Engine/Classes/Engine/BlueprintGeneratedClass.h
@@ -22,13 +22,11 @@ class UDynamicBlueprintBinding;
 class UInheritableComponentHandler;
 class UTimelineTemplate;
 class UClassCookedMetaData;
+class UEdGraphPin;
 
 DECLARE_MEMORY_STAT_EXTERN(TEXT("Persistent Uber Graph Frame memory"), STAT_PersistentUberGraphFrameMemory, STATGROUP_Memory, );
 DECLARE_MEMORY_STAT_EXTERN(TEXT("BPComp Instancing Fast Path memory"), STAT_BPCompInstancingFastPathMemory, STATGROUP_Memory, );
 
-class UEdGraphPin;
-
-
 USTRUCT()
 struct FNodeToCodeAssociation
 {
@@ -637,12 +635,6 @@ public:
 	UPROPERTY(AssetRegistrySearchable)
 	int32	NumReplicatedProperties;
 
-	// @todo: BP2CPP_remove
-	/** Flag used to indicate if this class has a nativized parent in a cooked build. */
-	UE_DEPRECATED(5.0, "This flag is no longer in use and will be removed.")
-	UPROPERTY()
-	uint8 bHasNativizedParent_DEPRECATED:1;
-
 	/** Flag used to indicate if this class has data to support the component instancing fast path. */
 	UPROPERTY()
 	uint8 bHasCookedComponentInstancingData:1;
@@ -687,9 +679,6 @@ public:
 	/** Stores data to override (in children classes) components (created by SCS) from parent classes */
 	UPROPERTY()
 	TObjectPtr<class UInheritableComponentHandler> InheritableComponentHandler;
-
-	UPROPERTY()
-	TObjectPtr<class UStructProperty> UberGraphFramePointerProperty_DEPRECATED;
 	
 	FStructProperty* UberGraphFramePointerProperty;
 
@@ -870,11 +859,6 @@ protected:
 	*/
 	static ENGINE_API bool InitPropertyFromSubPropertyList(const FProperty* Property, const FCustomPropertyListNode* SubPropertyList, uint8* PropertyValue, const uint8* DefaultPropertyValue);
 
-	// @todo: BP2CPP_remove
-	/** Check for and handle manual application of default value overrides to component subobjects that were inherited from a nativized parent class */
-	UE_DEPRECATED(5.0, "This API is no longer in use and will be removed.")
-	static void CheckAndApplyComponentTemplateOverrides(UObject* InClassDefaultObject) {}
-
 public:
 
 	/** Called when the custom list of properties used during post-construct initialization needs to be rebuilt (e.g. after serialization and recompilation). */
diff --git a/Engine/Source/Runtime/Engine/Private/Animation/AnimClassData.cpp b/Engine/Source/Runtime/Engine/Private/Animation/AnimClassData.cpp
deleted file mode 100644
diff --git a/Engine/Source/Runtime/Engine/Private/Blueprint.cpp b/Engine/Source/Runtime/Engine/Private/Blueprint.cpp
--- a/Engine/Source/Runtime/Engine/Private/Blueprint.cpp
+++ b/Engine/Source/Runtime/Engine/Private/Blueprint.cpp
@@ -349,10 +349,6 @@ UBlueprint::UBlueprint(const FObjectInitializer& ObjectInitializer)
 #if WITH_EDITORONLY_DATA
 	, bDuplicatingReadOnly(false)
 	, bCachedDependenciesUpToDate(false)
-	// @todo: BP2CPP_remove
-	PRAGMA_DISABLE_DEPRECATION_WARNINGS
-	, bHasAnyNonReducibleFunction(EIsBPNonReducible::Unkown)
-	PRAGMA_ENABLE_DEPRECATION_WARNINGS
 #endif
 {
 }
diff --git a/Engine/Source/Runtime/Engine/Private/BlueprintGeneratedClass.cpp b/Engine/Source/Runtime/Engine/Private/BlueprintGeneratedClass.cpp
--- a/Engine/Source/Runtime/Engine/Private/BlueprintGeneratedClass.cpp
+++ b/Engine/Source/Runtime/Engine/Private/BlueprintGeneratedClass.cpp
@@ -387,7 +387,6 @@ namespace UE::Runtime::Engine::Private
 	};
 }
 
-PRAGMA_DISABLE_DEPRECATION_WARNINGS
 UBlueprintGeneratedClass::UBlueprintGeneratedClass(const FObjectInitializer& ObjectInitializer)
 	: Super(ObjectInitializer)
 #if VALIDATE_UBER_GRAPH_PERSISTENT_FRAME
@@ -395,15 +394,12 @@ UBlueprintGeneratedClass::UBlueprintGeneratedClass(const FObjectInitializer& Obj
 #endif//VALIDATE_UBER_GRAPH_PERSISTENT_FRAME
 {
 	NumReplicatedProperties = 0;
-	// @todo: BP2CPP_remove
-	bHasNativizedParent_DEPRECATED = false;
 	bHasCookedComponentInstancingData = false;
 	bCustomPropertyListForPostConstructionInitialized = false;
 #if WITH_EDITORONLY_DATA
 	bIsSparseClassDataSerializable = false;
 #endif
 }
-PRAGMA_ENABLE_DEPRECATION_WARNINGS
 
 void UBlueprintGeneratedClass::PostInitProperties()
 {
@@ -657,7 +653,7 @@ FPrimaryAssetId UBlueprintGeneratedClass::GetPrimaryAssetId() const
 
 UClass* UBlueprintGeneratedClass::GetAuthoritativeClass()
 {
- 	if (nullptr == ClassGeneratedBy) // to track UE-11597 and UE-11595
+ 	if (nullptr == ClassGeneratedBy)
  	{
 		// If this is a cooked blueprint, the generatedby class will have been discarded so we'll just have to assume we're authoritative!
 		if (bCooked || RootPackageHasAnyFlags(PKG_Cooked))
@@ -672,61 +668,9 @@ UClass* UBlueprintGeneratedClass::GetAuthoritativeClass()
 
 	UBlueprint* GeneratingBP = CastChecked<UBlueprint>(ClassGeneratedBy);
 
-	check(GeneratingBP);
-
 	return (GeneratingBP->GeneratedClass != NULL) ? GeneratingBP->GeneratedClass : this;
 }
 
-struct FConditionalRecompileClassHepler
-{
-	enum class ENeededAction : uint8
-	{
-		None,
-		StaticLink,
-		Recompile,
-	};
-
-	static bool HasTheSameLayoutAsParent(const UStruct* Struct)
-	{
-		const UStruct* Parent = Struct ? Struct->GetSuperStruct() : NULL;
-		return FStructUtils::TheSameLayout(Struct, Parent);
-	}
-
-	static ENeededAction IsConditionalRecompilationNecessary(const UBlueprint* GeneratingBP)
-	{
-		if (FBlueprintEditorUtils::IsInterfaceBlueprint(GeneratingBP))
-		{
-			return ENeededAction::None;
-		}
-
-		if (FBlueprintEditorUtils::IsDataOnlyBlueprint(GeneratingBP))
-		{
-			// If my parent is native, my layout wasn't changed.
-			const UClass* ParentClass = *GeneratingBP->ParentClass;
-			if (!GeneratingBP->GeneratedClass || (GeneratingBP->GeneratedClass->GetSuperClass() != ParentClass))
-			{
-				return ENeededAction::Recompile;
-			}
-
-			if (ParentClass && ParentClass->HasAllClassFlags(CLASS_Native))
-			{
-				return ENeededAction::None;
-			}
-
-			if (HasTheSameLayoutAsParent(*GeneratingBP->GeneratedClass))
-			{
-				return ENeededAction::StaticLink;
-			}
-			else
-			{
-				UE_LOG(LogBlueprint, Log, TEXT("During ConditionalRecompilation the layout of DataOnly BP should not be changed. It will be handled, but it's bad for performence. Blueprint %s"), *GeneratingBP->GetName());
-			}
-		}
-
-		return ENeededAction::Recompile;
-	}
-};
-
 void UBlueprintGeneratedClass::ConditionalRecompileClass(FUObjectSerializeContext* InLoadContext)
 {
 	FBlueprintCompilationManager::FlushCompilationQueue(InLoadContext);
@@ -1870,7 +1814,6 @@ uint8* UBlueprintGeneratedClass::GetPersistentUberGraphFrame(UObject* Obj, UFunc
 
 void UBlueprintGeneratedClass::CreatePersistentUberGraphFrame(UObject* Obj, bool bCreateOnlyIfEmpty, bool bSkipSuperClass, UClass* OldClass) const
 {
-#if USE_UBER_GRAPH_PERSISTENT_FRAME
 #if WITH_EDITORONLY_DATA
 	/** Macros should not create uber graph frames as they have no uber graph. If UBlueprints are cooked out the macro class probably does not exist as well */
 	UBlueprint* Blueprint = Cast<UBlueprint>(ClassGeneratedBy);
@@ -1929,12 +1872,10 @@ void UBlueprintGeneratedClass::CreatePersistentUberGraphFrame(UObject* Obj, bool
 		checkSlow(ParentClass);
 		ParentClass->CreatePersistentUberGraphFrame(Obj, bCreateOnlyIfEmpty);
 	}
-#endif // USE_UBER_GRAPH_PERSISTENT_FRAME
 }
 
 void UBlueprintGeneratedClass::DestroyPersistentUberGraphFrame(UObject* Obj, bool bSkipSuperClass) const
 {
-#if USE_UBER_GRAPH_PERSISTENT_FRAME
 	ensure(!UberGraphFramePointerProperty == !UberGraphFunction);
 	if (Obj && UsePersistentUberGraphFrame() && UberGraphFramePointerProperty && UberGraphFunction)
 	{
@@ -1963,7 +1904,6 @@ void UBlueprintGeneratedClass::DestroyPersistentUberGraphFrame(UObject* Obj, boo
 		checkSlow(ParentClass);
 		ParentClass->DestroyPersistentUberGraphFrame(Obj);
 	}
-#endif // USE_UBER_GRAPH_PERSISTENT_FRAME
 }
 
 void UBlueprintGeneratedClass::GetPreloadDependencies(TArray<UObject*>& OutDeps)
@@ -2283,7 +2223,6 @@ void UBlueprintGeneratedClass::Link(FArchive& Ar, bool bRelinkExistingProperties
 {
 	Super::Link(Ar, bRelinkExistingProperties);
 
-#if USE_UBER_GRAPH_PERSISTENT_FRAME
 	if (UsePersistentUberGraphFrame())
 	{
 		if (UberGraphFunction)
@@ -2301,7 +2240,6 @@ void UBlueprintGeneratedClass::Link(FArchive& Ar, bool bRelinkExistingProperties
 			checkSlow(UberGraphFramePointerProperty);
 		}
 	}
-#endif
 
 	AssembleReferenceTokenStream(true);
 }
@@ -2322,7 +2260,7 @@ void UBlueprintGeneratedClass::PurgeClass(bool bRecompilingOnLoad)
 	FastCallPairs_DEPRECATED.Empty();
 #endif
 	CalledFunctions.Empty();
-#endif //WITH_EDITOR
+#endif //WITH_EDITORONLY_DATA
 }
 
 void UBlueprintGeneratedClass::Bind()
@@ -2470,7 +2408,6 @@ void UBlueprintGeneratedClass::AddReferencedObjectsInUbergraphFrame(UObject* InT
 	
 	while (true)
 	{
-#if USE_UBER_GRAPH_PERSISTENT_FRAME
 		UFunction* UberGraphFunction = BPGC->UberGraphFunction;
 		if (BPGC->UberGraphFramePointerProperty)
 		{
@@ -2509,7 +2446,6 @@ void UBlueprintGeneratedClass::AddReferencedObjectsInUbergraphFrame(UObject* InT
 
 			}
 		}
-#endif // USE_UBER_GRAPH_PERSISTENT_FRAME
 
 		if (SuperClass->HasAllClassFlags(CLASS_Native))
 		{
@@ -2532,12 +2468,8 @@ FName UBlueprintGeneratedClass::GetUberGraphFrameName()
 
 bool UBlueprintGeneratedClass::UsePersistentUberGraphFrame()
 {
-#if USE_UBER_GRAPH_PERSISTENT_FRAME
 	static const FBoolConfigValueHelper PersistentUberGraphFrame(TEXT("Kismet"), TEXT("bPersistentUberGraphFrame"), GEngineIni);
 	return PersistentUberGraphFrame;
-#else
-	return false;
-#endif
 }
 
 #if VALIDATE_UBER_GRAPH_PERSISTENT_FRAME
@@ -2627,13 +2559,6 @@ void UBlueprintGeneratedClass::Serialize(FArchive& Ar)
 		}
 	}
 
-#if WITH_EDITORONLY_DATA
-	if (Ar.IsLoading())
-	{
-		UberGraphFramePointerProperty_DEPRECATED = nullptr;
-	}
-#endif
-
 #if WITH_EDITORONLY_DATA
 	if (Ar.IsSaving() && Ar.IsCooking() && Ar.IsObjectReferenceCollector())
 	{
