Update UClass::HasProperty to take a const FProperty*

#rnx

diff --git a/Engine/Source/Runtime/CoreUObject/Private/UObject/Class.cpp b/Engine/Source/Runtime/CoreUObject/Private/UObject/Class.cpp
--- a/Engine/Source/Runtime/CoreUObject/Private/UObject/Class.cpp
+++ b/Engine/Source/Runtime/CoreUObject/Private/UObject/Class.cpp
@@ -5772,7 +5772,7 @@ bool UClass::IsFunctionImplementedInScript(FName InFunctionName) const
 	return false;
 }
 
-bool UClass::HasProperty(FProperty* InProperty) const
+bool UClass::HasProperty(const FProperty* InProperty) const
 {
 	if (InProperty->GetOwner<UObject>())
 	{
diff --git a/Engine/Source/Runtime/CoreUObject/Public/UObject/Class.h b/Engine/Source/Runtime/CoreUObject/Public/UObject/Class.h
--- a/Engine/Source/Runtime/CoreUObject/Public/UObject/Class.h
+++ b/Engine/Source/Runtime/CoreUObject/Public/UObject/Class.h
@@ -3445,7 +3445,7 @@ public:
 	 * @param InProperty	The property to check if it is contained in this or a parent class.
 	 * @return				True if the property exists on this or a parent class.
 	 */
-	COREUOBJECT_API virtual bool HasProperty(FProperty* InProperty) const;
+	COREUOBJECT_API virtual bool HasProperty(const FProperty* InProperty) const;
 
 	/** Finds the object that is used as the parent object when serializing properties, overridden for blueprints */
 	virtual UObject* FindArchetype(const UClass* ArchetypeClass, const FName ArchetypeName) const { return nullptr; }
