Added an API to edit the attributes of a text property, akin to what happens when editing it in a details panel

This is needed when upgrading localized text data in an editor-like environment via a BP, as any String -> Text conversion within a BP produces an invariant (unlocalizable) text

* Moved the implementation of IEditableTextProperty::StaticStableTextId to TextNamespaceUtil::GetTextIdForEdit.
* Added TextNamespaceUtil::EditTextProperty to safely edit a given text attribute to a given value.
* Added TextNamespaceUtil::GenerateDeterministicTextKey as a way to build deterministic text keys for a given object and property pair.
* Added UKismetTextLibrary::EditTextSourceString to expose the above to BPs in a way that can safely edit the source string of a text property.

diff --git a/Engine/Source/Editor/EditorWidgets/Private/STextPropertyEditableTextBox.cpp b/Engine/Source/Editor/EditorWidgets/Private/STextPropertyEditableTextBox.cpp
--- a/Engine/Source/Editor/EditorWidgets/Private/STextPropertyEditableTextBox.cpp
+++ b/Engine/Source/Editor/EditorWidgets/Private/STextPropertyEditableTextBox.cpp
@@ -2,7 +2,6 @@
 
 #include "STextPropertyEditableTextBox.h"
 #include "AssetRegistry/AssetData.h"
-#include "Internationalization/TextNamespaceUtil.h"
 #include "Widgets/SBoxPanel.h"
 #include "Widgets/Images/SImage.h"
 #include "Widgets/Text/STextBlock.h"
@@ -21,7 +20,6 @@
 #include "Misc/PackageName.h"
 #include "AssetRegistry/AssetRegistryModule.h"
 #include "Internationalization/StringTable.h"
-#include "Internationalization/TextPackageNamespaceUtil.h"
 #include "Internationalization/StringTableCore.h"
 #include "Internationalization/StringTableRegistry.h"
 #include "Serialization/TextReferenceCollector.h"
@@ -43,46 +41,7 @@ void IEditableTextProperty::StaticStableTextId(UObject* InObject, const ETextPro
 
 void IEditableTextProperty::StaticStableTextId(UPackage* InPackage, const ETextPropertyEditAction InEditAction, const FString& InTextSource, const FString& InProposedNamespace, const FString& InProposedKey, FString& OutStableNamespace, FString& OutStableKey)
 {
-	bool bPersistKey = false;
-
-	const FString PackageNamespace = TextNamespaceUtil::EnsurePackageNamespace(InPackage);
-	if (!PackageNamespace.IsEmpty())
-	{
-		// Make sure the proposed namespace is using the correct namespace for this package
-		OutStableNamespace = TextNamespaceUtil::BuildFullNamespace(InProposedNamespace, PackageNamespace, /*bAlwaysApplyPackageNamespace*/true);
-
-		if (InProposedNamespace.Equals(OutStableNamespace, ESearchCase::CaseSensitive) || InEditAction == ETextPropertyEditAction::EditedNamespace)
-		{
-			// If the proposal was already using the correct namespace (or we just set the namespace), attempt to persist the proposed key too
-			if (!InProposedKey.IsEmpty())
-			{
-				// If we changed the source text, then we can persist the key if this text is the *only* reference using that ID
-				// If we changed the identifier, then we can persist the key only if doing so won't cause an identify conflict
-				const FTextReferenceCollector::EComparisonMode ReferenceComparisonMode = InEditAction == ETextPropertyEditAction::EditedSource ? FTextReferenceCollector::EComparisonMode::MatchId : FTextReferenceCollector::EComparisonMode::MismatchSource;
-				const int32 RequiredReferenceCount = InEditAction == ETextPropertyEditAction::EditedSource ? 1 : 0;
-
-				int32 ReferenceCount = 0;
-				FTextReferenceCollector(InPackage, ReferenceComparisonMode, OutStableNamespace, InProposedKey, InTextSource, ReferenceCount);
-
-				if (ReferenceCount == RequiredReferenceCount)
-				{
-					bPersistKey = true;
-					OutStableKey = InProposedKey;
-				}
-			}
-		}
-		else if (InEditAction != ETextPropertyEditAction::EditedNamespace)
-		{
-			// If our proposed namespace wasn't correct for our package, and we didn't just set it (which doesn't include the package namespace)
-			// then we should clear out any user specified part of it
-			OutStableNamespace = TextNamespaceUtil::BuildFullNamespace(FString(), PackageNamespace, /*bAlwaysApplyPackageNamespace*/true);
-		}
-	}
-
-	if (!bPersistKey)
-	{
-		OutStableKey = FGuid::NewGuid().ToString();
-	}
+	TextNamespaceUtil::GetTextIdForEdit(InPackage, (TextNamespaceUtil::ETextEditAction)InEditAction, InTextSource, InProposedNamespace, InProposedKey, OutStableNamespace, OutStableKey);
 }
 
 #endif // USE_STABLE_LOCALIZATION_KEYS
diff --git a/Engine/Source/Editor/EditorWidgets/Public/STextPropertyEditableTextBox.h b/Engine/Source/Editor/EditorWidgets/Public/STextPropertyEditableTextBox.h
--- a/Engine/Source/Editor/EditorWidgets/Public/STextPropertyEditableTextBox.h
+++ b/Engine/Source/Editor/EditorWidgets/Public/STextPropertyEditableTextBox.h
@@ -14,7 +14,7 @@
 #include "HAL/PlatformCrt.h"
 #include "Input/Reply.h"
 #include "Internationalization/Text.h"
-#include "Internationalization/TextNamespaceFwd.h"
+#include "Internationalization/TextPackageNamespaceUtil.h"
 #include "Layout/Visibility.h"
 #include "Misc/Attribute.h"
 #include "Misc/Optional.h"
@@ -53,9 +53,9 @@ class IEditableTextProperty
 public:
 	enum class ETextPropertyEditAction : uint8
 	{
-		EditedNamespace,
-		EditedKey,
-		EditedSource,
+		EditedNamespace = (uint8)TextNamespaceUtil::ETextEditAction::Namespace,
+		EditedKey = (uint8)TextNamespaceUtil::ETextEditAction::Key,
+		EditedSource = (uint8)TextNamespaceUtil::ETextEditAction::SourceString,
 	};
 
 	virtual ~IEditableTextProperty() = default;
diff --git a/Engine/Source/Runtime/Core/Public/Internationalization/TextNamespaceUtil.h b/Engine/Source/Runtime/Core/Public/Internationalization/TextNamespaceUtil.h
--- a/Engine/Source/Runtime/Core/Public/Internationalization/TextNamespaceUtil.h
+++ b/Engine/Source/Runtime/Core/Public/Internationalization/TextNamespaceUtil.h
@@ -22,7 +22,7 @@ constexpr inline TCHAR PackageNamespaceEndMarker = TEXT(']');
  * @param InTextNamespace				The namespace currently used by the FText instance.
  * @param InPackageNamespace			The namespace of the package owning the FText instance.
  * @param bAlwaysApplyPackageNamespace	If true, this will always apply the package namespace to the text namespace.
- *										If false, this will only apply the package namespace if the text namespace already contains package namespace makers.
+ *										If false, this will only apply the package namespace if the text namespace already contains package namespace markers.
  *
  * @return The full namespace that should be used by the localization system.
  */
@@ -65,8 +65,8 @@ enum class ETextCopyMethod : uint8
  * @param InText						The current FText instance.
  * @param InPackageNamespace			The namespace of the destination package of the FText instance.
  * @param InCopyMethod					The method that should be used to copy the FText instance.
- * @param bAlwaysApplyPackageNamespace	If true, this will always apply the package namespace to the text namespace (always treated as Verbatim when USE_STABLE_LOCALIZATION_KEYS is false).
- *										If false, this will only apply the package namespace if the text namespace already contains package namespace makers.
+ * @param bAlwaysApplyPackageNamespace	If true, this will always apply the package namespace to the text namespace (always treated as ETextCopyMethod::Verbatim when USE_STABLE_LOCALIZATION_KEYS is false).
+ *										If false, this will only apply the package namespace if the text namespace already contains package namespace markers.
  *
  * @return A copy of the given text that's valid to use with the given package namespace.
  */
diff --git a/Engine/Source/Runtime/CoreUObject/Private/Internationalization/TextPackageNamespaceUtil.cpp b/Engine/Source/Runtime/CoreUObject/Private/Internationalization/TextPackageNamespaceUtil.cpp
--- a/Engine/Source/Runtime/CoreUObject/Private/Internationalization/TextPackageNamespaceUtil.cpp
+++ b/Engine/Source/Runtime/CoreUObject/Private/Internationalization/TextPackageNamespaceUtil.cpp
@@ -5,8 +5,11 @@
 #include "Hash/Blake3.h"
 #include "Misc/Guid.h"
 #include "Misc/PackageName.h"
-#include "UObject/Package.h"
+#include "Serialization/TextReferenceCollector.h"
+#include "UObject/Class.h"
 #include "UObject/MetaData.h"
+#include "UObject/Package.h"
+#include "UObject/TextProperty.h"
 
 #if USE_STABLE_LOCALIZATION_KEYS
 
@@ -168,3 +171,149 @@ FText TextNamespaceUtil::CopyTextToPackage(const FText& InText, UObject* InObjec
 	return CopyTextToPackage(InText, FString(), InCopyMethod, bAlwaysApplyPackageNamespace);
 #endif	// USE_STABLE_LOCALIZATION_KEYS
 }
+
+FString TextNamespaceUtil::GenerateRandomTextKey()
+{
+	return FGuid::NewGuid().ToString();
+}
+
+FString TextNamespaceUtil::GenerateDeterministicTextKey(UObject* InTextOwner, const FTextProperty* InTextProperty, const bool bApplyPackageNamespace)
+{
+	auto GetNameKeyHash = [](const FName Name)
+	{
+		FNameBuilder Builder(Name);
+		return TextKeyUtil::HashString(Builder.ToString(), Builder.Len());
+	};
+
+	auto GetObjectKeyHash = [](const UObject* Obj)
+	{
+		FNameBuilder Builder;
+		if (Obj)
+		{
+			Obj->GetPathName(nullptr, Builder);
+		}
+		return TextKeyUtil::HashString(Builder.ToString(), Builder.Len());
+	};
+
+	// Build a (hopefully) unique and deterministic key from a combination of the text owner and text property info
+	FGuid KeyGuid(GetObjectKeyHash(InTextOwner->GetOuter()), GetNameKeyHash(InTextOwner->GetFName()), GetNameKeyHash(InTextOwner->GetClass()->GetFName()), GetNameKeyHash(InTextProperty->GetFName()));
+#if USE_STABLE_LOCALIZATION_KEYS
+	if (const FString PackageNamespace = bApplyPackageNamespace ? EnsurePackageNamespace(InTextOwner) : FString();
+		!PackageNamespace.IsEmpty())
+	{
+		// Mix the package namespace hash into the outer hash
+		KeyGuid.A = TextKeyUtil::HashString(PackageNamespace, KeyGuid.A);
+	}
+#endif // USE_STABLE_LOCALIZATION_KEYS
+	return KeyGuid.ToString();
+}
+
+void TextNamespaceUtil::GetTextIdForEdit(UPackage* InPackage, const ETextEditAction InEditAction, const FString& InTextSource, const FString& InProposedNamespace, const FString& InProposedKey, FString& OutStableNamespace, FString& OutStableKey, TFunctionRef<FString()> InTextKeyGenerator, const bool bApplyPackageNamespace)
+{
+#if USE_STABLE_LOCALIZATION_KEYS
+	bool bPersistKey = false;
+
+	if (const FString PackageNamespace = bApplyPackageNamespace ? EnsurePackageNamespace(InPackage) : FString();
+		!PackageNamespace.IsEmpty())
+	{
+		// Make sure the proposed namespace is using the correct namespace for this package
+		OutStableNamespace = BuildFullNamespace(InProposedNamespace, PackageNamespace, /*bAlwaysApplyPackageNamespace*/true);
+
+		if (InProposedNamespace.Equals(OutStableNamespace, ESearchCase::CaseSensitive) || InEditAction == ETextEditAction::Namespace)
+		{
+			// If the proposal was already using the correct namespace (or we just set the namespace), attempt to persist the proposed key too
+			if (!InProposedKey.IsEmpty())
+			{
+				// If we changed the source text, then we can persist the key if this text is the *only* reference using that ID
+				// If we changed the identifier, then we can persist the key only if doing so won't cause an identify conflict
+				const FTextReferenceCollector::EComparisonMode ReferenceComparisonMode = InEditAction == ETextEditAction::SourceString ? FTextReferenceCollector::EComparisonMode::MatchId : FTextReferenceCollector::EComparisonMode::MismatchSource;
+				const int32 RequiredReferenceCount = InEditAction == ETextEditAction::SourceString ? 1 : 0;
+
+				int32 ReferenceCount = 0;
+				FTextReferenceCollector Collector(InPackage, ReferenceComparisonMode, OutStableNamespace, InProposedKey, InTextSource, ReferenceCount);
+
+				if (ReferenceCount == RequiredReferenceCount)
+				{
+					bPersistKey = true;
+					OutStableKey = InProposedKey;
+				}
+			}
+		}
+		else if (InEditAction != ETextEditAction::Namespace)
+		{
+			// If our proposed namespace wasn't correct for our package, and we didn't just set it (which doesn't include the package namespace)
+			// then we should clear out any user specified part of it
+			OutStableNamespace = BuildFullNamespace(FString(), PackageNamespace, /*bAlwaysApplyPackageNamespace*/true);
+		}
+	}
+
+	if (!bPersistKey)
+#endif // USE_STABLE_LOCALIZATION_KEYS
+	{
+		OutStableKey = InTextKeyGenerator();
+	}
+}
+
+bool TextNamespaceUtil::EditTextProperty(UObject* InTextOwner, const FTextProperty* InTextProperty, const ETextEditAction InEditAction, const FString& InEditValue, TFunctionRef<FString()> InTextKeyGenerator, const bool bApplyPackageNamespace)
+{
+	if (!InTextOwner->GetClass()->HasProperty(InTextProperty))
+	{
+		return false;
+	}
+
+	if (InEditAction == ETextEditAction::SourceString && InEditValue.IsEmpty())
+	{
+		// Empty source strings always produce an empty text
+		InTextProperty->SetPropertyValue_InContainer(InTextOwner, FText());
+		return true;
+	}
+
+	const FText CurrentTextValue = InTextProperty->GetPropertyValue_InContainer(InTextOwner);
+	const FTextId CurrentTextId = FTextInspector::GetTextId(CurrentTextValue);
+	const FString* CurrentSourceString = FTextInspector::GetSourceString(CurrentTextValue);
+	const bool bIsCurrentTextLocalized = !CurrentTextValue.IsCultureInvariant() && !CurrentTextValue.IsFromStringTable();
+
+	// Verify the edited attribute has actually changed
+	if (bIsCurrentTextLocalized)
+	{
+		switch (InEditAction)
+		{
+		case ETextEditAction::Namespace:
+			if (const FString CurrentTextNamespace = StripPackageNamespace(CurrentTextId.GetNamespace().GetChars());
+				CurrentTextNamespace.Equals(InEditValue, ESearchCase::CaseSensitive))
+			{
+				return true;
+			}
+			break;
+
+		case ETextEditAction::Key:
+			if (FCString::Strcmp(CurrentTextId.GetKey().GetChars(), *InEditValue) == 0)
+			{
+				return true;
+			}
+			break;
+
+		case ETextEditAction::SourceString:
+			if (CurrentSourceString && CurrentSourceString->Equals(InEditValue, ESearchCase::CaseSensitive))
+			{
+				return true;
+			}
+			break;
+
+		default:
+			checkf(false, TEXT("Unknown ETextEditAction!"));
+			break;
+		}
+	}
+
+	const FString ProposedNamespace = (InEditAction == ETextEditAction::Namespace ? InEditValue : CurrentTextId.GetNamespace().GetChars());
+	const FString ProposedKey = (InEditAction == ETextEditAction::Key ? InEditValue : CurrentTextId.GetKey().GetChars());
+	const FString SourceString = (InEditAction == ETextEditAction::SourceString ? InEditValue : CurrentSourceString ? *CurrentSourceString : FString());
+
+	FString StableNamespace;
+	FString StableKey;
+	GetTextIdForEdit(InTextOwner->GetPackage(), InEditAction, SourceString, ProposedNamespace, ProposedKey, StableNamespace, StableKey, InTextKeyGenerator, bApplyPackageNamespace);
+
+	InTextProperty->SetPropertyValue_InContainer(InTextOwner, FInternationalization::ForUseOnlyByLocMacroAndGraphNodeTextLiterals_CreateText(*SourceString, *StableNamespace, *StableKey));
+	return true;
+}
diff --git a/Engine/Source/Runtime/CoreUObject/Public/Internationalization/TextPackageNamespaceUtil.h b/Engine/Source/Runtime/CoreUObject/Public/Internationalization/TextPackageNamespaceUtil.h
--- a/Engine/Source/Runtime/CoreUObject/Public/Internationalization/TextPackageNamespaceUtil.h
+++ b/Engine/Source/Runtime/CoreUObject/Public/Internationalization/TextPackageNamespaceUtil.h
@@ -8,6 +8,7 @@
 #include "Internationalization/TextNamespaceFwd.h"
 #include "Internationalization/TextNamespaceUtil.h"
 
+class FTextProperty;
 class UObject;
 class UPackage;
 
@@ -91,12 +92,64 @@ COREUOBJECT_API void ForcePackageNamespace(UObject* InObject, const FString& InN
  * @param InText						The current FText instance.
  * @param InPackage/InObject			The package (or object to get the owner package from) to get the namespace for (will call EnsurePackageNamespace).
  * @param InCopyMethod					The method that should be used to copy the FText instance.
- * @param bAlwaysApplyPackageNamespace	If true, this will always apply the package namespace to the text namespace (always treated as Verbatim when USE_STABLE_LOCALIZATION_KEYS is false).
- *										If false, this will only apply the package namespace if the text namespace already contains package namespace makers.
+ * @param bAlwaysApplyPackageNamespace	If true, this will always apply the package namespace to the text namespace (always treated as ETextCopyMethod::Verbatim when USE_STABLE_LOCALIZATION_KEYS is false).
+ *										If false, this will only apply the package namespace if the text namespace already contains package namespace markers.
  *
  * @return A copy of the given text that's valid to use with the given package.
  */
 COREUOBJECT_API FText CopyTextToPackage(const FText& InText, UPackage* InPackage, const ETextCopyMethod InCopyMethod = ETextCopyMethod::NewKey, const bool bAlwaysApplyPackageNamespace = false);
 COREUOBJECT_API FText CopyTextToPackage(const FText& InText, UObject* InObject, const ETextCopyMethod InCopyMethod = ETextCopyMethod::NewKey, const bool bAlwaysApplyPackageNamespace = false);
 
+/**
+ * Generate a random text key.
+ * @note This key will be a GUID.
+ */
+COREUOBJECT_API FString GenerateRandomTextKey();
+
+/**
+ * Generate a deterministic text key based on the given object and property info.
+ * @note This key will be formatted like a GUID, but the value will actually be based on deterministic hashes.
+ *
+ * @param InTextOwner					The object that owns the given TextProperty.
+ * @param InTextProperty				The text property to generate the key for.
+ * @param bApplyPackageNamespace		If true, apply the package namespace to the generated key hash (when USE_STABLE_LOCALIZATION_KEYS is true).
+ */
+COREUOBJECT_API FString GenerateDeterministicTextKey(UObject* InTextOwner, const FTextProperty* InTextProperty, const bool bApplyPackageNamespace = true);
+
+enum class ETextEditAction : uint8
+{
+	Namespace,
+	Key,
+	SourceString,
+};
+
+/**
+ * Called when editing a text property to determine the new ID for the text, ideally using the proposed text ID when possible (and when USE_STABLE_LOCALIZATION_KEYS is true).
+ * 
+ * @param InPackage						The package to query the namespace for.
+ * @param InEditAction					How has the given text been edited?
+ * @param InTextSource					The current source string for the text being edited. Can be empty when InEditAction is ETextEditAction::EditedSource.
+ * @param InProposedNamespace			The namespace we'd like to assign to the edited text.
+ * @param InProposedKey					The key we'd like to assign to the edited text.
+ * @param OutStableNamespace			The namespace that should be assigned to the edited text.
+ * @param OutStableKey					The key that should be assigned to the edited text.
+ * @param InTextKeyGenerator			Generator for the new text key. Will generate a random key by default.
+ * @param bApplyPackageNamespace		If true, apply the package namespace to the generated text ID (when USE_STABLE_LOCALIZATION_KEYS is true).
+ */
+COREUOBJECT_API void GetTextIdForEdit(UPackage* InPackage, const ETextEditAction InEditAction, const FString& InTextSource, const FString& InProposedNamespace, const FString& InProposedKey, FString& OutStableNamespace, FString& OutStableKey, TFunctionRef<FString()> InTextKeyGenerator = &GenerateRandomTextKey, const bool bApplyPackageNamespace = true);
+
+/**
+ * Edit an attribute of the given text property, akin to what happens when editing a text property in a details panel.
+ *
+ * @param InTextOwner					The object that owns the given TextProperty to be edited.
+ * @param InTextProperty				The text property to edit. This must be a property that exists on TextOwner.
+ * @param InEditAction					How has the given text been edited?
+ * @param InEditValue					The new value of the attribute that was edited.
+ * @param InTextKeyGenerator			Generator for the new text key. Will generate a random key by default.
+ * @param bApplyPackageNamespace		If true, apply the package namespace to the generated text ID (when USE_STABLE_LOCALIZATION_KEYS is true).
+ * 
+ * @return True if edit was possible, or false if not.
+ */
+COREUOBJECT_API bool EditTextProperty(UObject* InTextOwner, const FTextProperty* InTextProperty, const ETextEditAction InEditAction, const FString& InEditValue, TFunctionRef<FString()> InTextKeyGenerator = &GenerateRandomTextKey, const bool bApplyPackageNamespace = true);
+
 }
diff --git a/Engine/Source/Runtime/Engine/Classes/Kismet/KismetTextLibrary.h b/Engine/Source/Runtime/Engine/Classes/Kismet/KismetTextLibrary.h
--- a/Engine/Source/Runtime/Engine/Classes/Kismet/KismetTextLibrary.h
+++ b/Engine/Source/Runtime/Engine/Classes/Kismet/KismetTextLibrary.h
@@ -313,4 +313,20 @@ class UKismetTextLibrary : public UBlueprintFunctionLibrary
 	 */
 	UFUNCTION(BlueprintPure, Category="Utilities|Text")
 	static ENGINE_API FText PolyglotDataToText(const FPolyglotTextData& PolyglotData);
+
+	/**
+	 * Edit the source string of the given text property, akin to what happens when editing a text property in a details panel.
+	 * This will attempt to preserve the existing ID of the text property being edited, or failing that will attempt to build a deterministic ID based on the object and property info.
+	 * 
+	 * @note This is an ADVANCED function that is ONLY safe to be used in environments where the modified text property will be gathered for localization (eg, in the editor, or a game mode that collects text properties to be localized).
+	 * 
+	 * @param TextOwner		The object that owns the given Text to be edited.
+	 * @param Text			The text property to edit. This must be a property that exists on TextOwner.
+	 * @param SourceString	The source string that the edited text property should use.
+	 * 
+	 * @return True if edit was possible, or false if not.
+	 */
+	UFUNCTION(BlueprintCallable, CustomThunk, Category="Utilities|Text", meta=(DefaultToSelf="TextOwner"))
+	static ENGINE_API bool EditTextSourceString(UObject* TextOwner, UPARAM(ref) FText& Text, const FString& SourceString);
+	DECLARE_FUNCTION(execEditTextSourceString);
 };
diff --git a/Engine/Source/Runtime/Engine/Private/KismetTextLibrary.cpp b/Engine/Source/Runtime/Engine/Private/KismetTextLibrary.cpp
--- a/Engine/Source/Runtime/Engine/Private/KismetTextLibrary.cpp
+++ b/Engine/Source/Runtime/Engine/Private/KismetTextLibrary.cpp
@@ -1,9 +1,13 @@
 // Copyright Epic Games, Inc. All Rights Reserved.
 
 #include "Kismet/KismetTextLibrary.h"
+
+#include "Engine/World.h"
 #include "Internationalization/PolyglotTextData.h"
 #include "Internationalization/TextFormatter.h"
 #include "Internationalization/TextKey.h"
+#include "Internationalization/TextPackageNamespaceUtil.h"
+#include "Misc/RuntimeErrors.h"
 
 #define LOCTEXT_NAMESPACE "Kismet"
 
@@ -339,4 +343,77 @@ FText UKismetTextLibrary::PolyglotDataToText(const FPolyglotTextData& PolyglotDa
 	return PolyglotData.GetText();
 }
 
+bool UKismetTextLibrary::EditTextSourceString(UObject* TextOwner, FText& Text, const FString& SourceString)
+{
+	// We should never hit this! Stubbed to avoid NoExport on the class.
+	check(0);
+	return false;
+}
+
+DEFINE_FUNCTION(UKismetTextLibrary::execEditTextSourceString)
+{
+	P_GET_OBJECT(UObject, TextOwner);
+
+	P_GET_PROPERTY_REF(FTextProperty, Text);
+	FTextProperty* TextProperty = CastField<FTextProperty>(Stack.MostRecentProperty);
+
+	P_GET_PROPERTY_REF(FStrProperty, SourceString);
+
+	P_FINISH;
+
+	P_NATIVE_BEGIN;
+	{
+		*(bool*)RESULT_PARAM = false;
+
+		if (!TextOwner)
+		{
+			LogRuntimeWarning(LOCTEXT("EditTextSourceString.Warning.NullTextOwner", "The given TextOwner was null!"));
+			return;
+		}
+
+		if (!TextProperty)
+		{
+			LogRuntimeWarning(LOCTEXT("EditTextSourceString.Warning.NullTextProperty", "The given Text value was not a TextProperty!"));
+			return;
+		}
+
+		if (!TextOwner->GetClass()->HasProperty(TextProperty))
+		{
+			LogRuntimeWarning(FText::Format(LOCTEXT("EditTextSourceString.Warning.InvalidTextProperty", "The given Text resolved to a TextProperty ({0}) that doesn't belong to the given TextOwner ({1})!"), FText::AsCultureInvariant(TextProperty->GetPathName()), FText::AsCultureInvariant(TextOwner->GetPathName())));
+			return;
+		}
+		
+		auto TextOwnerIsInEditorWorld = [TextOwner]()
+		{
+			const UWorld* WorldContext = TextOwner->GetWorld();
+			return WorldContext && WorldContext->WorldType == EWorldType::Editor;
+		};
+
+		auto TextOwnerIsInAsset = [TextOwner]()
+		{
+			for (const UObject* Obj = TextOwner; Obj; Obj = Obj->GetOuter())
+			{
+				if (Obj->IsAsset())
+				{
+					return true;
+				}
+			}
+			return false;
+		};
+
+		const bool bApplyPackageNamespace = GIsEditor && (TextOwnerIsInEditorWorld() || TextOwnerIsInAsset());
+		auto DeterministicTextKeyGenerator = [TextOwner, TextProperty, bApplyPackageNamespace]()
+		{
+			return TextNamespaceUtil::GenerateDeterministicTextKey(TextOwner, TextProperty, bApplyPackageNamespace);
+		};
+
+		if (TextNamespaceUtil::EditTextProperty(TextOwner, TextProperty, TextNamespaceUtil::ETextEditAction::SourceString, SourceString, DeterministicTextKeyGenerator, bApplyPackageNamespace))
+		{
+			*(bool*)RESULT_PARAM = true;
+			Text = TextProperty->GetPropertyValue_InContainer(TextOwner);
+		}
+	}
+	P_NATIVE_END;
+}
+
 #undef LOCTEXT_NAMESPACE
