Fixed incorrect comment

diff --git a/Engine/Source/Runtime/CoreUObject/Public/Internationalization/TextPackageNamespaceUtil.h b/Engine/Source/Runtime/CoreUObject/Public/Internationalization/TextPackageNamespaceUtil.h
--- a/Engine/Source/Runtime/CoreUObject/Public/Internationalization/TextPackageNamespaceUtil.h
+++ b/Engine/Source/Runtime/CoreUObject/Public/Internationalization/TextPackageNamespaceUtil.h
@@ -128,7 +128,7 @@ enum class ETextEditAction : uint8
  * 
  * @param InPackage						The package to query the namespace for.
  * @param InEditAction					How has the given text been edited?
- * @param InTextSource					The current source string for the text being edited. Can be empty when InEditAction is ETextEditAction::EditedSource.
+ * @param InTextSource					The current source string for the text being edited. Can be empty when InEditAction is ETextEditAction::SourceString.
  * @param InProposedNamespace			The namespace we'd like to assign to the edited text.
  * @param InProposedKey					The key we'd like to assign to the edited text.
  * @param OutStableNamespace			The namespace that should be assigned to the edited text.
