Core: Allow TScriptInterface to avoid a Cast call if we know at compile time that we got an appropriate subclass. This prevents TScriptInterface from requiring linking to a UInterface's module just to use the constructor.

diff --git a/Engine/Source/Runtime/CoreUObject/Public/UObject/ScriptInterface.h b/Engine/Source/Runtime/CoreUObject/Public/UObject/ScriptInterface.h
--- a/Engine/Source/Runtime/CoreUObject/Public/UObject/ScriptInterface.h
+++ b/Engine/Source/Runtime/CoreUObject/Public/UObject/ScriptInterface.h
@@ -163,9 +163,18 @@ public:
 		UObject* SourceObject = ImplicitConv<UObject*>(Source);
 		SetObject(SourceObject);
 
-		// Tries to set the native interface instance, this will set it to null for BP-implemented interfaces
-		InInterfaceType* SourceInterface = Cast<InInterfaceType>(SourceObject);
-		SetInterface(SourceInterface);
+		if constexpr (std::is_base_of<InInterfaceType, std::remove_pointer_t<std::remove_reference_t<U>>>::value)
+		{
+			// If we know at compile time that we got passed some subclass of InInterfaceType, set it
+			// without a cast (avoiding the cast also allows us to not require linking to its module)
+			SetInterface(Source);
+		}
+		else
+		{
+			// Tries to set the native interface instance, this will set it to null for BP-implemented interfaces
+			InInterfaceType* SourceInterface = Cast<InInterfaceType>(SourceObject);
+			SetInterface(SourceInterface);
+		}
 	}
 
 	/**
@@ -192,9 +201,18 @@ public:
 		// Always set the object
 		SetObject(SourceObject);
 
-		// Tries to set the native interface instance, this will set it to null for BP-implemented interfaces
-		InInterfaceType* SourceInterface = Cast<InInterfaceType>(ToRawPtr(SourceObject));
-		SetInterface(SourceInterface);
+		if constexpr (std::is_base_of<InInterfaceType, ObjectType>::value)
+		{
+			// If we know at compile time that we got passed some subclass of InInterfaceType, set it
+			// without a cast (avoiding the cast also allows us to not require linking to its module)
+			SetInterface(SourceObject.Get());
+		}
+		else
+		{
+			// Tries to set the native interface instance, this will set it to null for BP-implemented interfaces
+			InInterfaceType* SourceInterface = Cast<InInterfaceType>(ToRawPtr(SourceObject));
+			SetInterface(SourceInterface);
+		}
 	}
 
 	/**
