CookIndeterminism: AsyncLoading2: Set PackageNames to match capitalization of the file on disk so that differently capitalized imports from other packages will not cause cook indeterminism in the capitalization of the package's name.
#rnx

diff --git a/Engine/Source/Runtime/CoreUObject/Private/Serialization/AsyncLoading2.cpp b/Engine/Source/Runtime/CoreUObject/Private/Serialization/AsyncLoading2.cpp
--- a/Engine/Source/Runtime/CoreUObject/Private/Serialization/AsyncLoading2.cpp
+++ b/Engine/Source/Runtime/CoreUObject/Private/Serialization/AsyncLoading2.cpp
@@ -3576,12 +3576,34 @@ public:
 	}
 
 private:
-#if ALT2_ENABLE_LINKERLOAD_SUPPORT
-	bool ShouldLoadPackageFromFileSystem(FName PackageNameToLoad, FPackagePath& OutPackagePath)
+#if ALT2_ENABLE_LINKERLOAD_SUPPORT || WITH_EDITOR
+	bool TryGetPackagePathFromFileSystem(FName& PackageNameToLoad, FName& UPackageName, FPackagePath& OutPackagePath)
 	{
-		return !PackageNameToLoad.IsNone() &&
+#if WITH_EDITORONLY_DATA
+		// In editor, set MatchCaseOnDisk=true so that we set the capitalization of the Package's FName to match the
+		// capitalization on disk. Different capitalizations can arise from imports of the package that were somehow
+		// constructed with a different captialization (most often because the disk captialization changed).
+		// We need the captialization to match so that source control operations in case-significant source control
+		// depots succeed, and to avoid indetermism in the cook.
+		constexpr bool bMatchCaseOnDisk = true;
+#else
+		constexpr bool bMatchCaseOnDisk = false;
+#endif
+
+		if (!PackageNameToLoad.IsNone() &&
 			FPackagePath::TryFromPackageName(PackageNameToLoad, OutPackagePath) &&
-			FPackageName::DoesPackageExistEx(OutPackagePath, FPackageName::EPackageLocationFilter::FileSystem, false /* bMatchCaseOnDisk */, &OutPackagePath) != FPackageName::EPackageLocationFilter::None;
+			FPackageName::DoesPackageExistEx(OutPackagePath, FPackageName::EPackageLocationFilter::FileSystem,
+				bMatchCaseOnDisk, &OutPackagePath) != FPackageName::EPackageLocationFilter::None)
+		{
+			FName CaseCorrectedPackageName = OutPackagePath.GetPackageFName();
+			if (PackageNameToLoad == UPackageName)
+			{
+				UPackageName = CaseCorrectedPackageName;
+			}
+			PackageNameToLoad = CaseCorrectedPackageName;
+			return true;
+		}
+		return false;
 	}
 #endif
 
@@ -4370,7 +4392,8 @@ bool FAsyncLoadingThread2::CreateAsyncPackagesFromQueue(FAsyncLoadingThreadState
 
 #if ALT2_ENABLE_LINKERLOAD_SUPPORT
 			bool bIsZenPackage = true;
-			if (ShouldLoadPackageFromFileSystem(PackageNameToLoad, Request.PackagePath))
+			FName CorrectedPackageName;
+			if (TryGetPackagePathFromFileSystem(PackageNameToLoad, UPackageName, Request.PackagePath))
 			{
 				bIsZenPackage = false;
 				PackageStatus = EPackageStoreEntryStatus::Ok;
@@ -4978,7 +5001,8 @@ void FAsyncPackage2::ImportPackagesRecursiveInner(FAsyncLoadingThreadState2& Thr
 #if ALT2_ENABLE_LINKERLOAD_SUPPORT || WITH_EDITOR
 		bool bIsZenPackage = !LinkerLoadState.IsSet();
 		bool bIsZenPackageImport = true;
-		if (AsyncLoadingThread.ShouldLoadPackageFromFileSystem(ImportedPackageNameToLoad, ImportedPackagePath))
+		FName CorrectedPackageName;
+		if (AsyncLoadingThread.TryGetPackagePathFromFileSystem(ImportedPackageNameToLoad, ImportedPackageUPackageName, ImportedPackagePath))
 		{
 			bIsZenPackageImport = false;
 			ImportedPackageStatus = EPackageStoreEntryStatus::Ok;
