SparseClassData and GarbageCollection: Collect references from the SparseClassDataStruct using AddPropertyReferencesWithStructARO rather than with the minimal AddPropertyReferencesLimitedToObjectProperties. This is necessary because nested structs on some SparseClassDataStructs, such as a TInstancedStruct, were not having their StructARO called and were therefore having some of their references skipped, when the SparseClassDataStruct was collected with only AddPropertyReferencesLimitedToObjectProperties.
#rnx

diff --git a/Engine/Source/Runtime/CoreUObject/Private/UObject/Class.cpp b/Engine/Source/Runtime/CoreUObject/Private/UObject/Class.cpp
--- a/Engine/Source/Runtime/CoreUObject/Private/UObject/Class.cpp
+++ b/Engine/Source/Runtime/CoreUObject/Private/UObject/Class.cpp
@@ -4404,16 +4404,7 @@ void UClass::AddReferencedObjects(UObject* InThis, FReferenceCollector& Collecto
 	// Add sparse class data
 	if (This->SparseClassDataStruct && This->SparseClassData)
 	{
-		if (This->SparseClassDataStruct->StructFlags & STRUCT_AddStructReferencedObjects)
-		{
-			This->SparseClassDataStruct->GetCppStructOps()->AddStructReferencedObjects()(This->SparseClassData, Collector);
-		}
-		else
-		{
-			// Skip FFieldPathProperty and FInterfaceProperty references since original code iterated over
-			// FObjectProperty only. Unsure if sparse class author intended this.
-			Collector.AddPropertyReferencesLimitedToObjectProperties(This->SparseClassDataStruct, This->SparseClassData, This);
-		}
+		Collector.AddPropertyReferencesWithStructARO(This->SparseClassDataStruct, This->SparseClassData, This);
 	}
 	
 	Super::AddReferencedObjects( This, Collector );
