Fixed transaction with new Container Overridable Logic

diff --git a/Engine/Source/Runtime/CoreUObject/Private/UObject/PropertyArray.cpp b/Engine/Source/Runtime/CoreUObject/Private/UObject/PropertyArray.cpp
--- a/Engine/Source/Runtime/CoreUObject/Private/UObject/PropertyArray.cpp
+++ b/Engine/Source/Runtime/CoreUObject/Private/UObject/PropertyArray.cpp
@@ -393,7 +393,7 @@ void FArrayProperty::SerializeItem(FStructuredArchive::FSlot Slot, void* Value,
 			TArray<int32> AddedIndices;
 
 			bool bReplaceArray = false;
-			if (!Defaults)
+			if (!Defaults || !UnderlyingArchive.DoDelta() || UnderlyingArchive.IsTransacting())
 			{
 				bReplaceArray = true;
 			}
diff --git a/Engine/Source/Runtime/CoreUObject/Private/UObject/PropertyMap.cpp b/Engine/Source/Runtime/CoreUObject/Private/UObject/PropertyMap.cpp
--- a/Engine/Source/Runtime/CoreUObject/Private/UObject/PropertyMap.cpp
+++ b/Engine/Source/Runtime/CoreUObject/Private/UObject/PropertyMap.cpp
@@ -517,7 +517,7 @@ void FMapProperty::SerializeItem(FStructuredArchive::FSlot Slot, void* Value, co
 			TSet<int32> ModifiedIndices;
 
 			bool bReplaceMap = false;
-			if (!Defaults)
+			if (!Defaults || !UnderlyingArchive.DoDelta() || UnderlyingArchive.IsTransacting())
 			{
 				bReplaceMap = true;
 			}
