Added more type information to FPropertyPathName and improved tracking of SerializedPropertyPath

diff --git a/Engine/Source/Runtime/CoreUObject/Private/UObject/Class.cpp b/Engine/Source/Runtime/CoreUObject/Private/UObject/Class.cpp
--- a/Engine/Source/Runtime/CoreUObject/Private/UObject/Class.cpp
+++ b/Engine/Source/Runtime/CoreUObject/Private/UObject/Class.cpp
@@ -1088,6 +1088,16 @@ void UStruct::DestroyStruct(void* Dest, int32 ArrayDim) const
 //
 void UStruct::SerializeBin( FStructuredArchive::FSlot Slot, void* Data ) const
 {
+	FUObjectSerializeContext* SerializeContext = FUObjectThreadContext::Get().GetSerializeContext();
+	const bool bSaveSerializedPropertyPath = IsA<UClass>() && SerializeContext && !SerializeContext->SerializedPropertyPath.IsEmpty();
+	UE::FPropertyPathName PrevSerializedPropertyPath;
+
+	if (bSaveSerializedPropertyPath)
+	{
+		PrevSerializedPropertyPath = MoveTemp(SerializeContext->SerializedPropertyPath);
+		SerializeContext->SerializedPropertyPath.Reset();
+	}
+
 	FArchive& UnderlyingArchive = Slot.GetUnderlyingArchive();
 
 	FStructuredArchive::FStream PropertyStream = Slot.EnterStream();
@@ -1132,6 +1142,11 @@ void UStruct::SerializeBin( FStructuredArchive::FSlot Slot, void* Data ) const
 			Property->SerializeBinProperty(PropertyStream.EnterElement(), Data);
 		}
 	}
+
+	if (bSaveSerializedPropertyPath)
+	{
+		SerializeContext->SerializedPropertyPath = MoveTemp(PrevSerializedPropertyPath);
+	}
 }
 
 void UStruct::SerializeBinEx( FStructuredArchive::FSlot Slot, void* Data, void const* DefaultData, UStruct* DefaultStruct ) const
@@ -1142,10 +1157,25 @@ void UStruct::SerializeBinEx( FStructuredArchive::FSlot Slot, void* Data, void c
 		return;
 	}
 
+	FUObjectSerializeContext* SerializeContext = FUObjectThreadContext::Get().GetSerializeContext();
+	const bool bSaveSerializedPropertyPath = IsA<UClass>() && SerializeContext && !SerializeContext->SerializedPropertyPath.IsEmpty();
+	UE::FPropertyPathName PrevSerializedPropertyPath;
+
+	if (bSaveSerializedPropertyPath)
+	{
+		PrevSerializedPropertyPath = MoveTemp(SerializeContext->SerializedPropertyPath);
+		SerializeContext->SerializedPropertyPath.Reset();
+	}
+
 	for( TFieldIterator<FProperty> It(this); It; ++It )
 	{
 		It->SerializeNonMatchingBinProperty(Slot, Data, DefaultData, DefaultStruct);
 	}
+
+	if (bSaveSerializedPropertyPath)
+	{
+		SerializeContext->SerializedPropertyPath = MoveTemp(PrevSerializedPropertyPath);
+	}
 }
 
 void UStruct::LoadTaggedPropertiesFromText(FStructuredArchive::FSlot Slot, uint8* Data, UStruct* DefaultsStruct, uint8* Defaults, const UObject* BreakRecursionIfFullyLoad) const
@@ -1290,6 +1320,16 @@ void UStruct::LoadTaggedPropertiesFromText(FStructuredArchive::FSlot Slot, uint8
 
 void UStruct::SerializeTaggedProperties(FStructuredArchive::FSlot Slot, uint8* Data, UStruct* DefaultsStruct, uint8* Defaults, const UObject* BreakRecursionIfFullyLoad) const
 {
+	FUObjectSerializeContext* SerializeContext = FUObjectThreadContext::Get().GetSerializeContext();
+	const bool bSaveSerializedPropertyPath = IsA<UClass>() && SerializeContext && !SerializeContext->SerializedPropertyPath.IsEmpty();
+	UE::FPropertyPathName PrevSerializedPropertyPath;
+
+	if (bSaveSerializedPropertyPath)
+	{
+		PrevSerializedPropertyPath = MoveTemp(SerializeContext->SerializedPropertyPath);
+		SerializeContext->SerializedPropertyPath.Reset();
+	}
+
 	if (Slot.GetArchiveState().UseUnversionedPropertySerialization())
 	{
 		SerializeUnversionedProperties(this, Slot, Data, DefaultsStruct, Defaults);
@@ -1298,6 +1338,11 @@ void UStruct::SerializeTaggedProperties(FStructuredArchive::FSlot Slot, uint8* D
 	{
 		SerializeVersionedTaggedProperties(Slot, Data, DefaultsStruct, Defaults, BreakRecursionIfFullyLoad);
 	}
+
+	if (bSaveSerializedPropertyPath)
+	{
+		SerializeContext->SerializedPropertyPath = MoveTemp(PrevSerializedPropertyPath);
+	}
 }
 
 #if WITH_EDITORONLY_DATA
@@ -1357,7 +1402,7 @@ void UStruct::SerializeVersionedTaggedProperties(FStructuredArchive::FSlot Slot,
 	checkf(Data, TEXT("Expecting a non null data ptr"));
 
 	FArchive& UnderlyingArchive = Slot.GetUnderlyingArchive();
-	FUObjectSerializeContext* LoadContext = UnderlyingArchive.GetSerializeContext();
+	FUObjectSerializeContext* SerializeContext = FUObjectThreadContext::Get().GetSerializeContext();
 	//SCOPED_LOADTIMER(SerializeTaggedPropertiesTime);
 
 	// Setup serialization control data extensions, this is serialized only on root i.e. UObject and not structs!
@@ -1405,16 +1450,16 @@ void UStruct::SerializeVersionedTaggedProperties(FStructuredArchive::FSlot Slot,
 		else
 #endif // WITH_TEXT_ARCHIVE_SUPPORT
 		{
-			auto TryFindPropertyBag = [PropertyBag = (FPropertyBag*)nullptr, bSearched = false, LoadContext]() mutable -> FPropertyBag*
+			auto TryFindPropertyBag = [PropertyBag = (FPropertyBag*)nullptr, bSearched = false, SerializeContext]() mutable -> FPropertyBag*
 			{
 				if (bSearched)
 				{
 					return PropertyBag;
 				}
 				bSearched = true;
-				if (LoadContext && LoadContext->bSerializeUnknownProperty)
+				if (SerializeContext && SerializeContext->bSerializeUnknownProperty)
 				{
-					if (UObject* Object = LoadContext->SerializedObject)
+					if (UObject* Object = SerializeContext->SerializedObject)
 					{
 						PropertyBag = FPropertyBagRepository::Get().CreateOuterBag(Object);
 					}
@@ -1522,30 +1567,28 @@ void UStruct::SerializeVersionedTaggedProperties(FStructuredArchive::FSlot Slot,
 					Property = CustomFindProperty(Tag.Name);
 				}
 
-				if (LoadContext)
+				if (SerializeContext)
 				{
 					const FName Name = Property ? Property->GetFName() : Tag.Name;
 					const int32 Index = Tag.ArrayIndex > 0 || (Property && Property->ArrayDim > 1) ? Tag.ArrayIndex : INDEX_NONE;
+					SerializeContext->SerializedPropertyPath.Push({Name, Tag.Type, Index});
 
-					FName Type;
-					if (!Tag.InnerType.IsNone() && Tag.ValueType.IsNone())
+					if (!Tag.StructName.IsNone())
 					{
-						Type = Tag.InnerType;
-					}
-					else if (!Tag.StructName.IsNone())
-					{
-						Type = Tag.StructName;
+						SerializeContext->SerializedPropertyPath.PushType(Tag.StructName);
 					}
 					else if (!Tag.EnumName.IsNone())
 					{
-						Type = Tag.EnumName;
+						SerializeContext->SerializedPropertyPath.PushType(Tag.EnumName);
 					}
-					else
+					else if (!Tag.InnerType.IsNone())
 					{
-						Type = Tag.Type;
+						SerializeContext->SerializedPropertyPath.PushType(Tag.InnerType);
+						if (!Tag.ValueType.IsNone())
+						{
+							SerializeContext->SerializedPropertyPath.PushType(Tag.ValueType);
+						}
 					}
-
-					LoadContext->SerializedPropertyPath.Push({Name, Type, Index});
 				}
 
 				if (Property)
@@ -1617,7 +1660,7 @@ void UStruct::SerializeVersionedTaggedProperties(FStructuredArchive::FSlot Slot,
 									UnderlyingArchive.Seek(StartOfProperty);
 									FStructuredArchive::FSlot CopySlot = PropertyRecord.EnterField(TEXT("Value"));
 									Tag.Prop = nullptr;
-									PropertyBag->LoadPropertyByTag(LoadContext->SerializedPropertyPath, Tag, CopySlot, Property->ContainerPtrToValuePtrForDefaults<uint8>(DefaultsStruct, Defaults, Tag.ArrayIndex));
+									PropertyBag->LoadPropertyByTag(SerializeContext->SerializedPropertyPath, Tag, CopySlot, Property->ContainerPtrToValuePtrForDefaults<uint8>(DefaultsStruct, Defaults, Tag.ArrayIndex));
 								}
 								break;
 
@@ -1628,7 +1671,8 @@ void UStruct::SerializeVersionedTaggedProperties(FStructuredArchive::FSlot Slot,
 									if (FPropertyBag* PropertyBag = TryFindPropertyBag())
 									{
 										Tag.Prop = nullptr;
-										PropertyBag->LoadPropertyByTag(LoadContext->SerializedPropertyPath, Tag, ValueSlot, Property->ContainerPtrToValuePtrForDefaults<uint8>(DefaultsStruct, Defaults, Tag.ArrayIndex));
+										PropertyBag->LoadPropertyByTag(SerializeContext->SerializedPropertyPath, Tag, ValueSlot, Property->ContainerPtrToValuePtrForDefaults<uint8>(DefaultsStruct, Defaults, Tag.ArrayIndex));
+										bAdvanceProperty = !UnderlyingArchive.IsCriticalError();
 									}
 								}
 								else
@@ -1646,7 +1690,8 @@ void UStruct::SerializeVersionedTaggedProperties(FStructuredArchive::FSlot Slot,
 								if (FPropertyBag* PropertyBag = TryFindPropertyBag())
 								{
 									Tag.Prop = nullptr;
-									PropertyBag->LoadPropertyByTag(LoadContext->SerializedPropertyPath, Tag, ValueSlot, Property->ContainerPtrToValuePtrForDefaults<uint8>(DefaultsStruct, Defaults, Tag.ArrayIndex));
+									PropertyBag->LoadPropertyByTag(SerializeContext->SerializedPropertyPath, Tag, ValueSlot, Property->ContainerPtrToValuePtrForDefaults<uint8>(DefaultsStruct, Defaults, Tag.ArrayIndex));
+									bAdvanceProperty = !UnderlyingArchive.IsCriticalError();
 								}
 								break;
 
@@ -1655,16 +1700,16 @@ void UStruct::SerializeVersionedTaggedProperties(FStructuredArchive::FSlot Slot,
 						}
 					}
 				}
-				else if (FPropertyBag* PropertyBag = TryFindPropertyBag(); PropertyBag && LoadContext)
+				else if (FPropertyBag* PropertyBag = TryFindPropertyBag(); PropertyBag && SerializeContext)
 				{
 					// TODO: Might we find defaults in a property bag for Defaults?
 					FStructuredArchive::FSlot ValueSlot = PropertyRecord.EnterField(TEXT("Value"));
-					PropertyBag->LoadPropertyByTag(LoadContext->SerializedPropertyPath, Tag, ValueSlot);
+					PropertyBag->LoadPropertyByTag(SerializeContext->SerializedPropertyPath, Tag, ValueSlot);
 				}
 
-				if (LoadContext)
+				if (SerializeContext)
 				{
-					LoadContext->SerializedPropertyPath.Pop();
+					SerializeContext->SerializedPropertyPath.Pop();
 				}
 
 				int64 Loaded = UnderlyingArchive.Tell() - StartOfProperty;
diff --git a/Engine/Source/Runtime/CoreUObject/Private/UObject/LinkerLoad.cpp b/Engine/Source/Runtime/CoreUObject/Private/UObject/LinkerLoad.cpp
--- a/Engine/Source/Runtime/CoreUObject/Private/UObject/LinkerLoad.cpp
+++ b/Engine/Source/Runtime/CoreUObject/Private/UObject/LinkerLoad.cpp
@@ -4593,11 +4593,9 @@ void FLinkerLoad::Preload( UObject* Object )
 #endif // USE_CIRCULAR_DEPENDENCY_LOAD_DEFERRING
 						
 						check(CurrentLoadContext);
-						// Maintain the current SerializedObject and SerializedPropertyPath.
+						// Maintain the current SerializedObject.
 						UObject* PrevSerializedObject = CurrentLoadContext->SerializedObject;
-						UE::FPropertyPathName PrevSerializedPropertyPath = MoveTemp(CurrentLoadContext->SerializedPropertyPath);
 						CurrentLoadContext->SerializedObject = Object;
-						CurrentLoadContext->SerializedPropertyPath.Reset();
 
 #if WITH_EDITOR && WITH_TEXT_ARCHIVE_SUPPORT
 						if (IsTextFormat())
@@ -4622,7 +4620,6 @@ void FLinkerLoad::Preload( UObject* Object )
 						}
 
 						Object->SetFlags(RF_LoadCompleted);
-						CurrentLoadContext->SerializedPropertyPath = MoveTemp(PrevSerializedPropertyPath);
 						CurrentLoadContext->SerializedObject = PrevSerializedObject;
 					}
 					else
@@ -4640,11 +4637,9 @@ void FLinkerLoad::Preload( UObject* Object )
 						TRACE_CPUPROFILER_EVENT_SCOPE_TEXT_ON_CHANNEL(*Object->GetFullName(), AssetLoadTimeChannel);
 #endif
 
-						// Maintain the current SerializedObject and SerializedPropertyPath.
+						// Maintain the current SerializedObject.
 						UObject* PrevSerializedObject = CurrentLoadContext->SerializedObject;
-						UE::FPropertyPathName PrevSerializedPropertyPath = MoveTemp(CurrentLoadContext->SerializedPropertyPath);
 						CurrentLoadContext->SerializedObject = Object;
-						CurrentLoadContext->SerializedPropertyPath.Reset();
 
 #if WITH_EDITOR && WITH_TEXT_ARCHIVE_SUPPORT
 						if (IsTextFormat())
@@ -4670,7 +4665,6 @@ void FLinkerLoad::Preload( UObject* Object )
 						}
 
 						Object->SetFlags(RF_LoadCompleted);
-						CurrentLoadContext->SerializedPropertyPath = MoveTemp(PrevSerializedPropertyPath);
 						CurrentLoadContext->SerializedObject = PrevSerializedObject;
 					}
 				}
diff --git a/Engine/Source/Runtime/CoreUObject/Private/UObject/PropertyPathName.cpp b/Engine/Source/Runtime/CoreUObject/Private/UObject/PropertyPathName.cpp
--- a/Engine/Source/Runtime/CoreUObject/Private/UObject/PropertyPathName.cpp
+++ b/Engine/Source/Runtime/CoreUObject/Private/UObject/PropertyPathName.cpp
@@ -63,6 +63,52 @@ bool FPropertyPathName::operator<(const FPropertyPathName& Path) const
 	return SegmentCountA < SegmentCountB;
 }
 
+void FPropertyPathName::PushTypeInternal(FName Type)
+{
+	FNameEntryId& LastType = Segments.Last().Type;
+
+	TStringBuilder<256> CombinedTypes;
+	FName::GetEntry(LastType)->AppendNameToString(CombinedTypes);
+	CombinedTypes.AppendChar(' ');
+	Type.AppendString(CombinedTypes);
+
+	LastType = FName(CombinedTypes).GetDisplayIndex();
+}
+
+FName FPropertyPathName::PopType()
+{
+	if (Segments.IsEmpty())
+	{
+		return FName();
+	}
+
+	FNameEntryId& LastType = Segments.Last().Type;
+
+	if (LastType.IsNone())
+	{
+		return FName();
+	}
+
+	FName PoppedType;
+	TStringBuilder<256> CombinedTypes;
+	FName::GetEntry(LastType)->AppendNameToString(CombinedTypes);
+
+	const int32 Index = String::FindLastChar(CombinedTypes, TEXT(' '));
+	if (Index == INDEX_NONE)
+	{
+		PoppedType = FName::CreateFromDisplayId(LastType, NAME_NO_NUMBER_INTERNAL);
+		LastType = FNameEntryId();
+	}
+	else
+	{
+		const FStringView CombinedTypesView(CombinedTypes);
+		PoppedType = FName(CombinedTypesView.RightChop(Index + 1));
+		LastType = FName(CombinedTypesView.Left(Index)).GetDisplayIndex();
+	}
+
+	return PoppedType;
+}
+
 void FPropertyPathName::ToString(FStringBuilderBase& Out, FStringView Separator) const
 {
 	bool bFirst = true;
diff --git a/Engine/Source/Runtime/CoreUObject/Public/UObject/PropertyPathName.h b/Engine/Source/Runtime/CoreUObject/Public/UObject/PropertyPathName.h
--- a/Engine/Source/Runtime/CoreUObject/Public/UObject/PropertyPathName.h
+++ b/Engine/Source/Runtime/CoreUObject/Public/UObject/PropertyPathName.h
@@ -12,10 +12,16 @@
 namespace UE
 {
 
+/**
+ * Represents one segment of a property path including type names and container index.
+ */
 struct FPropertyPathNameSegment
 {
+	/** Name of the property. */
 	FName Name;
+	/** Type name(s) of the property, separated with spaces, like "ArrayProperty StructProperty TypeOfStruct" */
 	FName Type;
+	/** Index within a container, or INDEX_NONE if not in a container. Uses ElementId for associative containers. */
 	int32 Index = INDEX_NONE;
 
 	[[nodiscard]] inline FName PackNameWithIndex() const
@@ -35,7 +41,7 @@ struct FPropertyPathNameSegment
 };
 
 /**
- * Represents the path to a property, by name, including struct type names and the index in each container.
+ * Represents the path to a property, by name, including type names and the index in each container.
  *
  * Sequenced containers use the index directly and associative containers use their ElementId.
  */
@@ -94,6 +100,26 @@ public:
 		Segments.Pop(/*bAllowShrinking*/ false);
 	}
 
+	/** Pushes the type name on the end of the type of the last segment of the path. Ignored if the path is empty. */
+	inline void PushType(FName Type)
+	{
+		if (!Segments.IsEmpty())
+		{
+			FNameEntryId& LastType = Segments.Last().Type;
+			if (LastType.IsNone())
+			{
+				LastType = Type.GetDisplayIndex();
+			}
+			else
+			{
+				PushTypeInternal(Type);
+			}
+		}
+	}
+
+	/** Pops a type name off the end of the type of the last segment of the path. Returns None if the path is empty. */
+	UE_API FName PopType();
+
 	/** Sets the index of the last segment of the path. Ignored if the path is empty. */
 	inline void SetIndex(int32 Index)
 	{
@@ -126,6 +152,8 @@ public:
 	UE_API friend uint32 GetTypeHash(const FPropertyPathName& Path);
 
 private:
+	UE_API void PushTypeInternal(FName Type);
+
 	TArray<FSegment> Segments;
 };
 
