Cvar backed enum properties now can interpret cvars set to True\False for initialization. Added a Warning so user knows to update the cvar to proper int values.

diff --git a/Engine/Source/Runtime/CoreUObject/Private/UObject/PropertyByte.cpp b/Engine/Source/Runtime/CoreUObject/Private/UObject/PropertyByte.cpp
--- a/Engine/Source/Runtime/CoreUObject/Private/UObject/PropertyByte.cpp
+++ b/Engine/Source/Runtime/CoreUObject/Private/UObject/PropertyByte.cpp
@@ -460,7 +460,10 @@ const TCHAR* FByteProperty::ImportText_Internal( const TCHAR* InBuffer, void* Co
 	}
 	
 	// Interpret "True" and "False" as 1 and 0. This is mostly for importing a property that was exported as a bool and is imported as a non-enum byte.
-	if (!Enum)
+	// Also allow for ConsoleVariable-backed enums to attempt to convert True/False to 1/0 in case a bool cvar has been converted to an enum. 
+	// Enum properties backed by an integer CVar are stored as number values, so this code will only do anything when reading an old .ini file with True/False values
+	// We log a warning so users can fix up their .ini files to use integer values that map to the enum
+	if (!Enum || (PortFlags & PPF_ConsoleVariable))
 	{
 		FString Temp;
 		if (const TCHAR* Buffer = FPropertyHelpers::ReadToken(InBuffer, Temp))
@@ -478,6 +481,12 @@ const TCHAR* FByteProperty::ImportText_Internal( const TCHAR* InBuffer, void* Co
 				{
 					SetIntPropertyValue(PointerToValuePtr(ContainerOrPropertyPtr, PropertyPointerType), TrueValue);
 				}
+
+				if (Enum)
+				{
+					UE_LOG(LogClass, Warning, TEXT("ConsoleVariable-Backed Enum Property of type '%s' was set from a string. Please update the cvar in your ini files."), *Enum->GetPathName(), *Enum->GetName());
+				}
+
 				return Buffer;
 			}
 			else if (Temp == TEXT("False") || Temp == *(CoreTexts.False.ToString()))
@@ -491,6 +500,12 @@ const TCHAR* FByteProperty::ImportText_Internal( const TCHAR* InBuffer, void* Co
 				{
 					SetIntPropertyValue(PointerToValuePtr(ContainerOrPropertyPtr, PropertyPointerType), FalseValue);
 				}
+
+				if (Enum)
+				{
+					UE_LOG(LogClass, Warning, TEXT("ConsoleVariable-Backed Enum Property of type '%s' was set from a string. Please update the cvar in your ini files."), *Enum->GetPathName(), *Enum->GetName());
+				}
+
 				return Buffer;
 			}
 		}
