AssetDatas and GetAssetRegistryTags: Change API and change cook behavior:
API: Change UObject::GetAssetRegistryTags to take a Context object that reports whether the call is coming from SavePackage (and should calculate everything) or load (and should skip expensive tags).
Cook Behavior: AssetDatas are now included in the development AssetRegistry generated by a cook if they were included in a referenced package. Tags are updated during save rather than during load. For now, AssetDatas are added for non-referenced packages in non-dlc cooks, but this addition will eventually be removed, and their tags are stale (have the values saved from the last editor save of the package).
[RN] Miinor, Cooking

diff --git a/Engine/Plugins/Editor/AssetSearch/Source/Private/AssetSearchManager.cpp b/Engine/Plugins/Editor/AssetSearch/Source/Private/AssetSearchManager.cpp
--- a/Engine/Plugins/Editor/AssetSearch/Source/Private/AssetSearchManager.cpp
+++ b/Engine/Plugins/Editor/AssetSearch/Source/Private/AssetSearchManager.cpp
@@ -22,6 +22,7 @@
 #include "Sound/SoundCue.h"
 #include "Misc/ScopedSlowTask.h"
 #include "Editor.h"
+#include "UObject/AssetRegistryTagsContext.h"
 #include "UObject/ObjectSaveContext.h"
 #include "FileHelpers.h"
 #include "Misc/MessageDialog.h"
@@ -202,7 +203,7 @@ FAssetSearchManager::~FAssetSearchManager()
 
 	UPackage::PackageSavedWithContextEvent.RemoveAll(this);
 	FCoreUObjectDelegates::OnAssetLoaded.RemoveAll(this);
-	UObject::FAssetRegistryTag::OnGetExtendedAssetRegistryTagsForSave.RemoveAll(this);
+	UObject::FAssetRegistryTag::OnGetExtraObjectTagsWithContext.RemoveAll(this);
 
 	FTSTicker::GetCoreTicker().RemoveTicker(TickerHandle);
 }
@@ -241,7 +242,7 @@ void FAssetSearchManager::Start()
 
 	if (IntermediateStorage == ESearchIntermediateStorage::AssetTagData)
 	{
-		UObject::FAssetRegistryTag::OnGetExtendedAssetRegistryTagsForSave.AddRaw(this, &FAssetSearchManager::HandleGetExtendedAssetRegistryTagsForSave);
+		UObject::FAssetRegistryTag::OnGetExtraObjectTagsWithContext.AddRaw(this, &FAssetSearchManager::HandleOnGetExtraObjectTags);
 	}
 }
 
@@ -485,8 +486,15 @@ uint64 FAssetSearchManager::GetTextHash(FStringView PackageRelativeExportPath) c
 	return CityHash64(reinterpret_cast<const char*>(PackageRelativeExportPath.GetData() + 1), (PackageRelativeExportPath.Len() - 1) * sizeof(TCHAR));
 }
 
-void FAssetSearchManager::HandleGetExtendedAssetRegistryTagsForSave(const UObject* Object, const ITargetPlatform* TargetPlatform, TArray<UObject::FAssetRegistryTag>& OutTags)
+void FAssetSearchManager::HandleOnGetExtraObjectTags(FAssetRegistryTagsContext& Context)
 {
+	if (!Context.IsFullUpdate())
+	{
+		return;
+	}
+
+	const UObject* Object = Context.GetObject();
+	const ITargetPlatform* TargetPlatform = Context.GetTargetPlatform();
 	if (GIsEditor && !Object->GetOutermost()->HasAnyPackageFlags(PKG_ForDiffing) && !IsRunningCookCommandlet())
 	{
 		if (!GEditor->IsAutosaving())
@@ -502,9 +510,9 @@ void FAssetSearchManager::HandleGetExtendedAssetRegistryTagsForSave(const UObjec
 					const FString IndexVersion = GetBaseIndexKey(Object->GetClass());
 					const FString IndexHash = LexToString(GetTextHash(IndexedJson));
 
-					OutTags.Add(UObject::FAssetRegistryTag(FAssetSearchManager::AssetSearchIndexVersionTag, IndexVersion, UObject::FAssetRegistryTag::TT_Hidden));
-					OutTags.Add(UObject::FAssetRegistryTag(FAssetSearchManager::AssetSearchIndexHashTag, IndexHash, UObject::FAssetRegistryTag::TT_Hidden));
-					OutTags.Add(UObject::FAssetRegistryTag(FAssetSearchManager::AssetSearchIndexDataTag, IndexedJson, UObject::FAssetRegistryTag::TT_Hidden));
+					Context.AddTag(UObject::FAssetRegistryTag(FAssetSearchManager::AssetSearchIndexVersionTag, IndexVersion, UObject::FAssetRegistryTag::TT_Hidden));
+					Context.AddTag(UObject::FAssetRegistryTag(FAssetSearchManager::AssetSearchIndexHashTag, IndexHash, UObject::FAssetRegistryTag::TT_Hidden));
+					Context.AddTag(UObject::FAssetRegistryTag(FAssetSearchManager::AssetSearchIndexDataTag, IndexedJson, UObject::FAssetRegistryTag::TT_Hidden));
 
 					if (!IndexedJson.IsEmpty())
 					{
diff --git a/Engine/Plugins/Editor/AssetSearch/Source/Private/AssetSearchManager.h b/Engine/Plugins/Editor/AssetSearch/Source/Private/AssetSearchManager.h
--- a/Engine/Plugins/Editor/AssetSearch/Source/Private/AssetSearchManager.h
+++ b/Engine/Plugins/Editor/AssetSearch/Source/Private/AssetSearchManager.h
@@ -57,7 +57,7 @@ private:
 	void OnAssetRemoved(const FAssetData& InAssetData);
 	void OnAssetScanFinished();
 
-	void HandleGetExtendedAssetRegistryTagsForSave(const UObject* Object, const ITargetPlatform* TargetPlatform, TArray<UObject::FAssetRegistryTag>& OutTags);
+	void HandleOnGetExtraObjectTags(FAssetRegistryTagsContext& Context);
 	void HandlePackageSaved(const FString& PackageFilename, UPackage* Package, FObjectPostSaveContext ObjectSaveContext);
 	void OnAssetLoaded(UObject* InObject);
 
diff --git a/Engine/Source/Editor/UnrealEd/Private/Commandlets/AssetRegistryGenerator.cpp b/Engine/Source/Editor/UnrealEd/Private/Commandlets/AssetRegistryGenerator.cpp
--- a/Engine/Source/Editor/UnrealEd/Private/Commandlets/AssetRegistryGenerator.cpp
+++ b/Engine/Source/Editor/UnrealEd/Private/Commandlets/AssetRegistryGenerator.cpp
@@ -6,7 +6,6 @@
 #include "Algo/StableSort.h"
 #include "Algo/Unique.h"
 #include "AssetRegistry/AssetRegistryModule.h"
-#include "AssetRegistry/CookTagList.h"
 #include "CollectionManagerModule.h"
 #include "CollectionManagerTypes.h"
 #include "Cooker/CookPackageData.h"
@@ -259,9 +258,10 @@ int64 FAssetRegistryGenerator::GetMaxChunkSizePerPlatform(const ITargetPlatform*
 TArray<int32> FAssetRegistryGenerator::GetExistingPackageChunkAssignments(FName PackageFName)
 {
 	TArray<int32> ExistingChunkIDs;
+	int32 PackageFNameHash = GetTypeHash(PackageFName);
 	for (uint32 ChunkIndex = 0, MaxChunk = ChunkManifests.Num(); ChunkIndex < MaxChunk; ++ChunkIndex)
 	{
-		if (ChunkManifests[ChunkIndex] && ChunkManifests[ChunkIndex]->Contains(PackageFName))
+		if (ChunkManifests[ChunkIndex] && ChunkManifests[ChunkIndex]->ContainsByHash(PackageFNameHash, PackageFName))
 		{
 			ExistingChunkIDs.AddUnique(ChunkIndex);
 		}
@@ -1045,45 +1045,6 @@ void FAssetRegistryGenerator::UpdateKeptPackages()
 	PreviousPackagesToUpdate.Empty();
 }
 
-static void AppendCookTagsToTagMap(TArray<TPair<FName, FString>>&& InTags, FAssetDataTagMap& OutTags)
-{
-	for (TPair<FName, FString>& Tag : InTags)
-	{
-		// Don't add empty tags
-		if (!Tag.Key.IsNone() && !Tag.Value.IsEmpty())
-		{
-			// Prepend Cook_
-			// Do the accumulation in a stack buffer to avoid a bit of work, but we still
-			// eat the FName generation.
-			FName TagName(WriteToString<256>(TEXTVIEW("Cook_"), Tag.Key));
-			OutTags.Add(TagName, MoveTemp(Tag.Value));
-		}
-	}
-}
-
-static void AddCookTagsToState(TMap<FSoftObjectPath, TArray<TPair<FName, FString>>>&& InCookTags, FAssetRegistryState& InState)
-{
-	TRACE_CPUPROFILER_EVENT_SCOPE(AddCookTagToState)
-
-	uint32 ObjectCount = 0;
-	uint32 TagCount = 0;
-
-	for (TPair<FSoftObjectPath, TArray<TPair<FName, FString>>>& ObjectToTags : InCookTags)
-	{
-		ObjectCount++;
-		const FAssetData* AssetData = InState.GetAssetByObjectPath(FSoftObjectPath(WriteToString<256>(ObjectToTags.Key)));
-
-		// Migrate to FAssetDataTagMap
-		FAssetDataTagMap NewTags;
-		AppendCookTagsToTagMap(MoveTemp(ObjectToTags.Value), NewTags);
-		TagCount += NewTags.Num();
-		InState.AddTagsToAssetData(ObjectToTags.Key, MoveTemp(NewTags));
-	} // end for each object
-
-	UE_LOG(LogAssetRegistryGenerator, Verbose, TEXT("Added %d cook tags to %d objects"), TagCount, ObjectCount);
-	InCookTags.Reset();
-}
-
 void FAssetRegistryGenerator::UpdateCollectionAssetData()
 {
 	// Read out the per-platform settings use to build the list of collections to tag
@@ -1704,6 +1665,7 @@ bool FAssetRegistryGenerator::SaveAssetRegistry(const FString& SandboxPath, bool
 	// Write development first, this will always write
 	FAssetRegistrySerializationOptions DevelopmentSaveOptions;
 	AssetRegistry.InitializeSerializationOptions(DevelopmentSaveOptions, TargetPlatform->IniPlatformName(), UE::AssetRegistry::ESerializationTarget::ForDevelopment);
+	DevelopmentSaveOptions.bKeepDevelopmentAssetRegistryTags = true;
 
 	// Write runtime registry, this can be excluded per game/platform
 	FAssetRegistrySerializationOptions SaveOptions;
@@ -1725,13 +1687,10 @@ bool FAssetRegistryGenerator::SaveAssetRegistry(const FString& SandboxPath, bool
 
 		if (bSerializeDevelopmentAssetRegistry)
 		{
-			// Make a copy of the state so we can add tags to only the development registry.
+			// Make a copy of the state so it can be filtered independently
 			FAssetRegistryState DevelopmentState;
 			DevelopmentState.InitializeFromExisting(State, DevelopmentSaveOptions);
-
-			DevelopmentState.FilterTags(DevelopmentSaveOptions);
-
-			AddCookTagsToState(MoveTemp(CookTagsToAdd), DevelopmentState);
+			// No need to call FilterTags; it is called by InitializeFromExisting
 
 			// Create development registry data, used for DLC cooks, iterative cooks, and editor viewing
 			FArrayWriter SerializedAssetRegistry;
@@ -2589,36 +2548,9 @@ void FAssetRegistryGenerator::GetChunkAssignments(TArray<TSet<FName>>& OutAssign
 	}
 }
 
-FAssetRegistryGenerator::FCreateOrFindArray FAssetRegistryGenerator::CreateOrFindAssetDatas(const UPackage& Package)
-{
-	FCreateOrFindArray OutputAssets;
-
-	ForEachObjectWithOuter(&Package, [&OutputAssets, this](UObject* const Object)
-	{
-		if (Object->IsAsset())
-		{
-			OutputAssets.Add(CreateOrFindAssetData(*Object));
-		}
-	}, /*bIncludeNestedObjects*/ false);
-
-	return OutputAssets;
-}
-
-const FAssetData* FAssetRegistryGenerator::CreateOrFindAssetData(UObject& Object)
-{
-	const FAssetData* const AssetData = State.GetAssetByObjectPath(FSoftObjectPath(&Object));
-	if (!AssetData)
-	{
-		FAssetData* const NewAssetData = new FAssetData(&Object, true /* bAllowBlueprintClass */);
-		State.AddAssetData(NewAssetData);
-		return NewAssetData;
-	}
-	return AssetData;
-}
-
 void FAssetRegistryGenerator::UpdateAssetRegistryData(FName PackageName, const UPackage* Package,
-	UE::Cook::ECookResult CookResult, FSavePackageResultStruct* SavePackageResult, FCookTagList&& InArchiveCookTagList,
-	bool bIncludeOnlyDiskAssets, TOptional<FAssetPackageData>&& OverrideAssetPackageData,
+	UE::Cook::ECookResult CookResult, FSavePackageResultStruct* SavePackageResult,
+	TOptional<TArray<FAssetData>>&& AssetDatasFromSave, TOptional<FAssetPackageData>&& OverrideAssetPackageData,
 	TOptional<TArray<FAssetDependency>>&& OverridePackageDependencies)
 {
 	LLM_SCOPE_BYTAG(Cooker_GeneratedAssetRegistry);
@@ -2628,18 +2560,17 @@ void FAssetRegistryGenerator::UpdateAssetRegistryData(FName PackageName, const U
 	bool bSaveSucceeded = CookResult == UE::Cook::ECookResult::Succeeded;
 
 	// Copy latest data for all Assets in the package into the cooked registry. This should be done even
-	// if not successful so that editor-only packages are recorded as well
-	TArray<FAssetData> AssetDatas;
-	AssetRegistry.GetAssetsByPackageName(PackageName, AssetDatas, bIncludeOnlyDiskAssets,
-		false /* SkipARFilteredAssets */);
-	for (FAssetData& AssetData : AssetDatas)
+	// if not successful so that editor-only packages are recorded as well. When the AssetDatas were not
+	// calculated by SavePackage, copy them instead from the on-disk AssetRegistry.
+	if (!AssetDatasFromSave)
 	{
-		State.UpdateAssetData(MoveTemp(AssetData), true /* bCreateIfNotExists */);
+		AssetDatasFromSave.Emplace();
+		AssetRegistry.GetAssetsByPackageName(PackageName, *AssetDatasFromSave, true /* bIncludeOnlyDiskAssets */,
+			false /* SkipARFilteredAssets */);
 	}
-	// Create a record for assets that were created during PostLoad or cooking and are not in the global assetregistry
-	if (Package)
+	for (FAssetData& AssetData : *AssetDatasFromSave)
 	{
-		CreateOrFindAssetDatas(*Package);
+		State.UpdateAssetData(MoveTemp(AssetData), true /* bCreateIfNotExists */);
 	}
 
 	FAssetPackageData* AssetPackageData = GetAssetPackageData(PackageName);
@@ -2669,9 +2600,6 @@ void FAssetRegistryGenerator::UpdateAssetRegistryData(FName PackageName, const U
 	if (bSaveSucceeded)
 	{
 		check(SavePackageResult);
-		// Migrate cook tags over
-		CookTagsToAdd.Append(MoveTemp(InArchiveCookTagList.ObjectToTags));
-		InArchiveCookTagList.Reset();
 
 		// Set the PackageFlags to the recorded value from SavePackage
 		NewPackageFlags = SavePackageResult->SerializedPackageFlags;
@@ -2730,8 +2658,6 @@ void FAssetRegistryGenerator::UpdateAssetRegistryData(UE::Cook::FMPCollectorServ
 	{
 		SetOverridePackageDependencies(PackageName, *Message.OverridePackageDependencies);
 	}
-
-	CookTagsToAdd.Append(MoveTemp(Message.CookTags));
 }
 
 namespace UE::Cook
@@ -2744,8 +2670,8 @@ FAssetRegistryReporterRemote::FAssetRegistryReporterRemote(FCookWorkerClient& In
 }
 
 void FAssetRegistryReporterRemote::UpdateAssetRegistryData(FName PackageName, const UPackage* Package,
-	UE::Cook::ECookResult CookResult, FSavePackageResultStruct* SavePackageResult, FCookTagList&& InArchiveCookTagList,
-	bool bIncludeOnlyDiskAssets, TOptional<FAssetPackageData>&& OverrideAssetPackageData,
+	UE::Cook::ECookResult CookResult, FSavePackageResultStruct* SavePackageResult,
+	TOptional<TArray<FAssetData>>&& AssetDatasFromSave, TOptional<FAssetPackageData>&& OverrideAssetPackageData,
 	TOptional<TArray<FAssetDependency>>&& OverridePackageDependencies)
 {
 	uint32 NewPackageFlags = 0;
@@ -2774,33 +2700,14 @@ void FAssetRegistryReporterRemote::UpdateAssetRegistryData(FName PackageName, co
 	Message.OverrideAssetPackageData = MoveTemp(OverrideAssetPackageData);
 	Message.OverridePackageDependencies = MoveTemp(OverridePackageDependencies);
 
-	// Add to the message all the AssetDatas in the package from the global AssetRegistry
-	IAssetRegistry::Get()->GetAssetsByPackageName(PackageName, Message.AssetDatas,
-		bIncludeOnlyDiskAssets, false /* SkipARFilteredAssets */);
-
-	// Also add AssetDatas for any assets that were created during PostLoad or cooking and are not in the global assetregistry
-	TSet<FName> ExistingAssets;
-	for (const FAssetData& AssetData : Message.AssetDatas)
+	// Add to the message all the AssetDatas in the package
+	if (!AssetDatasFromSave)
 	{
-		ExistingAssets.Add(AssetData.AssetName);
+		AssetDatasFromSave.Emplace();
+		IAssetRegistry::Get()->GetAssetsByPackageName(PackageName, *AssetDatasFromSave,
+			true /* bIncludeOnlyDiskAssets */, false /* SkipARFilteredAssets */);
 	}
-	if (Package)
-	{
-		ForEachObjectWithOuter(Package, [&Message, &ExistingAssets](UObject* const Object)
-			{
-				if (Object->IsAsset())
-				{
-					FName AssetName = Object->GetFName();
-					if (!ExistingAssets.Contains(AssetName))
-					{
-						Message.AssetDatas.Emplace(Object, true /* bAllowBlueprintClass */);
-					}
-				}
-			}, /*bIncludeNestedObjects*/ false);
-	}
-
-	// Add the cooktags that were recorded during serialization
-	Message.CookTags = MoveTemp(InArchiveCookTagList.ObjectToTags);
+	Message.AssetDatas = MoveTemp(*AssetDatasFromSave);
 
 	// Send the message to the director
 	FCbWriter Writer;
@@ -2839,11 +2746,6 @@ void FAssetRegistryPackageMessage::Write(FCbWriter& Writer) const
 	{
 		Writer << "D" << *OverridePackageDependencies;
 	}
-	// We cast TMap<FSoftObjectPath,                     ValueType>
-	//      to TMap<FSoftObjectPathSerializationWrapper, ValueType>
-	// to workaround FSoftObjectPath's implicit constructor. See comment in CompactBinaryTCP.h
-	static_assert(std::is_same_v<decltype(CookTags)::KeyType, FSoftObjectPath>, "Expected KeyType of CookTags to be FSoftObjectPath");
-	Writer << "T" << reinterpret_cast<const TMap<FSoftObjectPathSerializationWrapper, decltype(CookTags)::ValueType>&>(CookTags);
 	Writer << "F" << PackageFlags;
 	Writer << "S" << DiskSize;
 }
@@ -2897,11 +2799,6 @@ bool FAssetRegistryPackageMessage::TryRead(FCbObjectView Object)
 		}
 	}
 
-	static_assert(std::is_same_v<decltype(CookTags)::KeyType, FSoftObjectPath>, "Expected KeyType of CookTags to be FSoftObjectPath");
-	if (!LoadFromCompactBinary(Object["T"], reinterpret_cast<TMap<FSoftObjectPathSerializationWrapper, decltype(CookTags)::ValueType>&>(CookTags)))
-	{
-		return false;
-	}
 	if (!LoadFromCompactBinary(Object["F"], PackageFlags))
 	{
 		return false;
diff --git a/Engine/Source/Editor/UnrealEd/Private/Commandlets/SavePackageUtilitiesCommandlet.cpp b/Engine/Source/Editor/UnrealEd/Private/Commandlets/SavePackageUtilitiesCommandlet.cpp
--- a/Engine/Source/Editor/UnrealEd/Private/Commandlets/SavePackageUtilitiesCommandlet.cpp
+++ b/Engine/Source/Editor/UnrealEd/Private/Commandlets/SavePackageUtilitiesCommandlet.cpp
@@ -41,7 +41,7 @@ int32 USavePackageUtilitiesCommandlet::Main(const FString& Params)
 		if (TargetPlatform != nullptr)
 		{
 			CookContext.Emplace(Package, UE::Cook::ECookType::Unknown,
-				UE::Cook::ECookingDLC::Unknown);
+				UE::Cook::ECookingDLC::Unknown, TargetPlatform);
 			CookData.Emplace(*TargetPlatform, *CookContext);
 		}
 
diff --git a/Engine/Source/Editor/UnrealEd/Private/CookOnTheFlyServer.cpp b/Engine/Source/Editor/UnrealEd/Private/CookOnTheFlyServer.cpp
--- a/Engine/Source/Editor/UnrealEd/Private/CookOnTheFlyServer.cpp
+++ b/Engine/Source/Editor/UnrealEd/Private/CookOnTheFlyServer.cpp
@@ -6035,7 +6035,8 @@ void FSaveCookedPackageContext::SetupPlatform(const ITargetPlatform* InTargetPla
 
 	ArchiveCookContext.Emplace(Package,
 		COTFS.IsDirectorCookByTheBook() ? UE::Cook::ECookType::ByTheBook : UE::Cook::ECookType::OnTheFly,
-		COTFS.IsCookingDLC() ? UE::Cook::ECookingDLC::Yes : UE::Cook::ECookingDLC::No);
+		COTFS.IsCookingDLC() ? UE::Cook::ECookingDLC::Yes : UE::Cook::ECookingDLC::No,
+		TargetPlatform);
 
 	// don't save Editor resources from the Engine if the target doesn't have editoronly data
 	if (COTFS.IsCookFlagSet(ECookInitializationFlags::SkipEditorContent) &&
@@ -6195,23 +6196,16 @@ void FSaveCookedPackageContext::FinishPlatform()
 	// Update asset registry
 	if (COTFS.IsDirectorCookByTheBook())
 	{
-		// Flush the AssetRegisty so any AssetData changes from the save are present
-		COTFS.AssetRegistry->WaitForCompletion();
 		IAssetRegistryReporter& Reporter = *(COTFS.PlatformManager->GetPlatformData(TargetPlatform)->RegistryReporter);
 
-		// UpdateAssetRegistryData will pull data from the global asset registry
-		// For most types, our contract is that we use the data in the asset registry that was updated on load and is in
-		// the disk asset storage of the global asset registry. Most types do not take significant action during cooking
-		// so we do not need to allow them to update their tags after they finish loading.
-		bool bIncludeOnlyDiskAssets = true;
+		// Calculate up-to-date dependencies for Generator and Generated packages
+		// For generated packages, additionally calculate FAssetPackageData; that struct is calculated
+		// during user saves for non-generated packages.
 		TOptional<TArray<FAssetDependency>> OverridePackageDependencies;
 		TOptional<FAssetPackageData> OverrideAssetPackageData;
 		FGeneratorPackage* GeneratorPackage;
 		if (GeneratorPackage = PackageData.GetGeneratorPackage(); GeneratorPackage)
 		{
-			// Generator packages do a lot of work during generation and may need to update their tags, so read the
-			// latest tags off of their asset by setting bIncludeOnlyDiskAssets=true.
-			bIncludeOnlyDiskAssets = false;
 			OverridePackageDependencies.Emplace();
 
 			// Set override dependencies equal to the global AssetRegistry dependencies plus a dependency on
@@ -6242,8 +6236,6 @@ void FSaveCookedPackageContext::FinishPlatform()
 			}
 			else
 			{
-				// Generated packages do not exist on disk, so read their assetdata from the in-memory package
-				bIncludeOnlyDiskAssets = false;
 				// There should be no package dependencies present for the package from the global assetregistry
 				// because it is newly created. Add on the dependencies declared for it from the CookPackageSplitter.
 				OverridePackageDependencies.Emplace(GeneratedInfo->PackageDependencies);
@@ -6276,9 +6268,13 @@ void FSaveCookedPackageContext::FinishPlatform()
 				OverrideAssetPackageData->ImportedClasses = ImportedClasses;
 			}
 		}
+		TOptional<TArray<FAssetData>> AssetDatasFromSave;
+		if (SavePackageResult.IsSuccessful())
+		{
+			AssetDatasFromSave.Emplace(MoveTemp(SavePackageResult.SavedAssets));
+		}
 		Reporter.UpdateAssetRegistryData(Package->GetFName(), Package, CookResult, &SavePackageResult,
-			MoveTemp(*ArchiveCookContext->GetCookTagList()), bIncludeOnlyDiskAssets,
-			MoveTemp(OverrideAssetPackageData), MoveTemp(OverridePackageDependencies));
+			MoveTemp(AssetDatasFromSave), MoveTemp(OverrideAssetPackageData), MoveTemp(OverridePackageDependencies));
 	}
 
 	if (bSuccessful && COTFS.bSkipOnlyEditorOnly)
@@ -6473,9 +6469,8 @@ void UCookOnTheFlyServer::RecordExternalActorDependencies(TConstArrayView<FName>
 
 				DependencyData->SetPlatformCooked(TargetPlatform, ECookResult::NeverCookPlaceholder);
 				Reporter.UpdateAssetRegistryData(DependencyName, nullptr /* Package */,
-					ECookResult::NeverCookPlaceholder, nullptr /* SavePackageResult */, FCookTagList(nullptr),
-					true /* bIncludeOnlyDiskAssets */, TOptional<FAssetPackageData>(),
-					TOptional<TArray<FAssetDependency>>());
+					ECookResult::NeverCookPlaceholder, nullptr /* SavePackageResult */,
+					TOptional<TArray<FAssetData>>(), TOptional<FAssetPackageData>(), TOptional<TArray<FAssetDependency>>());
 			}
 		}
 	}
diff --git a/Engine/Source/Editor/UnrealEd/Public/Commandlets/AssetRegistryGenerator.h b/Engine/Source/Editor/UnrealEd/Public/Commandlets/AssetRegistryGenerator.h
--- a/Engine/Source/Editor/UnrealEd/Public/Commandlets/AssetRegistryGenerator.h
+++ b/Engine/Source/Editor/UnrealEd/Public/Commandlets/AssetRegistryGenerator.h
@@ -24,7 +24,6 @@ class IChunkDataGenerator;
 class UChunkDependencyInfo;
 class UCookOnTheFlyServer;
 struct FChunkDependencyTreeNode;
-struct FCookTagList;
 struct FSoftObjectPath;
 namespace UE::Cook { class FAssetRegistryMPCollector; }
 namespace UE::Cook { class FAssetRegistryPackageMessage; }
@@ -253,7 +252,7 @@ public:
 	 */
 	void UpdateAssetRegistryData(FName PackageName, const UPackage* Package,
 		UE::Cook::ECookResult CookResult, FSavePackageResultStruct* SavePackageResult,
-		FCookTagList&& InArchiveCookTagList, bool bIncludeOnlyDiskAssets,
+		TOptional<TArray<FAssetData>>&& AssetDatasFromSave,
 		TOptional<FAssetPackageData>&& OverrideAssetPackageData,
 		TOptional<TArray<FAssetDependency>>&& OverridePackageDependencies);
 	void UpdateAssetRegistryData(UE::Cook::FMPCollectorServerMessageContext& Context,
@@ -266,14 +265,6 @@ public:
 	static void UpdateAssetManagerDatabase();
 
 private:
-	/**
-	 * Ensures all assets in the input package are present in the registry
-	 * @param Package - Package to process
-	 * @return - Array of FAssetData entries for all assets in the input package
-	 */
-	typedef TArray<const FAssetData*, TInlineAllocator<1>> FCreateOrFindArray;
-	FCreateOrFindArray CreateOrFindAssetDatas(const UPackage& Package);
-
 	/**
 	 * Updates all asset package flags in the specified package
 	 *
@@ -297,12 +288,6 @@ private:
 	/** State of the asset registry that is being built for this platform */
 	FAssetRegistryState State;
 	
-	/** 
-	 * The list of tags to add for each asset. This is populated during cook by the books,
-	 * and is only added to development registries.
-	 */
-	TMap<FSoftObjectPath, TArray<TPair<FName, FString>>> CookTagsToAdd;
-
 	struct FIterativelySkippedPackageUpdateData
 	{
 		TArray<FAssetData> AssetDatas;
@@ -478,11 +463,6 @@ private:
 	/** Initialize ChunkIdPakchunkIndexMapping and PakchunkIndexChunkIdMapping. */
 	void InitializeChunkIdPakchunkIndexMapping();
 
-	/**
-	 * Helper function to find or create asset data for the input object. If the asset is not in the registry it will be added.
-	 */
-	const FAssetData* CreateOrFindAssetData(UObject& Object);
-
 	/** If InState records PackageName is generated, return the name of the Generator, otherwise return NAME_None. */
 	static FName GetGeneratorPackage(FName PackageName, const FAssetRegistryState& InState);
 };
@@ -496,8 +476,8 @@ public:
 	virtual ~IAssetRegistryReporter() {}
 
 	virtual void UpdateAssetRegistryData(FName PackageName, const UPackage* Package, UE::Cook::ECookResult CookResult,
-		FSavePackageResultStruct* SavePackageResult, FCookTagList&& InArchiveCookTagList,
-		bool bIncludeOnlyDiskAssets, TOptional<FAssetPackageData>&& OverrideAssetPackageData, 
+		FSavePackageResultStruct* SavePackageResult,
+		TOptional<TArray<FAssetData>>&& AssetDatasFromSave, TOptional<FAssetPackageData>&& OverrideAssetPackageData,
 		TOptional<TArray<FAssetDependency>>&& OverridePackageDependencies) = 0;
 
 };
@@ -511,13 +491,12 @@ public:
 	}
 
 	virtual void UpdateAssetRegistryData(FName PackageName, const UPackage* Package, UE::Cook::ECookResult CookResult,
-		FSavePackageResultStruct* SavePackageResult, FCookTagList&& InArchiveCookTagList,
-		bool bIncludeOnlyDiskAssets, TOptional<FAssetPackageData>&& OverrideAssetPackageData,
+		FSavePackageResultStruct* SavePackageResult,
+		TOptional<TArray<FAssetData>>&& AssetDatasFromSave, TOptional<FAssetPackageData>&& OverrideAssetPackageData,
 		TOptional<TArray<FAssetDependency>>&& OverridePackageDependencies) override
 	{
 		Generator.UpdateAssetRegistryData(PackageName, Package, CookResult, SavePackageResult,
-			MoveTemp(InArchiveCookTagList), bIncludeOnlyDiskAssets, MoveTemp(OverrideAssetPackageData),
-			MoveTemp(OverridePackageDependencies));
+			MoveTemp(AssetDatasFromSave), MoveTemp(OverrideAssetPackageData), MoveTemp(OverridePackageDependencies));
 	}
 
 private:
@@ -530,8 +509,8 @@ public:
 	FAssetRegistryReporterRemote(FCookWorkerClient& InClient, const ITargetPlatform* InTargetPlatform);
 
 	virtual void UpdateAssetRegistryData(FName PackageName, const UPackage* Package, UE::Cook::ECookResult CookResult,
-		FSavePackageResultStruct* SavePackageResult, FCookTagList&& InArchiveCookTagList,
-		bool bIncludeOnlyDiskAssets, TOptional<FAssetPackageData>&& OverrideAssetPackageData, 
+		FSavePackageResultStruct* SavePackageResult,
+		TOptional<TArray<FAssetData>>&& AssetDatasFromSave, TOptional<FAssetPackageData>&& OverrideAssetPackageData,
 		TOptional<TArray<FAssetDependency>>&& OverridePackageDependencies) override;
 
 private:
@@ -555,7 +534,6 @@ public:
 	TArray<FAssetData> AssetDatas;
 	TOptional<FAssetPackageData> OverrideAssetPackageData;
 	TOptional<TArray<FAssetDependency>> OverridePackageDependencies;
-	TMap<FSoftObjectPath, TArray<TPair<FName, FString>>> CookTags;
 	uint32 PackageFlags = 0;
 	int64 DiskSize = -1;
 
diff --git a/Engine/Source/Runtime/AssetRegistry/Private/AssetRegistry.cpp b/Engine/Source/Runtime/AssetRegistry/Private/AssetRegistry.cpp
--- a/Engine/Source/Runtime/AssetRegistry/Private/AssetRegistry.cpp
+++ b/Engine/Source/Runtime/AssetRegistry/Private/AssetRegistry.cpp
@@ -38,6 +38,7 @@
 #include "String/RemoveFrom.h"
 #include "Templates/UnrealTemplate.h"
 #include "TelemetryRouter.h"
+#include "UObject/AssetRegistryTagsContext.h"
 #include "UObject/ConstructorHelpers.h"
 #include "UObject/CoreRedirects.h"
 #include "UObject/MetaData.h"
@@ -856,10 +857,17 @@ void FAssetRegistryImpl::Initialize(Impl::FInitializeContext& Context)
 #if !WITH_EDITOR
 	Context.bUpdateDiskCacheAfterLoad = false;
 #else
-	Context.bUpdateDiskCacheAfterLoad = true;
-	if (GConfig)
+	if (IsRunningCookCommandlet())
 	{
-		GConfig->GetBool(TEXT("AssetRegistry"), TEXT("bUpdateDiskCacheAfterLoad"), Context.bUpdateDiskCacheAfterLoad, GEngineIni);
+		Context.bUpdateDiskCacheAfterLoad = false;
+	}
+	else
+	{
+		Context.bUpdateDiskCacheAfterLoad = true;
+		if (GConfig)
+		{
+			GConfig->GetBool(TEXT("AssetRegistry"), TEXT("bUpdateDiskCacheAfterLoad"), Context.bUpdateDiskCacheAfterLoad, GEngineIni);
+		}
 	}
 #endif
 
@@ -1022,7 +1030,7 @@ void UAssetRegistryImpl::InitializeEvents(UE::AssetRegistry::Impl::FInitializeCo
 
 	if (bAddMetaDataTagsToOnGetExtraObjectTags)
 	{
-		UObject::FAssetRegistryTag::OnGetExtraObjectTags.AddUObject(this, &UAssetRegistryImpl::OnGetExtraObjectTags);
+		UObject::FAssetRegistryTag::OnGetExtraObjectTagsWithContext.AddUObject(this, &UAssetRegistryImpl::OnGetExtraObjectTags);
 	}
 	if (Context.bNeedsSearchAllAssetsAtStartSynchronous)
 	{
@@ -1564,7 +1572,7 @@ void UAssetRegistryImpl::FinishDestroy()
 
 		if (bAddMetaDataTagsToOnGetExtraObjectTags)
 		{
-			UObject::FAssetRegistryTag::OnGetExtraObjectTags.RemoveAll(this);
+			UObject::FAssetRegistryTag::OnGetExtraObjectTagsWithContext.RemoveAll(this);
 		}
 		FCoreDelegates::OnFEngineLoopInitComplete.RemoveAll(this);
 
@@ -2351,9 +2359,11 @@ void EnumerateMemoryAssets(const FARCompiledFilter& InFilter, TSet<FName>& OutPa
 	EnumerateMemoryAssetsHelper(InFilter, OutPackageNamesWithAssets, bOutStopIteration,
 		[&InFilter, &Callback, &InterfaceLock, &GuardedDataState](const UObject* Object, FAssetData&& PartialAssetData)
 		{
-			Object->GetAssetRegistryTags(PartialAssetData);
+			FAssetRegistryTagsContextData Context(Object, EAssetRegistryTagsCaller::AssetRegistryQuery);
+			Object->GetAssetRegistryTags(Context, PartialAssetData);
 			{
-				// GetAssetRegistryTags call does not add extended tags that may exist in the on-disk Asset.
+				// GetAssetRegistryTags with EAssetRegistryTagsCaller::AssetRegistryQuery does not add some tags that
+				// are too expensive to regularly compute but that exist in the on-disk Asset from SavePackage.
 				// Our contract for on-disk versus in-memory tags is that in-memory tags override on-disk tags, but we
 				// keep any on-disk tags that do not exist in the in-memory tags because they may be extended tags.
 				FReadScopeLock InterfaceScopeLock(InterfaceLock);
@@ -2427,7 +2437,8 @@ FAssetData UAssetRegistryImpl::GetAssetByObjectPath(const FSoftObjectPath& Objec
 		{
 			if (!bSkipARFilteredAssets || !UE::AssetRegistry::FFiltering::ShouldSkipAsset(Asset))
 			{
-				return FAssetData(Asset, FAssetData::ECreationFlags::None /** Do not allow blueprint classes */);
+				return FAssetData(Asset, FAssetData::ECreationFlags::None /** Do not allow blueprint classes */,
+					EAssetRegistryTagsCaller::AssetRegistryQuery);
 			}
 			else
 			{
@@ -3850,7 +3861,8 @@ void UAssetRegistryImpl::AssetCreated(UObject* NewAsset)
 		{
 			checkf(IsInGameThread(), TEXT("AssetCreated is not yet implemented as callable from other threads"));
 			// Let subscribers know that the new asset was added to the registry
-			FAssetData AssetData = FAssetData(NewAsset, FAssetData::ECreationFlags::AllowBlueprintClass);
+			FAssetData AssetData = FAssetData(NewAsset, FAssetData::ECreationFlags::AllowBlueprintClass,
+				EAssetRegistryTagsCaller::AssetRegistryQuery);
 			AssetAddedEvent.Broadcast(AssetData);
 			OnAssetsAdded().Broadcast({ AssetData });
 
@@ -3903,7 +3915,8 @@ void UAssetRegistryImpl::AssetDeleted(UObject* DeletedAsset)
 
 		if (!bShouldSkipAsset)
 		{
-			FAssetData AssetDataDeleted = FAssetData(DeletedAsset, FAssetData::ECreationFlags::AllowBlueprintClass);
+			FAssetData AssetDataDeleted = FAssetData(DeletedAsset, FAssetData::ECreationFlags::AllowBlueprintClass,
+				EAssetRegistryTagsCaller::AssetRegistryQuery);
 
 			checkf(IsInGameThread(), TEXT("AssetDeleted is not yet implemented as callable from other threads"));
 			// Let subscribers know that the asset was removed from the registry
@@ -3962,7 +3975,8 @@ void UAssetRegistryImpl::AssetRenamed(const UObject* RenamedAsset, const FString
 		{
 			checkf(IsInGameThread(), TEXT("AssetRenamed is not yet implemented as callable from other threads"));
 			AssetRenamedEvent.Broadcast(
-				FAssetData(RenamedAsset, FAssetData::ECreationFlags::AllowBlueprintClass),
+				FAssetData(RenamedAsset, FAssetData::ECreationFlags::AllowBlueprintClass,
+					EAssetRegistryTagsCaller::AssetRegistryQuery),
 				OldObjectPath
 			);
 		}
@@ -3971,7 +3985,7 @@ void UAssetRegistryImpl::AssetRenamed(const UObject* RenamedAsset, const FString
 
 void UAssetRegistryImpl::AssetSaved(const UObject& SavedAsset)
 {
-	AssetFullyUpdateTags(const_cast<UObject*>(&SavedAsset));
+	AssetUpdateTags(const_cast<UObject*>(&SavedAsset), EAssetRegistryTagsCaller::Fast);
 }
 
 void UAssetRegistryImpl::AssetsSaved(TArray<FAssetData>&& Assets)
@@ -3987,9 +4001,14 @@ void UAssetRegistryImpl::AssetsSaved(TArray<FAssetData>&& Assets)
 }
 
 void UAssetRegistryImpl::AssetFullyUpdateTags(UObject* Object)
+{
+	AssetUpdateTags(Object, EAssetRegistryTagsCaller::Fast);
+}
+
+void UAssetRegistryImpl::AssetUpdateTags(UObject* Object, EAssetRegistryTagsCaller Caller)
 {
 #if WITH_EDITOR
-	FAssetData AssetData(Object, FAssetData::ECreationFlags::None);
+	FAssetData AssetData(Object, FAssetData::ECreationFlags::None, Caller);
 	TArray<FAssetData> Assets;
 	Assets.Add(MoveTemp(AssetData));
 
@@ -6354,7 +6373,8 @@ void UAssetRegistryImpl::ProcessLoadedAssetsToUpdateCache(UE::AssetRegistry::Imp
 				// If the object has changed and is no longer an asset, ignore it. This can happen when an Actor is modified during cooking to no longer have an external package
 				continue;
 			}
-			BatchAssetDatas.Add(FAssetData(LoadedObject, FAssetData::ECreationFlags::AllowBlueprintClass));
+			BatchAssetDatas.Add(FAssetData(LoadedObject, FAssetData::ECreationFlags::AllowBlueprintClass,
+				EAssetRegistryTagsCaller::AssetRegistryLoad));
 
 			// Check to see if we have run out of time in this tick
 			if (!bFlushFullBuffer &&
@@ -6442,12 +6462,13 @@ void FAssetRegistryImpl::PushProcessLoadedAssetsBatch(Impl::FEventContext& Event
 		else
 		{
 			// When updating disk-based AssetData with the AssetData from a loaded UObject, we keep
-			// existing tags from disk even if they are no longer returned from the GetAssetRegistryTags
-			// function on the loaded UObject. We do this because they might come from GetAssetRegistryTagsExtended,
-			// which is only called during Save.
+			// existing tags from disk even if they are not returned from the
+			// GetAssetRegistryTags(EAssetRegistryTagsCaller::AssetRegistryLoad) function on the loaded UObject.
+			// We do this because the tags might be tags that are only calculated during
+			// GetAssetRegistryTags(EAssetRegistryTagsCaller::SavePackage).
 			// Modified tag values on the other hand do overwrite the old values from disk.
 			// This means that the only way to delete no-longer present tags from an AssetData
-			// is to resave the package, or to manually call AssetFullyUpdateTags from c++.
+			// is to resave the package, or to manually call AssetUpdateTags(EAssetRegistryTagsCaller::FullUpdate) from c++.
 			UpdateAssetData(EventContext, *DataFromGather, MoveTemp(NewAssetData), true /* bKeepDeletedTags */);
 		}
 	}
@@ -7010,12 +7031,12 @@ void UAssetRegistryImpl::GetInheritanceContextAfterVerifyingLock(uint64 CurrentG
 }
 
 #if WITH_EDITOR
-void UAssetRegistryImpl::OnGetExtraObjectTags(const UObject* Object, TArray<UObject::FAssetRegistryTag>& OutTags)
+void UAssetRegistryImpl::OnGetExtraObjectTags(FAssetRegistryTagsContext& Context)
 {
 	if (bAddMetaDataTagsToOnGetExtraObjectTags)
 	{
 		// It is critical that bIncludeOnlyOnDiskAssets=true otherwise this will cause an infinite loop
-		const FAssetData AssetData = GetAssetByObjectPath(FSoftObjectPath(Object), /*bIncludeOnlyOnDiskAssets=*/true);
+		const FAssetData AssetData = GetAssetByObjectPath(FSoftObjectPath(Context.GetObject()), /*bIncludeOnlyOnDiskAssets=*/true);
 		// Adding metadata tags from disk is only necessary for cooked assets; uncooked assets still have the metadata and add them elsewhere
 		// in UObject::GetAssetRegistryTags. Adding the tags from disk into uncooked assets would make the tags impossible to remove when
 		// the uncooked assets are resaved.
@@ -7025,12 +7046,12 @@ void UAssetRegistryImpl::OnGetExtraObjectTags(const UObject* Object, TArray<UObj
 			for (const FName MetaDataTag : MetaDataTags)
 			{
 				auto OutTagsContainsTagPredicate = [MetaDataTag](const UObject::FAssetRegistryTag& Tag) { return Tag.Name == MetaDataTag; };
-				if (!OutTags.ContainsByPredicate(OutTagsContainsTagPredicate))
+				if (!Context.ContainsTag(MetaDataTag))
 				{
 					FAssetTagValueRef TagValue = AssetData.TagsAndValues.FindTag(MetaDataTag);
 					if (TagValue.IsSet())
 					{
-						OutTags.Add(UObject::FAssetRegistryTag(MetaDataTag, TagValue.AsString(), UObject::FAssetRegistryTag::TT_Alphabetical));
+						Context.AddTag(UObject::FAssetRegistryTag(MetaDataTag, TagValue.AsString(), UObject::FAssetRegistryTag::TT_Alphabetical));
 					}
 				}
 			}
diff --git a/Engine/Source/Runtime/AssetRegistry/Private/AssetRegistry.h b/Engine/Source/Runtime/AssetRegistry/Private/AssetRegistry.h
--- a/Engine/Source/Runtime/AssetRegistry/Private/AssetRegistry.h
+++ b/Engine/Source/Runtime/AssetRegistry/Private/AssetRegistry.h
@@ -141,8 +141,11 @@ public:
 	virtual void AssetCreated(UObject* NewAsset) override;
 	virtual void AssetDeleted(UObject* DeletedAsset) override;
 	virtual void AssetRenamed(const UObject* RenamedAsset, const FString& OldObjectPath) override;
+	UE_DEPRECATED(5.2, "Use the new AssetsSaved function that takes FAssetData.")
 	virtual void AssetSaved(const UObject& SavedAsset) override;
 	virtual void AssetsSaved(TArray<FAssetData>&& SavedAssets) override;
+	virtual void AssetUpdateTags(UObject* Object, EAssetRegistryTagsCaller Caller) override;
+	UE_DEPRECATED(5.4, "Call AssetUpdateTags with EAssetRegistryTagsCaller::Fast")
 	virtual void AssetFullyUpdateTags(UObject* Object) override;
 	virtual void AssetTagsFinalized(const UObject& FinalizedAsset) override;
 
@@ -190,7 +193,7 @@ public:
 	virtual FFileLoadProgressUpdatedEvent& OnFileLoadProgressUpdated() override;
 
 	virtual bool IsLoadingAssets() const override;
-	virtual bool IsUpdateDiskCacheAfterLoad() const override
+	virtual bool ShouldUpdateDiskCacheAfterLoad() const override
 	{
 #if WITH_EDITORONLY_DATA
 		return bUpdateDiskCacheAfterLoad;
@@ -289,7 +292,7 @@ private:
 	 * Callback for FObject::FAssetRegistryTag::OnGetExtraObjectTags
 	 * If bAddMetaDataTagsToOnGetExtraObjectTags is true, this function will add missing UMetaData tags to cooked assets
 	 */
-	void OnGetExtraObjectTags(const UObject* Object, TArray<UObject::FAssetRegistryTag>& OutTags);
+	void OnGetExtraObjectTags(FAssetRegistryTagsContext& Context);
 #endif
 
 private:
diff --git a/Engine/Source/Runtime/AssetRegistry/Private/AssetRegistryState.cpp b/Engine/Source/Runtime/AssetRegistry/Private/AssetRegistryState.cpp
--- a/Engine/Source/Runtime/AssetRegistry/Private/AssetRegistryState.cpp
+++ b/Engine/Source/Runtime/AssetRegistry/Private/AssetRegistryState.cpp
@@ -23,6 +23,7 @@
 #include "Serialization/ArrayReader.h"
 #include "Serialization/LargeMemoryReader.h"
 #include "Serialization/MemoryWriter.h"
+#include "UObject/AssetRegistryTagsContext.h"
 #include "UObject/MetaData.h"
 #include "UObject/NameBatchSerialization.h"
 #include "UObject/PrimaryAssetId.h"
@@ -159,22 +160,35 @@ void FAssetRegistryState::FilterTags(const FAssetDataTagMapSharedView& InTagsAnd
 {
 	const TSet<FName>* AllClassesFilterList = Options.CookFilterlistTagsByClass.Find(UE::AssetRegistry::WildcardPathName);
 
+	TStringBuilder<64> TagNameStr;
+
 	// Exclude denied tags or include only allowed tags, based on how we were configured in ini
 	for (const auto& TagPair : InTagsAndValues)
 	{
 		bool bKeep = false;
 
-		const bool bInAllClassesList = AllClassesFilterList && (AllClassesFilterList->Contains(TagPair.Key) || AllClassesFilterList->Contains(UE::AssetRegistry::WildcardFName));
-		const bool bInClassSpecificList = ClassSpecificFilterList && (ClassSpecificFilterList->Contains(TagPair.Key) || ClassSpecificFilterList->Contains(UE::AssetRegistry::WildcardFName));
-		if (Options.bUseAssetRegistryTagsAllowListInsteadOfDenyList)
+		// Cook_ tags, aka DevelopmentAssetRegistryTags are special; they are kept depending on whether the Options
+		// are development or runtime and they do not use the options' filter list.
+		TagNameStr.Reset();
+		TagNameStr << TagPair.Key;
+		if (FStringView(TagNameStr).StartsWith(UE::AssetRegistry::CookTagPrefix, ESearchCase::IgnoreCase))
 		{
-			// It's an allow list, only include it if it is in the all classes list or in the class specific list
-			bKeep = bInAllClassesList || bInClassSpecificList;
+			bKeep = Options.bKeepDevelopmentAssetRegistryTags;
 		}
 		else
 		{
-			// It's a deny list, include it unless it is in the all classes list or in the class specific list
-			bKeep = !bInAllClassesList && !bInClassSpecificList;
+			const bool bInAllClassesList = AllClassesFilterList && (AllClassesFilterList->Contains(TagPair.Key) || AllClassesFilterList->Contains(UE::AssetRegistry::WildcardFName));
+			const bool bInClassSpecificList = ClassSpecificFilterList && (ClassSpecificFilterList->Contains(TagPair.Key) || ClassSpecificFilterList->Contains(UE::AssetRegistry::WildcardFName));
+			if (Options.bUseAssetRegistryTagsAllowListInsteadOfDenyList)
+			{
+			// It's an allow list, only include it if it is in the all classes list or in the class specific list
+				bKeep = bInAllClassesList || bInClassSpecificList;
+			}
+			else
+			{
+				// It's a deny list, include it unless it is in the all classes list or in the class specific list
+				bKeep = !bInAllClassesList && !bInClassSpecificList;
+			}
 		}
 		if (bKeep)
 		{
@@ -2571,6 +2585,30 @@ namespace AssetRegistry
 }
 }
 
+bool PrintAssetDataMapKeyIsLess(FName A, FName B)
+{
+	return A.Compare(B) < 0;
+}
+
+bool PrintAssetDataMapKeyIsLess(const FString& A, const FString& B)
+{
+	return A.Compare(B, ESearchCase::IgnoreCase) < 0;
+}
+
+bool PrintAssetDataMapKeyIsLess(const FTopLevelAssetPath& A, const FTopLevelAssetPath& B)
+{
+	return A.Compare(B) < 0;
+}
+
+template <typename KeyType>
+struct FPrintAssetDataMapKeyIsLess
+{
+	bool operator()(const KeyType& A, const KeyType& B) const
+	{
+		return PrintAssetDataMapKeyIsLess(A, B);
+	}
+};
+
 template <typename MapType>
 static void PrintAssetDataMap(FString Name, const MapType& AssetMap, TStringBuilder<16>& PageBuffer, const TFunctionRef<void()>& AddLine,
 	TUniqueFunction<void(const typename MapType::KeyType& Key, const FAssetData& Data)>&& PrintValue = {})
@@ -2581,14 +2619,7 @@ static void PrintAssetDataMap(FString Name, const MapType& AssetMap, TStringBuil
 	TArray<typename MapType::KeyType> Keys;
 	AssetMap.GenerateKeyArray(Keys);
 
-	struct FKeyTypeCompare
-	{
-		FORCEINLINE bool operator()(const typename MapType::KeyType& A, const typename MapType::KeyType& B) const
-		{
-			return A.Compare(B) < 0;
-		}
-	};
-	Keys.Sort(FKeyTypeCompare());
+	Keys.Sort(FPrintAssetDataMapKeyIsLess<typename MapType::KeyType>());
 
 	TArray<FAssetData*> Items;
 	Items.Reserve(1024);
@@ -2754,7 +2785,7 @@ void FAssetRegistryState::Dump(const TArray<FString>& Arguments, TArray<FString>
 			Keys.Emplace(FCachedAssetKey(*AssetData));
 		}
 		Keys.Sort([](const FCachedAssetKey& A, const FCachedAssetKey& B) {
-			return WriteToString<1024>(A).ToView().Compare(WriteToString<1024>(B).ToView()) < 0;
+			return WriteToString<1024>(A).ToView().Compare(WriteToString<1024>(B).ToView(), ESearchCase::IgnoreCase) < 0;
 		});
 
 		for (const FCachedAssetKey& Key : Keys)
diff --git a/Engine/Source/Runtime/AssetRegistry/Public/AssetRegistry/AssetRegistryState.h b/Engine/Source/Runtime/AssetRegistry/Public/AssetRegistry/AssetRegistryState.h
--- a/Engine/Source/Runtime/AssetRegistry/Public/AssetRegistry/AssetRegistryState.h
+++ b/Engine/Source/Runtime/AssetRegistry/Public/AssetRegistry/AssetRegistryState.h
@@ -76,6 +76,12 @@ struct FAssetRegistrySerializationOptions
 	/** Filter out searchable names from dependency data */
 	bool bFilterSearchableNames = false;
 
+	/**
+	 * Keep tags intended for the cooker's output DevelopmentAssetRegistry. this flag defaults to false and is set to
+	 * true only by the cooker.
+	 */
+	bool bKeepDevelopmentAssetRegistryTags = false;
+
 	/** The map of class pathname to tag set of tags that are allowed in cooked builds. This is either an allow list or deny list depending on bUseAssetRegistryTagsAllowListInsteadOfDenyList */
 	TMap<FTopLevelAssetPath, TSet<FName>> CookFilterlistTagsByClass;
 
diff --git a/Engine/Source/Runtime/AssetRegistry/Public/AssetRegistry/IAssetRegistry.h b/Engine/Source/Runtime/AssetRegistry/Public/AssetRegistry/IAssetRegistry.h
--- a/Engine/Source/Runtime/AssetRegistry/Public/AssetRegistry/IAssetRegistry.h
+++ b/Engine/Source/Runtime/AssetRegistry/Public/AssetRegistry/IAssetRegistry.h
@@ -27,6 +27,7 @@ class FAssetRegistryState;
 class FCbFieldView;
 class FCbWriter;
 class FDependsNode;
+enum class EAssetRegistryTagsCaller : uint8;
 struct FARFilter;
 struct FARCompiledFilter;
 struct FAssetData;
@@ -165,7 +166,8 @@ class UAssetRegistry : public UInterface
  *     up-to-date. The InMemory data will sometimes vary from the DiskGatheredData.
  * 
  *     When InMemoryData is returned some categories of data that are always missing from the object in memory (e.g.
- *     GetAssetRegistryTagsExtended) are read from the DiskGatheredData and added to the InMemoryData.
+ *     GetAssetRegistryTags(EAssetRegistryTagsCaller::SavePackage)) are read from the DiskGatheredData and added to
+ *     the InMemoryData.
  * 
  *     Setting this value to true will always be faster than setting it to false, because the same registry
  *     lookups are performed in either case, but the InMemoryData lookup is skipped in the true case.
@@ -679,12 +681,17 @@ public:
 
 	/**
 	 * Called on demand from systems that need to fully update an AssetData's tags. When an Asset is loaded its tags
-	 * are updated by calling UObject::GetAssetRegistryTags, but GetAssetRegistryTagsExtended is not called, and tags
-	 * that exist in the Old AssetData but not in the results from GetAssetRegistryTags are kept because they might be
-	 * extended tags. When an asset is saved, GetAssetRegistryTagsExtended is called and all old tags are deleted in
-	 * favor of the new list. AssetFullyUpdated allows a manual trigger of the on-SavePackage behavior:
-	 * GetAssetRegistryTagsExtended is called and all old tags are deleted in favor of the new list.
+	 * are updated by calling GetAssetRegistryTags(EAssetRegistryTagsCaller::AssetRegistryLoad), but that version of
+	 * the function is allowed to skip writing expensive tags, so tags that exist in the old AssetData but not in the
+	 * results from GetAssetRegistryTags(EAssetRegistryTagsCaller::AssetRegistryLoad) are kept because they might be
+	 * skipped expensive tags. When an asset is saved, all old tags are deleted and 
+	 * GetAssetRegistryTags(EAssetRegistryTagsCaller::SavePackage) is called. AssetUpdateTags allows a manual trigger
+	 * of the on-SavePackage behavior: all old tags are deleted and
+	 * GetAssetRegistryTags(Caller) is called. Pass in EAssetRegistryTagsCaller::FullUpdate to behave the same as
+	 * SavePackage.
 	 */
+	virtual void AssetUpdateTags(UObject* Object, EAssetRegistryTagsCaller Caller) = 0;
+	UE_DEPRECATED(5.4, "Call AssetUpdateTags with EAssetRegistryTagsCaller::Fast")
 	virtual void AssetFullyUpdateTags(UObject* Object) = 0;
 
 	/** Informs the asset registry that a Verse file has been created on disk */
@@ -778,7 +785,7 @@ public:
 	virtual bool IsLoadingAssets() const = 0;
 
 	/** If true, the AssetRegistry updates its on-disk information for an Asset whenever that Asset loads. */
-	virtual bool IsUpdateDiskCacheAfterLoad() const = 0;
+	virtual bool ShouldUpdateDiskCacheAfterLoad() const = 0;
 
 	/** Tick the asset registry */
 	virtual void Tick (float DeltaTime) = 0;
diff --git a/Engine/Source/Runtime/CoreUObject/Private/AssetRegistry/AssetData.cpp b/Engine/Source/Runtime/CoreUObject/Private/AssetRegistry/AssetData.cpp
--- a/Engine/Source/Runtime/CoreUObject/Private/AssetRegistry/AssetData.cpp
+++ b/Engine/Source/Runtime/CoreUObject/Private/AssetRegistry/AssetData.cpp
@@ -16,6 +16,7 @@
 #include "Serialization/CompactBinaryWriter.h"
 #include "Serialization/CustomVersion.h"
 #include "String/Find.h"
+#include "UObject/AssetRegistryTagsContext.h"
 #include "UObject/LinkerLoad.h"
 #include "UObject/PropertyPortFlags.h"
 
@@ -331,6 +332,11 @@ PRAGMA_ENABLE_DEPRECATION_WARNINGS
 }
 
 FAssetData::FAssetData(const UObject* InAsset, FAssetData::ECreationFlags InCreationFlags)
+	: FAssetData(InAsset, InCreationFlags, EAssetRegistryTagsCaller::Uncategorized)
+{
+}
+
+FAssetData::FAssetData(const UObject* InAsset, FAssetData::ECreationFlags InCreationFlags, EAssetRegistryTagsCaller Caller)
 {
 	if (InAsset != nullptr)
 	{
@@ -362,7 +368,8 @@ PRAGMA_ENABLE_DEPRECATION_WARNINGS
 
 		if (!EnumHasAnyFlags(InCreationFlags, FAssetData::ECreationFlags::SkipAssetRegistryTagsGathering))
 		{
-			InAsset->GetAssetRegistryTags(*this);
+			FAssetRegistryTagsContextData Context(InAsset, Caller);
+			InAsset->GetAssetRegistryTags(Context, *this);
 		}
 
 		PackageFlags = Package->GetPackageFlags();
diff --git a/Engine/Source/Runtime/CoreUObject/Private/UObject/AssetRegistryTagsContext.cpp b/Engine/Source/Runtime/CoreUObject/Private/UObject/AssetRegistryTagsContext.cpp
new file mode 100644
--- /dev/null
+++ b/Engine/Source/Runtime/CoreUObject/Private/UObject/AssetRegistryTagsContext.cpp
@@ -0,0 +1,261 @@
+// Copyright Epic Games, Inc. All Rights Reserved.
+
+#include "UObject/AssetRegistryTagsContext.h"
+
+#include "Logging/LogMacros.h"
+#include "Misc/StringBuilder.h"
+
+FAssetRegistryTagsContextData::FAssetRegistryTagsContextData(const UObject* CallTarget,
+	EAssetRegistryTagsCaller InCaller)
+: Object(CallTarget)
+, Caller(InCaller)
+{
+}
+
+FAssetRegistryTagsContext::FAssetRegistryTagsContext(FAssetRegistryTagsContextData& InData)
+	: Data(InData)
+{
+}
+
+FAssetRegistryTagsContext::FAssetRegistryTagsContext(FAssetRegistryTagsContext& Other)
+	: Data(Other.Data)
+{
+}
+
+const UObject* FAssetRegistryTagsContext::GetObject() const
+{
+	return Data.Object;
+}
+
+EAssetRegistryTagsCaller FAssetRegistryTagsContext::GetCaller() const
+{
+	return Data.Caller;
+}
+
+bool FAssetRegistryTagsContext::IsFullUpdate() const
+{
+	return Data.Caller == EAssetRegistryTagsCaller::FullUpdate || Data.bFullUpdateRequested;
+}
+
+bool FAssetRegistryTagsContext::IsSaving() const
+{
+	return Data.Caller == EAssetRegistryTagsCaller::SavePackage;
+}
+
+bool FAssetRegistryTagsContext::IsProceduralSave() const
+{
+	return Data.bProceduralSave;
+}
+
+bool FAssetRegistryTagsContext::IsCooking() const
+{
+	return Data.TargetPlatform != nullptr;
+}
+
+const ITargetPlatform* FAssetRegistryTagsContext::GetTargetPlatform() const
+{
+	return Data.TargetPlatform;
+}
+
+bool FAssetRegistryTagsContext::IsCookByTheBook() const
+{
+	return Data.CookType == UE::Cook::ECookType::ByTheBook;
+}
+
+bool FAssetRegistryTagsContext::IsCookOnTheFly() const
+{
+	return Data.CookType == UE::Cook::ECookType::OnTheFly;
+}
+
+bool FAssetRegistryTagsContext::IsCookTypeUnknown() const
+{
+	return Data.CookType == UE::Cook::ECookType::Unknown;
+}
+
+UE::Cook::ECookType FAssetRegistryTagsContext::GetCookType() const
+{
+	return Data.CookType;
+}
+
+UE::Cook::ECookingDLC FAssetRegistryTagsContext::GetCookingDLC() const
+{
+	return Data.CookingDLC;
+}
+
+UObject::FAssetRegistryTag& FAssetRegistryTagsContext::FindOrAddTag(FName TagName, bool* bOutAlreadyExists)
+{
+	return FindOrAddTagInternal(TagName, FStringView(), bOutAlreadyExists);
+}
+
+UObject::FAssetRegistryTag* FAssetRegistryTagsContext::FindTag(FName TagName)
+{
+	return FindTagInternal(TagName, FStringView());
+}
+
+bool FAssetRegistryTagsContext::ContainsTag(FName TagName) const
+{
+	return ContainsTagInternal(TagName, FStringView());
+}
+
+void FAssetRegistryTagsContext::AddTag(UObject::FAssetRegistryTag TagResult)
+{
+	AddTagInternal(MoveTemp(TagResult), FStringView());
+}
+
+bool FAssetRegistryTagsContext::WantsBundleResult() const
+{
+	return Data.bWantsBundleResult;
+}
+
+const FAssetBundleData* FAssetRegistryTagsContext::GetBundleResult()
+{
+	return Data.BundleResult;
+}
+
+void FAssetRegistryTagsContext::SetBundleResult(const FAssetBundleData* InBundleResult)
+{
+	Data.BundleResult = InBundleResult;
+}
+
+bool FAssetRegistryTagsContext::WantsCookTags() const
+{
+	return Data.bWantsCookTags;
+}
+
+UObject::FAssetRegistryTag& FAssetRegistryTagsContext::FindOrAddCookTag(FName TagName, bool* bOutAlreadyExists)
+{
+	if (!WantsCookTags())
+	{
+		static UObject::FAssetRegistryTag PlaceholderTag;
+		PlaceholderTag.Name = TagName;
+		return PlaceholderTag;
+	}
+	return FindOrAddTagInternal(TagName, UE::AssetRegistry::CookTagPrefix, bOutAlreadyExists);
+}
+
+UObject::FAssetRegistryTag* FAssetRegistryTagsContext::FindCookTag(FName TagName)
+{
+	if (!WantsCookTags())
+	{
+		return nullptr;
+	}
+	return FindTagInternal(TagName, UE::AssetRegistry::CookTagPrefix);
+}
+
+bool FAssetRegistryTagsContext::ContainsCookTag(FName TagName) const
+{
+	if (!WantsCookTags())
+	{
+		return false;
+	}
+	return ContainsTagInternal(TagName, UE::AssetRegistry::CookTagPrefix);
+}
+
+void FAssetRegistryTagsContext::AddCookTag(UObject::FAssetRegistryTag TagResult)
+{
+	if (!WantsCookTags())
+	{
+		return;
+	}
+	AddTagInternal(MoveTemp(TagResult), UE::AssetRegistry::CookTagPrefix);
+}
+
+bool FAssetRegistryTagsContext::TryValidateTagName(FName TagName) const
+{
+	using namespace UE::AssetRegistry;
+
+	if (!ensureMsgf(!TagName.IsNone(), TEXT("Invalid empty TagName passed to FAssetRegistryTagsContext.")))
+	{
+		return false;
+	}
+	TStringBuilder<64> TagNameStr(InPlace, TagName);
+	bool bProtectedKeyword = FStringView(TagNameStr).StartsWith(CookTagPrefix, ESearchCase::IgnoreCase);
+	if (!ensureMsgf(!bProtectedKeyword,
+		TEXT("Invalid TagName %s that starts with protected keyword %.*s passed to FAssetRegistryTagsContext."),
+		*TagNameStr, CookTagPrefix.Len(), CookTagPrefix.GetData()))
+	{
+		return false;
+	}
+	return true;
+}
+
+FName FAssetRegistryTagsContext::TransformTagName(FName TagName, FStringView PrefixToAdd) const
+{
+	if (!TryValidateTagName(TagName))
+	{
+		return NAME_None;
+	}
+	if (!PrefixToAdd.IsEmpty())
+	{
+		TStringBuilder<64> TagNameStr(InPlace, PrefixToAdd, TagName);
+		TagName = FName(TagNameStr);
+	}
+	return TagName;
+}
+
+UObject::FAssetRegistryTag& FAssetRegistryTagsContext::FindOrAddTagInternal(FName TagName, FStringView PrefixToAdd, bool* bOutAlreadyExists)
+{
+	TagName = TransformTagName(TagName, PrefixToAdd);
+	if (TagName.IsNone())
+	{
+		static UObject::FAssetRegistryTag ErrorResult;
+		ErrorResult.Name = TagName;
+		if (bOutAlreadyExists)
+		{
+			*bOutAlreadyExists = false;
+		}
+		return ErrorResult;
+	}
+
+	UObject::FAssetRegistryTag& Result = Data.Tags.FindOrAdd(TagName);
+	if (Result.Name != TagName)
+	{
+		Result.Name = TagName;
+		if (bOutAlreadyExists)
+		{
+			*bOutAlreadyExists = false;
+		}
+	}
+	else
+	{
+		if (bOutAlreadyExists)
+		{
+			*bOutAlreadyExists = true;
+		}
+	}
+	return Result;
+}
+
+UObject::FAssetRegistryTag* FAssetRegistryTagsContext::FindTagInternal(FName TagName, FStringView PrefixToAdd)
+{
+	TagName = TransformTagName(TagName, PrefixToAdd);
+	if (TagName.IsNone())
+	{
+		return nullptr;
+	}
+
+	return Data.Tags.Find(TagName);
+}
+
+bool FAssetRegistryTagsContext::ContainsTagInternal(FName TagName, FStringView PrefixToAdd) const
+{
+	TagName = TransformTagName(TagName, PrefixToAdd);
+	if (TagName.IsNone())
+	{
+		return false;
+	}
+
+	return Data.Tags.Contains(TagName);
+}
+
+void FAssetRegistryTagsContext::AddTagInternal(UObject::FAssetRegistryTag&& TagResult, FStringView PrefixToAdd)
+{
+	FName TagName = TransformTagName(TagResult.Name, PrefixToAdd);
+	if (TagName.IsNone())
+	{
+		return;
+	}
+	UObject::FAssetRegistryTag& AddedTag = Data.Tags.FindOrAdd(TagName);
+	AddedTag = MoveTemp(TagResult);
+	AddedTag.Name = TagName;
+}
\ No newline at end of file
diff --git a/Engine/Source/Runtime/CoreUObject/Private/UObject/Obj.cpp b/Engine/Source/Runtime/CoreUObject/Private/UObject/Obj.cpp
--- a/Engine/Source/Runtime/CoreUObject/Private/UObject/Obj.cpp
+++ b/Engine/Source/Runtime/CoreUObject/Private/UObject/Obj.cpp
@@ -31,6 +31,7 @@
 #include "UObject/Package.h"
 #include "UObject/MetaData.h"
 #include "Templates/Casts.h"
+#include "UObject/AssetRegistryTagsContext.h"
 #include "UObject/LazyObjectPtr.h"
 #include "UObject/SoftObjectPtr.h"
 #include "UObject/PropertyPortFlags.h"
@@ -132,6 +133,7 @@ static UPackage*			GObjTransientPkg								= NULL;
 	};
 	static FUObjectAnnotationSparse<FPropagatedEditChangeAnnotation, true> PropagatedEditChangeAnnotation;
 	UObject::FAssetRegistryTag::FOnGetObjectAssetRegistryTags UObject::FAssetRegistryTag::OnGetExtraObjectTags;
+	UObject::FAssetRegistryTag::FOnGetObjectAssetRegistryTagsWithContext UObject::FAssetRegistryTag::OnGetExtraObjectTagsWithContext;
 	UObject::FAssetRegistryTag::FOnGetExtendedAssetRegistryTagsForSave UObject::FAssetRegistryTag::OnGetExtendedAssetRegistryTagsForSave;
 	UObject::FOnGetPreviewPlatform UObject::OnGetPreviewPlatform;
 #endif // WITH_EDITOR
@@ -2081,26 +2083,23 @@ FString GetConfigFilename( UObject* SourceObject )
 namespace UE::Object::Private
 {
 
-// Thread local state to avoid UObject::GetAssetRegistryTags() API change
-thread_local FAssetBundleData const** TGetAssetRegistryTags_OutBundles = nullptr;
-
-static void GetAssetRegistryTagFromProperty(const void* BaseMemoryLocation, const UObject* OwnerObject, FProperty* Prop, TArray<UObject::FAssetRegistryTag>& OutTags)
+static void GetAssetRegistryTagFromProperty(const void* BaseMemoryLocation, const UObject* OwnerObject, FProperty* Prop, FAssetRegistryTagsContext Context)
 {
 	FStructProperty* StructProp = CastField<FStructProperty>(Prop);
 	if (StructProp && StructProp->Struct && StructProp->Struct->GetFName() == GAssetBundleDataName)
 	{
 		const FAssetBundleData* Bundles = reinterpret_cast<const FAssetBundleData*>(Prop->ContainerPtrToValuePtr<uint8>(BaseMemoryLocation));
 
-		if (FAssetBundleData const** OutBundles = TGetAssetRegistryTags_OutBundles)
+		if (Context.WantsBundleResult())
 		{
-			checkf(*OutBundles == nullptr, TEXT("Object %s has more than one FAssetBundleData!"), *OwnerObject->GetPathName());
-			*OutBundles = Bundles;
+			checkf(Context.GetBundleResult() == nullptr, TEXT("Object %s has more than one FAssetBundleData!"), *OwnerObject->GetPathName());
+			Context.SetBundleResult(Bundles);
 		}
 		else
 		{
 			FString PropertyStr;
 			Prop->ExportTextItem_Direct(PropertyStr, Bundles, Bundles, nullptr, PPF_None);
-			OutTags.Add(UObject::FAssetRegistryTag(GAssetBundleDataName, MoveTemp(PropertyStr), UObject::FAssetRegistryTag::ETagType::TT_Alphabetical));
+			Context.AddTag(UObject::FAssetRegistryTag(GAssetBundleDataName, MoveTemp(PropertyStr), UObject::FAssetRegistryTag::ETagType::TT_Alphabetical));
 		}
 	}
 	else if (Prop->HasAnyPropertyFlags(CPF_AssetRegistrySearchable))
@@ -2143,17 +2142,17 @@ static void GetAssetRegistryTagFromProperty(const void* BaseMemoryLocation, cons
 		const uint8* PropertyAddr = Prop->ContainerPtrToValuePtr<uint8>(BaseMemoryLocation);
 		Prop->ExportTextItem_Direct(PropertyStr, PropertyAddr, PropertyAddr, nullptr, PPF_None);
 
-		OutTags.Add(UObject::FAssetRegistryTag(Prop->GetFName(), MoveTemp(PropertyStr), TagType));
+		Context.AddTag(UObject::FAssetRegistryTag(Prop->GetFName(), MoveTemp(PropertyStr), TagType));
 	}
 }
 
-static void GetAssetRegistryTagsFromSearchableProperties(const UObject* Object, TArray<UObject::FAssetRegistryTag>& OutTags)
+static void GetAssetRegistryTagsFromSearchableProperties(const UObject* Object, FAssetRegistryTagsContext Context)
 {
 	check(nullptr != Object);
 
 	for (TFieldIterator<FProperty> FieldIt( Object->GetClass() ); FieldIt; ++FieldIt)
 	{
-		GetAssetRegistryTagFromProperty(Object, Object, CastField<FProperty>(*FieldIt), OutTags);
+		GetAssetRegistryTagFromProperty(Object, Object, CastField<FProperty>(*FieldIt), Context);
 	}
 
 	UScriptStruct* SparseClassDataStruct = Object->GetClass()->GetSparseClassDataStruct();
@@ -2162,7 +2161,7 @@ static void GetAssetRegistryTagsFromSearchableProperties(const UObject* Object,
 		const void* SparseClassData = Object->GetClass()->GetSparseClassData(EGetSparseClassDataMethod::ArchetypeIfNull);
 		for (TFieldIterator<FProperty> FieldIt(SparseClassDataStruct); FieldIt; ++FieldIt)
 		{
-			GetAssetRegistryTagFromProperty(SparseClassData, Object, CastField<FProperty>(*FieldIt), OutTags);
+			GetAssetRegistryTagFromProperty(SparseClassData, Object, CastField<FProperty>(*FieldIt), Context);
 		}
 	}
 }
@@ -2173,24 +2172,89 @@ const FName FPrimaryAssetId::PrimaryAssetTypeTag(TEXT("PrimaryAssetType"));
 const FName FPrimaryAssetId::PrimaryAssetNameTag(TEXT("PrimaryAssetName"));
 const FName FPrimaryAssetId::PrimaryAssetDisplayNameTag(TEXT("PrimaryAssetDisplayName"));
 
+// This list of the objects that are forwarding from one GetAssetRegistry tags to the other has to be a list rather
+// than merely a bool because some GetAssetRegistryTags calls call GetAssetRegistryTags on other objects, e.g. UWorld
+// calls it on its LevelBlueprint.
+// UE_DEPRECATED(5.4, "Used to provide backwards compatibility for the deprecated GetAssetRegistryTags function") // UE_DEPRECATED seems not to work with thread_local
+thread_local TArray<const UObject*, TInlineAllocator<2>> GAssetRegistryTagsObjectsBeingForwarded;
+// UE_DEPRECATED(5.4, "Used to provide backwards compatibility for the deprecated GetAssetRegistryTags function") // UE_DEPRECATED seems not to work with thread_local
+thread_local bool bGLegacyTagsWantsBundleResult = false;
+
 void UObject::GetAssetRegistryTags(TArray<FAssetRegistryTag>& OutTags) const
+{
+	PRAGMA_DISABLE_DEPRECATION_WARNINGS;
+	if (!GAssetRegistryTagsObjectsBeingForwarded.Contains(this))
+	{
+		FAssetRegistryTagsContextData Context(this, EAssetRegistryTagsCaller::Uncategorized);
+		Context.bWantsBundleResult = bGLegacyTagsWantsBundleResult;
+
+		GAssetRegistryTagsObjectsBeingForwarded.Add(this);
+		GetAssetRegistryTags(Context);
+		check(!GAssetRegistryTagsObjectsBeingForwarded.IsEmpty() && GAssetRegistryTagsObjectsBeingForwarded.Last() == this);
+		GAssetRegistryTagsObjectsBeingForwarded.Pop(false /* bAllowShrinking */);
+
+		OutTags.Reserve(OutTags.Num() + Context.Tags.Num());
+		for (TPair<FName, FAssetRegistryTag>& Pair : Context.Tags)
+		{
+			OutTags.Add(MoveTemp(Pair.Value));
+		}
+	}
+	PRAGMA_ENABLE_DEPRECATION_WARNINGS;
+}
+
+void UObject::GetAssetRegistryTags(FAssetRegistryTagsContext Context) const
 {
 	using namespace UE::Object::Private;
+
+	PRAGMA_DISABLE_DEPRECATION_WARNINGS;
+	auto AddLegacyTags = [&Context](auto&& Callback)
+	{
+		TArray<FAssetRegistryTag> LegacyTags;
+		Callback(LegacyTags);
+		for (FAssetRegistryTag& Tag : LegacyTags)
+		{
+			Context.AddTag(MoveTemp(Tag));
+		}
+	};
+
+	// Forward this call to the legacy version for classes that have not converted yet.
+	if (!GAssetRegistryTagsObjectsBeingForwarded.Contains(this))
+	{
+		TGuardValue<bool> WantsBundleScope(bGLegacyTagsWantsBundleResult, Context.WantsBundleResult());
+		GAssetRegistryTagsObjectsBeingForwarded.Add(this);
+		AddLegacyTags([this](TArray<FAssetRegistryTag>& Tags) { GetAssetRegistryTags(Tags); });
+		check(!GAssetRegistryTagsObjectsBeingForwarded.IsEmpty() && GAssetRegistryTagsObjectsBeingForwarded.Last() == this);
+		GAssetRegistryTagsObjectsBeingForwarded.Pop(false /* bAllowShrinking */);
+	}
+	PRAGMA_ENABLE_DEPRECATION_WARNINGS;
+
 	UE::Core::Private::FPlayInEditorLoadingScope Scope(INDEX_NONE);
 
 	// Add primary asset info if valid
 	FPrimaryAssetId PrimaryAssetId = GetPrimaryAssetId();
 	if (PrimaryAssetId.IsValid())
 	{
-		OutTags.Add(FAssetRegistryTag(FPrimaryAssetId::PrimaryAssetTypeTag, PrimaryAssetId.PrimaryAssetType.ToString(), UObject::FAssetRegistryTag::TT_Alphabetical));
-		OutTags.Add(FAssetRegistryTag(FPrimaryAssetId::PrimaryAssetNameTag, PrimaryAssetId.PrimaryAssetName.ToString(), UObject::FAssetRegistryTag::TT_Alphabetical));
+		Context.AddTag(FAssetRegistryTag(FPrimaryAssetId::PrimaryAssetTypeTag, PrimaryAssetId.PrimaryAssetType.ToString(),
+			UObject::FAssetRegistryTag::TT_Alphabetical));
+		Context.AddTag(FAssetRegistryTag(FPrimaryAssetId::PrimaryAssetNameTag, PrimaryAssetId.PrimaryAssetName.ToString(),
+			UObject::FAssetRegistryTag::TT_Alphabetical));
 	}
 
-	GetAssetRegistryTagsFromSearchableProperties(this, OutTags);
+	GetAssetRegistryTagsFromSearchableProperties(this, Context);
 
 #if WITH_EDITOR
 	// Notify external sources that we need tags.
-	FAssetRegistryTag::OnGetExtraObjectTags.Broadcast(this, OutTags);
+	FAssetRegistryTag::OnGetExtraObjectTagsWithContext.Broadcast(Context);
+
+	PRAGMA_DISABLE_DEPRECATION_WARNINGS;
+	AddLegacyTags([this](TArray<FAssetRegistryTag>& Tags)
+		{ FAssetRegistryTag::OnGetExtraObjectTags.Broadcast(this, Tags); });
+	if (Context.IsFullUpdate())
+	{
+		AddLegacyTags([this, &Context](TArray<FAssetRegistryTag>& Tags)
+			{ GetExtendedAssetRegistryTagsForSave(Context.GetTargetPlatform(), Tags); });
+	}
+	PRAGMA_ENABLE_DEPRECATION_WARNINGS;
 
 	// Check if there's a UMetaData for this object that has tags that are requested in the settings to be transferred to the Asset Registry
 	const TSet<FName>& MetaDataTagsForAR = GetMetaDataTagsForAssetRegistry();
@@ -2204,7 +2268,7 @@ void UObject::GetAssetRegistryTags(TArray<FAssetRegistryTag>& OutTags) const
 				FName Tag = It->Key;
 				if (!Tag.IsNone() && MetaDataTagsForAR.Contains(Tag))
 				{
-					OutTags.Add(FAssetRegistryTag(Tag, It->Value, UObject::FAssetRegistryTag::TT_Alphabetical));
+					Context.AddTag(FAssetRegistryTag(Tag, It->Value, UObject::FAssetRegistryTag::TT_Alphabetical));
 				}
 			}
 		}
@@ -2213,19 +2277,30 @@ void UObject::GetAssetRegistryTags(TArray<FAssetRegistryTag>& OutTags) const
 }
 
 #if WITH_EDITOR
+void UObject::GetAdditionalAssetDataObjectsForCook(FArchiveCookContext& CookContext,
+	TArray<UObject*>& OutObjects) const
+{
+}
+
 void UObject::GetExtendedAssetRegistryTagsForSave(const ITargetPlatform* TargetPlatform, TArray<FAssetRegistryTag>& OutTags) const
 {
-	// Notify external sources that we need tags for save.
+	// DEPRECATION Note: This function will not return the data from classes that have been converted to use the new
+	// FAssetRegistryTagsContext API. We could make it do so, but it would require extra effort because this function is
+	// supposed to return only the expensive tags, and not the common tags that are also returned when called with
+	// EAssetRegistryTagsCaller::AssetRegistryLoad. Because this function was designed only to be called from SavePackage,
+	// and we have removed SavePackage's dependence on it, we decided not to make that extra effort. Any licensee calling
+	// this function should instead call GetAssetRegistryTags with EAssetRegistryTagsCaller::SavePackage.
 	FAssetRegistryTag::OnGetExtendedAssetRegistryTagsForSave.Broadcast(this, TargetPlatform, OutTags);
 }
 #endif // WITH_EDITOR
 
-static FAssetDataTagMapSharedView MakeSharedTagMap(TArray<UObject::FAssetRegistryTag>&& Tags)
+static FAssetDataTagMapSharedView MakeSharedTagMap(TMap<FName, UObject::FAssetRegistryTag>&& Tags)
 {
 	FAssetDataTagMap Out;
 	Out.Reserve(Tags.Num());
-	for (UObject::FAssetRegistryTag& Tag : Tags)
+	for (TPair<FName,UObject::FAssetRegistryTag>& Pair : Tags)
 	{
+		UObject::FAssetRegistryTag& Tag = Pair.Value;
 		// Don't add empty tags
 		if (!Tag.Name.IsNone() && !Tag.Value.IsEmpty())
 		{
@@ -2247,18 +2322,21 @@ static TSharedPtr<FAssetBundleData, ESPMode::ThreadSafe> MakeSharedBundles(const
 }
 
 void UObject::GetAssetRegistryTags(FAssetData& Out) const
+{
+	FAssetRegistryTagsContextData Context(this, EAssetRegistryTagsCaller::Uncategorized);
+	GetAssetRegistryTags(Context, Out);
+}
+
+void UObject::GetAssetRegistryTags(FAssetRegistryTagsContext Context, FAssetData& Out) const
 {
 	using namespace UE::Object::Private;
 
 	const FAssetBundleData* Bundles = nullptr;
-
-	TArray<FAssetRegistryTag> Tags;
-	TGetAssetRegistryTags_OutBundles = &Bundles;
-	GetAssetRegistryTags(Tags);
-	TGetAssetRegistryTags_OutBundles = nullptr;
-
-	Out.TagsAndValues = MakeSharedTagMap(MoveTemp(Tags));
-	Out.TaggedAssetBundles = MakeSharedBundles(Bundles);
+	FAssetRegistryTagsContextData& ContextData = Context.Data;
+	ContextData.bWantsBundleResult = true;
+	GetAssetRegistryTags(Context);
+	Out.TagsAndValues = MakeSharedTagMap(MoveTemp(ContextData.Tags));
+	Out.TaggedAssetBundles = MakeSharedBundles(ContextData.BundleResult);
 }
 
 const FName& UObject::SourceFileTagName()
diff --git a/Engine/Source/Runtime/CoreUObject/Private/UObject/SavePackage.cpp b/Engine/Source/Runtime/CoreUObject/Private/UObject/SavePackage.cpp
--- a/Engine/Source/Runtime/CoreUObject/Private/UObject/SavePackage.cpp
+++ b/Engine/Source/Runtime/CoreUObject/Private/UObject/SavePackage.cpp
@@ -30,7 +30,7 @@ bool UPackage::SavePackage(UPackage* InOuter, UObject* Base, EObjectFlags TopLev
 	if (TargetPlatform != nullptr)
 	{
 		CookContext.Emplace(InOuter, UE::Cook::ECookType::Unknown,
-			UE::Cook::ECookingDLC::Unknown);
+			UE::Cook::ECookingDLC::Unknown, TargetPlatform);
 		CookData.Emplace(*TargetPlatform, *CookContext);
 	}
 	FSavePackageArgs SaveArgs = { nullptr /* deprecated target platform */, CookData.GetPtrOrNull(), TopLevelFlags, SaveFlags, bForceByteSwapping,
diff --git a/Engine/Source/Runtime/CoreUObject/Private/UObject/SavePackage/SavePackageUtilities.cpp b/Engine/Source/Runtime/CoreUObject/Private/UObject/SavePackage/SavePackageUtilities.cpp
--- a/Engine/Source/Runtime/CoreUObject/Private/UObject/SavePackage/SavePackageUtilities.cpp
+++ b/Engine/Source/Runtime/CoreUObject/Private/UObject/SavePackage/SavePackageUtilities.cpp
@@ -5,6 +5,7 @@
 #include "Algo/Sort.h"
 #include "Algo/Unique.h"
 #include "AssetRegistry/AssetData.h"
+#include "AssetRegistry/CookTagList.h"
 #include "Blueprint/BlueprintSupport.h"
 #include "CoreMinimal.h"
 #include "HAL/FileManager.h"
@@ -25,6 +26,8 @@
 #include "Serialization/LargeMemoryWriter.h"
 #include "Serialization/PackageWriter.h"
 #include "Tasks/Task.h"
+#include "UObject/ArchiveCookContext.h"
+#include "UObject/AssetRegistryTagsContext.h"
 #include "UObject/Class.h"
 #include "UObject/GCScopeLock.h"
 #include "UObject/Linker.h"
@@ -2150,13 +2153,46 @@ bool UPackage::IsEmptyPackage(UPackage* Package, const UObject* LastReferencer)
 namespace UE::AssetRegistry
 {
 
-// See the corresponding ReadPackageDataMain and ReadPackageDataDependencies defined in PackageReader.cpp in AssetRegistry module
 void WritePackageData(FStructuredArchiveRecord& ParentRecord, bool bIsCooking, const UPackage* Package,
-		FLinkerSave* Linker, const TSet<TObjectPtr<UObject>>& ImportsUsedInGame, const TSet<FName>& SoftPackagesUsedInGame,
+	FLinkerSave* Linker, const TSet<TObjectPtr<UObject>>& ImportsUsedInGame, const TSet<FName>& SoftPackagesUsedInGame,
 	const ITargetPlatform* TargetPlatform, TArray<FAssetData>* OutAssetDatas)
 {
+	if (TargetPlatform)
+	{
+		FArchiveCookContext CookContext(const_cast<UPackage*>(Package), UE::Cook::ECookType::Unknown,
+			UE::Cook::ECookingDLC::Unknown, TargetPlatform);
+		WritePackageData(ParentRecord, &CookContext, Package, Linker, ImportsUsedInGame, SoftPackagesUsedInGame,
+			OutAssetDatas, true /* bProceduralSave */);
+	}
+	else
+	{
+		WritePackageData(ParentRecord, nullptr, Package, Linker, ImportsUsedInGame, SoftPackagesUsedInGame,
+			OutAssetDatas, false/* bProceduralSave */);
+	}
+}
+
+// See the corresponding ReadPackageDataMain and ReadPackageDataDependencies defined in PackageReader.cpp in AssetRegistry module
+void WritePackageData(FStructuredArchiveRecord& ParentRecord, FArchiveCookContext* CookContext, const UPackage* Package,
+	FLinkerSave* Linker, const TSet<TObjectPtr<UObject>>& ImportsUsedInGame, const TSet<FName>& SoftPackagesUsedInGame,
+	TArray<FAssetData>* OutAssetDatas, bool bProceduralSave)
+{
+	bProceduralSave = bProceduralSave | (CookContext != nullptr);
+	IAssetRegistryInterface* AssetRegistry = IAssetRegistryInterface::GetPtr();
+
 	// To avoid large patch sizes, we have frozen cooked package format at the format before VER_UE4_ASSETREGISTRY_DEPENDENCYFLAGS
-	bool bPreDependencyFormat = bIsCooking;
+	// Non-cooked saves do a full update. Orthogonally, they also store the assets in the package in addition to the output variable.
+	// Cooked saves do an additive update and do not store the assets in the package.
+	bool bPreDependencyFormat = false;
+	bool bWriteAssetsToPackage = true;
+	// Editor saves do a full update, but procedural saves (including cook) do not
+	bool bFullUpdate = !bProceduralSave;
+	FCookTagList* CookTagList = nullptr;
+	if (CookContext)
+	{
+		bPreDependencyFormat = true;
+		bWriteAssetsToPackage = false;
+		CookTagList = CookContext->GetCookTagList();
+	}	
 
 	// WritePackageData is currently only called if not bTextFormat; we rely on that to save offsets
 	FArchive& BinaryArchive = ParentRecord.GetUnderlyingArchive();
@@ -2177,26 +2213,32 @@ void WritePackageData(FStructuredArchiveRecord& ParentRecord, bool bIsCooking, c
 		check(BinaryArchive.Tell() == OffsetToAssetRegistryDependencyDataOffset + sizeof(AssetRegistryDependencyDataOffset));
 	}
 
-	// Collect the tag map
 	TArray<UObject*> AssetObjects;
-	if (!(Linker->Summary.GetPackageFlags() & PKG_FilterEditorOnly))
+	for (int32 i = 0; i < Linker->ExportMap.Num(); i++)
 	{
-		// Find any exports which are not in the tag map
-		for (int32 i = 0; i < Linker->ExportMap.Num(); i++)
+		FObjectExport& Export = Linker->ExportMap[i];
+		if (Export.Object && Export.Object->IsAsset())
 		{
-			FObjectExport& Export = Linker->ExportMap[i];
-			if (Export.Object && Export.Object->IsAsset())
+#if WITH_EDITOR
+			if (CookContext)
 			{
-				AssetObjects.Add(Export.Object);
+				TArray<UObject*> AdditionalObjects;
+				Export.Object->GetAdditionalAssetDataObjectsForCook(*CookContext, AdditionalObjects);
+				for (UObject* Object : AdditionalObjects)
+				{
+					if (Object->IsAsset())
+					{
+						AssetObjects.Add(Object);
+					}
+				}
 			}
+#endif
+			AssetObjects.Add(Export.Object);
 		}
 	}
-	int32 ObjectCount = AssetObjects.Num();
-	FStructuredArchive::FArray AssetArray = AssetRegistryRecord.EnterArray(TEXT("TagMap"), ObjectCount);
-	if (OutAssetDatas)
-	{
-		OutAssetDatas->Reset();
-	}
+
+	int32 ObjectCountInPackage = bWriteAssetsToPackage ? AssetObjects.Num() : 0;
+	FStructuredArchive::FArray AssetArray = AssetRegistryRecord.EnterArray(TEXT("TagMap"), ObjectCountInPackage);
 	FString PackageName = Package->GetName();
 
 	for (int32 ObjectIdx = 0; ObjectIdx < AssetObjects.Num(); ++ObjectIdx)
@@ -2207,46 +2249,75 @@ void WritePackageData(FStructuredArchiveRecord& ParentRecord, bool bIsCooking, c
 		FString ObjectPath = Object->GetPathName(Package);
 		FString ObjectClassName = Object->GetClass()->GetPathName();
 
-		TArray<UObject::FAssetRegistryTag> SourceTags;
-		Object->GetAssetRegistryTags(SourceTags);
-#if WITH_EDITOR
-		Object->GetExtendedAssetRegistryTagsForSave(TargetPlatform, SourceTags);
-#endif // WITH_EDITOR
+		FAssetRegistryTagsContextData TagsContextData(Object, EAssetRegistryTagsCaller::SavePackage);
+		TagsContextData.bProceduralSave = bProceduralSave;
+		TagsContextData.TargetPlatform = nullptr;
+		if (CookContext)
+		{
+			TagsContextData.TargetPlatform = CookContext->GetTargetPlatform();
+			TagsContextData.CookType = CookContext->GetCookType();
+			TagsContextData.CookingDLC = CookContext->GetCookingDLC();
+			TagsContextData.bWantsCookTags = TagsContextData.CookType == UE::Cook::ECookType::ByTheBook;
+		}
+		TagsContextData.bFullUpdateRequested = bFullUpdate;
+		FAssetRegistryTagsContext TagsContext(TagsContextData);
 
-		TArray<UObject::FAssetRegistryTag> Tags;
-		for (UObject::FAssetRegistryTag& SourceTag : SourceTags)
+		if (AssetRegistry && !TagsContextData.bFullUpdateRequested)
 		{
-			UObject::FAssetRegistryTag* Existing = Tags.FindByPredicate([SourceTag](const UObject::FAssetRegistryTag& InTag) { return InTag.Name == SourceTag.Name; });
-			if (Existing)
+			FAssetData ExistingAssetData;
+			if (AssetRegistry->TryGetAssetByObjectPath(FSoftObjectPath(Object), ExistingAssetData)
+				== UE::AssetRegistry::EExists::Exists)
 			{
-				Existing->Value = SourceTag.Value;
+				TagsContextData.Tags.Reserve(ExistingAssetData.TagsAndValues.Num());
+				ExistingAssetData.TagsAndValues.ForEach(
+					[&TagsContextData](const TPair<FName, FAssetTagValueRef>& Pair)
+					{
+						TagsContextData.Tags.Add(Pair.Key, UObject::FAssetRegistryTag(Pair.Key, Pair.Value.GetStorageString(),
+							UObject::FAssetRegistryTag::TT_Alphabetical));
+					});
 			}
-			else
+		}
+		if (CookTagList)
+		{
+			TArray<FCookTagList::FTagNameValuePair>* CookTags = CookTagList->ObjectToTags.Find(Object);
+			if (CookTags)
 			{
-				Tags.Add(SourceTag);
+				for (FCookTagList::FTagNameValuePair& Pair : *CookTags)
+				{
+					TagsContext.AddCookTag(UObject::FAssetRegistryTag(Pair.Key, Pair.Value,
+						UObject::FAssetRegistryTag::TT_Alphabetical));
+				}
 			}
 		}
+		
+		Object->GetAssetRegistryTags(TagsContext);
 
-		int32 TagCount = Tags.Num();
+		int32 TagCount = TagsContextData.Tags.Num();
+		TagsContextData.Tags.KeySort(FNameLexicalLess());
 
-		FStructuredArchive::FRecord AssetRecord = AssetArray.EnterElement().EnterRecord();
-		AssetRecord << SA_VALUE(TEXT("Path"), ObjectPath) << SA_VALUE(TEXT("Class"), ObjectClassName);
+		if (bWriteAssetsToPackage)
+		{
+			FStructuredArchive::FRecord AssetRecord = AssetArray.EnterElement().EnterRecord();
+			AssetRecord << SA_VALUE(TEXT("Path"), ObjectPath) << SA_VALUE(TEXT("Class"), ObjectClassName);
 
-		FStructuredArchive::FMap TagMap = AssetRecord.EnterField(TEXT("Tags")).EnterMap(TagCount);
+			FStructuredArchive::FMap TagMap = AssetRecord.EnterField(TEXT("Tags")).EnterMap(TagCount);
 
-		for (TArray<UObject::FAssetRegistryTag>::TConstIterator TagIter(Tags); TagIter; ++TagIter)
-		{
-			FString Key = TagIter->Name.ToString();
-			FString Value = TagIter->Value;
+			for (const TPair<FName, UObject::FAssetRegistryTag>& TagPair : TagsContextData.Tags)
+			{
+				const UObject::FAssetRegistryTag& Tag = TagPair.Value;
+				FString Key = Tag.Name.ToString();
+				FString Value = Tag.Value;
 
-			TagMap.EnterElement(Key) << Value;
+				TagMap.EnterElement(Key) << Value;
+			}
 		}
 
 		if (OutAssetDatas)
 		{
 			FAssetDataTagMap TagsAndValues;
-			for (UObject::FAssetRegistryTag& Tag : Tags)
+			for (TPair<FName, UObject::FAssetRegistryTag>& TagPair : TagsContextData.Tags)
 			{
+				UObject::FAssetRegistryTag& Tag = TagPair.Value;
 				if (!Tag.Name.IsNone() && !Tag.Value.IsEmpty())
 				{
 					TagsAndValues.Add(Tag.Name, MoveTemp(Tag.Value));
diff --git a/Engine/Source/Runtime/CoreUObject/Private/UObject/SavePackage2.cpp b/Engine/Source/Runtime/CoreUObject/Private/UObject/SavePackage2.cpp
--- a/Engine/Source/Runtime/CoreUObject/Private/UObject/SavePackage2.cpp
+++ b/Engine/Source/Runtime/CoreUObject/Private/UObject/SavePackage2.cpp
@@ -1895,9 +1895,11 @@ ESavePackageResult WritePackageHeader(FStructuredArchive::FRecord& StructuredArc
 	{
 		// Save asset registry data so the editor can search for information about assets in this package
 		SCOPED_SAVETIMER(UPackage_Save_SaveAssetRegistryData);
-		UE::AssetRegistry::WritePackageData(StructuredArchiveRoot, SaveContext.IsCooking(), SaveContext.GetPackage(),
+		FArchiveCookData* CookData = SaveContext.GetCookData();
+		FArchiveCookContext* CookContext = CookData ? &CookData->CookContext : nullptr;
+		UE::AssetRegistry::WritePackageData(StructuredArchiveRoot, CookContext, SaveContext.GetPackage(),
 			Linker, SaveContext.GetImportsUsedInGame(), SaveContext.GetSoftPackagesUsedInGame(),
-			SaveContext.GetTargetPlatform(), &SaveContext.GetSavedAssets());
+			&SaveContext.GetSavedAssets(), SaveContext.IsProceduralSave());
 	}
 	// Save level information used by World browser
 	{
diff --git a/Engine/Source/Runtime/CoreUObject/Public/AssetRegistry/AssetData.h b/Engine/Source/Runtime/CoreUObject/Public/AssetRegistry/AssetData.h
--- a/Engine/Source/Runtime/CoreUObject/Public/AssetRegistry/AssetData.h
+++ b/Engine/Source/Runtime/CoreUObject/Public/AssetRegistry/AssetData.h
@@ -61,6 +61,7 @@ struct FAssetBundleData;
 struct FCustomVersion;
 class FCbWriter;
 class FCbFieldView;
+enum class EAssetRegistryTagsCaller : uint8;
 
 COREUOBJECT_API DECLARE_LOG_CATEGORY_EXTERN(LogAssetData, Log, All);
 COREUOBJECT_API extern const FName GAssetBundleDataName;
@@ -261,6 +262,7 @@ PRAGMA_ENABLE_DEPRECATION_WARNINGS
 
 	/** Constructor taking a UObject. By default trying to create one for a blueprint class will create one for the UBlueprint instead, but this can be overridden */
 	COREUOBJECT_API FAssetData(const UObject* InAsset, FAssetData::ECreationFlags InCreationFlags = ECreationFlags::None);
+	COREUOBJECT_API FAssetData(const UObject* InAsset, FAssetData::ECreationFlags InCreationFlags, EAssetRegistryTagsCaller Caller);
 
 	/** Constructor taking a UObject. By default trying to create one for a blueprint class will create one for the UBlueprint instead, but this can be overridden */
 	inline FAssetData(const UObject* InAsset, bool bAllowBlueprintClass)
diff --git a/Engine/Source/Runtime/CoreUObject/Public/AssetRegistry/CookTagList.h b/Engine/Source/Runtime/CoreUObject/Public/AssetRegistry/CookTagList.h
--- a/Engine/Source/Runtime/CoreUObject/Public/AssetRegistry/CookTagList.h
+++ b/Engine/Source/Runtime/CoreUObject/Public/AssetRegistry/CookTagList.h
@@ -5,12 +5,11 @@
 #include "UObject/SoftObjectPath.h"
 
 /**
-*	Used for passing key-value pairs to the asset registry during cook, and is accessed off of
-*	FArchive::CookContext()
+*	Used for passing key-value pairs to the asset registry during serialization during cook,
+*	and is accessed off of FArchive::CookContext()
 * 
-*	This exists so that the asset registry can have per-platform data, as the general purpose 
-*	GetAssetRegistryTags function can't have a TargetPlatform and more generally can't assume 
-*	that platform data has had a chance to build.
+*	This duplicates functionality in GetAssetRegistryTags, but it exists because some tags are
+*	expensive to calculate and the calculation overlaps with other calculations done during serialize.
 *
 *	All tags added to this list have "Cook_" prepended to their names prior to being
 *	added to the asset registry.
@@ -33,7 +32,7 @@ struct FCookTagList
 {
 	UPackage* Package;
 	using FTagNameValuePair = TPair<FName, FString>;
-	TMap<FSoftObjectPath, TArray<FTagNameValuePair>> ObjectToTags;
+	TMap<UObject*, TArray<FTagNameValuePair>> ObjectToTags;
 
 	FCookTagList(UPackage* InPackage) : Package(InPackage) {}
 
@@ -54,7 +53,7 @@ struct FCookTagList
 
 		if (InObject->IsAsset())
 		{
-			TArray<FTagNameValuePair>& ObjectTags = ObjectToTags.FindOrAdd(FSoftObjectPath(InObject));
+			TArray<FTagNameValuePair>& ObjectTags = ObjectToTags.FindOrAdd(InObject);
 			ObjectTags.Add(FTagNameValuePair(InTagName, InTagValue));
 		}
 	}
diff --git a/Engine/Source/Runtime/CoreUObject/Public/Misc/AssetRegistryInterface.h b/Engine/Source/Runtime/CoreUObject/Public/Misc/AssetRegistryInterface.h
--- a/Engine/Source/Runtime/CoreUObject/Public/Misc/AssetRegistryInterface.h
+++ b/Engine/Source/Runtime/CoreUObject/Public/Misc/AssetRegistryInterface.h
@@ -21,6 +21,7 @@ class FStructuredArchiveRecord;
 class ITargetPlatform;
 class UObject;
 class UPackage;
+struct FArchiveCookContext;
 struct FAssetData;
 
 namespace EAssetRegistryDependencyType
@@ -191,10 +192,13 @@ namespace UE::AssetRegistry
 	};
 
 	// Functions to read and write the data used by the AssetRegistry in each package; the format of this data is separate from the format of the data in the asset registry
-	COREUOBJECT_API void WritePackageData(FStructuredArchiveRecord& ParentRecord, bool bIsCooking,
+	COREUOBJECT_API void WritePackageData(FStructuredArchiveRecord& ParentRecord, FArchiveCookContext* CookContext,
 		const UPackage* Package, FLinkerSave* Linker, const TSet<TObjectPtr<UObject>>& ImportsUsedInGame,
-		const TSet<FName>& SoftPackagesUsedInGame, const ITargetPlatform* TargetPlatform,
-		TArray<FAssetData>* OutAssetDatas);
+		const TSet<FName>& SoftPackagesUsedInGame, TArray<FAssetData>* OutAssetDatas, bool bProceduralSave);
+	UE_DEPRECATED(5.4, "Use version that takes FArchiveCookContext");
+	COREUOBJECT_API void WritePackageData(FStructuredArchiveRecord& ParentRecord, bool bIsCooking, const UPackage* Package,
+		FLinkerSave* Linker, const TSet<TObjectPtr<UObject>>& ImportsUsedInGame, const TSet<FName>& SoftPackagesUsedInGame,
+		const ITargetPlatform* TargetPlatform, TArray<FAssetData>* OutAssetDatas);
 	// ReadPackageDataMain and ReadPackageDataDependencies are declared in IAssetRegistry.h, in the AssetRegistry module, because they depend upon some structures defined in the AssetRegistry module
 
 	namespace Private
diff --git a/Engine/Source/Runtime/CoreUObject/Public/UObject/ArchiveCookContext.h b/Engine/Source/Runtime/CoreUObject/Public/UObject/ArchiveCookContext.h
--- a/Engine/Source/Runtime/CoreUObject/Public/UObject/ArchiveCookContext.h
+++ b/Engine/Source/Runtime/CoreUObject/Public/UObject/ArchiveCookContext.h
@@ -32,17 +32,26 @@ private:
 
 	/** CookTagList is only valid for cook by the book; it is not publically accessible otherwise. */
 	FCookTagList CookTagList;
+	const ITargetPlatform* TargetPlatform = nullptr;
 	bool bCookTagListEnabled = false;
 	UE::Cook::ECookType CookType = UE::Cook::ECookType::Unknown;
 	UE::Cook::ECookingDLC CookingDLC = UE::Cook::ECookingDLC::Unknown;
 
 public:
 
-	FArchiveCookContext(UPackage* InPackage, UE::Cook::ECookType InCookType, UE::Cook::ECookingDLC CookingDLC)
+	UE_DEPRECATED(5.4, "Call version that takes the TargetPlatform")
+	FArchiveCookContext(UPackage* InPackage, UE::Cook::ECookType InCookType, UE::Cook::ECookingDLC InCookingDLC)
+		: FArchiveCookContext(InPackage, InCookType, InCookingDLC, nullptr)
+	{
+	}
+
+	FArchiveCookContext(UPackage* InPackage, UE::Cook::ECookType InCookType, UE::Cook::ECookingDLC InCookingDLC,
+		const ITargetPlatform* InTargetPlatform)
 		: CookTagList(InPackage)
+		, TargetPlatform(InTargetPlatform)
 		, bCookTagListEnabled(InPackage && InCookType == UE::Cook::ECookType::ByTheBook)
 		, CookType(InCookType)
-		, CookingDLC(CookingDLC)
+		, CookingDLC(InCookingDLC)
 	{
 	}
 
@@ -52,6 +61,7 @@ public:
 	}
 
 	FCookTagList* GetCookTagList() { return bCookTagListEnabled ? &CookTagList : nullptr; }
+	const ITargetPlatform* GetTargetPlatform() const { return TargetPlatform; }
 
 	bool IsCookByTheBook() const { return CookType == UE::Cook::ECookType::ByTheBook; }
 	bool IsCookOnTheFly() const { return CookType == UE::Cook::ECookType::OnTheFly; }
diff --git a/Engine/Source/Runtime/CoreUObject/Public/UObject/AssetRegistryTagsContext.h b/Engine/Source/Runtime/CoreUObject/Public/UObject/AssetRegistryTagsContext.h
new file mode 100644
--- /dev/null
+++ b/Engine/Source/Runtime/CoreUObject/Public/UObject/AssetRegistryTagsContext.h
@@ -0,0 +1,193 @@
+// Copyright Epic Games, Inc. All Rights Reserved.
+
+#pragma once
+
+#include "Containers/Array.h"
+#include "Containers/Map.h"
+#include "Containers/StringView.h"
+#include "Misc/Optional.h"
+#include "Templates/UnrealTemplate.h"
+#include "UObject/CookEnums.h"
+#include "UObject/Object.h" // UObject::FAssetRegistryTag
+
+class ITargetPlatform;
+class UObject;
+struct FAssetBundleData;
+
+namespace UE::AssetRegistry
+{
+
+/**
+ * We have to store all tags as Name/String pairs, without any extra data for e.g. bIsCookTag,
+ * so we use a naming convention on the Name. The FAssetRegistryTagsContextData API prepends CookTagPrefix
+ * to tags specified as CookTags, and gives a warning and ignores any regular tags starting with CookTagPrefix.
+ */
+constexpr FStringView CookTagPrefix = TEXTVIEW("Cook_");
+
+}
+
+/** Marked up at the callsite with the category of callsite. */
+enum class EAssetRegistryTagsCaller : uint8
+{
+	/** The AssetRegistry is calling FAssetData(UObject*) to return results for a query from an in-memory UObject. */
+	AssetRegistryQuery,
+	/**
+	 * The AssetRegistry is calling FAssetData(UObject*) on a UObject that has just loaded and will write them into
+	 * the AR's registered tags for that object.
+	 */
+	AssetRegistryLoad,
+	/** Caller has requested a full update. */
+	FullUpdate,
+	/**
+	 * SavePackage is calling GetAssetRegistryTags and will store the result in the package header.
+	 * See other functions on the context to determine what kind of save it is.
+	 */
+	SavePackage,
+	/**
+	 * System-specific code outside of SavePackage or the AssetRegistry is constructing an FAssetData, and
+	 * wants to skip any expensive tags.
+	 */
+	Fast,
+	/** System-specific code outside of SavePackage or the AssetRegistry is constructing an FAssetData. */
+	Uncategorized,
+};
+
+/**
+ * Data used to provide information about the calling context and receive results from GetAssetRegistryTags.
+ * The fields provided by this class are hidden from implementers of GetAssetRegistryTags so that we can change the
+ * implementation of this class without needing to add deprecation handling in the GetAssetRegistryTags functions.
+ */
+struct FAssetRegistryTagsContextData
+{
+public:
+	COREUOBJECT_API FAssetRegistryTagsContextData(const UObject* CallTarget,
+		EAssetRegistryTagsCaller InCaller = EAssetRegistryTagsCaller::Uncategorized);
+
+	TMap<FName, UObject::FAssetRegistryTag> Tags;
+	const UObject* Object = nullptr;
+	const ITargetPlatform* TargetPlatform = nullptr;
+	const FAssetBundleData* BundleResult = nullptr;
+	EAssetRegistryTagsCaller Caller = EAssetRegistryTagsCaller::Uncategorized;
+	UE::Cook::ECookType CookType = UE::Cook::ECookType::Unknown;
+	UE::Cook::ECookingDLC CookingDLC = UE::Cook::ECookingDLC::Unknown;
+	bool bProceduralSave = false;
+	bool bWantsBundleResult = false;
+	bool bWantsCookTags = false;
+	bool bFullUpdateRequested = false;
+};
+
+/**
+ * Interface used by GetAssetRegistryTags to access the calling context data.
+ * 
+ * API Notes:
+ * CookTag functions on this class add CookTags - tags that are added only when cooking and that are added to the
+ * development AssetRegistry used for patches patches, DLC, and iterative cooking, but are removed from the runtime
+ * AssetRegistry passed into the game. CookTags can only be added when WantsCookTags() is true, calls in other Context
+ * states are ignored. All tags added with the CookTags interface have "Cook_" prepended to their names prior to being
+ * saved into the development to the asset registry.
+ *
+ * CookType: When GetAssetRegistryTags is called for package saves during cook, the cook type is set depending on
+ * the type of cook. Cooks in Unreal can be either CookByTheBook or CookOnTheFly. If GetAssetRegistryTags was called
+ * for other reasons, the CookType is reported as unknown. CookByTheBook means that the project's packaging settings
+ * are read by the cooker and a list of all referenced packages are found and those packages are saved and packaged.
+ * CookOnTheFly means that the game is running without packaged data, and is connecting to the cooker process to 
+ * request the cooked version of packages that it demands. Only the packages requested by the runtime session are cooked.
+ */
+class FAssetRegistryTagsContext
+{
+public:
+	// This is an implicit constructor to reduce the boilerplate for calling GetAssetRegistryTags
+	COREUOBJECT_API FAssetRegistryTagsContext(FAssetRegistryTagsContextData& InData);
+	COREUOBJECT_API FAssetRegistryTagsContext(FAssetRegistryTagsContext& Other);
+
+	COREUOBJECT_API EAssetRegistryTagsCaller GetCaller() const;
+	COREUOBJECT_API const UObject* GetObject() const;
+
+	/** Whether all tags should be calculated, as set by the creator of the context (e.g. SavePackage). */
+	COREUOBJECT_API bool IsFullUpdate() const;
+	/** Whether the call is coming from SavePackage. See functions below for determining what kind of save. */
+	COREUOBJECT_API bool IsSaving() const;
+	/**
+	 * Return whether the package is being saved due to a procedural save. False if not saving.
+	 * Any save without the possibility of user-generated edits to the package is a procedural save (Cooking, EditorDomain).
+	 * This allows us to add or update tags that only need to be modified in response to new user data.
+	 */
+	COREUOBJECT_API bool IsProceduralSave() const;
+	/** Report whether this is a save into a target-specific cooked format. False if not saving. */
+	COREUOBJECT_API bool IsCooking() const;
+	/** Return the targetplatform of the save, if cooking. Null if not saving or not cooking. */
+	COREUOBJECT_API const ITargetPlatform* GetTargetPlatform() const;
+
+	/**
+	 * If IsCooking is true, reports whether the cook is CookByTheBook. @see FAssetRegistryTagsContext.
+	 * If IsCooking is false, returns false.
+	 */
+	COREUOBJECT_API bool IsCookByTheBook() const;
+	/**
+	 * If IsCooking is true, reports whether the cook is CookByTheBook. @see FAssetRegistryTagsContext.
+	 * If IsCooking is false, returns false.
+	 */
+	COREUOBJECT_API bool IsCookOnTheFly() const;
+	/**
+	 * If IsCooking is true, reports whether the cook type is unknown and could be any type of cook.
+	 * If IsCooking is false, returns true.
+	 */
+	COREUOBJECT_API bool IsCookTypeUnknown() const;
+	/**
+	 * If IsCooking is true, returns the cook type, @see FAssetRegistryTagsContext.
+	 * If IsCooking is false, returns ECookType::Unknown.
+	 */
+	COREUOBJECT_API UE::Cook::ECookType GetCookType() const;
+	/**
+	 * If IsCooking is true, returns whether whether the cook is for a DLC rather than base game.
+	 * If IsCooking is false, returns ECookingDLC::Unknown.
+	 */
+	COREUOBJECT_API UE::Cook::ECookingDLC GetCookingDLC() const;
+
+	/**
+	 * Output function: Add a tag if it does not already exist. Return reference to the existing tag and optionally
+	 * report whether it already existed. Issues an error and returns referenced to unused global if TagName is invalid.
+	 */
+	COREUOBJECT_API UObject::FAssetRegistryTag& FindOrAddTag(FName TagName, bool* bOutAlreadyExists=nullptr);
+	/** Output function: Find a tag in the output if it already exists. Issues an error if TagName is invalid. */
+	COREUOBJECT_API UObject::FAssetRegistryTag* FindTag(FName TagName);
+	/** Output function: Report whether a tag is already in the output. Issues an error if TagName is invalid. */
+	COREUOBJECT_API bool ContainsTag(FName TagName) const;
+	/** Output function: move the given tag into the results, overwriting previous result if it exists. */
+	COREUOBJECT_API void AddTag(UObject::FAssetRegistryTag TagResult);
+
+	/**
+	 * True if the caller wants the FAssetBundleData structure (if present) to be stored as a pointer.
+	 * If false, the FAssetBundleData structure should instead be converted to text and stored as a tag.
+	 */
+	COREUOBJECT_API bool WantsBundleResult() const;
+	COREUOBJECT_API const FAssetBundleData* GetBundleResult();
+	COREUOBJECT_API void SetBundleResult(const FAssetBundleData* InBundleResult);
+
+	/**
+	 * Return whether the context is collecting cook tags. CookTag functions are noops if this returns
+	 * false. CookTags are only collected for CookByTheBook.
+	 */
+	COREUOBJECT_API bool WantsCookTags() const;
+
+	/** FindOrAddCookTag but for a cooktag. @see FindOrAddTag and @see FAssetRegistryTagsContext */
+	COREUOBJECT_API UObject::FAssetRegistryTag& FindOrAddCookTag(FName TagName, bool* bOutAlreadyExists);
+	/** FindTag but for a cooktag. @see FindTag and @see FAssetRegistryTagsContext */
+	COREUOBJECT_API UObject::FAssetRegistryTag* FindCookTag(FName TagName);
+	/** ContainsTag but for a cooktag. @see ContainsTag and @see FAssetRegistryTagsContext */
+	COREUOBJECT_API bool ContainsCookTag(FName TagName) const;
+	/** AddTag but for a cooktag. @see AddTag and @see FAssetRegistryTagsContext */
+	COREUOBJECT_API void AddCookTag(UObject::FAssetRegistryTag TagResult);
+
+private:
+	bool TryValidateTagName(FName TagName) const;
+	FName TransformTagName(FName TagName, FStringView PrefixToAdd) const;
+	UObject::FAssetRegistryTag& FindOrAddTagInternal(FName TagName, FStringView PrefixToAdd, bool* bOutAlreadyExists);
+	UObject::FAssetRegistryTag* FindTagInternal(FName TagName, FStringView PrefixToAdd);
+	bool ContainsTagInternal(FName TagName, FStringView PrefixToAdd) const;
+	void AddTagInternal(UObject::FAssetRegistryTag&& TagResult, FStringView PrefixToAdd);
+
+private:
+	FAssetRegistryTagsContextData& Data;
+	friend class UObject;
+};
\ No newline at end of file
diff --git a/Engine/Source/Runtime/CoreUObject/Public/UObject/Object.h b/Engine/Source/Runtime/CoreUObject/Public/UObject/Object.h
--- a/Engine/Source/Runtime/CoreUObject/Public/UObject/Object.h
+++ b/Engine/Source/Runtime/CoreUObject/Public/UObject/Object.h
@@ -15,6 +15,7 @@
 #include "Containers/VersePathFwd.h"
 
 struct FAssetData;
+class FAssetRegistryTagsContext;
 class FConfigCacheIni;
 class FCustomPropertyConditionState;
 class FEditPropertyChain;
@@ -25,6 +26,7 @@ class FObjectPreSaveRootContext;
 class ITargetPlatform;
 class ITransactionObjectAnnotation;
 class FTransactionObjectEvent;
+struct FArchiveCookContext;
 struct FAppendToClassSchemaContext;
 struct FFrame;
 struct FObjectInstancingGraph;
@@ -859,11 +861,12 @@ public:
 			: Name(InName), Value(MoveTemp(InValue)), Type(InType), DisplayFlags(InDisplayFlags) {}
 
 #if WITH_EDITOR
-		/** Callback  */
+		/** Event for listeners who want to add tags to some UObjects' GetAssetRegistryTags. */
+		DECLARE_MULTICAST_DELEGATE_OneParam(FOnGetObjectAssetRegistryTagsWithContext, FAssetRegistryTagsContext&);
+		COREUOBJECT_API static FOnGetObjectAssetRegistryTagsWithContext OnGetExtraObjectTagsWithContext;
+
 		DECLARE_MULTICAST_DELEGATE_TwoParams(FOnGetObjectAssetRegistryTags, const UObject* /*Object*/, TArray<FAssetRegistryTag>& /*InOutTags*/);
 		COREUOBJECT_API static FOnGetObjectAssetRegistryTags OnGetExtraObjectTags;
-
-		/** Callback for GetExtendedAssetRegistryTagsForSave */
 		DECLARE_MULTICAST_DELEGATE_ThreeParams(FOnGetExtendedAssetRegistryTagsForSave, const UObject* /*Object*/, const ITargetPlatform* TargetPlatform, TArray<FAssetRegistryTag>& /*InOutTags*/);
 		COREUOBJECT_API static FOnGetExtendedAssetRegistryTagsForSave OnGetExtendedAssetRegistryTagsForSave;
 #endif // WITH_EDITOR
@@ -876,8 +879,15 @@ public:
 	 * @param	OutTags		A list of key-value pairs associated with this object and their types
 	 */
 	COREUOBJECT_API virtual void GetAssetRegistryTags(TArray<FAssetRegistryTag>& OutTags) const;
+	COREUOBJECT_API virtual void GetAssetRegistryTags(FAssetRegistryTagsContext Context) const;
 
 #if WITH_EDITOR
+	/**
+	 * Get an additional object that should be added to the AssetDatas stored for a cooked package, to handle
+	 * objects stripped from cooked packages. e.g. UBlueprintGeneratedClass returns its UBlueprint.
+	 */
+	COREUOBJECT_API virtual void GetAdditionalAssetDataObjectsForCook(FArchiveCookContext& CookContext,
+		TArray<UObject*>& OutObjects) const;
 	/**
 	 * Temporary interim solution to gather asset registry data at save time only.  Can depend on the target platform being saved for.
 	 *
@@ -887,8 +897,12 @@ public:
 	COREUOBJECT_API virtual void GetExtendedAssetRegistryTagsForSave(const ITargetPlatform* TargetPlatform, TArray<FAssetRegistryTag>& OutTags) const;
 #endif // WITH_EDITOR
 
-	/** Gathers a list of asset registry tags for an FAssetData  */
 	COREUOBJECT_API void GetAssetRegistryTags(FAssetData& Out) const;
+	/**
+	 * Gathers a list of asset registry tags for an FAssetData. Output data will be removed from the Context and
+	 * moved onto the Out FAssetData.
+	 */
+	COREUOBJECT_API void GetAssetRegistryTags(FAssetRegistryTagsContext Context, FAssetData& Out) const;
 
 	/** Get the common tag name used for all asset source file import paths */
 	static COREUOBJECT_API const FName& SourceFileTagName();
diff --git a/Engine/Source/Runtime/Engine/Classes/Engine/BlueprintGeneratedClass.h b/Engine/Source/Runtime/Engine/Classes/Engine/BlueprintGeneratedClass.h
--- a/Engine/Source/Runtime/Engine/Classes/Engine/BlueprintGeneratedClass.h
+++ b/Engine/Source/Runtime/Engine/Classes/Engine/BlueprintGeneratedClass.h
@@ -774,6 +774,8 @@ public:
 	ENGINE_API virtual void GetPreloadDependencies(TArray<UObject*>& OutDeps) override;
 	ENGINE_API virtual void GetAssetRegistryTags(TArray<FAssetRegistryTag>& OutTags) const override;
 #if WITH_EDITOR
+	ENGINE_API virtual void GetAdditionalAssetDataObjectsForCook(FArchiveCookContext& CookContext, 
+		TArray<UObject*>& OutObjects) const override;
 	ENGINE_API virtual void PostLoadAssetRegistryTags(const FAssetData& InAssetData, TArray<FAssetRegistryTag>& OutTagsAndValuesToUpdate) const;
 #endif //~ WITH_EDITOR
 	ENGINE_API virtual FPrimaryAssetId GetPrimaryAssetId() const override;
diff --git a/Engine/Source/Runtime/Engine/Private/BlueprintGeneratedClass.cpp b/Engine/Source/Runtime/Engine/Private/BlueprintGeneratedClass.cpp
--- a/Engine/Source/Runtime/Engine/Private/BlueprintGeneratedClass.cpp
+++ b/Engine/Source/Runtime/Engine/Private/BlueprintGeneratedClass.cpp
@@ -602,6 +602,17 @@ void UBlueprintGeneratedClass::GetAssetRegistryTags(TArray<FAssetRegistryTag>& O
 }
 
 #if WITH_EDITOR
+void UBlueprintGeneratedClass::GetAdditionalAssetDataObjectsForCook(FArchiveCookContext& CookContext,
+	TArray<UObject*>& OutObjects) const
+{
+	Super::GetAdditionalAssetDataObjectsForCook(CookContext, OutObjects);
+	UBlueprint* Blueprint = Cast<UBlueprint>(ClassGeneratedBy);
+	if (Blueprint && Blueprint->GetPackage() == GetPackage())
+	{
+		OutObjects.Add(Blueprint);
+	}
+}
+
 void UBlueprintGeneratedClass::PostLoadAssetRegistryTags(const FAssetData& InAssetData, TArray<FAssetRegistryTag>& OutTagsAndValuesToUpdate) const
 {
 	Super::PostLoadAssetRegistryTags(InAssetData, OutTagsAndValuesToUpdate);
diff --git a/Engine/Source/Runtime/Engine/Private/ShaderCompiler/ShaderCompiler.cpp b/Engine/Source/Runtime/Engine/Private/ShaderCompiler/ShaderCompiler.cpp
--- a/Engine/Source/Runtime/Engine/Private/ShaderCompiler/ShaderCompiler.cpp
+++ b/Engine/Source/Runtime/Engine/Private/ShaderCompiler/ShaderCompiler.cpp
@@ -9394,7 +9394,7 @@ FString SaveGlobalShaderFile(EShaderPlatform Platform, FString SavePath, class I
 		if (TargetPlatform != nullptr)
 		{
 			CookContext.Emplace(nullptr /*InPackage*/, UE::Cook::ECookType::Unknown,
-				UE::Cook::ECookingDLC::Unknown);
+				UE::Cook::ECookingDLC::Unknown, TargetPlatform);
 			CookData.Emplace(*TargetPlatform, *CookContext);
 			MemoryWriter.SetCookData(CookData.GetPtrOrNull());
 		}
@@ -10020,7 +10020,7 @@ void CompileGlobalShaderMapForRemote(
 	if (TargetPlatform != nullptr)
 	{
 		CookContext.Emplace(nullptr /*InPackage*/, UE::Cook::ECookType::Unknown,
-			UE::Cook::ECookingDLC::Unknown);
+			UE::Cook::ECookingDLC::Unknown, TargetPlatform);
 		CookData.Emplace(*TargetPlatform, *CookContext);
 		Ar.SetCookData(CookData.GetPtrOrNull());
 	}
@@ -10040,7 +10040,7 @@ void SaveShaderMapsForRemote(ITargetPlatform* TargetPlatform, const TMap<FString
 	if (TargetPlatform != nullptr)
 	{
 		CookContext.Emplace(nullptr /*InPackage*/, UE::Cook::ECookType::Unknown,
-			UE::Cook::ECookingDLC::Unknown);
+			UE::Cook::ECookingDLC::Unknown, TargetPlatform);
 		CookData.Emplace(*TargetPlatform, *CookContext);
 		Ar.SetCookData(CookData.GetPtrOrNull());
 	}
diff --git a/Engine/Source/Runtime/Engine/Private/WorldPartition/WorldPartitionRuntimeLevelStreamingCell.cpp b/Engine/Source/Runtime/Engine/Private/WorldPartition/WorldPartitionRuntimeLevelStreamingCell.cpp
--- a/Engine/Source/Runtime/Engine/Private/WorldPartition/WorldPartitionRuntimeLevelStreamingCell.cpp
+++ b/Engine/Source/Runtime/Engine/Private/WorldPartition/WorldPartitionRuntimeLevelStreamingCell.cpp
@@ -10,6 +10,7 @@
 #include "Misc/HierarchicalLogArchive.h"
 #include "Misc/Paths.h"
 #include "AssetRegistry/IAssetRegistry.h"
+#include "UObject/AssetRegistryTagsContext.h"
 
 #include UE_INLINE_GENERATED_CPP_BY_NAME(WorldPartitionRuntimeLevelStreamingCell)
 
@@ -330,7 +331,7 @@ bool UWorldPartitionRuntimeLevelStreamingCell::PopulateGeneratorPackageForCook(T
 
 		// Make sure Asset Registry tags are updated here synchronously now that Package contains all its actors
 		// ex: AFunctionalTest actors need to be part of the Worlds asset tags once they are not longer external so they can be discovered at runtime
-		IAssetRegistry::Get()->AssetFullyUpdateTags(OuterWorld);
+		IAssetRegistry::Get()->AssetUpdateTags(OuterWorld, EAssetRegistryTagsCaller::Fast);
 
 		// Empty cell's package list (this ensures that no one can rely on cell's content).
 		Packages.Empty();
