Added partial support for loading unknown arrays, structs, and enums into FPropertyBag

diff --git a/Engine/Source/Runtime/CoreUObject/Private/UObject/EnumProperty.cpp b/Engine/Source/Runtime/CoreUObject/Private/UObject/EnumProperty.cpp
--- a/Engine/Source/Runtime/CoreUObject/Private/UObject/EnumProperty.cpp
+++ b/Engine/Source/Runtime/CoreUObject/Private/UObject/EnumProperty.cpp
@@ -557,6 +557,16 @@ bool FEnumProperty::LoadFromTag(const FPropertyTag& Tag)
 	}
 
 	// Update FByteProperty when making changes here.
+	TStringBuilder<256> EnumName(InPlace, Tag.EnumName);
+	if (UEnum* LocalEnum = FindFirstObject<UEnum>(*EnumName, EFindFirstObjectOptions::NativeFirst))
+	{
+		Enum = LocalEnum;
+		UE_CLOG(Enum->GetMaxEnumValue() >= 256, LogClass, Warning,
+			TEXT("Enum '%s' does not fit in a byte property loading property '%s'."), *Enum->GetName(), *GetName());
+		AddCppProperty(new FByteProperty(this, GetFName(), RF_NoFlags));
+		return true;
+	}
+
 	return false;
 }
 
@@ -564,17 +574,17 @@ void FEnumProperty::SaveToTag(FPropertyTag& Tag)
 {
 	Super::SaveToTag(Tag);
 
-	if (const UEnum* LocalEnum = Enum)
+	const UEnum* LocalEnum = Enum;
+	check(LocalEnum);
+
+	// RobM: Ugly hack so that we can avoid content changes in most of the packages
+	// Update FByteProperty when making changes here.
+	if (LocalEnum->GetPackage()->HasAnyPackageFlags(PKG_CompiledIn))
+	{				
+		Tag.EnumName = LocalEnum->GetFName();
+	}
+	else
 	{
-		// RobM: Ugly hack so that we can avoid content changes in most of the packages
-		// Update FByteProperty when making changes here.
-		if (LocalEnum->GetPackage()->HasAnyPackageFlags(PKG_CompiledIn))
-		{				
-			Tag.EnumName = LocalEnum->GetFName();
-		}
-		else
-		{
-			Tag.EnumName = FName(*LocalEnum->GetPathName());
-		}
+		Tag.EnumName = FName(*LocalEnum->GetPathName());
 	}
 }
diff --git a/Engine/Source/Runtime/CoreUObject/Private/UObject/PropertyArray.cpp b/Engine/Source/Runtime/CoreUObject/Private/UObject/PropertyArray.cpp
--- a/Engine/Source/Runtime/CoreUObject/Private/UObject/PropertyArray.cpp
+++ b/Engine/Source/Runtime/CoreUObject/Private/UObject/PropertyArray.cpp
@@ -612,7 +612,11 @@ void FArrayProperty::SerializeItem(FStructuredArchive::FSlot Slot, void* Value,
 		}
 	}
 
-	FUObjectSerializeContext* Context = UnderlyingArchive.GetSerializeContext();
+	FUObjectSerializeContext* Context = FUObjectThreadContext::Get().GetSerializeContext();
+	if (Context && MaybeInnerTag)
+	{
+		Context->SerializedPropertyPath.PushType(MaybeInnerTag->StructName);
+	}
 
 	// need to know how much data this call to SerializeItem consumes, so mark where we are
 	int64 DataOffset = UnderlyingArchive.Tell();
@@ -1230,6 +1234,25 @@ bool FArrayProperty::LoadFromTag(const FPropertyTag& Tag)
 		return false;
 	}
 
+	FField* Field = FField::Construct(Tag.InnerType, {}, Tag.Name, RF_NoFlags);
+	if (FProperty* Property = CastField<FProperty>(Field))
+	{
+		FPropertyTag InnerTag = Tag;
+		InnerTag.Type = Tag.InnerType;
+		InnerTag.InnerType = {};
+		// Skip property types that are missing the name of the inner type.
+		// Structs have their name in a tag in the serialized data, but we cannot
+		// proceed safely unless we know if the struct used native serialization.
+		if (!Property->IsA<FStructProperty>() &&
+			!Property->IsA<FByteProperty>() &&
+			!Property->IsA<FEnumProperty>() &&
+			Property->LoadFromTag(InnerTag))
+		{
+			Inner = Property;
+			return true;
+		}
+	}
+	delete Field;
 	return false;
 }
 
diff --git a/Engine/Source/Runtime/CoreUObject/Private/UObject/PropertyBag.cpp b/Engine/Source/Runtime/CoreUObject/Private/UObject/PropertyBag.cpp
--- a/Engine/Source/Runtime/CoreUObject/Private/UObject/PropertyBag.cpp
+++ b/Engine/Source/Runtime/CoreUObject/Private/UObject/PropertyBag.cpp
@@ -112,6 +112,7 @@ void FPropertyBag::LoadPropertyByTag(const FPropertyPathName& Path, const FPrope
 	FField* Field = FField::Construct(Tag.Type, {}, Tag.Name, RF_NoFlags);
 	if (FProperty* Property = CastField<FProperty>(Field); Property && Property->LoadFromTag(Value.Tag))
 	{
+		Property->Link(UnderlyingArchive);
 		Value.bOwnsProperty = true;
 		Value.Tag.Prop = Property;
 		Value.AllocateAndInitializeValue();
diff --git a/Engine/Source/Runtime/CoreUObject/Private/UObject/PropertyByte.cpp b/Engine/Source/Runtime/CoreUObject/Private/UObject/PropertyByte.cpp
--- a/Engine/Source/Runtime/CoreUObject/Private/UObject/PropertyByte.cpp
+++ b/Engine/Source/Runtime/CoreUObject/Private/UObject/PropertyByte.cpp
@@ -534,8 +534,20 @@ bool FByteProperty::LoadFromTag(const FPropertyTag& Tag)
 		return false;
 	}
 
+	if (Tag.EnumName.IsNone())
+	{
+		return true;
+	}
+
 	// Update FEnumProperty when making changes here.
-	return Tag.EnumName.IsNone();
+	TStringBuilder<256> EnumName(InPlace, Tag.EnumName);
+	if (UEnum* LocalEnum = FindFirstObject<UEnum>(*EnumName, EFindFirstObjectOptions::NativeFirst))
+	{
+		Enum = LocalEnum;
+		return true;
+	}
+
+	return false;
 }
 
 void FByteProperty::SaveToTag(FPropertyTag& Tag)
diff --git a/Engine/Source/Runtime/CoreUObject/Private/UObject/PropertyStruct.cpp b/Engine/Source/Runtime/CoreUObject/Private/UObject/PropertyStruct.cpp
--- a/Engine/Source/Runtime/CoreUObject/Private/UObject/PropertyStruct.cpp
+++ b/Engine/Source/Runtime/CoreUObject/Private/UObject/PropertyStruct.cpp
@@ -11,6 +11,7 @@
 #include "Serialization/ArchiveUObjectFromStructuredArchive.h"
 #include "Hash/Blake3.h"
 #include "IO/IoHash.h"
+#include "Misc/StringBuilder.h"
 
 static inline void PreloadInnerStructMembers(FStructProperty* StructProperty)
 {
@@ -473,6 +474,16 @@ bool FStructProperty::LoadFromTag(const FPropertyTag& Tag)
 		return false;
 	}
 
+	if (!Tag.StructName.IsNone())
+	{
+		TStringBuilder<256> StructName(InPlace, Tag.StructName);
+		if (UScriptStruct* LocalStruct = FindFirstObject<UScriptStruct>(*StructName, EFindFirstObjectOptions::NativeFirst))
+		{
+			Struct = LocalStruct;
+			return true;
+		}
+	}
+
 	return false;
 }
 
