Fixing possible division by 0 when calculating the number of threads to run GatherUnreachableObjects with on server exit

diff --git a/Engine/Source/Runtime/CoreUObject/Private/UObject/ReachabilityAnalysisState.h b/Engine/Source/Runtime/CoreUObject/Private/UObject/ReachabilityAnalysisState.h
--- a/Engine/Source/Runtime/CoreUObject/Private/UObject/ReachabilityAnalysisState.h
+++ b/Engine/Source/Runtime/CoreUObject/Private/UObject/ReachabilityAnalysisState.h
@@ -142,7 +142,7 @@ inline int32 GetNumThreadsForGather(const EGatherOptions Options, const int32 Nu
 	if (!!(Options & EGatherOptions::Parallel) && NumObjects > 0)
 	{
 		constexpr int32 MinNumObjectsPerThread = 100;
-		const int32 MaxNumThreads = FTaskGraphInterface::Get().GetNumWorkerThreads();
+		const int32 MaxNumThreads = FMath::Max(1, FTaskGraphInterface::Get().GetNumWorkerThreads());
 		const int32 NumThreadsForNumObjects = (NumObjects + MinNumObjectsPerThread - 1) / MinNumObjectsPerThread;
 		return FMath::Min(NumThreadsForNumObjects, MaxNumThreads);
 	}
