[ZenLoader]
 - Add progress callback when IO finishes reading

#tests UnitTests

diff --git a/Engine/Source/Runtime/CoreUObject/Private/Serialization/AsyncLoading2.cpp b/Engine/Source/Runtime/CoreUObject/Private/Serialization/AsyncLoading2.cpp
--- a/Engine/Source/Runtime/CoreUObject/Private/Serialization/AsyncLoading2.cpp
+++ b/Engine/Source/Runtime/CoreUObject/Private/Serialization/AsyncLoading2.cpp
@@ -2999,8 +2999,7 @@ public:
 
 	/** [EDL] End Event driven loader specific stuff */
 
-	void CallCompletionCallbacks(EAsyncLoadingResult::Type LoadingResult);
-
+	void CallProgressCallbacks(EAsyncLoadingProgress ProgressType);
 private:
 	void InitializeExportArchive(FExportArchive& Ar, bool bIsOptionalSegment);
 	void CreatePackageNodes(const FAsyncLoadEventSpec* EventSpecs);
@@ -3221,7 +3220,6 @@ private:
 				}
 				ProgressCallbacks.Empty();
 			}
-
 		}
 	};
 	TArray<FCompletedPackageRequest> CompletedPackageRequests;
@@ -5273,22 +5271,7 @@ void FAsyncPackage2::StartLoading(FAsyncLoadingThreadState2& ThreadState, FIoBat
 
 	LoadStartTime = FPlatformTime::Seconds();
 
-	if (ProgressCallbacks.Num() != 0)
-	{
-		TRACE_CPUPROFILER_EVENT_SCOPE(PackageProgressCallbacks_Started);
-
-		const FLoadPackageAsyncProgressParams Params
-		{
-			.PackageName = Desc.UPackageName,
-			.LoadedPackage = GetLinkerRoot(),
-			.ProgressType = EAsyncLoadingProgress::Started
-		};
-
-		for (TUniquePtr<FLoadPackageAsyncProgressDelegate>& ProgressCallback : ProgressCallbacks)
-		{
-			ProgressCallback->ExecuteIfBound(Params);
-		}
-	}
+	CallProgressCallbacks(EAsyncLoadingProgress::Started);
 
 	AsyncPackageLoadingState = EAsyncPackageLoadingState2::WaitingForIo;
 
@@ -5985,6 +5968,8 @@ EEventLoadNodeExecutionResult FAsyncPackage2::Event_ProcessPackageSummary(FAsync
 	UE::Core::Private::FPlayInEditorLoadingScope PlayInEditorIDScope(Package->Desc.PIEInstanceID);
 #endif
 
+	Package->CallProgressCallbacks(EAsyncLoadingProgress::Read);
+
 #if ALT2_ENABLE_LINKERLOAD_SUPPORT
 	if (Package->LinkerLoadState.IsSet())
 	{
@@ -6102,6 +6087,26 @@ EEventLoadNodeExecutionResult FAsyncPackage2::Event_DependenciesReady(FAsyncLoad
 	return EEventLoadNodeExecutionResult::Complete;
 }
 
+void FAsyncPackage2::CallProgressCallbacks(EAsyncLoadingProgress ProgressType)
+{
+	if (ProgressCallbacks.Num() != 0)
+	{
+		TRACE_CPUPROFILER_EVENT_SCOPE(FAsyncPackage2::CallProgressCallbacks);
+
+		const FLoadPackageAsyncProgressParams Params
+		{
+			.PackageName = Desc.UPackageName,
+			.LoadedPackage = GetLinkerRoot(),
+			.ProgressType = ProgressType
+		};
+
+		for (TUniquePtr<FLoadPackageAsyncProgressDelegate>& ProgressCallback : ProgressCallbacks)
+		{
+			ProgressCallback->ExecuteIfBound(Params);
+		}
+	}
+}
+
 void FAsyncPackage2::InitializeExportArchive(FExportArchive& Ar, bool bIsOptionalSegment)
 {
 	Ar.SetUEVer(LinkerRoot->GetLinkerPackageVersion());
@@ -6914,22 +6919,7 @@ EEventLoadNodeExecutionResult FAsyncPackage2::Event_ExportsDone(FAsyncLoadingThr
 		FCoreDelegates::ReleasePreloadedPackageShaderMaps.ExecuteIfBound(Package->Data.ShaderMapHashes);
 	}
 
-	if (Package->ProgressCallbacks.Num())
-	{
-		TRACE_CPUPROFILER_EVENT_SCOPE(PackageProgressCallbacks_Serialized);
-
-		const FLoadPackageAsyncProgressParams Params
-		{
-			.PackageName = Package->Desc.UPackageName,
-			.LoadedPackage = Package->GetLinkerRoot(),
-			.ProgressType = EAsyncLoadingProgress::Serialized
-		};
-
-		for (TUniquePtr<FLoadPackageAsyncProgressDelegate>& ProgressCallback : Package->ProgressCallbacks)
-		{
-			ProgressCallback->ExecuteIfBound(Params);
-		}
-	}
+	Package->CallProgressCallbacks(EAsyncLoadingProgress::Serialized);
 
 	FAsyncLoadingPostLoadGroup* PostLoadGroup = Package->PostLoadGroup;
 	check(PostLoadGroup);
diff --git a/Engine/Source/Runtime/CoreUObject/Public/UObject/UObjectGlobals.h b/Engine/Source/Runtime/CoreUObject/Public/UObject/UObjectGlobals.h
--- a/Engine/Source/Runtime/CoreUObject/Public/UObject/UObjectGlobals.h
+++ b/Engine/Source/Runtime/CoreUObject/Public/UObject/UObjectGlobals.h
@@ -669,6 +669,8 @@ enum class EAsyncLoadingProgress : uint32
 	Failed,
 	/** Package has started loading. */
 	Started,
+	/** Package I/O has been read. */
+	Read,
 	/** Package has finished its serialization phase. */
 	Serialized,
 	/** Package has finished all loading phase successfully */
diff --git a/Engine/Source/Runtime/Engine/Private/Tests/Loading/UnattendedAsyncLoadingTests.cpp b/Engine/Source/Runtime/Engine/Private/Tests/Loading/UnattendedAsyncLoadingTests.cpp
--- a/Engine/Source/Runtime/Engine/Private/Tests/Loading/UnattendedAsyncLoadingTests.cpp
+++ b/Engine/Source/Runtime/Engine/Private/Tests/Loading/UnattendedAsyncLoadingTests.cpp
@@ -125,6 +125,8 @@ namespace UnattendedLoadTestImpl
 									{
 										// Filter out some progress types as we wait until later to start loading the next package
 										case EAsyncLoadingProgress::Started:
+											[[fallthrough]];
+										case EAsyncLoadingProgress::Read:
 											return;
 
 										// Whenever we receive a progress type that indicate the request has finished being processed
