- Add GetFNameSafe functions for UObject and FField. These complement the existing GetNameSafe functions.
- Minor updates to similar functions to comply with coding standards.

diff --git a/Engine/Source/Runtime/CoreUObject/Public/UObject/Field.h b/Engine/Source/Runtime/CoreUObject/Public/UObject/Field.h
--- a/Engine/Source/Runtime/CoreUObject/Public/UObject/Field.h
+++ b/Engine/Source/Runtime/CoreUObject/Public/UObject/Field.h
@@ -1051,6 +1051,20 @@ inline void SerializeSingleField(FArchive& Ar, FieldType*& Field, FFieldVariant
 	}
 }
 
+/**
+ * Gets the name of the provided field. If the field pointer is null, the result is "none"
+ */
+inline FName GetFNameSafe(const FField* InField)
+{
+	if (InField)
+	{
+		return InField->GetFName();
+	}
+	else
+	{
+		return NAME_None;
+	}
+}
 /**
  * Gets the name of the provided field. If the field pointer is null, the result is "none"
  */
diff --git a/Engine/Source/Runtime/CoreUObject/Public/UObject/UObjectBaseUtility.h b/Engine/Source/Runtime/CoreUObject/Public/UObject/UObjectBaseUtility.h
--- a/Engine/Source/Runtime/CoreUObject/Public/UObject/UObjectBaseUtility.h
+++ b/Engine/Source/Runtime/CoreUObject/Public/UObject/UObjectBaseUtility.h
@@ -817,13 +817,30 @@ FORCEINLINE bool IsPossiblyAllocatedUObjectPointer(UObject* Ptr)
 }
 
 /**
- * Returns the name of this object (with no path information)
- * @param Object object to retrieve the name for; NULL gives "None"
+ * Returns the logical name of this object.
+ * @param Object object to retrieve the name for; null gives NAME_None.
  * @return Name of the object.
 */
-FORCEINLINE FString GetNameSafe(const UObjectBaseUtility *Object)
+FORCEINLINE FName GetFNameSafe(const UObjectBaseUtility* Object)
 {
-	if( Object == NULL )
+	if (Object == nullptr)
+	{
+		return NAME_None;
+	}
+	else
+	{
+		return Object->GetFName();
+	}
+}
+
+/**
+ * Returns the name of this object (with no path information).
+ * @param Object object to retrieve the name for; null gives "None".
+ * @return Name of the object.
+*/
+FORCEINLINE FString GetNameSafe(const UObjectBaseUtility* Object)
+{
+	if (Object == nullptr)
 	{
 		return TEXT("None");
 	}
@@ -834,13 +851,13 @@ FORCEINLINE FString GetNameSafe(const UObjectBaseUtility *Object)
 }
 
 /**
- * Returns the path name of this object
- * @param Object object to retrieve the path name for; NULL gives "None"
+ * Returns the path name of this object.
+ * @param Object object to retrieve the path name for; null gives "None".
  * @return path name of the object.
 */
-FORCEINLINE FString GetPathNameSafe(const UObjectBaseUtility *Object)
+FORCEINLINE FString GetPathNameSafe(const UObjectBaseUtility* Object)
 {
-	if( Object == NULL )
+	if (Object == nullptr)
 	{
 		return TEXT("None");
 	}
@@ -851,13 +868,13 @@ FORCEINLINE FString GetPathNameSafe(const UObjectBaseUtility *Object)
 }
 
 /**
- * Returns the full name of this object
- * @param Object object to retrieve the full name for; NULL (or a null class!) gives "None"
+ * Returns the full name of this object.
+ * @param Object object to retrieve the full name for; null (or a null class!) gives "None".
  * @return full name of the object.
 */
-FORCEINLINE FString GetFullNameSafe(const UObjectBaseUtility *Object)
+FORCEINLINE FString GetFullNameSafe(const UObjectBaseUtility* Object)
 {
-	if( !Object || !Object->GetClass())
+	if (!Object || !Object->GetClass())
 	{
 		return TEXT("None");
 	}
