Moved FPropertyBag to be internal

#rnx

diff --git a/Engine/Source/Runtime/CoreUObject/Internal/UObject/PropertyBag.h b/Engine/Source/Runtime/CoreUObject/Internal/UObject/PropertyBag.h
new file mode 100644
--- /dev/null
+++ b/Engine/Source/Runtime/CoreUObject/Internal/UObject/PropertyBag.h
@@ -0,0 +1,130 @@
+// Copyright Epic Games, Inc. All Rights Reserved.
+
+#pragma once
+
+#include "CoreTypes.h"
+#include "Containers/Map.h"
+#include "Misc/TVariant.h"
+#include "UObject/PropertyPathName.h"
+#include "UObject/PropertyTag.h"
+
+#define UE_API COREUOBJECT_API
+
+enum EPropertyFlags : uint64;
+
+namespace UE { class FPropertyBagElement; }
+
+namespace UE
+{
+
+/**
+ * A property bag is an unordered collection of properties and their values.
+ */
+class FPropertyBag
+{
+	struct FValue
+	{
+		~FValue();
+
+		void AllocateAndInitializeValue();
+		void Destroy();
+
+		UE_API int32 GetSize() const;
+
+		inline explicit operator bool() const { return !!Data; }
+
+		FPropertyTag Tag;
+		void* Data = nullptr;
+		bool bOwnsProperty = false;
+	};
+
+	struct FNode;
+	using FNodeMap = TMap<FName, TUniquePtr<FNode>>;
+
+	struct FNode
+	{
+		FName Type;
+		FValue Value;
+		FNodeMap Nodes;
+	};
+
+	FValue& FindOrCreateValue(const FPropertyPathName& Path);
+
+public:
+	UE_API FPropertyBag();
+	UE_API ~FPropertyBag();
+
+	UE_API FPropertyBag(FPropertyBag&&);
+	UE_API FPropertyBag& operator=(FPropertyBag&&);
+
+	FPropertyBag(const FPropertyBag&) = delete;
+	FPropertyBag& operator=(const FPropertyBag&) = delete;
+
+	/** True if the bag contains no properties. */
+	inline bool IsEmpty() const { return Properties.IsEmpty(); }
+
+	/**
+	 * Remove every property from the property bag.
+	 */
+	UE_API void Empty();
+
+	/**
+	 * Add the property to the bag.
+	 *
+	 * Any property in the bag with the same path will be overwritten.
+	 */
+	UE_API void Add(const FPropertyPathName& Path, FProperty* Property, void* Data, int32 ArrayIndex = 0);
+
+	/**
+	 * Remove a property from the bag by path.
+	 */
+	UE_API void Remove(const FPropertyPathName& Path);
+
+	/**
+	 * Load a property from the slot into the bag based on the tag.
+	 *
+	 * If the property has a known type then it will be deserialized and accessible.
+	 * If not, the serialized data will be retained to avoid losing the unrecognized property.
+	 */
+	UE_API void LoadPropertyByTag(const FPropertyPathName& Path, const FPropertyTag& Tag, FStructuredArchiveSlot& ValueSlot, const void* Defaults = nullptr);
+
+	class FConstIterator
+	{
+	public:
+		using FNodeIterator = typename FNodeMap::TConstIterator;
+
+		inline explicit FConstIterator(const FNodeIterator& NodeIt)
+		{
+			NodeIterators.Push(NodeIt);
+			EnterNode();
+		}
+
+		inline FConstIterator& operator++() { EnterNode(); return *this; }
+
+		inline explicit operator bool() const { return !!CurrentValue; }
+
+		inline bool operator==(const FConstIterator& Rhs) const { return CurrentValue == Rhs.CurrentValue; }
+		inline bool operator!=(const FConstIterator& Rhs) const { return CurrentValue != Rhs.CurrentValue; }
+
+		inline const FPropertyPathName& GetPath() const { return CurrentPath; }
+		inline const FProperty* GetProperty() const { return CurrentValue->Tag.Prop; }
+		inline void* GetValue() const { return CurrentValue->Data; }
+		inline int32 GetValueSize() const { return CurrentValue->GetSize(); }
+
+	private:
+		UE_API void EnterNode();
+
+		TArray<FNodeIterator, TInlineAllocator<8>> NodeIterators;
+		FPropertyPathName CurrentPath;
+		FValue* CurrentValue = nullptr;
+	};
+
+	inline FConstIterator CreateConstIterator() const { return FConstIterator(Properties.CreateConstIterator()); }
+
+private:
+	FNodeMap Properties;
+};
+
+} // UE
+
+#undef UE_API
diff --git a/Engine/Source/Runtime/CoreUObject/Public/UObject/PropertyBag.h b/Engine/Source/Runtime/CoreUObject/Public/UObject/PropertyBag.h
deleted file mode 100644
