CIS fix

diff --git a/Engine/Source/Runtime/CoreUObject/Public/UObject/ObjectPtr.h b/Engine/Source/Runtime/CoreUObject/Public/UObject/ObjectPtr.h
--- a/Engine/Source/Runtime/CoreUObject/Public/UObject/ObjectPtr.h
+++ b/Engine/Source/Runtime/CoreUObject/Public/UObject/ObjectPtr.h
@@ -5,6 +5,7 @@
 #include "HAL/Platform.h"
 #include "Serialization/StructuredArchive.h"
 #include "Templates/IsTObjectPtr.h"
+#include "UObject/GarbageCollectionGlobals.h"
 #include "UObject/ObjectHandle.h"
 #include "UObject/UObjectGlobals.h"
 #include "Templates/NonNullPointer.h"
