[Feature] Archetypes being used for the Scenegraph Serialization
miro: https://miro.com/app/board/uXjVNVqTAsU=/?share_link_id=446429147441

This is a WIP feature. Known bugs/missing support with this feature includes:
- I haven't put much thought/testing into what to do when the class being fixed up already has an inheritance structure beyond just inheriting directly from UObject. currently, the archetype flattens the inheritance structure
- I haven't put much thought/testing into how archetypes should handle sub-objects. I think that would take place external to this code though during serialization
- sets, maps, and enums still don't have full support in property bags so they don't fully work here either

diff --git a/Engine/Source/Runtime/CoreUObject/Private/UObject/ArchetypeUtils.cpp b/Engine/Source/Runtime/CoreUObject/Private/UObject/ArchetypeUtils.cpp
new file mode 100644
--- /dev/null
+++ b/Engine/Source/Runtime/CoreUObject/Private/UObject/ArchetypeUtils.cpp
@@ -0,0 +1,624 @@
+// Copyright Epic Games, Inc. All Rights Reserved.
+
+#include "UObject/ArchetypeUtils.h"
+#include "UObject/PropertyBag.h"
+#include "UObject/UnrealType.h"
+#include "UObject/EnumProperty.h"
+#include "UObject/Field.h"
+
+namespace UE
+{
+	// typedef to help make it clearer when a pathName has indices and when the indices are wildcarded away
+	using FWildcardPropertyPathName = FPropertyPathName;
+	
+	static FName Name_ValuesSetBySerialization(TEXT("_ValuesSetBySerialization"));
+	
+	struct ResolvePropertyPathNameHelperParams
+	{
+		void* Data = nullptr;
+		FProperty* ResultProperty = nullptr;
+		const UE::FPropertyPathName& Path;
+		int32 CurPathIndex = 0;
+		int32 EndPathIndex = INDEX_NONE; // INDEX_NONE has the same behavior as Path.GetSegmentCount()
+		bool bAddIfNeeded = false;
+	};
+
+	static void BuildSegmentTypeFromProperty(const FProperty* Property, TArray<FName>& OutType)
+	{
+#if true // TODO: @jordan.hoffmann when complete type info is finished use this branch instead
+		OutType.Add(Property->GetID());
+		if (const FStructProperty* AsStructProperty = CastField<FStructProperty>(Property))
+		{
+#if WITH_EDITORONLY_DATA
+			if (const FString* TPSOverrideStructName = AsStructProperty->FindMetaData("TPSOverrideStructName"))
+			{
+				OutType.Add(FName(*TPSOverrideStructName));
+			}
+			else
+#endif
+			{
+				OutType.Add(AsStructProperty->Struct->GetFName());
+			}
+		}
+		else if (const FObjectProperty* AsObjectProperty = CastField<FObjectProperty>(Property))
+		{
+			OutType.Add(AsObjectProperty->GetID());
+		}
+		else if (const FEnumProperty* AsEnumProperty = CastField<FEnumProperty>(Property))
+		{
+			OutType.Add(AsEnumProperty->GetEnum()->GetFName());
+		}
+		else if (const FArrayProperty* AsArrayProperty = CastField<FArrayProperty>(Property))
+		{
+			BuildSegmentTypeFromProperty(AsArrayProperty->Inner, OutType);
+		}
+		else if (const FSetProperty* AsSetProperty = CastField<FSetProperty>(Property))
+		{
+			BuildSegmentTypeFromProperty(AsSetProperty->ElementProp, OutType);
+		}
+		else if (const FMapProperty* AsMapProperty = CastField<FMapProperty>(Property))
+		{
+			BuildSegmentTypeFromProperty(AsMapProperty->KeyProp, OutType);
+			BuildSegmentTypeFromProperty(AsMapProperty->ValueProp, OutType);
+		}
+#else
+		OutType.Add(Property->GetID());
+		if (const FStructProperty* AsStructProperty = CastField<FStructProperty>(Property))
+		{
+#if WITH_EDITORONLY_DATA
+			if (const FString* TPSOverrideStructName = AsStructProperty->FindMetaData("TPSOverrideStructName"))
+			{
+				OutType.Add(FName(*TPSOverrideStructName));
+			}
+			else
+#endif
+			{
+				OutType.Add(AsStructProperty->Struct->GetFName());
+			}
+		}
+		else if (const FEnumProperty* AsEnumProperty = CastField<FEnumProperty>(Property))
+		{
+			OutType.Add(AsEnumProperty->GetEnum()->GetFName());
+		}
+		else if (const FArrayProperty* AsArrayProperty = CastField<FArrayProperty>(Property))
+		{
+			if (AsArrayProperty->Inner->IsA<FEnumProperty>())
+			{
+				// enum paths currently don't recurse when they're in containers
+				OutType.Add(AsArrayProperty->Inner->GetClass()->GetFName());
+			}
+			else
+			{
+				BuildSegmentTypeFromProperty(AsArrayProperty->Inner, OutType);
+			}
+		}
+		else if (const FSetProperty* AsSetProperty = CastField<FSetProperty>(Property))
+		{
+			// sets currently don't recurse their element property types
+			OutType.Add(AsSetProperty->ElementProp->GetClass()->GetFName());
+		}
+		else if (const FMapProperty* AsMapProperty = CastField<FMapProperty>(Property))
+		{
+			// maps currently don't recurse their key/value property types
+			OutType.Add(AsMapProperty->KeyProp->GetClass()->GetFName());
+			OutType.Add(AsMapProperty->ValueProp->GetClass()->GetFName());
+		}
+#endif
+	}
+	
+	static FName GetSegmentTypeFromProperty(const FProperty* Property)
+	{
+		TArray<FName> TypeList;
+		BuildSegmentTypeFromProperty(Property, TypeList);
+		TStringBuilder<512> Type;
+		Type.Join(TypeList, TEXT(' '));
+		
+		return FName(Type);
+	}
+	
+	static bool ResolvePropertyPathNameHelper(const UStruct* Struct, ResolvePropertyPathNameHelperParams& Params);
+	static bool ResolvePropertyPathNameHelper(ResolvePropertyPathNameHelperParams& Params)
+	{
+		if (Params.CurPathIndex == Params.Path.GetSegmentCount() || Params.CurPathIndex == Params.EndPathIndex)
+		{
+			return true;
+		}
+		FPropertyPathNameSegment Segment = Params.Path.GetSegment(Params.CurPathIndex);
+		if (const FStructProperty* AsStructProperty = CastField<FStructProperty>(Params.ResultProperty))
+		{
+			return ResolvePropertyPathNameHelper(AsStructProperty->Struct, Params);
+		}
+		if (const FArrayProperty* AsArrayProperty = CastField<FArrayProperty>(Params.ResultProperty))
+		{
+			FScriptArrayHelper Array(AsArrayProperty, Params.Data);
+			if (!Array.IsValidIndex(Segment.Index))
+			{
+				if (Params.bAddIfNeeded)
+				{
+					Array.Resize(Segment.Index + 1);
+				}
+				else
+				{
+					return false;
+				}
+			}
+			++Params.CurPathIndex;
+			Params.ResultProperty = AsArrayProperty->Inner;
+			Params.Data = Array.GetElementPtr(Segment.Index);
+			return ResolvePropertyPathNameHelper(Params);
+		}
+		if (const FSetProperty* AsSetProperty = CastField<FSetProperty>(Params.ResultProperty))
+		{
+			FScriptSetHelper Set(AsSetProperty, Params.Data);
+			if (!Set.IsValidIndex(Segment.Index))
+			{
+				if (Params.bAddIfNeeded)
+				{
+					Segment.Index = Set.AddDefaultValue_Invalid_NeedsRehash();
+				}
+				else
+				{
+					return false;
+				}
+			}
+			++Params.CurPathIndex;
+			Params.ResultProperty = AsSetProperty->ElementProp;
+			Params.Data = Set.GetElementPtr(Segment.Index);
+			return ResolvePropertyPathNameHelper(Params);
+		}
+		if (const FMapProperty* AsMapProperty = CastField<FMapProperty>(Params.ResultProperty))
+		{
+			FScriptMapHelper Map(AsMapProperty, Params.Data);
+			if (!Map.IsValidIndex(Segment.Index))
+			{
+				if (Params.bAddIfNeeded)
+				{
+					Segment.Index = Map.AddDefaultValue_Invalid_NeedsRehash();
+				}
+				else
+				{
+					return false;
+				}
+			}
+			++Params.CurPathIndex;
+			Params.ResultProperty = AsMapProperty->KeyProp;
+			Params.Data = Map.GetKeyPtr(Segment.Index);
+			if (ResolvePropertyPathNameHelper(Params))
+			{
+				return true;
+			}
+			Params.ResultProperty = AsMapProperty->ValueProp;
+			Params.Data = Map.GetValuePtr(Segment.Index);
+			return ResolvePropertyPathNameHelper(Params);
+		}
+		check(Params.CurPathIndex == Params.Path.GetSegmentCount() - 1)
+		return true;
+	}
+	
+	static bool ResolvePropertyPathNameHelper(const UStruct* Struct, ResolvePropertyPathNameHelperParams& Params)
+    {
+		const FPropertyPathNameSegment Segment = Params.Path.GetSegment(Params.CurPathIndex);
+		for (FProperty* Property : TFieldRange<FProperty>(Struct))
+		{
+			// find a property that matches the segment
+			if (Segment.Name != Property->GetFName() || Segment.Type != GetSegmentTypeFromProperty(Property))
+			{
+				continue;
+			}
+			const bool bStaticArrayIndex = !Property->IsA<FArrayProperty>() && !Property->IsA<FSetProperty>() && !Property->IsA<FMapProperty>();
+			
+			if (bStaticArrayIndex && Segment.Index != INDEX_NONE)
+			{
+				if (Segment.Index < Property->ArrayDim)
+				{
+					++Params.CurPathIndex;
+					Params.ResultProperty = Property;
+					Params.Data = Property->ContainerPtrToValuePtr<void>(Params.Data, Segment.Index);
+					return ResolvePropertyPathNameHelper(Params);
+				}
+				else
+				{
+					// out of static array bounds
+					return false;
+				}
+			}
+			else
+			{
+				++Params.CurPathIndex;
+				Params.ResultProperty = Property;
+				Params.Data = Property->ContainerPtrToValuePtr<void>(Params.Data);
+				return ResolvePropertyPathNameHelper(Params);
+			}
+		}
+		
+		return false; // segment not found in struct
+    }
+
+	static UStruct* CreatePropertyBagArchetypeStructRec(const UClass* StructClass, UStruct* TemplateStruct,
+		UObject* Outer, const TMap<FWildcardPropertyPathName, TArray<const FProperty*>>& LooseProperties, FWildcardPropertyPathName& Path);
+	template <typename TStructType>
+	TStructType* CreatePropertyBagArchetypeStructRec(UStruct* TemplateStruct, UObject* Outer,
+		const TMap<FWildcardPropertyPathName, TArray<const FProperty*>>& LooseProperties, FWildcardPropertyPathName& Path)
+	{
+		return CastChecked<TStructType>(CreatePropertyBagArchetypeStructRec(TStructType::StaticClass(), TemplateStruct, Outer, LooseProperties, Path));
+	}
+
+	static FPropertyPathNameSegment CreateSegmentFromProperty(const FProperty* Inner, int32 Index = INDEX_NONE)
+	{
+		FPropertyPathNameSegment Result;
+		Result.Index = INDEX_NONE;
+		Result.Name = Inner->GetFName();
+		Result.Type = GetSegmentTypeFromProperty(Inner);
+		return Result;
+	}
+
+	// recursively re-instances all structs contained by this property to include loose properties
+	static void ConvertToArchetypeProperty(FProperty* Property, UObject* Outer,
+		const TMap<FWildcardPropertyPathName, TArray<const FProperty*>>& LooseProperties, FWildcardPropertyPathName& Path)
+	{
+		if (FStructProperty* AsStructProperty = CastField<FStructProperty>(Property))
+		{
+			const FName TPSOverrideStructName = AsStructProperty->Struct->GetFName();
+			AsStructProperty->Struct = CreatePropertyBagArchetypeStructRec<UScriptStruct>(AsStructProperty->Struct, Outer, LooseProperties, Path);
+#if WITH_EDITORONLY_DATA
+			AsStructProperty->SetMetaData("TPSOverrideStructName", TPSOverrideStructName.ToString());
+#endif
+		}
+		else if (const FArrayProperty* AsArrayProperty = CastField<FArrayProperty>(Property))
+		{
+			Path.Push(CreateSegmentFromProperty(AsArrayProperty->Inner));
+			ConvertToArchetypeProperty(AsArrayProperty->Inner, Outer, LooseProperties, Path);
+			Path.Pop();
+		}
+		else if (const FSetProperty* AsSetProperty = CastField<FSetProperty>(Property))
+		{
+			Path.Push(CreateSegmentFromProperty(AsSetProperty->ElementProp));
+			ConvertToArchetypeProperty(AsSetProperty->ElementProp, Outer, LooseProperties, Path);
+			Path.Pop();
+		}
+		else if (const FMapProperty* AsMapProperty = CastField<FMapProperty>(Property))
+		{
+			Path.Push(CreateSegmentFromProperty(AsMapProperty->KeyProp));
+			ConvertToArchetypeProperty(AsMapProperty->KeyProp, Outer, LooseProperties, Path);
+			Path.Pop();
+			
+			Path.Push(CreateSegmentFromProperty(AsMapProperty->ValueProp));
+			ConvertToArchetypeProperty(AsMapProperty->ValueProp, Outer, LooseProperties, Path);
+			Path.Pop();
+		}
+	}
+	
+	// copy template property then convert it into an archetype property by adding loose properties
+	static FProperty* CreateArchetypeProperty(const FProperty* TemplateProperty, UObject* Outer,
+		const TMap<FWildcardPropertyPathName, TArray<const FProperty*>>& LooseProperties, FWildcardPropertyPathName& Path)
+	{
+		FProperty* ArchetypeProperty = CastFieldChecked<FProperty>(FField::Duplicate(TemplateProperty, Outer));
+#if WITH_EDITORONLY_DATA
+		FField::CopyMetaData(TemplateProperty, ArchetypeProperty);
+#endif
+		ConvertToArchetypeProperty(ArchetypeProperty, Outer, LooseProperties, Path);
+		return ArchetypeProperty;
+	}
+
+	// return a copy of Path with all the indices set to -1. This way all container elements will have the same wildcard path
+	static FWildcardPropertyPathName ConvertToWildcardPath(const FPropertyPathName& Path)
+	{
+		FWildcardPropertyPathName Result = Path;
+		// make path a wildcard path
+		for (int I = 0; I < Result.GetSegmentCount(); ++I)
+		{
+			FPropertyPathNameSegment Segment = Result.GetSegment(I);
+			Segment.Index = INDEX_NONE;
+			Result.SetSegment(I, Segment);
+		}
+		return Result;
+	}
+
+	// recursively add all the wildcard paths of both Property and all it's sub-Properties to OutLooseProperties
+	static void AddWildcardedProperties(TMap<FWildcardPropertyPathName, TArray<const FProperty*>>& OutProperties, FWildcardPropertyPathName& ParentPath, const FProperty* Property)
+	{
+		OutProperties.FindOrAdd(ParentPath).Add(Property);
+		
+		ParentPath.Push(CreateSegmentFromProperty(Property));
+		if (const FStructProperty* AsStructProperty = CastField<FStructProperty>(Property))
+		{
+			for (const FProperty* SubProperty : TFieldRange<FProperty>(AsStructProperty->Struct))
+			{
+				AddWildcardedProperties(OutProperties, ParentPath, SubProperty);
+			}
+		}
+		else if (const FArrayProperty* AsArrayProperty = CastField<FArrayProperty>(Property))
+		{
+			AddWildcardedProperties(OutProperties, ParentPath, AsArrayProperty->Inner);
+		}
+		else if (const FSetProperty* AsSetProperty = CastField<FSetProperty>(Property))
+		{
+			AddWildcardedProperties(OutProperties, ParentPath, AsSetProperty->ElementProp);
+		}
+		else if (const FMapProperty* AsMapProperty = CastField<FMapProperty>(Property))
+		{
+			AddWildcardedProperties(OutProperties, ParentPath, AsMapProperty->KeyProp);
+			AddWildcardedProperties(OutProperties, ParentPath, AsMapProperty->ValueProp);
+		}
+		ParentPath.Pop();
+	}
+
+	// construct a map that keys a parent struct by it's wildcard path and returns an array of all it's loose properties
+	static TMap<FWildcardPropertyPathName, TArray<const FProperty*>> GetWildcardedLooseProperties(const FPropertyBag* PropertyBag)
+	{
+		
+		TMap<FWildcardPropertyPathName, TArray<const FProperty*>> LooseProperties;
+		if (PropertyBag)
+		{
+			for (FPropertyBag::FConstIterator Itr = PropertyBag->CreateConstIterator(); Itr; ++Itr)
+			{
+				FWildcardPropertyPathName ParentPath = ConvertToWildcardPath(Itr.GetPath());
+				ParentPath.Pop();
+				AddWildcardedProperties(LooseProperties, ParentPath, Itr.GetProperty());
+			}
+		}
+		
+		return LooseProperties;
+	}
+
+	// recursively gives a property the metadata and flags of a loose property
+	static void MarkPropertyAsLoose(FProperty* Property)
+	{
+#if WITH_EDITORONLY_DATA
+		Property->SetMetaData(TEXT("isLoose"), TEXT("True"));
+		Property->SetMetaData(TEXT("category"), TEXT("Loose Properties"));
+#endif
+		Property->SetPropertyFlags(CPF_Edit | CPF_EditConst);
+		if (const FArrayProperty* AsArrayProperty = CastField<FArrayProperty>(Property))
+        {
+			MarkPropertyAsLoose(AsArrayProperty->Inner);
+        }
+        else if (const FSetProperty* AsSetProperty =CastField<FSetProperty>(Property))
+        {
+			MarkPropertyAsLoose(AsSetProperty->ElementProp);
+        }
+        else if (const FMapProperty* AsMapProperty =CastField<FMapProperty>(Property))
+        {
+			MarkPropertyAsLoose(AsMapProperty->KeyProp);
+			MarkPropertyAsLoose(AsMapProperty->ValueProp);
+        }
+	}
+
+	// constructs an archetype struct by merging the properties in 
+	static UStruct* CreatePropertyBagArchetypeStructRec(const UClass* StructClass, UStruct* TemplateStruct,
+		UObject* Outer, const TMap<FWildcardPropertyPathName, TArray<const FProperty*>>& LooseProperties, FWildcardPropertyPathName& Path)
+	{
+		UStruct* Super = nullptr;
+
+		const TArray<const FProperty*>* BagProperties = LooseProperties.Find(Path);
+
+		auto MatchesBagProperty = [&BagProperties](const FProperty* Property)
+		{
+			return BagProperties && BagProperties->ContainsByPredicate([Property](const FProperty* Found)
+			{
+				return Found->SameType(Property) && Found->GetFName() == Property->GetFName();
+			});
+		};
+		
+		if (TemplateStruct)
+		{
+			const FName SuperName(TemplateStruct->GetName() + TEXT("_Super"));
+			Super = NewObject<UStruct>(Outer, StructClass, MakeUniqueObjectName(Outer, StructClass, SuperName));
+			
+			// Gather properties for Super Struct
+			TArray<FProperty*> SuperProperties;
+			for (const FProperty* TemplateProperty : TFieldRange<FProperty>(TemplateStruct))
+			{
+				if (MatchesBagProperty(TemplateProperty))
+				{
+					// this property was determined to be loose despite it being in the template.
+					// this likely occured due to an entire struct instance being loose and that instance becoming a template
+					continue;
+				}
+				Path.Push(CreateSegmentFromProperty(TemplateProperty));
+				FProperty* SuperProperty = CreateArchetypeProperty(TemplateProperty, Super, LooseProperties, Path);
+				Path.Pop();
+				SuperProperties.Add(SuperProperty);
+			}
+
+			if (StructClass == UClass::StaticClass())
+			{
+				// UClasses are required to inherit from a UObject class
+				Super->SetSuperStruct(UObject::StaticClass());
+			}
+		    
+			// AddCppProperty expects reverse property order for StaticLink to work correctly
+			for (int32 I = SuperProperties.Num() - 1; I >= 0; --I)
+			{
+				Super->AddCppProperty(SuperProperties[I]);
+			}
+			Super->Bind();
+			Super->StaticLink(/*RelinkExistingProperties*/true);
+		}
+		else if (StructClass == UClass::StaticClass())
+		{
+			// UClasses are required to inherit from a UObject class
+			Super->SetSuperStruct(UObject::StaticClass());
+		}
+
+#if false // TODO: fix FStructProperty::ConvertFromType to allow converting this struct to the template struct then enable this.
+		const FName ArchetypeName = (TemplateStruct) ? FName(TemplateStruct->GetName() + TEXT("_Archetype")) : FName(TEXT("Archetype"));
+#else
+		FName ArchetypeName;
+		if (StructClass->IsChildOf(UClass::StaticClass()))
+		{
+			ArchetypeName = MakeUniqueObjectName(Outer, StructClass, (TemplateStruct) ? FName(TemplateStruct->GetName() + TEXT("_Archetype")) : FName(TEXT("Archetype")));
+		}
+		else
+		{
+			// struct archetype serialization breaks currently if the name doesn't match.
+			ArchetypeName = (TemplateStruct) ? TemplateStruct->GetFName() : FName(TEXT("Archetype"));
+		}
+#endif
+		UStruct* Result = NewObject<UStruct>(Outer, StructClass, ArchetypeName);
+
+		// Gather "loose" properties for child Struct
+		TArray<FProperty*> LooseArchetypeProperties;
+		if (BagProperties)
+		{
+			for (const FProperty* BagProperty : *BagProperties)
+			{
+				Path.Push(CreateSegmentFromProperty(BagProperty));
+				FProperty* LooseProperty = CreateArchetypeProperty(BagProperty, Result, LooseProperties, Path);
+				Path.Pop();
+				
+				MarkPropertyAsLoose(LooseProperty);
+				LooseArchetypeProperties.Add(LooseProperty);
+			}
+		}
+
+		// add a hidden set property used to record whether this struct's properties were set serialization.
+		{
+			FSetProperty* ValuesSetBySerializationProperty = CastFieldChecked<FSetProperty>(FSetProperty::Construct(Result, Name_ValuesSetBySerialization, RF_Transient | RF_MarkAsNative));
+			static FName Name_PropertyName(TEXT("PropertyName"));
+			ValuesSetBySerializationProperty->ElementProp = CastFieldChecked<FProperty>(FInt64Property::Construct(ValuesSetBySerializationProperty, Name_PropertyName, RF_Transient));
+			ValuesSetBySerializationProperty->SetPropertyFlags(CPF_Transient | CPF_EditorOnly | CPF_NativeAccessSpecifierPrivate);
+			Result->AddCppProperty(ValuesSetBySerializationProperty);
+		}
+		
+		Result->SetSuperStruct(Super);
+		
+		// AddCppProperty expects reverse property order for StaticLink to work correctly
+		for (int32 I = LooseArchetypeProperties.Num() - 1; I >= 0; --I)
+		{
+			Result->AddCppProperty(LooseArchetypeProperties[I]);
+		}
+		Result->Bind();
+		Result->StaticLink(/*RelinkExistingProperties*/true);
+		return Result;
+	}
+	
+	UClass* CreatePropertyBagArchetypeClass(const FPropertyBag* PropertyBag, UStruct* TemplateStruct, UObject* Outer)
+	{
+		const TMap<FWildcardPropertyPathName, TArray<const FProperty*>> LooseProperties = GetWildcardedLooseProperties(PropertyBag);
+		FWildcardPropertyPathName ParentPath;
+		return CreatePropertyBagArchetypeStructRec<UClass>(TemplateStruct, Outer, LooseProperties, ParentPath);
+	}
+	
+	void MarkPropertySetBySerialization(UObject* Object, const FPropertyPathName& Path)
+	{
+		const FPropertyPathNameSegment Segment = Path.GetSegment(Path.GetSegmentCount() - 1);
+
+		// partially resolve path to find the parent struct and data pointer
+		const UStruct* ParentStruct = Object->GetClass();
+		const void* ParentData = Object;
+		
+		ResolvePropertyPathNameHelperParams Params {
+            .Data = Object,
+            .Path = Path,
+            .CurPathIndex = 0,
+			.EndPathIndex = Path.GetSegmentCount() - 1,
+            .bAddIfNeeded = false
+        };
+		if (Params.CurPathIndex != Params.EndPathIndex)
+		{
+			if (!(ensure(ResolvePropertyPathNameHelper(Object->GetClass(), Params))))
+            {
+                return;
+            }
+            if (const FStructProperty* ParentAsStructProperty = CastField<FStructProperty>(Params.ResultProperty))
+            {
+                ParentData = Params.Data;
+                ParentStruct = ParentAsStructProperty->Struct;
+            }
+		}
+
+		// resolve the last segment of the path
+		Params.EndPathIndex += 1;
+        if (ensure(ResolvePropertyPathNameHelper(ParentStruct, Params)))
+        {
+            MarkPropertySetBySerialization(ParentStruct, ParentData, Params.ResultProperty, Segment.Index);
+        }
+	}
+
+	void MarkPropertySetBySerialization(const UStruct* Struct, const void* StructData, const FProperty* Property, int32 ArrayIndex)
+	{
+		if (const FSetProperty* ValuesSetByPropertyBagProperty = CastField<FSetProperty>(Struct->FindPropertyByName(Name_ValuesSetBySerialization)))
+		{
+			const uint8* PropertyDataPtr;
+			if (ArrayIndex == INDEX_NONE || Property->IsA<FArrayProperty>() || Property->IsA<FMapProperty>() || Property->IsA<FSetProperty>())
+			{
+				PropertyDataPtr = Property->ContainerPtrToValuePtr<uint8>(StructData);
+			}
+			else
+			{
+				PropertyDataPtr = Property->ContainerPtrToValuePtr<uint8>(StructData, ArrayIndex);
+			}
+			
+			FScriptSetHelper ValuesSetByPropertyBag(ValuesSetByPropertyBagProperty, ValuesSetByPropertyBagProperty->ContainerPtrToValuePtr<void>(StructData));
+			const int64 ValueOffset = PropertyDataPtr - static_cast<const uint8*>(StructData);
+			const int32 FoundIndex = ValuesSetByPropertyBag.FindElementIndex(&ValueOffset);
+			if (FoundIndex == INDEX_NONE)
+			{
+				ValuesSetByPropertyBag.AddElement(&ValueOffset);
+			}
+		}
+	}
+
+	bool WasPropertySetBySerialization(UObject* Object, const FPropertyPathName& Path)
+	{
+		const FPropertyPathNameSegment Segment = Path.GetSegment(Path.GetSegmentCount() - 1);
+
+		// partially resolve path to find the parent struct and data pointer
+		const UStruct* ParentStruct = Object->GetClass();
+		const void* ParentData = Object;
+		
+		ResolvePropertyPathNameHelperParams Params {
+			.Data = Object,
+			.Path = Path,
+			.CurPathIndex = 0,
+			.EndPathIndex = Path.GetSegmentCount() - 1,
+			.bAddIfNeeded = false
+		};
+		if (Params.CurPathIndex != Params.EndPathIndex)
+		{
+			if (!(ensure(ResolvePropertyPathNameHelper(Object->GetClass(), Params))))
+			{
+				return false;
+			}
+			if (const FStructProperty* ParentAsStructProperty = CastField<FStructProperty>(Params.ResultProperty))
+			{
+				ParentData = Params.Data;
+				ParentStruct = ParentAsStructProperty->Struct;
+			}
+		}
+
+		// resolve the last segment of the path
+		Params.EndPathIndex += 1;
+		if (ensure(ResolvePropertyPathNameHelper(ParentStruct, Params)))
+		{
+			return WasPropertySetBySerialization(ParentStruct, ParentData, Params.ResultProperty, Segment.Index);
+		}
+		return false;
+	}
+	
+	bool WasPropertySetBySerialization(const UStruct* Struct, const void* StructData, const FProperty* Property, int32 ArrayIndex)
+	{
+		if (ArrayIndex == INDEX_NONE)
+		{
+			ArrayIndex = 0;
+		}
+		if (const FSetProperty* ValuesSetByPropertyBagProperty = CastField<FSetProperty>(Struct->FindPropertyByName(Name_ValuesSetBySerialization)))
+		{
+			const uint8* PropertyDataPtr;
+			if (ArrayIndex == INDEX_NONE || Property->IsA<FArrayProperty>() || Property->IsA<FMapProperty>() || Property->IsA<FSetProperty>())
+			{
+				PropertyDataPtr = Property->ContainerPtrToValuePtr<uint8>(StructData);
+			}
+			else
+			{
+				PropertyDataPtr = Property->ContainerPtrToValuePtr<uint8>(StructData, ArrayIndex);
+			}
+
+			const FScriptSetHelper ValuesSetByPropertyBag(ValuesSetByPropertyBagProperty, ValuesSetByPropertyBagProperty->ContainerPtrToValuePtr<void>(StructData));
+			const int64 ValueOffset = PropertyDataPtr - static_cast<const uint8*>(StructData);
+			return ValuesSetByPropertyBag.FindElementIndex(&ValueOffset) != INDEX_NONE;
+		}
+		return false;
+	}
+} // UE
\ No newline at end of file
diff --git a/Engine/Source/Runtime/CoreUObject/Public/UObject/ArchetypeUtils.h b/Engine/Source/Runtime/CoreUObject/Public/UObject/ArchetypeUtils.h
new file mode 100644
--- /dev/null
+++ b/Engine/Source/Runtime/CoreUObject/Public/UObject/ArchetypeUtils.h
@@ -0,0 +1,32 @@
+// Copyright Epic Games, Inc. All Rights Reserved.
+
+#pragma once
+#include "HAL/Platform.h"
+
+class UClass;
+class UObject;
+class UStruct;
+class FProperty;
+
+namespace UE
+{
+	class FPropertyBag;
+	class FPropertyPathName;
+
+	// generate a UClass that unions the properties of PropertyBag and TemplateStruct
+	COREUOBJECT_API UClass* CreatePropertyBagArchetypeClass(const FPropertyBag* PropertyBag, UStruct* TemplateStruct, UObject* Outer);
+
+	// copy the values from a property bag to an archetype
+	COREUOBJECT_API void CopyPropertyBagValuesToArchetype(const FPropertyBag* PropertyBag, const UObject* Destination);
+
+	// notify that a property in an object was set when the object was deserialized
+	COREUOBJECT_API void MarkPropertySetBySerialization(UObject* Object, const FPropertyPathName& Path);
+	// notify that a property in in Struct was set when the struct was deserialized
+	COREUOBJECT_API void MarkPropertySetBySerialization(const UStruct* Struct, const void* StructData, const FProperty* Property, int32 ArrayIndex = 0);
+	
+	// query whether a property in an object was set when the object was deserialized
+	COREUOBJECT_API bool WasPropertySetBySerialization(UObject* Object, const FPropertyPathName& Path);
+	// query whether a property in Struct was set when the struct was deserialized
+	COREUOBJECT_API bool WasPropertySetBySerialization(const UStruct* Struct, const void* StructData, const FProperty* Property, int32 ArrayIndex = 0);
+}
+
