GetAssetRegistryTags: Fix deprecation-handling code to pass along the AssetBundle from the deprecated path into the new FAssetRegistryTagsContext.
#rnx

diff --git a/Engine/Source/Runtime/CoreUObject/Private/UObject/Obj.cpp b/Engine/Source/Runtime/CoreUObject/Private/UObject/Obj.cpp
--- a/Engine/Source/Runtime/CoreUObject/Private/UObject/Obj.cpp
+++ b/Engine/Source/Runtime/CoreUObject/Private/UObject/Obj.cpp
@@ -2178,7 +2178,8 @@ const FName FPrimaryAssetId::PrimaryAssetDisplayNameTag(TEXT("PrimaryAssetDispla
 // UE_DEPRECATED(5.4, "Used to provide backwards compatibility for the deprecated GetAssetRegistryTags function") // UE_DEPRECATED seems not to work with thread_local
 thread_local TArray<const UObject*, TInlineAllocator<2>> GAssetRegistryTagsObjectsBeingForwarded;
 // UE_DEPRECATED(5.4, "Used to provide backwards compatibility for the deprecated GetAssetRegistryTags function") // UE_DEPRECATED seems not to work with thread_local
-thread_local bool bGLegacyTagsWantsBundleResult = false;
+thread_local FAssetBundleData const** TLegacyGetAssetRegistryTags_OutBundles = nullptr;
+
 
 void UObject::GetAssetRegistryTags(TArray<FAssetRegistryTag>& OutTags) const
 {
@@ -2186,7 +2187,7 @@ void UObject::GetAssetRegistryTags(TArray<FAssetRegistryTag>& OutTags) const
 	if (!GAssetRegistryTagsObjectsBeingForwarded.Contains(this))
 	{
 		FAssetRegistryTagsContextData Context(this, EAssetRegistryTagsCaller::Uncategorized);
-		Context.bWantsBundleResult = bGLegacyTagsWantsBundleResult;
+		Context.bWantsBundleResult = TLegacyGetAssetRegistryTags_OutBundles != nullptr;
 
 		GAssetRegistryTagsObjectsBeingForwarded.Add(this);
 		GetAssetRegistryTags(Context);
@@ -2198,6 +2199,11 @@ void UObject::GetAssetRegistryTags(TArray<FAssetRegistryTag>& OutTags) const
 		{
 			OutTags.Add(MoveTemp(Pair.Value));
 		}
+		if (TLegacyGetAssetRegistryTags_OutBundles)
+		{
+			checkf(*TLegacyGetAssetRegistryTags_OutBundles == nullptr, TEXT("Object %s has more than one FAssetBundleData!"), *GetPathName());
+			*TLegacyGetAssetRegistryTags_OutBundles = Context.BundleResult;
+		}
 	}
 	PRAGMA_ENABLE_DEPRECATION_WARNINGS;
 }
@@ -2220,11 +2226,20 @@ void UObject::GetAssetRegistryTags(FAssetRegistryTagsContext Context) const
 	// Forward this call to the legacy version for classes that have not converted yet.
 	if (!GAssetRegistryTagsObjectsBeingForwarded.Contains(this))
 	{
-		TGuardValue<bool> WantsBundleScope(bGLegacyTagsWantsBundleResult, Context.WantsBundleResult());
+		const FAssetBundleData* BundleResult = nullptr;
+		TOptional<TGuardValue<const FAssetBundleData**>> WantsBundleScope;
+		if (Context.WantsBundleResult())
+		{
+			WantsBundleScope.Emplace(TLegacyGetAssetRegistryTags_OutBundles, &BundleResult);
+		}
 		GAssetRegistryTagsObjectsBeingForwarded.Add(this);
 		AddLegacyTags([this](TArray<FAssetRegistryTag>& Tags) { GetAssetRegistryTags(Tags); });
 		check(!GAssetRegistryTagsObjectsBeingForwarded.IsEmpty() && GAssetRegistryTagsObjectsBeingForwarded.Last() == this);
 		GAssetRegistryTagsObjectsBeingForwarded.Pop(false /* bAllowShrinking */);
+		if (Context.WantsBundleResult() && BundleResult)
+		{
+			Context.SetBundleResult(BundleResult);
+		}
 	}
 	PRAGMA_ENABLE_DEPRECATION_WARNINGS;
 
