[BugFix] Accidentally enabled a branch of code that shouldn't be enabled until FPropertyPathNames type info is fully finished

diff --git a/Engine/Source/Runtime/CoreUObject/Private/UObject/ArchetypeUtils.cpp b/Engine/Source/Runtime/CoreUObject/Private/UObject/ArchetypeUtils.cpp
--- a/Engine/Source/Runtime/CoreUObject/Private/UObject/ArchetypeUtils.cpp
+++ b/Engine/Source/Runtime/CoreUObject/Private/UObject/ArchetypeUtils.cpp
@@ -25,7 +25,7 @@ namespace UE
 
 	static void BuildSegmentTypeFromProperty(const FProperty* Property, TArray<FName>& OutType)
 	{
-#if true // TODO: @jordan.hoffmann when complete type info is finished use this branch instead
+#if false // TODO: @jordan.hoffmann when complete type info is finished use this branch instead
 		OutType.Add(Property->GetID());
 		if (const FStructProperty* AsStructProperty = CastField<FStructProperty>(Property))
 		{
