[BugFix] FProperty::ConvertFromType fixups to support property bags and archetypes

diff --git a/Engine/Source/Runtime/CoreUObject/Private/UObject/ArchetypeUtils.cpp b/Engine/Source/Runtime/CoreUObject/Private/UObject/ArchetypeUtils.cpp
--- a/Engine/Source/Runtime/CoreUObject/Private/UObject/ArchetypeUtils.cpp
+++ b/Engine/Source/Runtime/CoreUObject/Private/UObject/ArchetypeUtils.cpp
@@ -441,20 +441,7 @@ namespace UE
 			Super->SetSuperStruct(UObject::StaticClass());
 		}
 
-#if false // TODO: fix FStructProperty::ConvertFromType to allow converting this struct to the template struct then enable this.
 		const FName ArchetypeName = (TemplateStruct) ? FName(TemplateStruct->GetName() + TEXT("_Archetype")) : FName(TEXT("Archetype"));
-#else
-		FName ArchetypeName;
-		if (StructClass->IsChildOf(UClass::StaticClass()))
-		{
-			ArchetypeName = MakeUniqueObjectName(Outer, StructClass, (TemplateStruct) ? FName(TemplateStruct->GetName() + TEXT("_Archetype")) : FName(TEXT("Archetype")));
-		}
-		else
-		{
-			// struct archetype serialization breaks currently if the name doesn't match.
-			ArchetypeName = (TemplateStruct) ? TemplateStruct->GetFName() : FName(TEXT("Archetype"));
-		}
-#endif
 		UStruct* Result = NewObject<UStruct>(Outer, StructClass, ArchetypeName);
 
 		// Gather "loose" properties for child Struct
diff --git a/Engine/Source/Runtime/CoreUObject/Private/UObject/Class.cpp b/Engine/Source/Runtime/CoreUObject/Private/UObject/Class.cpp
--- a/Engine/Source/Runtime/CoreUObject/Private/UObject/Class.cpp
+++ b/Engine/Source/Runtime/CoreUObject/Private/UObject/Class.cpp
@@ -1655,13 +1655,19 @@ void UStruct::SerializeVersionedTaggedProperties(FStructuredArchive::FSlot Slot,
 						{
 							case EConvertFromTypeResult::Converted:
 								bAdvanceProperty = true;
-								if (FPropertyBag* PropertyBag = TryFindPropertyBag())
+								// TODO: ConvertFromType always returns EConvertFromTypeResult::Converted for numeric types. this type check tries to fix
+								// that but we should find a more robust solution
+								if (Tag.Type != PropID)
 								{
-									UnderlyingArchive.Seek(StartOfProperty);
-									FStructuredArchive::FSlot CopySlot = PropertyRecord.EnterField(TEXT("Value"));
-									Tag.Prop = nullptr;
-									PropertyBag->LoadPropertyByTag(SerializeContext->SerializedPropertyPath, Tag, CopySlot, Property->ContainerPtrToValuePtrForDefaults<uint8>(DefaultsStruct, Defaults, Tag.ArrayIndex));
+									if (FPropertyBag* PropertyBag = TryFindPropertyBag())
+                                    {
+                                    	UnderlyingArchive.Seek(StartOfProperty);
+                                    	FStructuredArchive::FSlot CopySlot = PropertyRecord.EnterField(TEXT("Value"));
+                                    	Tag.Prop = nullptr;
+                                    	PropertyBag->LoadPropertyByTag(SerializeContext->SerializedPropertyPath, Tag, CopySlot, Property->ContainerPtrToValuePtrForDefaults<uint8>(DefaultsStruct, Defaults, Tag.ArrayIndex));
+                                    }
 								}
+								
 								break;
 
 							case EConvertFromTypeResult::UseSerializeItem:
diff --git a/Engine/Source/Runtime/CoreUObject/Private/UObject/PropertyStruct.cpp b/Engine/Source/Runtime/CoreUObject/Private/UObject/PropertyStruct.cpp
--- a/Engine/Source/Runtime/CoreUObject/Private/UObject/PropertyStruct.cpp
+++ b/Engine/Source/Runtime/CoreUObject/Private/UObject/PropertyStruct.cpp
@@ -403,6 +403,18 @@ EConvertFromTypeResult FStructProperty::ConvertFromType(const FPropertyTag& Tag,
 			// Old Implementation
 			return StructProperty && !StructProperty->UseBinaryOrNativeSerialization(InAr);
 		}
+#if WITH_EDITORONLY_DATA
+		if (StructProperty && StructProperty->Struct)
+		{
+			if (const FString* TPSOverrideStructName = StructProperty->FindMetaData("TPSOverrideStructName"))
+			{
+				if (*TPSOverrideStructName == PropertyTag.StructName.ToString())
+				{
+					return true;
+				}
+			}
+		}
+#endif
 		return PropertyTag.StructGuid.IsValid() && StructProperty && StructProperty->Struct && (PropertyTag.StructGuid == StructProperty->Struct->GetCustomGuid());
 	};
 
