Change GC IsReadyForFinishDestroy timeout logging to only log every 10 seconds instead of every iteration after timeout has passed.
Add a cvar to control maximum number of items logged.

diff --git a/Engine/Source/Runtime/CoreUObject/Private/UObject/GarbageCollection.cpp b/Engine/Source/Runtime/CoreUObject/Private/UObject/GarbageCollection.cpp
--- a/Engine/Source/Runtime/CoreUObject/Private/UObject/GarbageCollection.cpp
+++ b/Engine/Source/Runtime/CoreUObject/Private/UObject/GarbageCollection.cpp
@@ -329,6 +329,14 @@ static FAutoConsoleVariableRef CVarForceEnableDebugGCProcessor(
 	ECVF_Default
 );
 
+static int32 GMaxFinishDestroyTimeoutObjectsToLog = 10;
+static FAutoConsoleVariableRef CVarMaxFinishDestroyTimeoutObjectsToLog (
+	TEXT("gc.MaxFinishDestroyTimeoutObjectsToLog"),
+	GMaxFinishDestroyTimeoutObjectsToLog ,
+	TEXT("Maximum number of objects to log out when object destruction takes longer than expected"),
+	ECVF_Default
+);
+
 namespace UE::GC
 {
 
@@ -4410,6 +4418,7 @@ bool IncrementalDestroyGarbage(bool bUseTimeLimit, double TimeLimit)
 
 	// Keep track of time it took to destroy objects for stats
 	double IncrementalDestroyGarbageStartTime = FPlatformTime::Seconds();
+	double LastTimeoutWarningTime = FPlatformTime::Seconds();
 
 	// Depending on platform FPlatformTime::Seconds might take a noticeable amount of time if called thousands of times so we avoid
 	// enforcing the time limit too often, especially as neither Destroy nor actual deletion should take significant
@@ -4570,13 +4579,14 @@ bool IncrementalDestroyGarbage(bool bUseTimeLimit, double TimeLimit)
 #endif
 						// Check if we spent too much time on waiting for FinishDestroy without making any progress
 						if (LastLoopObjectsPendingDestructionCount == GGCObjectsPendingDestructionCount && bPollTimeLimit &&
-							((FPlatformTime::Seconds() - GCStartTime) > MaxTimeForFinishDestroy))
+							((FPlatformTime::Seconds() - LastTimeoutWarningTime) > MaxTimeForFinishDestroy))
 						{
+							LastTimeoutWarningTime = FPlatformTime::Seconds();
 							if (GEnableTimeoutOnPendingDestroyedObjectGC)
 							{
 								UE_LOG(LogGarbage, Warning, TEXT("Spent more than %.2fs on routing FinishDestroy to objects (objects in queue: %d)"), MaxTimeForFinishDestroy, GGCObjectsPendingDestructionCount);
 								UObject* LastObjectNotReadyForFinishDestroy = nullptr;
-								for (int32 ObjectIndex = 0; ObjectIndex < GGCObjectsPendingDestructionCount; ++ObjectIndex)
+								for (int32 ObjectIndex = 0; ObjectIndex < GGCObjectsPendingDestructionCount && ObjectIndex < GMaxFinishDestroyTimeoutObjectsToLog; ++ObjectIndex)
 								{
 									UObject* Obj = GGCObjectsPendingDestruction[ObjectIndex];
 									bool bReady = Obj->IsReadyForFinishDestroy();
