[BugFix] copy paste error fix that can cause nullptr to get dereferenced

diff --git a/Engine/Source/Runtime/CoreUObject/Private/UObject/ArchetypeUtils.cpp b/Engine/Source/Runtime/CoreUObject/Private/UObject/ArchetypeUtils.cpp
--- a/Engine/Source/Runtime/CoreUObject/Private/UObject/ArchetypeUtils.cpp
+++ b/Engine/Source/Runtime/CoreUObject/Private/UObject/ArchetypeUtils.cpp
@@ -438,7 +438,7 @@ namespace UE
 		else if (StructClass == UClass::StaticClass())
 		{
 			// UClasses are required to inherit from a UObject class
-			Super->SetSuperStruct(UObject::StaticClass());
+			Super = UObject::StaticClass();
 		}
 
 		const FName ArchetypeName = (TemplateStruct) ? FName(TemplateStruct->GetName() + TEXT("_Archetype")) : FName(TEXT("Archetype"));
@@ -467,7 +467,7 @@ namespace UE
 			ValuesSetBySerializationProperty->SetPropertyFlags(CPF_Transient | CPF_EditorOnly | CPF_NativeAccessSpecifierPrivate);
 			Result->AddCppProperty(ValuesSetBySerializationProperty);
 		}
-		
+
 		Result->SetSuperStruct(Super);
 		
 		// AddCppProperty expects reverse property order for StaticLink to work correctly
