[BugFix] made archetype struct names fully unique by calling MakeUniqueObjectName

diff --git a/Engine/Source/Runtime/CoreUObject/Private/UObject/ArchetypeUtils.cpp b/Engine/Source/Runtime/CoreUObject/Private/UObject/ArchetypeUtils.cpp
--- a/Engine/Source/Runtime/CoreUObject/Private/UObject/ArchetypeUtils.cpp
+++ b/Engine/Source/Runtime/CoreUObject/Private/UObject/ArchetypeUtils.cpp
@@ -442,7 +442,7 @@ namespace UE
 		}
 
 		const FName ArchetypeName = (TemplateStruct) ? FName(TemplateStruct->GetName() + TEXT("_Archetype")) : FName(TEXT("Archetype"));
-		UStruct* Result = NewObject<UStruct>(Outer, StructClass, ArchetypeName);
+		UStruct* Result = NewObject<UStruct>(Outer, StructClass, MakeUniqueObjectName(Outer, StructClass, ArchetypeName));
 
 		// Gather "loose" properties for child Struct
 		TArray<FProperty*> LooseArchetypeProperties;
