[ObjectAnnotation]
 - Fix potential race condition on unprotected members in AddAnnotationInternal

diff --git a/Engine/Source/Runtime/CoreUObject/Public/UObject/UObjectAnnotation.h b/Engine/Source/Runtime/CoreUObject/Public/UObject/UObjectAnnotation.h
--- a/Engine/Source/Runtime/CoreUObject/Public/UObject/UObjectAnnotation.h
+++ b/Engine/Source/Runtime/CoreUObject/Public/UObject/UObjectAnnotation.h
@@ -79,9 +79,8 @@ private:
 	void AddAnnotationInternal(const UObjectBase* Object, T&& Annotation)
 	{
 		check(Object);
-		AnnotationCacheKey = Object;
-		AnnotationCacheValue = Forward<T>(Annotation);
-		if (AnnotationCacheValue.IsDefault())
+		TAnnotation LocalAnnotation = Forward<T>(Annotation);
+		if (LocalAnnotation.IsDefault())
 		{
 			RemoveAnnotation(Object); // adding the default annotation is the same as removing an annotation
 		}
@@ -90,6 +89,8 @@ private:
 			bool bWasEmpty = false;
 			{
 				FScopeLock AnnotationMapLock(&AnnotationMapCritical);
+				AnnotationCacheKey = Object;
+				AnnotationCacheValue = MoveTemp(LocalAnnotation);
 				bWasEmpty = (AnnotationMap.Num() == 0);
 				AnnotationMap.Add(AnnotationCacheKey, AnnotationCacheValue);
 			}
