[Backout] - CL30233397
Original CL Desc
-----------------------------------------------------------------
Setting any of the root flags on UObjects while Incremental Reachability is running will now trigger GC barrier

Moved UE::GC::Private::GIsIncrementalReachabilityPending to GarbageCollectionGlobals.h to be able to access it from UObjectArray.h
Moved UObject::MarkAsReachable() to UObjectBase::MarkAsReachable() to be able to use it from FUObjectItem (UObjectArray.h)

diff --git a/Engine/Source/Runtime/CoreUObject/Private/UObject/GarbageCollection.cpp b/Engine/Source/Runtime/CoreUObject/Private/UObject/GarbageCollection.cpp
--- a/Engine/Source/Runtime/CoreUObject/Private/UObject/GarbageCollection.cpp
+++ b/Engine/Source/Runtime/CoreUObject/Private/UObject/GarbageCollection.cpp
@@ -1022,7 +1022,7 @@ static bool MarkClusterMutableObjectsAsReachable(FUObjectCluster& Cluster, Conta
 				else if (ReferencedMutableObjectItem->GetOwnerIndex() > 0 && !ReferencedMutableObjectItem->HasAnyFlags(EInternalObjectFlags::ReachableInCluster))
 				{
 					// This is a clustered object that maybe hasn't been processed yet
-					if (ReferencedMutableObjectItem->ThisThreadAtomicallySetFlag_ForGC(EInternalObjectFlags::ReachableInCluster))
+					if (ReferencedMutableObjectItem->ThisThreadAtomicallySetFlag(EInternalObjectFlags::ReachableInCluster))
 					{
 						// Needs doing, we need to get its cluster root and process it too
 						FUObjectItem* ReferencedMutableObjectsClusterRootItem = GUObjectArray.IndexToObjectUnsafeForGC(ReferencedMutableObjectItem->GetOwnerIndex());
@@ -2967,7 +2967,7 @@ public:
 				FUObjectItem* RootObjectItem = GUObjectArray.IndexToObjectUnsafeForGC(Metadata.ObjectItem->GetOwnerIndex());
 				checkSlow(RootObjectItem->HasAnyFlags(EInternalObjectFlags::ClusterRoot));
 
-				if (Metadata.ObjectItem->ThisThreadAtomicallySetFlag_ForGC(EInternalObjectFlags::ReachableInCluster))
+				if (Metadata.ObjectItem->ThisThreadAtomicallySetFlag(EInternalObjectFlags::ReachableInCluster))
 				{
 					if (ClearMaybeUnreachableInterlocked(RootObjectItem->Flags))
 					{
@@ -4040,7 +4040,7 @@ public:
 					bool bNeedsDoing = !ObjectItem->HasAnyFlags(EInternalObjectFlags::ReachableInCluster);
 					if (bNeedsDoing)
 					{
-						ObjectItem->ThisThreadAtomicallySetFlag_ForGC(EInternalObjectFlags::ReachableInCluster);
+						ObjectItem->SetFlags(EInternalObjectFlags::ReachableInCluster);
 						// Make sure cluster root object is reachable too
 						const int32 OwnerIndex = ObjectItem->GetOwnerIndex();
 						FUObjectItem* RootObjectItem = GUObjectArray.IndexToObjectUnsafeForGC(OwnerIndex);
@@ -5875,9 +5875,9 @@ FORCEINLINE static void MarkObjectItemAsReachable(FUObjectItem* ObjectItem)
 }
 
 template <bool bIsVerse>
-FORCEINLINE static void MarkAsReachable(const UObjectBase* Obj)
+FORCEINLINE static void MarkAsReachable(const UObject* Obj)
 {
-	FUObjectItem* ObjectItem = GUObjectArray.ObjectToObjectItem(Obj);
+	FUObjectItem* ObjectItem = GUObjectArray.ObjectToObjectItem(const_cast<UObject*>(Obj));
 	if (ObjectItem->IsMaybeUnreachable())
 	{
 		MarkObjectItemAsReachable<bIsVerse>(ObjectItem);
@@ -5892,7 +5892,7 @@ FORCEINLINE static void MarkAsReachable(const UObjectBase* Obj)
 	}
 }
 
-void UObjectBase::MarkAsReachable() const
+void UObject::MarkAsReachable() const
 {
 	// It is safe to perform mark as reachable in the open - the worst case is that we'll mark an object reachable that
 	// should/would be destroyed, and so in the next GC iteration it will be destroyed instead of in this iteration.
diff --git a/Engine/Source/Runtime/CoreUObject/Public/UObject/GarbageCollectionGlobals.h b/Engine/Source/Runtime/CoreUObject/Public/UObject/GarbageCollectionGlobals.h
deleted file mode 100644
diff --git a/Engine/Source/Runtime/CoreUObject/Public/UObject/UObjectArray.h b/Engine/Source/Runtime/CoreUObject/Public/UObject/UObjectArray.h
--- a/Engine/Source/Runtime/CoreUObject/Public/UObject/UObjectArray.h
+++ b/Engine/Source/Runtime/CoreUObject/Public/UObject/UObjectArray.h
@@ -8,7 +8,7 @@
 
 #include "HAL/ThreadSafeCounter.h"
 #include "Containers/LockFreeList.h"
-#include "UObject/GarbageCollectionGlobals.h"
+#include "UObject/ObjectMacros.h"
 #include "UObject/UObjectBase.h"
 
 #if WITH_VERSE_VM || defined(__INTELLISENSE__)
@@ -149,12 +149,7 @@ struct
 		return bIChangedIt;
 	}
 
-	/**
-	 * Uses atomics to set the specified flag(s). GC internal version.
-	 * @param FlagToSet
-	 * @return True if this call set the flag, false if it has been set by another thread.
-	 */
-	FORCEINLINE bool ThisThreadAtomicallySetFlag_ForGC(EInternalObjectFlags FlagToSet)
+	FORCEINLINE bool ThisThreadAtomicallySetFlag(EInternalObjectFlags FlagToSet)
 	{
 		static_assert(sizeof(int32) == sizeof(Flags), "Flags must be 32-bit for atomics.");
 		bool bIChangedIt = false;
@@ -199,11 +194,6 @@ struct
 		return !!(GetFlagsInternal() & int32(InFlags));
 	}
 
-	FORCEINLINE bool HasAllFlags(EInternalObjectFlags InFlags) const
-	{
-		return (GetFlagsInternal() & int32(InFlags)) == int32(InFlags);
-	}
-
 	FORCEINLINE void SetUnreachable()
 	{
 		ThisThreadAtomicallySetFlag_ForGC(UE::GC::GUnreachableObjectFlag);
diff --git a/Engine/Source/Runtime/CoreUObject/Public/UObject/UObjectBase.h b/Engine/Source/Runtime/CoreUObject/Public/UObject/UObjectBase.h
--- a/Engine/Source/Runtime/CoreUObject/Public/UObject/UObjectBase.h
+++ b/Engine/Source/Runtime/CoreUObject/Public/UObject/UObjectBase.h
@@ -181,11 +181,6 @@ public:
 	/** Returns the external UPackage for this object, if any, NOT THREAD SAFE, used by internal gc reference collecting. */
 	COREUOBJECT_API UPackage* GetExternalPackageInternal() const;
 
-	/**
-	 * Marks the object as Reachable if it's currently marked as MaybeUnreachable by incremental GC.
-	*/
-	COREUOBJECT_API void MarkAsReachable() const;
-
 protected:
 	/**
 	 * Set the object flags directly
diff --git a/Engine/Source/Runtime/CoreUObject/Public/UObject/UObjectGlobals.h b/Engine/Source/Runtime/CoreUObject/Public/UObject/UObjectGlobals.h
--- a/Engine/Source/Runtime/CoreUObject/Public/UObject/UObjectGlobals.h
+++ b/Engine/Source/Runtime/CoreUObject/Public/UObject/UObjectGlobals.h
@@ -3888,6 +3888,13 @@ enum class EDataValidationResult : uint8
  */
 COREUOBJECT_API EDataValidationResult CombineDataValidationResults(EDataValidationResult Result1, EDataValidationResult Result2);
 
+
+namespace UE::GC::Private
+{
+	/** true if incremental reachability analysis is in progress (global for faster access in low level structs and functions otherwise use IsIncrementalReachabilityAnalisysPending()) */
+	extern COREUOBJECT_API bool GIsIncrementalReachabilityPending;
+}
+
 #if UE_ENABLE_INCLUDE_ORDER_DEPRECATED_IN_5_2
 #include "CoreMinimal.h"
 #endif
