[Backout] - Finishing backout of CL30233419
Original CL Desc
-----------------------------------------------------------------
Setting any of the root flags on UObjects while Incremental Reachability is running will now trigger GC barrier

Moved UE::GC::Private::GIsIncrementalReachabilityPending to GarbageCollectionGlobals.h to be able to access it from UObjectArray.h
Moved UObject::MarkAsReachable() to UObjectBase::MarkAsReachable() to be able to use it from FUObjectItem (UObjectArray.h)

diff --git a/Engine/Source/Runtime/CoreUObject/Public/UObject/GarbageCollectionGlobals.h b/Engine/Source/Runtime/CoreUObject/Public/UObject/GarbageCollectionGlobals.h
new file mode 100644
--- /dev/null
+++ b/Engine/Source/Runtime/CoreUObject/Public/UObject/GarbageCollectionGlobals.h
@@ -0,0 +1,21 @@
+// Copyright Epic Games, Inc. All Rights Reserved.
+
+/*=============================================================================
+	GarbageCollectionGlobals.h: Garbage Collection Global State Vars
+=============================================================================*/
+
+#pragma once
+
+#include "UObject/ObjectMacros.h"
+
+namespace UE::GC
+{
+	/** Current EInternalObjectFlags value representing an unreachable object */
+	extern COREUOBJECT_API EInternalObjectFlags GUnreachableObjectFlag;
+
+	/** Current EInternalObjectFlags value representing a maybe unreachable object */
+	extern COREUOBJECT_API EInternalObjectFlags GMaybeUnreachableObjectFlag;
+
+	/** true if incremental reachability analysis is in progress (global for faster access in low level structs and functions otherwise use IsIncrementalReachabilityAnalisysPending()) */
+	extern COREUOBJECT_API bool GIsIncrementalReachabilityPending;
+}
