Fix for null verse function of verse object using overridable serialization.

diff --git a/Engine/Source/Runtime/CoreUObject/Private/UObject/OverriddenPropertySet.cpp b/Engine/Source/Runtime/CoreUObject/Private/UObject/OverriddenPropertySet.cpp
--- a/Engine/Source/Runtime/CoreUObject/Private/UObject/OverriddenPropertySet.cpp
+++ b/Engine/Source/Runtime/CoreUObject/Private/UObject/OverriddenPropertySet.cpp
@@ -36,9 +36,13 @@ EOverriddenPropertyOperation FOverridableSerializationLogic::GetOverriddenProper
 	// It does not mean that if we have no record of an overriden operation that a subobject might have one, need to traverse all possible subobjects. 
 	if (const FProperty* CurrentProperty = Property ? Property : (CurrentPropertyChain ? CurrentPropertyChain->GetPropertyFromStack(0) : nullptr) )
 	{
+		if (CurrentProperty->HasAnyPropertyFlags(CPF_ExperimentalAlwaysOverriden))
+		{
+			return EOverriddenPropertyOperation::Replace;
+		}
 		// Here we should just use CurrentProperty->ContainsInstancedObjectProperty() but somehow this CPF_InstanceReference seems to always be there on verse classes
 		// This should probably be fix in verse at some point.
-		if (CurrentProperty->HasAnyPropertyFlags(CPF_PersistentInstance))
+		else if (CurrentProperty->HasAnyPropertyFlags(CPF_PersistentInstance))
 		{
 			return EOverriddenPropertyOperation::Modified;
 		}
diff --git a/Engine/Source/Runtime/CoreUObject/Public/UObject/ObjectMacros.h b/Engine/Source/Runtime/CoreUObject/Public/UObject/ObjectMacros.h
--- a/Engine/Source/Runtime/CoreUObject/Public/UObject/ObjectMacros.h
+++ b/Engine/Source/Runtime/CoreUObject/Public/UObject/ObjectMacros.h
@@ -452,7 +452,8 @@ enum EPropertyFlags : uint64
 	CPF_NativeAccessSpecifierPrivate	= 0x0040000000000000,	///< Private native access specifier
 	CPF_SkipSerialization				= 0x0080000000000000,	///< Property shouldn't be serialized, can still be exported to text
 	CPF_TObjectPtr						= 0x0100000000000000,	///< Property is a TObjectPtr<T> instead of a USomething*. Need to differentiate between TObjectclassOf and TObjectPtr
-	CPF_ExperimentalOverridableLogic	= 0x0200000000000000,	///< ****Experimental*** Property will use different logic to serialize knowing what changes are done against its default use the overridable information provided by the overidable manager on the object
+	CPF_ExperimentalOverridableLogic	= 0x0200000000000000,	///< ****Experimental*** Property will use different logic to serialize knowing what changes are done against its default use the overridable information provided by the overridable manager on the object
+	CPF_ExperimentalAlwaysOverriden		= 0x0400000000000000,	///< ****Experimental*** Property should never inherit from the parent when using overridable serialization
 };
 
 /** All Native Access Specifier flags */
