Extended ConvertFromType to indicate whether it serialized the value without conversion

diff --git a/Engine/Source/Runtime/CoreUObject/Private/UObject/Class.cpp b/Engine/Source/Runtime/CoreUObject/Private/UObject/Class.cpp
--- a/Engine/Source/Runtime/CoreUObject/Private/UObject/Class.cpp
+++ b/Engine/Source/Runtime/CoreUObject/Private/UObject/Class.cpp
@@ -1289,6 +1289,7 @@ void UStruct::LoadTaggedPropertiesFromText(FStructuredArchive::FSlot Slot, uint8
 						switch (Property->ConvertFromType(Tag, ItemSlot.GetValue(), Data, DefaultsStruct, Defaults))
 						{
 						case EConvertFromTypeResult::Converted:
+						case EConvertFromTypeResult::Serialized:
 							break;
 
 						case EConvertFromTypeResult::UseSerializeItem:
@@ -1656,19 +1657,17 @@ void UStruct::SerializeVersionedTaggedProperties(FStructuredArchive::FSlot Slot,
 						{
 							case EConvertFromTypeResult::Converted:
 								bAdvanceProperty = true;
-								// TODO: ConvertFromType always returns EConvertFromTypeResult::Converted for numeric types. this type check tries to fix
-								// that but we should find a more robust solution
-								if (Tag.Type != PropID)
+								if (FPropertyBag* PropertyBag = TryFindPropertyBag())
 								{
-									if (FPropertyBag* PropertyBag = TryFindPropertyBag())
-                                    {
-                                    	UnderlyingArchive.Seek(StartOfProperty);
-                                    	FStructuredArchive::FSlot CopySlot = PropertyRecord.EnterField(TEXT("Value"));
-                                    	Tag.Prop = nullptr;
-                                    	PropertyBag->LoadPropertyByTag(SerializeContext->SerializedPropertyPath, Tag, CopySlot, Property->ContainerPtrToValuePtrForDefaults<uint8>(DefaultsStruct, Defaults, Tag.ArrayIndex));
-                                    }
+									UnderlyingArchive.Seek(StartOfProperty);
+									FStructuredArchive::FSlot CopySlot = PropertyRecord.EnterField(TEXT("Value"));
+									Tag.Prop = nullptr;
+									PropertyBag->LoadPropertyByTag(SerializeContext->SerializedPropertyPath, Tag, CopySlot, Property->ContainerPtrToValuePtrForDefaults<uint8>(DefaultsStruct, Defaults, Tag.ArrayIndex));
 								}
-								
+								break;
+
+							case EConvertFromTypeResult::Serialized:
+								bAdvanceProperty = true;
 								break;
 
 							case EConvertFromTypeResult::UseSerializeItem:
diff --git a/Engine/Source/Runtime/CoreUObject/Private/UObject/PropertyArray.cpp b/Engine/Source/Runtime/CoreUObject/Private/UObject/PropertyArray.cpp
--- a/Engine/Source/Runtime/CoreUObject/Private/UObject/PropertyArray.cpp
+++ b/Engine/Source/Runtime/CoreUObject/Private/UObject/PropertyArray.cpp
@@ -199,6 +199,7 @@ void FArrayProperty::SerializeItem(FStructuredArchive::FSlot Slot, void* Value,
 			switch(StructProperty->ConvertFromType(SerializeFromMismatchedTag.GetValue(), Slot, Item, nullptr, nullptr))
 			{
 				case EConvertFromTypeResult::Converted:
+				case EConvertFromTypeResult::Serialized:
 					return;
 				case EConvertFromTypeResult::CannotConvert:
 					// FStructProperty::ConvertFromType doesn't handle setting the default, so do it here.
@@ -1170,11 +1171,13 @@ EConvertFromTypeResult FArrayProperty::ConvertFromType(const FPropertyTag& Tag,
 		{
 			FStructuredArchive::FStream ValueStream = Slot.EnterStream();
 
-			if (Inner->ConvertFromType(InnerPropertyTag, ValueStream.EnterElement(), ScriptArrayHelper.GetRawPtr(0), DefaultsStruct, nullptr) == EConvertFromTypeResult::Converted)
+			EConvertFromTypeResult ConvertResult = Inner->ConvertFromType(InnerPropertyTag, ValueStream.EnterElement(), ScriptArrayHelper.GetRawPtr(0), DefaultsStruct, nullptr);
+			if (ConvertResult == EConvertFromTypeResult::Converted || ConvertResult == EConvertFromTypeResult::Serialized)
 			{
 				for (int32 i = 1; i < ElementCount; ++i)
 				{
-					verify(Inner->ConvertFromType(InnerPropertyTag, ValueStream.EnterElement(), ScriptArrayHelper.GetRawPtr(i), DefaultsStruct, nullptr) == EConvertFromTypeResult::Converted);
+					ConvertResult = Inner->ConvertFromType(InnerPropertyTag, ValueStream.EnterElement(), ScriptArrayHelper.GetRawPtr(i), DefaultsStruct, nullptr);
+					check(ConvertResult == EConvertFromTypeResult::Converted || ConvertResult == EConvertFromTypeResult::Serialized);
 				}
 
 				return EConvertFromTypeResult::Converted;
diff --git a/Engine/Source/Runtime/CoreUObject/Private/UObject/PropertyMap.cpp b/Engine/Source/Runtime/CoreUObject/Private/UObject/PropertyMap.cpp
--- a/Engine/Source/Runtime/CoreUObject/Private/UObject/PropertyMap.cpp
+++ b/Engine/Source/Runtime/CoreUObject/Private/UObject/PropertyMap.cpp
@@ -1431,11 +1431,9 @@ EConvertFromTypeResult FMapProperty::ConvertFromType(const FPropertyTag& Tag, FS
 			CurrentType->SerializeItem(InnerSlot, DestAddress, nullptr);
 			return true;
 		}
-		else if( CurrentType->ConvertFromType(InTag, InnerSlot, InData, InDefaultsStruct, nullptr) == EConvertFromTypeResult::Converted )
-		{
-			return true;
-		}
-		return false;
+
+		EConvertFromTypeResult ConvertResult = CurrentType->ConvertFromType(InTag, InnerSlot, InData, InDefaultsStruct, nullptr);
+		return ConvertResult == EConvertFromTypeResult::Converted || ConvertResult == EConvertFromTypeResult::Serialized;
 	};
 
 	if (Tag.Type == NAME_MapProperty)
diff --git a/Engine/Source/Runtime/CoreUObject/Private/UObject/PropertyOptional.cpp b/Engine/Source/Runtime/CoreUObject/Private/UObject/PropertyOptional.cpp
--- a/Engine/Source/Runtime/CoreUObject/Private/UObject/PropertyOptional.cpp
+++ b/Engine/Source/Runtime/CoreUObject/Private/UObject/PropertyOptional.cpp
@@ -224,6 +224,7 @@ void FOptionalProperty::SerializeItem(FStructuredArchive::FSlot Slot, void* Data
 			switch (GetValueProperty()->ConvertFromType(ValueTag, ValueSlot, static_cast<uint8*>(ValueData), GetOwnerStruct(), static_cast<const uint8*>(ValueDefaults)))
 			{
 				case EConvertFromTypeResult::Converted:
+				case EConvertFromTypeResult::Serialized:
 					bSuccessfullyDeserialized = true;
 					break;
 				case EConvertFromTypeResult::UseSerializeItem:
diff --git a/Engine/Source/Runtime/CoreUObject/Private/UObject/PropertySet.cpp b/Engine/Source/Runtime/CoreUObject/Private/UObject/PropertySet.cpp
--- a/Engine/Source/Runtime/CoreUObject/Private/UObject/PropertySet.cpp
+++ b/Engine/Source/Runtime/CoreUObject/Private/UObject/PropertySet.cpp
@@ -926,7 +926,8 @@ EConvertFromTypeResult FSetProperty::ConvertFromType(const FPropertyTag& Tag, FS
 				TempElementStorage = (uint8*)FMemory::Malloc(SetLayout.Size);
 				ElementProp->InitializeValue(TempElementStorage);
 
-				if (ElementProp->ConvertFromType(InnerPropertyTag, ElementsToRemoveArray.EnterElement(), TempElementStorage, DefaultsStruct, nullptr) == EConvertFromTypeResult::Converted)
+				EConvertFromTypeResult ConvertResult = ElementProp->ConvertFromType(InnerPropertyTag, ElementsToRemoveArray.EnterElement(), TempElementStorage, DefaultsStruct, nullptr);
+				if (ConvertResult == EConvertFromTypeResult::Converted || ConvertResult == EConvertFromTypeResult::Serialized)
 				{
 					int32 Found = ScriptSetHelper.FindElementIndex(TempElementStorage);
 					if (Found != INDEX_NONE)
@@ -936,7 +937,8 @@ EConvertFromTypeResult FSetProperty::ConvertFromType(const FPropertyTag& Tag, FS
 
 					for (int32 I = 1; I < NumElementsToRemove; ++I)
 					{
-						verify(ElementProp->ConvertFromType(InnerPropertyTag, ElementsToRemoveArray.EnterElement(), TempElementStorage, DefaultsStruct, nullptr) == EConvertFromTypeResult::Converted);
+						ConvertResult = ElementProp->ConvertFromType(InnerPropertyTag, ElementsToRemoveArray.EnterElement(), TempElementStorage, DefaultsStruct, nullptr);
+						check(ConvertResult == EConvertFromTypeResult::Converted || ConvertResult == EConvertFromTypeResult::Serialized);
 
 						Found = ScriptSetHelper.FindElementIndex(TempElementStorage);
 						if (Found != INDEX_NONE)
@@ -967,7 +969,8 @@ EConvertFromTypeResult FSetProperty::ConvertFromType(const FPropertyTag& Tag, FS
 
 					// and read the first entry, we have to check for conversion possibility again because 
 					// NumElementsToRemove may not have run (in fact, it likely did not):
-					if (ElementProp->ConvertFromType(InnerPropertyTag, ElementsArray.EnterElement(), TempElementStorage, DefaultsStruct, nullptr) == EConvertFromTypeResult::Converted)
+					EConvertFromTypeResult ConvertResult = ElementProp->ConvertFromType(InnerPropertyTag, ElementsArray.EnterElement(), TempElementStorage, DefaultsStruct, nullptr);
+					if (ConvertResult == EConvertFromTypeResult::Converted || ConvertResult == EConvertFromTypeResult::Serialized)
 					{
 						if (ScriptSetHelper.FindElementIndex(TempElementStorage) == INDEX_NONE)
 						{
@@ -982,7 +985,8 @@ EConvertFromTypeResult FSetProperty::ConvertFromType(const FPropertyTag& Tag, FS
 						for (int32 I = 1; I < Num; ++I)
 						{
 							// Read key into temporary storage
-							verify(ElementProp->ConvertFromType(InnerPropertyTag, ElementsArray.EnterElement(), TempElementStorage, DefaultsStruct, nullptr) == EConvertFromTypeResult::Converted);
+							ConvertResult = ElementProp->ConvertFromType(InnerPropertyTag, ElementsArray.EnterElement(), TempElementStorage, DefaultsStruct, nullptr);
+							check(ConvertResult == EConvertFromTypeResult::Converted || ConvertResult == EConvertFromTypeResult::Serialized);
 
 							// Add a new entry if the element doesn't currently exist in the set
 							if (ScriptSetHelper.FindElementIndex(TempElementStorage) == INDEX_NONE)
diff --git a/Engine/Source/Runtime/CoreUObject/Public/UObject/UnrealType.h b/Engine/Source/Runtime/CoreUObject/Public/UObject/UnrealType.h
--- a/Engine/Source/Runtime/CoreUObject/Public/UObject/UnrealType.h
+++ b/Engine/Source/Runtime/CoreUObject/Public/UObject/UnrealType.h
@@ -124,9 +124,14 @@ namespace EExportedDeclaration
 
 enum class EConvertFromTypeResult
 {
+	/** No conversion was performed. Use SerializeItem to serialize the property value. */
 	UseSerializeItem,
+	/** No conversion was performed. The property value was serialized. Skip SerializeItem. */
+	Serialized,
+	/** No conversion is possible. Skip SerializeItem. */
 	CannotConvert,
-	Converted
+	/** Conversion of the property value was performed. Skip SerializeItem. */
+	Converted,
 };
 
 enum class EPropertyMemoryAccess : uint8
@@ -368,8 +373,9 @@ public:
 	* @param	Defaults	if available, a pointer to the container containing the default value for this property, or null
 	*
 	* @return	A state which tells the tagged property system how the property dealt with the data.
-	*			Converted:        the function has handled the tag.
+	*			Converted:        the function handled conversion.
 	*			CannotConvert:    the tag is not something that the property can convert.
+	*			Serialized:       the function handled serialization without conversion.
 	*			UseSerializeItem: no conversion was done on the property - this can mean that the tag is correct and normal serialization applies or that the tag is incompatible.
 	*/
 	COREUOBJECT_API virtual EConvertFromTypeResult ConvertFromType(const FPropertyTag& Tag, FStructuredArchive::FSlot Slot, uint8* Data, UStruct* DefaultsStruct, const uint8* Defaults);
@@ -1746,37 +1752,38 @@ private:
 public:
 	virtual EConvertFromTypeResult ConvertFromType(const FPropertyTag& Tag, FStructuredArchive::FSlot Slot, uint8* Data, UStruct* DefaultsStruct, const uint8* Defaults) override
 	{
-		if (const EName * TagType = Tag.Type.ToEName())
+		if (const EName* TagType = Tag.Type.ToEName())
 		{
 			switch (*TagType)
 			{
 			case NAME_Int8Property:
 				ConvertFromArithmeticValue<int8>(Slot, Data, Tag);
-				return EConvertFromTypeResult::Converted;
+				return std::is_same_v<TCppType, int8> ? EConvertFromTypeResult::Serialized : EConvertFromTypeResult::Converted;
 
 			case NAME_Int16Property:
 				ConvertFromArithmeticValue<int16>(Slot, Data, Tag);
-				return EConvertFromTypeResult::Converted;
+				return std::is_same_v<TCppType, int16> ? EConvertFromTypeResult::Serialized : EConvertFromTypeResult::Converted;
 
 			case NAME_IntProperty:
 				ConvertFromArithmeticValue<int32>(Slot, Data, Tag);
-				return EConvertFromTypeResult::Converted;
+				return std::is_same_v<TCppType, int32> ? EConvertFromTypeResult::Serialized : EConvertFromTypeResult::Converted;
 
 			case NAME_Int64Property:
 				ConvertFromArithmeticValue<int64>(Slot, Data, Tag);
-				return EConvertFromTypeResult::Converted;
+				return std::is_same_v<TCppType, int64> ? EConvertFromTypeResult::Serialized : EConvertFromTypeResult::Converted;
 
 			case NAME_ByteProperty:
 				if (!Tag.EnumName.IsNone())
 				{
 					int64 PreviousValue = this->ReadEnumAsInt64(Slot, DefaultsStruct, Tag);
 					this->SetPropertyValue_InContainer(Data, (TCppType)PreviousValue, Tag.ArrayIndex);
+					return EConvertFromTypeResult::Converted;
 				}
 				else
 				{
 					ConvertFromArithmeticValue<int8>(Slot, Data, Tag);
+					return std::is_same_v<TCppType, uint8> ? EConvertFromTypeResult::Serialized : EConvertFromTypeResult::Converted;
 				}
-				return EConvertFromTypeResult::Converted;
 
 			case NAME_EnumProperty:
 			{
@@ -1787,27 +1794,27 @@ public:
 
 			case NAME_UInt16Property:
 				ConvertFromArithmeticValue<uint16>(Slot, Data, Tag);
-				return EConvertFromTypeResult::Converted;
+				return std::is_same_v<TCppType, uint16> ? EConvertFromTypeResult::Serialized : EConvertFromTypeResult::Converted;
 
 			case NAME_UInt32Property:
 				ConvertFromArithmeticValue<uint32>(Slot, Data, Tag);
-				return EConvertFromTypeResult::Converted;
+				return std::is_same_v<TCppType, uint32> ? EConvertFromTypeResult::Serialized : EConvertFromTypeResult::Converted;
 
 			case NAME_UInt64Property:
 				ConvertFromArithmeticValue<uint64>(Slot, Data, Tag);
-				return EConvertFromTypeResult::Converted;
+				return std::is_same_v<TCppType, uint64> ? EConvertFromTypeResult::Serialized : EConvertFromTypeResult::Converted;
 
 			case NAME_FloatProperty:
 				ConvertFromArithmeticValue<float>(Slot, Data, Tag);
-				return EConvertFromTypeResult::Converted;
+				return std::is_same_v<TCppType, float> ? EConvertFromTypeResult::Serialized : EConvertFromTypeResult::Converted;
 
 			case NAME_DoubleProperty:
 				ConvertFromArithmeticValue<double>(Slot, Data, Tag);
-				return EConvertFromTypeResult::Converted;
+				return std::is_same_v<TCppType, double> ? EConvertFromTypeResult::Serialized : EConvertFromTypeResult::Converted;
 
 			case NAME_BoolProperty:
 				this->SetPropertyValue_InContainer(Data, (TCppType)Tag.BoolVal, Tag.ArrayIndex);
-				return EConvertFromTypeResult::Converted;
+				return std::is_same_v<TCppType, bool> ? EConvertFromTypeResult::Serialized : EConvertFromTypeResult::Converted;
 
 			default:
 				// We didn't convert it
