FHarvestedRealm::NameExists changed to consume const FName& rather than FNameEntryId as it leaked the internal storage details for export names of the realm into the API.

Since we only store display names in the realm, we can directly compare FNames using display names when verifying we saved correctly.

114ms -> 3ms for SaveContext_NameExists over 71,680 instances for a QAGame cook

diff --git a/Engine/Source/Runtime/CoreUObject/Private/UObject/SavePackage/SaveContext.h b/Engine/Source/Runtime/CoreUObject/Private/UObject/SavePackage/SaveContext.h
--- a/Engine/Source/Runtime/CoreUObject/Private/UObject/SavePackage/SaveContext.h
+++ b/Engine/Source/Runtime/CoreUObject/Private/UObject/SavePackage/SaveContext.h
@@ -269,23 +269,14 @@ struct FHarvestedRealm
 		return ExportNativeObjectDependencies;
 	}
 
-	bool NameExists(FNameEntryId ComparisonId) const
+	bool NameExists(const FName& Name) const
 	{
-		for (FNameEntryId DisplayId : NamesReferencedFromExportData)
-		{
-			if (FName::GetComparisonIdFromDisplayId(DisplayId) == ComparisonId)
-			{
-				return true;
-			}
-		}
-		for (FNameEntryId DisplayId : NamesReferencedFromPackageHeader)
-		{
-			if (FName::GetComparisonIdFromDisplayId(DisplayId) == ComparisonId)
-			{
-				return true;
-			}
-		}
-		return false;
+		// Normally FName comparisons would be case insensitive and done using the comparisonIndex however 
+		// NamesReferencedFromExportData and NamesReferencedFromPackageHeader contain DisplayIndices and the passed in 'Name' 
+		// comes from memory (rather than disk) where we will not have case-sensitive descrepancies. Using the more restrictive 
+		// case-sensitive search in this case is valid and faster.
+		const FNameEntryId DisplayId = Name.GetDisplayIndex();
+		return NamesReferencedFromExportData.Find(DisplayId) || NamesReferencedFromPackageHeader.Find(DisplayId);
 	}
 
 	FLinkerSave* GetLinker() const
@@ -877,9 +868,9 @@ public:
 		return GetHarvestedRealm().GetNativeObjectDependencies();
 	}
 
-	bool NameExists(FNameEntryId ComparisonId) const
+	bool NameExists(const FName& Name) const
 	{
-		return GetHarvestedRealm().NameExists(ComparisonId);
+		return GetHarvestedRealm().NameExists(Name);
 	}
 
 	const FCustomVersionContainer& GetCustomVersions() const
diff --git a/Engine/Source/Runtime/CoreUObject/Private/UObject/SavePackage2.cpp b/Engine/Source/Runtime/CoreUObject/Private/UObject/SavePackage2.cpp
--- a/Engine/Source/Runtime/CoreUObject/Private/UObject/SavePackage2.cpp
+++ b/Engine/Source/Runtime/CoreUObject/Private/UObject/SavePackage2.cpp
@@ -455,7 +455,7 @@ ESavePackageResult HarvestPackage(FSaveContext& SaveContext)
 	return ReturnSuccessOrCancel();
 }
 
-static FNameEntryId NAME_UniqueObjectNameForCookingComparisonIndex = FName("UniqueObjectNameForCooking").GetComparisonIndex();
+static const FName NAME_UniqueObjectNameForCooking = FName("UniqueObjectNameForCooking");
 
 ESavePackageResult ValidateRealms(FSaveContext& SaveContext)
 {
@@ -630,16 +630,18 @@ ESavePackageResult ValidateExports(FSaveContext& SaveContext)
 	}
 
 	// Cooking checks
-#if WITH_EDITOR
+#if WITH_EDITOR && !UE_FNAME_OUTLINE_NUMBER
 	if (GOutputCookingWarnings)
 	{
-		// check the name list for UniqueObjectNameForCooking cooking
-		if (SaveContext.NameExists(NAME_UniqueObjectNameForCookingComparisonIndex))
+		// check the name list for UniqueObjectNameForCooking cooking. Since check is a fast check for any NAME_UniqueObjectNameForCooking
+		// we can only perform this check when UE_FNAME_OUTLINE_NUMBER=0 as otherwise all prefixed names will have a unique ComparisonIndex
+		if (SaveContext.NameExists(NAME_UniqueObjectNameForCooking))
 		{
+			const FNameEntryId UniqueObjectNameForCookingComparisonId = NAME_UniqueObjectNameForCooking.GetComparisonIndex();
 			for (const FTaggedExport& Export : SaveContext.GetExports())
 			{
-				FName NameInUse = Export.Obj->GetFName();
-				if (NameInUse.GetComparisonIndex() == NAME_UniqueObjectNameForCookingComparisonIndex)
+				const FName NameInUse = Export.Obj->GetFName();
+				if (NameInUse.GetComparisonIndex() == UniqueObjectNameForCookingComparisonId)
 				{
 					UObject* Outer = Export.Obj->GetOuter();
 					UE_LOG(LogSavePackage, Warning, TEXT("Saving object into cooked package %s which was created at cook time, Object Name %s, Full Path %s, Class %s, Outer %s, Outer class %s"), SaveContext.GetFilename(), *NameInUse.ToString(), *Export.Obj->GetFullName(), *Export.Obj->GetClass()->GetName(), Outer ? *Outer->GetName() : TEXT("None"), Outer ? *Outer->GetClass()->GetName() : TEXT("None"));
@@ -812,13 +814,13 @@ ESavePackageResult ValidateImports(FSaveContext& SaveContext)
 	{
 		TObjectPtr<UPackage> ImportPackage = Import.GetPackage();
 		// All names should be properly harvested at this point
-		ensureAlwaysMsgf(SaveContext.NameExists(Import.GetFName().GetComparisonIndex())
+		ensureAlwaysMsgf(SaveContext.NameExists(Import.GetFName())
 			, TEXT("Missing import name %s while saving package %s. Did you rename an import during serialization?"), *Import->GetName(), *PackageName);
-		ensureAlwaysMsgf(SaveContext.NameExists(ImportPackage.GetFName().GetComparisonIndex())
+		ensureAlwaysMsgf(SaveContext.NameExists(ImportPackage.GetFName())
 			, TEXT("Missing import package name %s while saving package %s. Did you rename an import during serialization?"), *ImportPackage->GetName(), *PackageName);
-		ensureAlwaysMsgf(SaveContext.NameExists(Import.GetClass()->GetFName().GetComparisonIndex())
+		ensureAlwaysMsgf(SaveContext.NameExists(Import.GetClass()->GetFName())
 			, TEXT("Missing import class name %s while saving package %s"), *Import->GetClass()->GetName(), *PackageName);
-		ensureAlwaysMsgf(SaveContext.NameExists(Import.GetClass()->GetOuter()->GetFName().GetComparisonIndex())
+		ensureAlwaysMsgf(SaveContext.NameExists(Import.GetClass()->GetOuter()->GetFName())
 			, TEXT("Missing import class package name %s while saving package %s"), *Import.GetClass()->GetOuter()->GetName(), *PackageName);
 
 		// if the import is marked as a prestream package, we dont need to validate further
