[ZenLoader]
  - Fix duplication of PIE world assets under many different in-memory only package which is not supposed to happen

diff --git a/Engine/Source/Runtime/CoreUObject/Private/Serialization/AsyncLoading2.cpp b/Engine/Source/Runtime/CoreUObject/Private/Serialization/AsyncLoading2.cpp
--- a/Engine/Source/Runtime/CoreUObject/Private/Serialization/AsyncLoading2.cpp
+++ b/Engine/Source/Runtime/CoreUObject/Private/Serialization/AsyncLoading2.cpp
@@ -5079,7 +5079,7 @@ void FAsyncPackage2::ImportPackagesRecursiveInner(FAsyncLoadingThreadState2& Thr
 		bool bInserted = false;
 		bool bIsFullyLoaded = ImportedPackageRef.AreAllPublicExportsLoaded();
 #if ALT2_ENABLE_LINKERLOAD_SUPPORT
-		if (!bIsZenPackageImport && (!ImportedPackageRef.HasPackage() || !ImportedPackageRef.GetPackage()->GetLinker()))
+		if (!bIsZenPackageImport && (!ImportedPackageRef.HasPackage() || (!ImportedPackageRef.GetPackage()->GetLinker() && !ImportedPackageRef.GetPackage()->HasAnyPackageFlags(PKG_InMemoryOnly))))
 		{
 			// If we're importing a linker load package and it doesn't have its linker we need to reload it, otherwise we can't reliably link to its imports
 			// Note: Legacy loader path appears to do this only for uncooked packages in the editor?
