Fixed FrankenGC race condition where UObject references would attempt to be marked prior to references being cleared.  This would result in the UObject reference being stomped and the object collected.

#rnx

diff --git a/Engine/Source/Runtime/CoreUObject/Private/UObject/GarbageCollection.cpp b/Engine/Source/Runtime/CoreUObject/Private/UObject/GarbageCollection.cpp
--- a/Engine/Source/Runtime/CoreUObject/Private/UObject/GarbageCollection.cpp
+++ b/Engine/Source/Runtime/CoreUObject/Private/UObject/GarbageCollection.cpp
@@ -113,15 +113,6 @@ static FGCTimingInfo GTimingInfo;
 
 bool GIsGarbageCollecting = false;
 
-#if WITH_VERSE_VM || defined(__INTELLISENSE__)
-namespace UE::GC
-{
-bool GIsFrankenGCCollecting = false;
-}
-static bool bFrankenGCEnabled = false;
-static Verse::FCollectionCycleRequest VerseCycleRequest;
-#endif
-
 /**
 * Call back into the async loading code to inform of the destruction of serialized objects
 */
@@ -3755,6 +3746,75 @@ void FReferenceFinder::HandleObjectReference( UObject*& InObject, const UObject*
 
 //////////////////////////////////////////////////////////////////////////
 
+namespace UE::GC
+{
+#if WITH_VERSE_VM || defined(__INTELLISENSE__)
+bool GIsFrankenGCCollecting = false;
+static bool bFrankenGCEnabled = false;
+static Verse::FCollectionCycleRequest VerseCycleRequest;
+
+static bool UpdateFrankenGCMode()
+{
+	bool bNewState = GEnableFrankenGC;
+	if (bFrankenGCEnabled != bNewState)
+	{
+		bFrankenGCEnabled = bNewState;
+		if (bFrankenGCEnabled)
+		{
+			Verse::FIOContext::Create([](Verse::FIOContext Context) {
+				Verse::FHeap::EnableExternalControl(Context);
+				});
+		}
+		else
+		{
+			Verse::FHeap::DisableExternalControl();
+		}
+	}
+	return bFrankenGCEnabled;
+}
+
+void EnableFrankenGCMode(bool bEnable)
+{
+	GEnableFrankenGC = bEnable;
+	UpdateFrankenGCMode();
+}
+
+static FORCEINLINE void StartVerseGC()
+{
+	GIsFrankenGCCollecting = UpdateFrankenGCMode();
+	if (GIsFrankenGCCollecting)
+	{
+		Verse::FIOContext::Create([](Verse::FIOContext Context) { 
+			Verse::FHeap::ExternallySynchronouslyStartGC(Context);
+			VerseCycleRequest = Verse::FHeap::StartCollectingIfNotCollecting();
+		});
+	}
+}
+
+static FORCEINLINE void StopVerseGC()
+{
+	if (GIsFrankenGCCollecting)
+	{
+		GIsFrankenGCCollecting = false;
+		Verse::FIOContext::Create([](Verse::FIOContext Context) { 
+			Verse::FHeap::ExternallySynchronouslyTerminateGC(Context);
+			VerseCycleRequest.Wait(Context);
+		});
+	}
+}
+#else
+static FORCEINLINE void StartVerseGC()
+{
+}
+
+static FORCEINLINE void StopVerseGC()
+{
+}
+#endif
+} // namespace UE::GC
+
+//////////////////////////////////////////////////////////////////////////
+
 namespace UE::GC
 {
 
@@ -4183,6 +4243,8 @@ public:
 		if (!GReachabilityState.IsSuspended())
 		{
 			StartReachabilityAnalysis(KeepFlags, Options);
+			// We start verse GC here so that the objects are unmarked prior to verse marking them
+			StartVerseGC();
 		}
 
 		{
@@ -5117,89 +5179,6 @@ EGatherOptions GetObjectGatherOptions()
 	return (ShouldForceSingleThreadedGC() ? EGatherOptions::None : EGatherOptions::Parallel);
 }
 
-#if WITH_VERSE_VM || defined(__INTELLISENSE__)
-static bool UpdateFrankenGCMode()
-{
-	bool bNewState = GEnableFrankenGC;
-	if (bFrankenGCEnabled != bNewState)
-	{
-		bFrankenGCEnabled = bNewState;
-		if (bFrankenGCEnabled)
-		{
-			Verse::FIOContext::Create([](Verse::FIOContext Context) {
-				Verse::FHeap::EnableExternalControl(Context);
-				});
-		}
-		else
-		{
-			Verse::FHeap::DisableExternalControl();
-		}
-	}
-	return bFrankenGCEnabled;
-}
-
-void EnableFrankenGCMode(bool bEnable)
-{
-	GEnableFrankenGC = bEnable;
-	UpdateFrankenGCMode();
-}
-
-static FORCEINLINE void StartVerseGCBody(Verse::FIOContext Context)
-{
-	Verse::FHeap::ExternallySynchronouslyStartGC(Context);
-	VerseCycleRequest = Verse::FHeap::StartCollectingIfNotCollecting();
-}
-
-static FORCEINLINE void StartVerseGC()
-{
-	GIsFrankenGCCollecting = UpdateFrankenGCMode();
-	if (GIsFrankenGCCollecting)
-	{
-		Verse::FIOContext::Create([](Verse::FIOContext Context) { StartVerseGCBody(Context); });
-	}
-}
-
-static FORCEINLINE void StopVerseGCBody(Verse::FIOContext Context)
-{
-	Verse::FHeap::ExternallySynchronouslyTerminateGC(Context);
-	VerseCycleRequest.Wait(Context);
-}
-
-static FORCEINLINE void StopVerseGC()
-{
-	if (GIsFrankenGCCollecting)
-	{
-		GIsFrankenGCCollecting = false;
-		Verse::FIOContext::Create([](Verse::FIOContext Context) { StopVerseGCBody(Context); });
-	}
-}
-
-static FORCEINLINE void RestartVerseGC()
-{
-	// If we are going to scan again, we need to cycle verse GC so the cells are unmarked.
-	if (GIsFrankenGCCollecting)
-	{
-		Verse::FIOContext::Create([](Verse::FIOContext Context) {
-			StopVerseGCBody(Context);
-			StartVerseGCBody(Context);
-			});
-	}
-}
-#else
-static FORCEINLINE void StartVerseGC()
-{
-}
-
-static FORCEINLINE void StopVerseGC()
-{
-
-}
-
-static FORCEINLINE void RestartVerseGC()
-{
-}
-#endif
-
 /** 
  * Deletes all unreferenced objects, keeping objects that have any of the passed in KeepFlags set
  *
@@ -5274,8 +5253,6 @@ void PreCollectGarbageImpl(EObjectFlags KeepFlags)
 			check(!GObjIncrementalPurgeIsInProgress);
 			check(!GObjPurgeIsRequired);
 
-			StartVerseGC();
-
 			// This can happen if someone disables clusters from the console (gc.CreateGCClusters)
 			if (!GCreateGCClusters && GUObjectClusters.GetNumAllocatedClusters())
 			{
@@ -5551,7 +5528,7 @@ void FReachabilityAnalysisState::PerformReachabilityAnalysisAndConditionallyPurg
 		const double StartTime = FPlatformTime::Seconds();
 		{
 			// If we are going to scan again, we need to cycle verse GC so the cells are unmarked.
-			RestartVerseGC();
+			StopVerseGC();
 			TGuardValue<bool> GuardReachabilityUsingTimeLimit(bReachabilityUsingTimeLimit, false);
 			FRealtimeGC GC;
 			GC.Stats = Stats; // This is to pass Stats.bFoundGarbageRef to CG
