Updated calls to GetSerializeContext on archives to use the thread context directly

Archives do not always set the context when it is needed. Archives that do set the context set it to the thread context anyway.

#rnx

diff --git a/Engine/Source/Runtime/CoreUObject/Private/Serialization/ArchiveStackTrace.cpp b/Engine/Source/Runtime/CoreUObject/Private/Serialization/ArchiveStackTrace.cpp
--- a/Engine/Source/Runtime/CoreUObject/Private/Serialization/ArchiveStackTrace.cpp
+++ b/Engine/Source/Runtime/CoreUObject/Private/Serialization/ArchiveStackTrace.cpp
@@ -1625,14 +1625,14 @@ static void DumpPackageHeaderDiffs_LinkerLoad(
 		TRefCountPtr<FUObjectSerializeContext> LinkerLoadContext(FUObjectThreadContext::Get().GetSerializeContext());
 		BeginLoad(LinkerLoadContext);
 		SourceLinker = FArchiveStackTraceWriter::CreateLinkerForPackage(LinkerLoadContext, SourceAssetPackageName, AssetFilename, SourcePackage);
-		EndLoad(SourceLinker ? SourceLinker->GetSerializeContext() : LinkerLoadContext.GetReference());
+		EndLoad(LinkerLoadContext);
 	}
-	
+
 	{
 		TRefCountPtr<FUObjectSerializeContext> LinkerLoadContext(FUObjectThreadContext::Get().GetSerializeContext());
 		BeginLoad(LinkerLoadContext);
 		DestLinker = FArchiveStackTraceWriter::CreateLinkerForPackage(LinkerLoadContext, DestAssetPackageName, AssetFilename, DestPackage);
-		EndLoad(DestLinker ? DestLinker->GetSerializeContext() : LinkerLoadContext.GetReference());
+		EndLoad(LinkerLoadContext);
 	}
 
 	if (SourceLinker && DestLinker)
diff --git a/Engine/Source/Runtime/CoreUObject/Private/UObject/Enum.cpp b/Engine/Source/Runtime/CoreUObject/Private/UObject/Enum.cpp
--- a/Engine/Source/Runtime/CoreUObject/Private/UObject/Enum.cpp
+++ b/Engine/Source/Runtime/CoreUObject/Private/UObject/Enum.cpp
@@ -716,12 +716,9 @@ int32 UEnum::GetIndexByNameString(const FString& InSearchString, EGetByNameFlags
 	{
 		// None is passed in by blueprints at various points, isn't an error. Any other failed resolve should be fixed
 		UObject* SerializedObject = nullptr;
-		if (FLinkerLoad* Linker = GetLinker())
+		if (FUObjectSerializeContext* LoadContext = FUObjectThreadContext::Get().GetSerializeContext())
 		{
-			if (FUObjectSerializeContext* LoadContext = Linker->GetSerializeContext())
-			{
-				SerializedObject = LoadContext->SerializedObject;
-			}
+			SerializedObject = LoadContext->SerializedObject;
 		}
 		const bool bIsNativeOrLoaded = (!HasAnyFlags(RF_WasLoaded) || HasAnyFlags(RF_LoadCompleted));
 		UE_LOG(LogEnum, Warning, TEXT("UEnum: In asset '%s', there is an enum property of type '%s' with an invalid value of '%s' - %s - %d"), 
diff --git a/Engine/Source/Runtime/CoreUObject/Private/UObject/EnumProperty.cpp b/Engine/Source/Runtime/CoreUObject/Private/UObject/EnumProperty.cpp
--- a/Engine/Source/Runtime/CoreUObject/Private/UObject/EnumProperty.cpp
+++ b/Engine/Source/Runtime/CoreUObject/Private/UObject/EnumProperty.cpp
@@ -350,12 +350,9 @@ const TCHAR* FEnumProperty::ImportText_Internal(const TCHAR* InBuffer, void* Con
 			// return null so that the caller of ImportText can generate a more meaningful
 			// warning/error
 			UObject* SerializedObject = nullptr;
-			if (FLinkerLoad* Linker = GetLinker())
+			if (FUObjectSerializeContext* LoadContext = FUObjectThreadContext::Get().GetSerializeContext())
 			{
-				if (FUObjectSerializeContext* LoadContext = Linker->GetSerializeContext())
-				{
-					SerializedObject = LoadContext->SerializedObject;
-				}
+				SerializedObject = LoadContext->SerializedObject;
 			}
 			const bool bIsNativeOrLoaded = (!Enum->HasAnyFlags(RF_WasLoaded) || Enum->HasAnyFlags(RF_LoadCompleted));
 			ErrorText->Logf(ELogVerbosity::Warning, TEXT("FEP: In asset '%s', there is an enum property of type '%s' with an invalid value of '%s' - %s"), 
diff --git a/Engine/Source/Runtime/CoreUObject/Private/UObject/PropertyBaseObject.cpp b/Engine/Source/Runtime/CoreUObject/Private/UObject/PropertyBaseObject.cpp
--- a/Engine/Source/Runtime/CoreUObject/Private/UObject/PropertyBaseObject.cpp
+++ b/Engine/Source/Runtime/CoreUObject/Private/UObject/PropertyBaseObject.cpp
@@ -384,9 +384,8 @@ const TCHAR* FObjectPropertyBase::ImportText_Internal( const TCHAR* InBuffer, vo
 {
 	const TCHAR* Buffer = InBuffer;
 	TObjectPtr<UObject> Result = nullptr;
-	FLinkerLoad* Linker = GetLinker();
 
-	bool bOk = ParseObjectPropertyValue(this, Parent, PropertyClass, PortFlags, Buffer, Result, Linker ? Linker->GetSerializeContext() : nullptr);
+	bool bOk = ParseObjectPropertyValue(this, Parent, PropertyClass, PortFlags, Buffer, Result, FUObjectThreadContext::Get().GetSerializeContext());
 
 	if (Result && (PortFlags & PPF_InstanceSubobjects) != 0 && HasAnyPropertyFlags(CPF_InstancedReference))
 	{
diff --git a/Engine/Source/Runtime/CoreUObject/Private/UObject/PropertyByte.cpp b/Engine/Source/Runtime/CoreUObject/Private/UObject/PropertyByte.cpp
--- a/Engine/Source/Runtime/CoreUObject/Private/UObject/PropertyByte.cpp
+++ b/Engine/Source/Runtime/CoreUObject/Private/UObject/PropertyByte.cpp
@@ -442,12 +442,9 @@ const TCHAR* FByteProperty::ImportText_Internal( const TCHAR* InBuffer, void* Co
 			// return null so that the caller of ImportText can generate a more meaningful
 			// warning/error
 			UObject* SerializedObject = nullptr;
-			if (FLinkerLoad* Linker = GetLinker())
+			if (FUObjectSerializeContext* LoadContext = FUObjectThreadContext::Get().GetSerializeContext())
 			{
-				if (FUObjectSerializeContext* LoadContext = Linker->GetSerializeContext())
-				{
-					SerializedObject = LoadContext->SerializedObject;
-				}
+				SerializedObject = LoadContext->SerializedObject;
 			}
 			const bool bIsNativeOrLoaded = (!Enum->HasAnyFlags(RF_WasLoaded) || Enum->HasAnyFlags(RF_LoadCompleted));
 			ErrorText->Logf(ELogVerbosity::Warning, TEXT("FBP: In asset '%s', there is an enum property of type '%s' with an invalid value of '%s' - %s"), 
diff --git a/Engine/Source/Runtime/CoreUObject/Private/UObject/PropertyMap.cpp b/Engine/Source/Runtime/CoreUObject/Private/UObject/PropertyMap.cpp
--- a/Engine/Source/Runtime/CoreUObject/Private/UObject/PropertyMap.cpp
+++ b/Engine/Source/Runtime/CoreUObject/Private/UObject/PropertyMap.cpp
@@ -674,7 +674,7 @@ void FMapProperty::SerializeItem(FStructuredArchive::FSlot Slot, void* Value, co
 
 	if (UnderlyingArchive.IsLoading())
 	{
-		FUObjectSerializeContext* Context = UnderlyingArchive.GetSerializeContext();
+		FUObjectSerializeContext* Context = FUObjectThreadContext::Get().GetSerializeContext();
 
 		// Delete any explicitly-removed elements
 		int32 NumKeysToRemove = 0;
diff --git a/Engine/Source/Runtime/CoreUObject/Private/UObject/PropertySet.cpp b/Engine/Source/Runtime/CoreUObject/Private/UObject/PropertySet.cpp
--- a/Engine/Source/Runtime/CoreUObject/Private/UObject/PropertySet.cpp
+++ b/Engine/Source/Runtime/CoreUObject/Private/UObject/PropertySet.cpp
@@ -261,7 +261,7 @@ void FSetProperty::SerializeItem(FStructuredArchive::FSlot Slot, void* Value, co
 
 	if (UnderlyingArchive.IsLoading())
 	{
-		FUObjectSerializeContext* Context = UnderlyingArchive.GetSerializeContext();
+		FUObjectSerializeContext* Context = FUObjectThreadContext::Get().GetSerializeContext();
 
 		if (Defaults)
 		{
diff --git a/Engine/Source/Runtime/CoreUObject/Private/UObject/SoftObjectPath.cpp b/Engine/Source/Runtime/CoreUObject/Private/UObject/SoftObjectPath.cpp
--- a/Engine/Source/Runtime/CoreUObject/Private/UObject/SoftObjectPath.cpp
+++ b/Engine/Source/Runtime/CoreUObject/Private/UObject/SoftObjectPath.cpp
@@ -776,7 +776,7 @@ bool FSoftObjectPathThreadContext::GetSerializationOptions(FName& OutPackageName
 	}
 	
 	// Check UObject serialize context as a backup
-	FUObjectSerializeContext* LoadContext = Archive ? Archive->GetSerializeContext() : nullptr;
+	FUObjectSerializeContext* LoadContext = FUObjectThreadContext::Get().GetSerializeContext();
 	if (LoadContext && LoadContext->SerializedObject)
 	{
 		FLinkerLoad* Linker = LoadContext->SerializedObject->GetLinker();
diff --git a/Engine/Source/Runtime/CoreUObject/Private/UObject/UObjectGlobals.cpp b/Engine/Source/Runtime/CoreUObject/Private/UObject/UObjectGlobals.cpp
--- a/Engine/Source/Runtime/CoreUObject/Private/UObject/UObjectGlobals.cpp
+++ b/Engine/Source/Runtime/CoreUObject/Private/UObject/UObjectGlobals.cpp
@@ -1805,7 +1805,7 @@ UPackage* LoadPackageInternal(UPackage* InOuter, const FPackagePath& PackagePath
 
 		auto EndLoadAndCopyLocalizationGatherFlag = [&]
 		{
-			EndLoad(Linker->GetSerializeContext(), &LoadedPackages);
+			EndLoad(LoadContext, &LoadedPackages);
 			// Set package-requires-localization flags from archive after loading. This reinforces flagging of packages that haven't yet been resaved.
 			Result->ThisRequiresLocalizationGather(Linker->RequiresLocalizationGather());
 		};
@@ -1889,7 +1889,7 @@ UPackage* LoadPackageInternal(UPackage* InOuter, const FPackagePath& PackagePath
 			// @todo: remove me when loading can be self-contained (and EndLoad doesn't check for IsInAsyncLoadingThread) or there's just one loading path
 			// If we start a non-async loading during async loading and the serialization context is not associated with any other package and
 			// doesn't come from an async package, queue this package to be async loaded, otherwise we'll end up not loading its exports
-			if (!Linker->AsyncRoot && Linker->GetSerializeContext()->GetBeginLoadCount() == 1 && IsInAsyncLoadingThread())
+			if (!Linker->AsyncRoot && LoadContext->GetBeginLoadCount() == 1 && IsInAsyncLoadingThread())
 			{
 				LoadPackageAsync(Linker->LinkerRoot->GetName());
 			}
@@ -1962,14 +1962,12 @@ UPackage* LoadPackageInternal(UPackage* InOuter, const FPackagePath& PackagePath
 				// Async loading removes delayed linkers on the game thread after streaming has finished
 				else
 				{
-					check(Linker->GetSerializeContext());
-					Linker->GetSerializeContext()->AddDelayedLinkerClosePackage(Linker);
+					LoadContext->AddDelayedLinkerClosePackage(Linker);
 				}
 			}
 			else
 			{
-				check(Linker->GetSerializeContext());
-				Linker->GetSerializeContext()->AddDelayedLinkerClosePackage(Linker);
+				LoadContext->AddDelayedLinkerClosePackage(Linker);
 			}
 		}
 	}
