Initial support for replicating optional properties. Tested and works with most inner types like int, float, logic, string and object. Doesn't work yet with array, map, set etc

diff --git a/Engine/Source/Runtime/CoreUObject/Private/UObject/PropertyOptional.cpp b/Engine/Source/Runtime/CoreUObject/Private/UObject/PropertyOptional.cpp
--- a/Engine/Source/Runtime/CoreUObject/Private/UObject/PropertyOptional.cpp
+++ b/Engine/Source/Runtime/CoreUObject/Private/UObject/PropertyOptional.cpp
@@ -290,8 +290,33 @@ void FOptionalProperty::SerializeItem(FStructuredArchive::FSlot Slot, void* Data
 
 bool FOptionalProperty::NetSerializeItem(FArchive& Ar, UPackageMap* Map, void* Data, TArray<uint8>* MetaData /*= nullptr*/) const
 {
-	UE_LOG(LogOptionalProperty, Fatal, TEXT("Unsupported code path"));
-	return true;
+	bool bSuccess = true;
+	bool bIsLoading = Ar.IsLoading();
+
+	uint8 IsSetValue = IsSet(Data);
+	Ar.SerializeBits(&IsSetValue, 1);
+
+	if (IsSetValue > 0)
+	{
+		void* ValueData = bIsLoading
+			? MarkSetAndGetInitializedValuePointerToReplace(Data)
+			: GetValuePointerForReadOrReplace(Data);
+
+		// The NetSerializeItem code path is not supported by all property types and they will not be supported within replicated optionals yet either
+		bSuccess = ValueProperty->NetSerializeItem(Ar, Map, ValueData, MetaData);
+	}
+	
+	if (bIsLoading && (!bSuccess || !IsSetValue))
+	{
+		MarkUnset(Data);
+	}
+
+	return bSuccess;
+}
+
+bool FOptionalProperty::SupportsNetSharedSerialization() const
+{
+	return ValueProperty->SupportsNetSharedSerialization();
 }
 
 void FOptionalProperty::ExportText_Internal(FString& ValueStr, const void* ContainerOrPropertyPtr, EPropertyPointerType PropertyPointerType, const void* DefaultValue, UObject* Parent, int32 PortFlags, UObject* ExportRootScope) const
diff --git a/Engine/Source/Runtime/CoreUObject/Public/UObject/PropertyOptional.h b/Engine/Source/Runtime/CoreUObject/Public/UObject/PropertyOptional.h
--- a/Engine/Source/Runtime/CoreUObject/Public/UObject/PropertyOptional.h
+++ b/Engine/Source/Runtime/CoreUObject/Public/UObject/PropertyOptional.h
@@ -215,6 +215,7 @@ public:
 	virtual bool Identical(const void* A, const void* B, uint32 PortFlags) const override;
 	virtual void SerializeItem(FStructuredArchive::FSlot Slot, void* Value, void const* Defaults) const override;
 	virtual bool NetSerializeItem(FArchive& Ar, UPackageMap* Map, void* Data, TArray<uint8>* MetaData = NULL) const override;
+	virtual bool SupportsNetSharedSerialization() const override;
 	virtual void ExportText_Internal(FString& ValueStr, const void* ContainerOrPropertyPtr, EPropertyPointerType PropertyPointerType, const void* DefaultValue, UObject* Parent, int32 PortFlags, UObject* ExportRootScope) const override;
 	virtual const TCHAR* ImportText_Internal(const TCHAR* Buffer, void* ContainerOrPropertyPtr, EPropertyPointerType PropertyPointerType, UObject* Parent, int32 PortFlags, FOutputDevice* ErrorText) const override;
 	virtual void CopyValuesInternal(void* Dest, void const* Src, int32 Count) const override;
