[ZenLoader]
  - Relax fatal error to a warning to maintain backward compatibility with old loader as recursive flush can happen for various reasons and sometimes the result is not even looked at.

diff --git a/Engine/Source/Runtime/CoreUObject/Private/Serialization/AsyncLoading2.cpp b/Engine/Source/Runtime/CoreUObject/Private/Serialization/AsyncLoading2.cpp
--- a/Engine/Source/Runtime/CoreUObject/Private/Serialization/AsyncLoading2.cpp
+++ b/Engine/Source/Runtime/CoreUObject/Private/Serialization/AsyncLoading2.cpp
@@ -7551,7 +7551,7 @@ void FAsyncLoadingThread2::UpdateSyncLoadContext(FAsyncLoadingThreadState2& Thre
 			{
 				// Flushing a package while it's already being processed on the stack, if we're done preloading we let it pass and remove the request id
 				bool bPreloadIsDone = RequestedPackage->AsyncPackageLoadingState >= EAsyncPackageLoadingState2::DeferredPostLoad;
-				UE_CLOG(!bPreloadIsDone, LogStreaming, Fatal, TEXT("Flushing package %s while it's being preloaded in the same callstack is not permitted"), *RequestedPackage->Desc.UPackageName.ToString());
+				UE_CLOG(!bPreloadIsDone, LogStreaming, Warning, TEXT("Flushing package %s while it's being preloaded in the same callstack is not possible. Releasing request %d to unblock."), *RequestedPackage->Desc.UPackageName.ToString(), RequestID);
 				RemovePendingRequests(ThreadState, {RequestID});
 			}
 		}
