Fix for components disapearing on verse builds
Converted the FName key for overridden properties to a FOverriddenPropertyNodeID which keeps a pointer to the Object if it was created from.
Fixed an infinite loop while HandlingReinstantiatedObjects in CoreEntityLevel
Change the DuplicateEntity to use StaticDuplicateObject so it will reconstruct the Overridden information for the new Object.
Fix a problem in the PropertyMap where it would nullout the Defaults on StaticDuplicateObject, added exception for OverridableLogic only.
Had to support reintrance in the TLS for the current OverriddenPropertySet, because seems that StaticDuplicateObject saves into the transaction buffer during the duplication serialization.
And upon Creation of Prefab, we are now asking to replace references upon reload when copying value from the instance to the CDO.

diff --git a/Engine/Source/Runtime/CoreUObject/Private/UObject/OverridableManager.cpp b/Engine/Source/Runtime/CoreUObject/Private/UObject/OverridableManager.cpp
--- a/Engine/Source/Runtime/CoreUObject/Private/UObject/OverridableManager.cpp
+++ b/Engine/Source/Runtime/CoreUObject/Private/UObject/OverridableManager.cpp
@@ -2,6 +2,7 @@
 
 #include "UObject/OverridableManager.h"
 #include "InstancedReferenceSubobjectHelper.h"
+#include "UObject/UObjectGlobals.h"
 
 /*
  *************************************************************************************
@@ -207,4 +208,23 @@ void FOverridableManager::SerializeOverriddenProperties(UObject& Object, FStruct
 	{
 		Disable(Object);
 	}
-}
\ No newline at end of file
+}
+
+FOverridableManager::FOverridableManager()
+{
+#if WITH_EDITOR
+	FCoreUObjectDelegates::OnObjectsReinstanced.AddRaw(this, &FOverridableManager::HandleObjectsReInstantiated);
+#endif
+}
+
+void FOverridableManager::HandleObjectsReInstantiated(const TMap<UObject*, UObject*>& OldToNewInstanceMap)
+{
+	const TMap<const UObjectBase *, FOverriddenPropertyAnnotation>& AnnotationMap = OverriddenObjectAnnotations.GetAnnotationMap();
+	for (const auto& Pair : AnnotationMap)
+	{
+		if( FOverriddenPropertySet* OverridenProperties = Pair.Value.OverriddenProperties.Get())
+		{
+			OverridenProperties->HandleObjectsReInstantiated(OldToNewInstanceMap);
+		}
+	}
+}
diff --git a/Engine/Source/Runtime/CoreUObject/Private/UObject/OverriddenPropertySet.cpp b/Engine/Source/Runtime/CoreUObject/Private/UObject/OverriddenPropertySet.cpp
--- a/Engine/Source/Runtime/CoreUObject/Private/UObject/OverriddenPropertySet.cpp
+++ b/Engine/Source/Runtime/CoreUObject/Private/UObject/OverriddenPropertySet.cpp
@@ -66,6 +66,12 @@ FEnableOverridableSerializationScope::FEnableOverridableSerializationScope(bool
 {
 	if (bEnableOverridableSerialization)
 	{
+		if (FOverridableSerializationLogic::IsEnabled())
+		{
+			bWasOverridableSerializationEnabled = true;
+			SavedOverriddenProperties = FOverridableSerializationLogic::GetOverriddenProperties();
+			FOverridableSerializationLogic::Disable();
+		}
 		FOverridableSerializationLogic::Enable(OverriddenProperties);
 		bOverridableSerializationEnabled = true;
 	}
@@ -76,6 +82,26 @@ FEnableOverridableSerializationScope::~FEnableOverridableSerializationScope()
 	if (bOverridableSerializationEnabled)
 	{
 		FOverridableSerializationLogic::Disable();
+		if (bWasOverridableSerializationEnabled)
+		{
+			FOverridableSerializationLogic::Enable(SavedOverriddenProperties);
+		}
+	}
+}
+
+//----------------------------------------------------------------------//
+// FOverriddenPropertyNodeID
+//----------------------------------------------------------------------//
+void FOverriddenPropertyNodeID::HandleObjectsReInstantiated(const TMap<UObject*, UObject*>& Map)
+{
+	if (!Object)
+	{
+		return;
+	}
+
+	if (UObject*const* ReplacedObject = Map.Find(Object))
+	{
+		Object = *ReplacedObject;
 	}
 }
 
@@ -83,12 +109,12 @@ FEnableOverridableSerializationScope::~FEnableOverridableSerializationScope()
 //----------------------------------------------------------------------//
 // FOverriddenPropertySet
 //----------------------------------------------------------------------//
-FOverriddenPropertyNode& FOverriddenPropertySet::FindOrAddNode(FOverriddenPropertyNode& ParentNode, FName NodeID)
+FOverriddenPropertyNode& FOverriddenPropertySet::FindOrAddNode(FOverriddenPropertyNode& ParentNode, FOverriddenPropertyNodeID NodeID)
 {
-	FName& SubNodeKey = ParentNode.SubPropertyNodeKeys.FindOrAdd(NodeID, NAME_None);
-	if (!SubNodeKey.IsNone())
+	FOverriddenPropertyNodeID& SubNodeID = ParentNode.SubPropertyNodeKeys.FindOrAdd(NodeID, FOverriddenPropertyNodeID());
+	if (SubNodeID.IsValid())
 	{
-		FOverriddenPropertyNode* FoundNode = OverriddenPropertyNodes.FindByHash(GetTypeHash(SubNodeKey), SubNodeKey);
+		FOverriddenPropertyNode* FoundNode = OverriddenPropertyNodes.FindByHash(GetTypeHash(SubNodeID), SubNodeID);
 		checkf(FoundNode, TEXT("Expecting a node"));
 		return *FoundNode;
 	}
@@ -101,11 +127,12 @@ FOverriddenPropertyNode& FOverriddenPropertySet::FindOrAddNode(FOverriddenProper
 
 	// Not found add the node
 	FStringBuilderBase SubPropertyKeyBuilder;
-	SubPropertyKeyBuilder = *ParentNode.PropertyPath.ToString();
+	SubPropertyKeyBuilder = *ParentNode.NodeID.ToString();
 	SubPropertyKeyBuilder.Append(TEXT("."));
 	SubPropertyKeyBuilder.Append(*NodeID.ToString());
-	SubNodeKey = FName(SubPropertyKeyBuilder.ToString());
-	const FSetElementId NewID = OverriddenPropertyNodes.Emplace(SubNodeKey);
+	SubNodeID.Path = FName(SubPropertyKeyBuilder.ToString());
+	SubNodeID.Object = NodeID.Object;
+	const FSetElementId NewID = OverriddenPropertyNodes.Emplace(SubNodeID);
 	return OverriddenPropertyNodes.Get(NewID);
 }
 
@@ -125,7 +152,7 @@ void FOverriddenPropertySet::NotifyPropertyChange(FOverriddenPropertyNode& Paren
 			RemoveOverriddenSubProperties(ParentPropertyNode);
 
 			// If we are overriding the root node, need to propagate the overrides to all instanced sub object
-			const FOverriddenPropertyNode* RootNode = OverriddenPropertyNodes.FindByHash(GetTypeHash(RootNodeName), RootNodeName);
+			const FOverriddenPropertyNode* RootNode = OverriddenPropertyNodes.FindByHash(GetTypeHash(RootNodeID), RootNodeID);
 			checkf(RootNode, TEXT("Expecting to always have a "));
 			if (RootNode == &ParentPropertyNode)
 			{
@@ -139,7 +166,7 @@ void FOverriddenPropertySet::NotifyPropertyChange(FOverriddenPropertyNode& Paren
 	const FProperty* Property = PropertyNode->GetValue();
 	checkf(Property, TEXT("Expecting a valid property"));
 
-	const FName PropID = Property->GetFName();
+	const FOverriddenPropertyNodeID PropID(Property->GetFName());
 	const void* SubValuePtr = Property->ContainerPtrToValuePtr<void>(Data, 0); //@todo support static arrays
 
 	FOverriddenPropertyNode& SubPropertyNode = FindOrAddNode(ParentPropertyNode, PropID);
@@ -185,7 +212,7 @@ void FOverriddenPropertySet::NotifyPropertyChange(FOverriddenPropertyNode& Paren
 						{
 							if (UObject* SubObject = InnerObjectProperty->GetObjectPropertyValue(ArrayHelper.GetElementPtr(i)))
 							{
-								const FName SubObjectID = SubObject->GetFName();
+								const FOverriddenPropertyNodeID SubObjectID(*SubObject);
 								FOverriddenPropertyNode& SubObjectNode = FindOrAddNode(SubPropertyNode, SubObjectID);
 								SubObjectNode.Operation = EOverriddenPropertyOperation::Replace;
 
@@ -199,14 +226,21 @@ void FOverriddenPropertySet::NotifyPropertyChange(FOverriddenPropertyNode& Paren
 						checkf(ArrayHelper.IsValidIndex(ArrayIndex), TEXT("ArrayAdd change type expected to have an valid index"));
 						if (UObject* AddedSubObject = InnerObjectProperty->GetObjectPropertyValue(ArrayHelper.GetElementPtr(ArrayIndex)))
 						{
-							const FName AddedSubObjectID = AddedSubObject->GetFName();
+							const FOverriddenPropertyNodeID  AddedSubObjectID(*AddedSubObject);
 							FOverriddenPropertyNode& AddedSubObjectNode = FindOrAddNode(SubPropertyNode, AddedSubObjectID);
-							if (AddedSubObjectNode.Operation == EOverriddenPropertyOperation::Remove)
+							AddedSubObjectNode.Operation = EOverriddenPropertyOperation::Add;
+
+							// Check if this could be a readd
+							UObject* RemovedSubObjectArchetype = AddedSubObject->GetArchetype();
+							if(RemovedSubObjectArchetype && !RemovedSubObjectArchetype->HasAnyFlags(RF_ClassDefaultObject))
 							{
-								// This is a readd, make sure all the subobjects values are overridden
-								OverridableManager.OverrideInstancedSubObject(*Owner, *AddedSubObject);
+								const FOverriddenPropertyNodeID RemovedSubObjectID(*RemovedSubObjectArchetype);
+								if (SubPropertyNode.SubPropertyNodeKeys.Remove(RemovedSubObjectID))
+								{
+									// This is a readd, make sure all the subobjects values are overridden
+									OverridableManager.OverrideInstancedSubObject(*Owner, *AddedSubObject);
+								}
 							}
-							AddedSubObjectNode.Operation = EOverriddenPropertyOperation::Add;
 						}
 					};
 
@@ -215,7 +249,8 @@ void FOverriddenPropertySet::NotifyPropertyChange(FOverriddenPropertyNode& Paren
 						checkf(PreEditSubObjectsArrayHelper.IsValidIndex(ArrayIndex), TEXT("ArrayRemove change type expected to have an valid index"));
 						if (UObject* RemovedSubObject = InnerObjectProperty->GetObjectPropertyValue(PreEditSubObjectsArrayHelper.GetElementPtr(ArrayIndex)))
 						{
-							const FName RemovedSubObjectID = RemovedSubObject->GetFName();
+							UObject* RemovedSubObjectArchetype = RemovedSubObject->GetArchetype();
+							const FOverriddenPropertyNodeID RemovedSubObjectID (!RemovedSubObjectArchetype || RemovedSubObjectArchetype->HasAnyFlags(RF_ClassDefaultObject) ? *RemovedSubObject : *RemovedSubObjectArchetype);
 							FOverriddenPropertyNode& RemovedSubObjectNode = FindOrAddNode(SubPropertyNode, RemovedSubObjectID);
 							OverridableManager.ClearInstancedSubObjectOverrides(*Owner, *RemovedSubObject);
 							if (RemovedSubObjectNode.Operation == EOverriddenPropertyOperation::Add)
@@ -381,10 +416,10 @@ void FOverriddenPropertySet::NotifyPropertyChange(FOverriddenPropertyNode& Paren
 				// Need to loop through every sub object and setup them up as overridden
 				for (FScriptMapHelper::FIterator It(MapHelper); It; ++It)
 				{
-					FName OverriddenKeyID = NAME_None;
+					FOverriddenPropertyNodeID OverriddenKeyID;
 					if (const UObject* OverriddenKeyObject = KeyObjectProperty ? KeyObjectProperty->GetObjectPropertyValue(MapHelper.GetKeyPtr(It.GetInternalIndex())) : nullptr)
 					{
-						OverriddenKeyID = FName(OverriddenKeyObject->GetPathName());;
+						OverriddenKeyID = FOverriddenPropertyNodeID(*OverriddenKeyObject);
 					}
 					else
 					{
@@ -410,10 +445,10 @@ void FOverriddenPropertySet::NotifyPropertyChange(FOverriddenPropertyNode& Paren
 			auto MapAddImpl = [&]()
 			{
 				checkf(MapHelper.IsValidIndex(InternalMapIndex), TEXT("ArrayAdd change type expected to have an valid index"));
-				FName AddedKeyID = NAME_None;
+				FOverriddenPropertyNodeID AddedKeyID;
 				if (const UObject* AddedKeyObject = KeyObjectProperty ? KeyObjectProperty->GetObjectPropertyValue(MapHelper.GetKeyPtr(InternalMapIndex)) : nullptr)
 				{
-					AddedKeyID = FName(AddedKeyObject->GetPathName());
+					AddedKeyID = FOverriddenPropertyNodeID(*AddedKeyObject);
 				}
 				else
 				{
@@ -446,10 +481,10 @@ void FOverriddenPropertySet::NotifyPropertyChange(FOverriddenPropertyNode& Paren
 			auto MapRemoveImpl = [&]()
 			{
 				checkf(PreEditMapHelper.IsValidIndex(InternalPreEditMapIndex), TEXT("ArrayRemove change type expected to have an valid index"));
-				FName RemovedKeyID = NAME_None;
+				FOverriddenPropertyNodeID RemovedKeyID;
 				if (const UObject* RemovedKeyObject = KeyObjectProperty ? KeyObjectProperty->GetObjectPropertyValue(PreEditMapHelper.GetKeyPtr(InternalPreEditMapIndex)) : nullptr)
 				{
-					RemovedKeyID = FName(RemovedKeyObject->GetPathName());;
+					RemovedKeyID = FOverriddenPropertyNodeID(*RemovedKeyObject);
 				}
 				else
 				{
@@ -614,12 +649,12 @@ void FOverriddenPropertySet::RemoveOverriddenSubProperties(FOverriddenPropertyNo
 
 void FOverriddenPropertySet::NotifyPropertyChange(const EPropertyNotificationType Notification, const FPropertyChangedEvent& PropertyEvent, const FEditPropertyChain::TDoubleLinkedListNode* PropertyNode, const void* Data)
 {
-	NotifyPropertyChange(OverriddenPropertyNodes.FindOrAddByHash(GetTypeHash(RootNodeName), RootNodeName), Notification, PropertyEvent, PropertyNode, Data);
+	NotifyPropertyChange(OverriddenPropertyNodes.FindOrAddByHash(GetTypeHash(RootNodeID), RootNodeID), Notification, PropertyEvent, PropertyNode, Data);
 }
 
 EOverriddenPropertyOperation FOverriddenPropertySet::GetOverriddenPropertyOperation(const FArchiveSerializedPropertyChain* CurrentPropertyChain, FProperty* Property) const
 {
-	if (const FOverriddenPropertyNode* RootNode = OverriddenPropertyNodes.FindByHash(GetTypeHash(RootNodeName), RootNodeName))
+	if (const FOverriddenPropertyNode* RootNode = OverriddenPropertyNodes.FindByHash(GetTypeHash(RootNodeID), RootNodeID))
 	{
 		return GetOverriddenPropertyOperation(*RootNode, CurrentPropertyChain, Property);
 	}
@@ -628,12 +663,12 @@ EOverriddenPropertyOperation FOverriddenPropertySet::GetOverriddenPropertyOperat
 
 FOverriddenPropertyNode* FOverriddenPropertySet::SetOverriddenPropertyOperation(EOverriddenPropertyOperation Operation, const FArchiveSerializedPropertyChain* CurrentPropertyChain, FProperty* Property)
 {
-	return SetOverriddenPropertyOperation(Operation, OverriddenPropertyNodes.FindOrAddByHash(GetTypeHash(RootNodeName), RootNodeName), CurrentPropertyChain, Property);
+	return SetOverriddenPropertyOperation(Operation, OverriddenPropertyNodes.FindOrAddByHash(GetTypeHash(RootNodeID), RootNodeID), CurrentPropertyChain, Property);
 }
 
 const FOverriddenPropertyNode* FOverriddenPropertySet::GetOverriddenPropertyNode(const FArchiveSerializedPropertyChain* CurrentPropertyChain) const
 {
-	if (const FOverriddenPropertyNode* RootNode = OverriddenPropertyNodes.FindByHash(GetTypeHash(RootNodeName), RootNodeName))
+	if (const FOverriddenPropertyNode* RootNode = OverriddenPropertyNodes.FindByHash(GetTypeHash(RootNodeID), RootNodeID))
 	{
 		return GetOverriddenPropertyNode(*RootNode, CurrentPropertyChain);
 	}
@@ -663,7 +698,7 @@ EOverriddenPropertyOperation FOverriddenPropertySet::GetOverriddenPropertyOperat
 	{
 		const FProperty* CurrentProperty = (*PropertyIterator);
 		const FName CurrentPropID = CurrentProperty->GetFName();
-		if (const FName* CurrentPropKey = OverriddenPropertyNode->SubPropertyNodeKeys.Find(CurrentPropID))
+		if (const FOverriddenPropertyNodeID* CurrentPropKey = OverriddenPropertyNode->SubPropertyNodeKeys.Find(CurrentPropID))
 		{
 			OverriddenPropertyNode = OverriddenPropertyNodes.FindByHash(GetTypeHash(*CurrentPropKey), *CurrentPropKey);
 			checkf(OverriddenPropertyNode, TEXT("Expecting a node"));
@@ -716,15 +751,15 @@ FOverriddenPropertyNode* FOverriddenPropertySet::SetOverriddenPropertyOperation(
 	return nullptr;
 }
 
-EOverriddenPropertyOperation FOverriddenPropertySet::GetSubPropertyOperation(FName PropertyKey) const
+EOverriddenPropertyOperation FOverriddenPropertySet::GetSubPropertyOperation(FOverriddenPropertyNodeID NodeID) const
 {
-	const FOverriddenPropertyNode* OverriddenPropertyNode = OverriddenPropertyNodes.FindByHash(GetTypeHash(PropertyKey), PropertyKey);
+	const FOverriddenPropertyNode* OverriddenPropertyNode = OverriddenPropertyNodes.FindByHash(GetTypeHash(NodeID), NodeID);
 	return OverriddenPropertyNode ? OverriddenPropertyNode->Operation : EOverriddenPropertyOperation::None;
 }
 
-FOverriddenPropertyNode* FOverriddenPropertySet::SetSubPropertyOperation(EOverriddenPropertyOperation Operation, FOverriddenPropertyNode& Node, FName PropID)
+FOverriddenPropertyNode* FOverriddenPropertySet::SetSubPropertyOperation(EOverriddenPropertyOperation Operation, FOverriddenPropertyNode& Node, FOverriddenPropertyNodeID NodeID)
 {
-	FOverriddenPropertyNode& OverriddenPropertyNode = FindOrAddNode(Node, PropID);
+	FOverriddenPropertyNode& OverriddenPropertyNode = FindOrAddNode(Node, NodeID);
 	OverriddenPropertyNode.Operation = Operation;
 	return &OverriddenPropertyNode;
 }
@@ -734,6 +769,19 @@ void FOverriddenPropertySet::Reset()
 	OverriddenPropertyNodes.Reset();
 }
 
+void FOverriddenPropertySet::HandleObjectsReInstantiated(const TMap<UObject*, UObject*>& Map)
+{
+	for (FOverriddenPropertyNode& Node : OverriddenPropertyNodes)
+	{
+		Node.NodeID.HandleObjectsReInstantiated(Map);
+		for (auto& Pair : Node.SubPropertyNodeKeys)
+		{
+			Pair.Key.HandleObjectsReInstantiated(Map);
+			Pair.Value.HandleObjectsReInstantiated(Map);
+		}
+	}
+}
+
 const FOverriddenPropertyNode* FOverriddenPropertySet::GetOverriddenPropertyNode(const FOverriddenPropertyNode& ParentPropertyNode, const FArchiveSerializedPropertyChain* CurrentPropertyChain) const
 {
 	if (!CurrentPropertyChain)
@@ -747,7 +795,7 @@ const FOverriddenPropertyNode* FOverriddenPropertySet::GetOverriddenPropertyNode
 	{
 		const FProperty* CurrentProperty = (*PropertyIterator);
 		const FName CurrentPropID = CurrentProperty->GetFName();
-		if (const FName* CurrentPropKey = OverriddenPropertyNode->SubPropertyNodeKeys.Find(CurrentPropID))
+		if (const FOverriddenPropertyNodeID* CurrentPropKey = OverriddenPropertyNode->SubPropertyNodeKeys.Find(CurrentPropID))
 		{
 			OverriddenPropertyNode = OverriddenPropertyNodes.FindByHash(GetTypeHash(*CurrentPropKey), *CurrentPropKey);
 			checkf(OverriddenPropertyNode, TEXT("Expecting a node"));
diff --git a/Engine/Source/Runtime/CoreUObject/Private/UObject/PropertyArray.cpp b/Engine/Source/Runtime/CoreUObject/Private/UObject/PropertyArray.cpp
--- a/Engine/Source/Runtime/CoreUObject/Private/UObject/PropertyArray.cpp
+++ b/Engine/Source/Runtime/CoreUObject/Private/UObject/PropertyArray.cpp
@@ -305,7 +305,7 @@ void FArrayProperty::SerializeItem(FStructuredArchive::FSlot Slot, void* Value,
 							if (FOverriddenPropertyNode* ArrayOverriddenPropertyNode = OverriddenProperties ? OverriddenProperties->SetOverriddenPropertyOperation(EOverriddenPropertyOperation::Modified, UnderlyingArchive.GetSerializedPropertyChain(), /*Property*/nullptr) : nullptr)
 							{
 								// Rebuild the overridden info
-								const FName RemovedSubObjectID = RemovedSubObject->GetFName();
+								const FOverriddenPropertyNodeID RemovedSubObjectID(*RemovedSubObject);
 								OverriddenProperties->SetSubPropertyOperation(EOverriddenPropertyOperation::Remove, *ArrayOverriddenPropertyNode, RemovedSubObjectID);
 							}
 						}
@@ -378,7 +378,7 @@ void FArrayProperty::SerializeItem(FStructuredArchive::FSlot Slot, void* Value,
 							if (FOverriddenPropertyNode* ArrayOverriddenPropertyNode = OverriddenProperties ? OverriddenProperties->SetOverriddenPropertyOperation(EOverriddenPropertyOperation::Modified, UnderlyingArchive.GetSerializedPropertyChain(), /*Property*/nullptr) : nullptr)
 							{
 								// Rebuild the overridden info
-								const FName AddedSubObjectID = AddedSubObject->GetFName();
+								const FOverriddenPropertyNodeID AddedSubObjectID(*AddedSubObject);
 								OverriddenProperties->SetSubPropertyOperation(EOverriddenPropertyOperation::Add, *ArrayOverriddenPropertyNode, AddedSubObjectID);
 							}
 						}
@@ -428,15 +428,17 @@ void FArrayProperty::SerializeItem(FStructuredArchive::FSlot Slot, void* Value,
 					if (OverriddenProperties && ArrayOverrideOp != EOverriddenPropertyOperation::None)
 					{
 
-						auto FindObject = [InnerObjectProperty](const FName ObjectToFind, FScriptArrayHelper& ArrayHelper) -> int32
+						auto FindObject = [InnerObjectProperty](const FOverriddenPropertyNodeID ObjectToFind, FScriptArrayHelper& ArrayHelper) -> int32
 						{
 							const int32 ArrayNum = ArrayHelper.Num();
 							for (int i = 0; i < ArrayNum; ++i)
 							{
-								UObject* CurrentObject = InnerObjectProperty->GetObjectPropertyValue(ArrayHelper.GetElementPtr(i));
-								if (CurrentObject && CurrentObject->GetName() == ObjectToFind)
+								if (UObject* CurrentObject = InnerObjectProperty->GetObjectPropertyValue(ArrayHelper.GetElementPtr(i)))
 								{
-									return i;
+									if (ObjectToFind == FOverriddenPropertyNodeID(*CurrentObject))
+									{
+										return i;
+									}
 								}
 							}
 							return INDEX_NONE;
diff --git a/Engine/Source/Runtime/CoreUObject/Private/UObject/PropertyMap.cpp b/Engine/Source/Runtime/CoreUObject/Private/UObject/PropertyMap.cpp
--- a/Engine/Source/Runtime/CoreUObject/Private/UObject/PropertyMap.cpp
+++ b/Engine/Source/Runtime/CoreUObject/Private/UObject/PropertyMap.cpp
@@ -296,7 +296,7 @@ void FMapProperty::SerializeItem(FStructuredArchive::FSlot Slot, void* Value, co
 	};
 
 	// If we're doing delta serialization within this property, act as if there are no defaults
-	if (!UnderlyingArchive.DoIntraPropertyDelta())
+	if (!UnderlyingArchive.DoIntraPropertyDelta() && !bExperimentalOverridableLogic)
 	{
 		Defaults = nullptr;
 	}
@@ -318,24 +318,24 @@ void FMapProperty::SerializeItem(FStructuredArchive::FSlot Slot, void* Value, co
 	{
 		checkf(!UnderlyingArchive.ArUseCustomPropertyList, TEXT("Using custom property list is not supported by overridable serialization"));
 
-		auto GetIDFromKey = [&](uint8* KeyData) -> FName
+		auto GetIDFromKey = [&](uint8* KeyData) -> FOverriddenPropertyNodeID
 		{
 			if (FObjectProperty* KeyObjectProperty = CastField<FObjectProperty>(KeyProp))
 			{
 				if (const UObject* Object = KeyObjectProperty->GetObjectPropertyValue(KeyData))
 				{
-					return FName(Object->GetPathName());
+					return FOverriddenPropertyNodeID(*Object);
 				}
 			}
 			else
 			{
 				FString KeyString;
 				KeyProp->ExportTextItem_Direct(KeyString, KeyData, nullptr, nullptr, PPF_None);
-				return FName(KeyString);
+				return FOverriddenPropertyNodeID(FName(KeyString));
 			}
 	
 			checkf(false, TEXT("This case is not handled"))
-			return NAME_None;
+			return FOverriddenPropertyNodeID();
 		};
 
 		if (UnderlyingArchive.IsLoading())
@@ -401,7 +401,7 @@ void FMapProperty::SerializeItem(FStructuredArchive::FSlot Slot, void* Value, co
 						if (FOverriddenPropertyNode* MapOverriddenPropertyNode = OverriddenProperties ? OverriddenProperties->SetOverriddenPropertyOperation(EOverriddenPropertyOperation::Modified, UnderlyingArchive.GetSerializedPropertyChain(), /*Property*/nullptr) : nullptr)
 						{
 							// Rebuild the overridden info
-							FName RemovedKeyID = GetIDFromKey(TempKeyValueStorage);
+							FOverriddenPropertyNodeID RemovedKeyID = GetIDFromKey(TempKeyValueStorage);
 							OverriddenProperties->SetSubPropertyOperation(EOverriddenPropertyOperation::Remove, *MapOverriddenPropertyNode, RemovedKeyID);
 						}
 					}
@@ -471,7 +471,7 @@ void FMapProperty::SerializeItem(FStructuredArchive::FSlot Slot, void* Value, co
 						if (FOverriddenPropertyNode* MapOverriddenPropertyNode = OverriddenProperties ? OverriddenProperties->SetOverriddenPropertyOperation(EOverriddenPropertyOperation::Modified, UnderlyingArchive.GetSerializedPropertyChain(), /*Property*/nullptr) : nullptr)
 						{
 							// Rebuild the overridden info
-							FName AddedKeyID = GetIDFromKey(TempKeyValueStorage);
+							FOverriddenPropertyNodeID AddedKeyID = GetIDFromKey(TempKeyValueStorage);
 							OverriddenProperties->SetSubPropertyOperation(EOverriddenPropertyOperation::Add, *MapOverriddenPropertyNode, AddedKeyID);
 						}
 					}
@@ -480,16 +480,18 @@ void FMapProperty::SerializeItem(FStructuredArchive::FSlot Slot, void* Value, co
 		}
 		else
 		{
-			auto FindKeyInternalIndex = [this](const FName KeyIDToFind, FScriptMapHelper& MapHelper) -> int32
+			auto FindKeyInternalIndex = [this](const FOverriddenPropertyNodeID& KeyIDToFind, FScriptMapHelper& MapHelper) -> int32
 			{
 				if (const FObjectProperty* KeyObjectProperty = CastField<FObjectProperty>(KeyProp))
 				{
 					for (FScriptMapHelper::FIterator It(MapHelper); It; ++It) 
 					{
-						UObject* CurrentObject = KeyObjectProperty->GetObjectPropertyValue(MapHelper.GetKeyPtr(It.GetInternalIndex()));
-						if (FName(CurrentObject->GetPathName()) == KeyIDToFind)
+						if (UObject* CurrentObject = KeyObjectProperty->GetObjectPropertyValue(MapHelper.GetKeyPtr(It.GetInternalIndex())))
 						{
-							return It.GetInternalIndex();
+							if (KeyIDToFind == FOverriddenPropertyNodeID(*CurrentObject))
+							{
+								return It.GetInternalIndex();
+							}
 						}
 					}
 				}
diff --git a/Engine/Source/Runtime/CoreUObject/Public/UObject/OverridableManager.h b/Engine/Source/Runtime/CoreUObject/Public/UObject/OverridableManager.h
--- a/Engine/Source/Runtime/CoreUObject/Public/UObject/OverridableManager.h
+++ b/Engine/Source/Runtime/CoreUObject/Public/UObject/OverridableManager.h
@@ -28,7 +28,7 @@ struct FOverriddenPropertyAnnotation
 /**
  * Global container of overriden object annotations
  */
-class FOverriddenPropertyAnnotations : public FUObjectAnnotationChunked<FOverriddenPropertyAnnotation, true/*bAutoRemove*/>
+class FOverriddenPropertyAnnotations : public FUObjectAnnotationSparse<FOverriddenPropertyAnnotation, true/*bAutoRemove*/>
 {
 public:
 	bool IsEnabled(const UObject& Object)
@@ -202,7 +202,11 @@ public:
 	 */
 	COREUOBJECT_API void SerializeOverriddenProperties(UObject& Object, FStructuredArchive::FRecord ObjectRecord);
 
+	void HandleObjectsReInstantiated(const TMap<UObject*, UObject*>& OldToNewInstanceMap);
+
 protected:
+	FOverridableManager();
+
 	FOverriddenPropertyAnnotations OverriddenObjectAnnotations;
 };
 
diff --git a/Engine/Source/Runtime/CoreUObject/Public/UObject/OverriddenPropertySet.h b/Engine/Source/Runtime/CoreUObject/Public/UObject/OverriddenPropertySet.h
--- a/Engine/Source/Runtime/CoreUObject/Public/UObject/OverriddenPropertySet.h
+++ b/Engine/Source/Runtime/CoreUObject/Public/UObject/OverriddenPropertySet.h
@@ -85,6 +85,8 @@ struct FEnableOverridableSerializationScope
 
 protected:
 	bool bOverridableSerializationEnabled = false;
+	bool bWasOverridableSerializationEnabled = false;
+	FOverriddenPropertySet* SavedOverriddenProperties = nullptr;
 };
 
 /*
@@ -100,6 +102,51 @@ enum class EOverriddenPropertyOperation : uint8
 	Remove,		/* this element was removed from the container */
 };
 
+USTRUCT()
+struct FOverriddenPropertyNodeID
+{
+	GENERATED_BODY()
+
+	FOverriddenPropertyNodeID(FName InPath = NAME_None)
+		: Path(InPath)
+		, Object(nullptr)
+	{}
+
+	FOverriddenPropertyNodeID(const UObject& InObject)
+		: Path(FString::Printf(TEXT("%d"), GUObjectArray.ObjectToIndex(&InObject)))
+		, Object(&InObject)
+	{
+	}
+
+	bool operator==(const FOverriddenPropertyNodeID& Other) const
+	{
+		return Path == Other.Path || (Object && Other.Object && Object == Other.Object);
+	}
+
+	FString ToString() const
+	{
+		return Path.ToString();
+	}
+
+	bool IsValid() const
+	{
+		return !Path.IsNone();
+	}
+
+	void HandleObjectsReInstantiated(const TMap<UObject*, UObject*>& Map);
+
+	UPROPERTY()
+	FName Path;
+
+	UPROPERTY()
+	TObjectPtr<const UObject> Object;
+
+	friend uint32 GetTypeHash(const FOverriddenPropertyNodeID& NodeID)
+	{
+		return GetTypeHash(NodeID.Path);
+	}
+};
+
 /*
  *************************************************************************************
  * Overridable serialization is experimental, not supported and use at your own risk *
@@ -112,27 +159,27 @@ struct FOverriddenPropertyNode
 {
 	GENERATED_BODY()
 
-	FOverriddenPropertyNode(FName InPropertyPath = NAME_None)
-		: PropertyPath(InPropertyPath)
+	FOverriddenPropertyNode(FOverriddenPropertyNodeID InNodeID = FOverriddenPropertyNodeID())
+		: NodeID(InNodeID)
 	{}
 
 	UPROPERTY()
-	FName PropertyPath;
+	FOverriddenPropertyNodeID NodeID;
 
 	UPROPERTY()
 	EOverriddenPropertyOperation Operation = EOverriddenPropertyOperation::None;
 
 	UPROPERTY()
-	TMap<FName,FName> SubPropertyNodeKeys;
+	TMap<FOverriddenPropertyNodeID, FOverriddenPropertyNodeID> SubPropertyNodeKeys;
 
 	bool operator==(const FOverriddenPropertyNode& Other) const
 	{
-		return PropertyPath == Other.PropertyPath;
+		return NodeID == Other.NodeID;
 	}
 
 	friend uint32 GetTypeHash(const FOverriddenPropertyNode& Node)
 	{
-		return GetTypeHash(Node.PropertyPath);
+		return GetTypeHash(Node.NodeID);
 	}
 };
 
@@ -196,21 +243,22 @@ public:
 
 	/**
 	 * Retrieve the overridable operation given the property key
-	 * @param PropertyKey that uniquely identify the property within the object 
+	 * @param NodeID that uniquely identify the property within the object 
 	 * @return the current type of override operation on the property */
-	EOverriddenPropertyOperation GetSubPropertyOperation(FName PropertyKey) const;
+	EOverriddenPropertyOperation GetSubPropertyOperation(FOverriddenPropertyNodeID NodeID) const;
 
 	/**
 	 * Set the overridable operation of a sub property of the specified node.
 	 * @param Operation to set for this property
 	 * @param Node from where the sub property is owned by
-	 * @param PropID the ID of the sub property 
+	 * @param NodeID the ID of the sub property 
 	 * @return the node to the sub property */
-	FOverriddenPropertyNode* SetSubPropertyOperation(EOverriddenPropertyOperation Operation, FOverriddenPropertyNode& Node, FName PropID);
+	FOverriddenPropertyNode* SetSubPropertyOperation(EOverriddenPropertyOperation Operation, FOverriddenPropertyNode& Node, FOverriddenPropertyNodeID NodeID);
 
 	/**
 	 * Resets all overrides of the object */
 	void Reset();
+	void HandleObjectsReInstantiated(const TMap<UObject*, UObject*>& Map);
 
 protected:
 
@@ -222,7 +270,7 @@ protected:
 
 	void RemoveOverriddenSubProperties(FOverriddenPropertyNode& PropertyNode);
 
-	FOverriddenPropertyNode& FindOrAddNode(FOverriddenPropertyNode& ParentPropertyNode, FName PropID);
+	FOverriddenPropertyNode& FindOrAddNode(FOverriddenPropertyNode& ParentPropertyNode, FOverriddenPropertyNodeID NodeID);
 
 private:
 	UPROPERTY()
@@ -231,5 +279,5 @@ private:
 	UPROPERTY()
 	TSet<FOverriddenPropertyNode> OverriddenPropertyNodes;
 
-	static inline FName RootNodeName = TEXT("root");
+	static inline FOverriddenPropertyNodeID RootNodeID = FOverriddenPropertyNodeID(FName(TEXT("root")));
 };
