Fix unversion save to setup override serialization when available while saving.

diff --git a/Engine/Source/Runtime/CoreUObject/Private/Serialization/UnversionedPropertySerialization.cpp b/Engine/Source/Runtime/CoreUObject/Private/Serialization/UnversionedPropertySerialization.cpp
--- a/Engine/Source/Runtime/CoreUObject/Private/Serialization/UnversionedPropertySerialization.cpp
+++ b/Engine/Source/Runtime/CoreUObject/Private/Serialization/UnversionedPropertySerialization.cpp
@@ -5,6 +5,7 @@
 #include "Hash/Blake3.h"
 #include "Interfaces/ITargetPlatform.h"
 #include "Misc/ScopeRWLock.h"
+#include "UObject/OverridableManager.h"
 #include "UObject/PropertyOptional.h"
 #include "UObject/UnrealType.h"
 
@@ -922,6 +923,17 @@ void SerializeUnversionedProperties(const UStruct* Struct, FStructuredArchive::F
 	}
 	else
 	{
+		bool bEnableOverridableSerialization = false;
+		FOverriddenPropertySet* OverriddenProperties = nullptr;
+		if (FOverriddenPropertySet* ObjectOverriddenProperties = Struct->IsA<UClass>() ? FOverridableManager::Get().GetOverriddenProperties(*(UObject*)Data) : nullptr)
+		{
+			bEnableOverridableSerialization = true;
+			OverriddenProperties = ObjectOverriddenProperties;
+		}
+
+		// Scope that enables the overridable serialization for this object
+		FEnableOverridableSerializationScope OverridableSerializationScope(bEnableOverridableSerialization, OverriddenProperties);
+
 		FUnversionedPropertyTestRunner TestRunner({Struct, Data, DefaultsStruct, DefaultsData});
 		FUnversionedPropertyTestCollector TestCollector;
 
