Fixed Mesh being reverted to cube since enabled OS
I will add a unittest to repro this problem!

diff --git a/Engine/Source/Runtime/CoreUObject/Private/UObject/Class.cpp b/Engine/Source/Runtime/CoreUObject/Private/UObject/Class.cpp
--- a/Engine/Source/Runtime/CoreUObject/Private/UObject/Class.cpp
+++ b/Engine/Source/Runtime/CoreUObject/Private/UObject/Class.cpp
@@ -1777,7 +1777,7 @@ void UStruct::SerializeVersionedTaggedProperties(FStructuredArchive::FSlot Slot,
 					uint8* DataPtr      = Property->ContainerPtrToValuePtr           <uint8>(Data, Idx);
 					uint8* DefaultValue = Property->ContainerPtrToValuePtrForDefaults<uint8>(DefaultsStruct, Defaults, Idx);
 					if (StaticArrayContainer.IsSet() || CustomPropertyNode || !bDoDeltaSerialization ||
-						(FOverridableSerializationLogic::IsEnabled() && FOverridableSerializationLogic::GetOverriddenPropertyOperation(UnderlyingArchive, Property) != EOverriddenPropertyOperation::None) ||
+						(FOverridableSerializationLogic::IsEnabled() && FOverridableSerializationLogic::GetOverriddenPropertyOperation(UnderlyingArchive, Property, DataPtr) != EOverriddenPropertyOperation::None) ||
 						(!FOverridableSerializationLogic::IsEnabled() && !Property->Identical(DataPtr, DefaultValue, UnderlyingArchive.GetPortFlags())))
 					{
 						if (bUseAtomicSerialization)
diff --git a/Engine/Source/Runtime/CoreUObject/Private/UObject/OverriddenPropertySet.cpp b/Engine/Source/Runtime/CoreUObject/Private/UObject/OverriddenPropertySet.cpp
--- a/Engine/Source/Runtime/CoreUObject/Private/UObject/OverriddenPropertySet.cpp
+++ b/Engine/Source/Runtime/CoreUObject/Private/UObject/OverriddenPropertySet.cpp
@@ -22,7 +22,7 @@ DEFINE_LOG_CATEGORY(LogOverridableObject);
 thread_local bool FOverridableSerializationLogic::bUseOverridableSerialization = false;
 thread_local FOverriddenPropertySet* FOverridableSerializationLogic::OverriddenProperties = nullptr;
 
-EOverriddenPropertyOperation FOverridableSerializationLogic::GetOverriddenPropertyOperation(const FArchive& Ar, FProperty* Property /*= nullptr*/)
+EOverriddenPropertyOperation FOverridableSerializationLogic::GetOverriddenPropertyOperation(const FArchive& Ar, FProperty* Property /*= nullptr*/, uint8* DataPtr /*= nullptr*/)
 {
 	checkf(bUseOverridableSerialization, TEXT("Nobody should use this method if it is not setup to use overridable serialization"));
 
@@ -44,7 +44,16 @@ EOverriddenPropertyOperation FOverridableSerializationLogic::GetOverriddenProper
 		// This should probably be fix in verse at some point.
 		else if (CurrentProperty->HasAnyPropertyFlags(CPF_PersistentInstance))
 		{
-			return EOverriddenPropertyOperation::Modified;
+			if (const FObjectProperty* ObjectProperty = CastField<FObjectProperty>(CurrentProperty))
+			{
+				if (UObject* SubObject = ObjectProperty->GetObjectPropertyValue(DataPtr))
+				{
+					if (FOverridableManager::Get().GetOverriddenState(*SubObject) != EOverriddenState::NoOverrides)
+					{
+						return EOverriddenPropertyOperation::Modified;
+					}
+				}
+			}
 		}
 		else if (const FArrayProperty* CurrentArrayProperty = CastField<FArrayProperty>(CurrentProperty))
 		{
diff --git a/Engine/Source/Runtime/CoreUObject/Public/UObject/OverriddenPropertySet.h b/Engine/Source/Runtime/CoreUObject/Public/UObject/OverriddenPropertySet.h
--- a/Engine/Source/Runtime/CoreUObject/Public/UObject/OverriddenPropertySet.h
+++ b/Engine/Source/Runtime/CoreUObject/Public/UObject/OverriddenPropertySet.h
@@ -64,8 +64,9 @@ public:
 	 * Retrieve from the Archive and the current property the overridden property operation to know if it has to be serialized or not
 	 * @param Ar currently being used to serialize the current object (will be used to retrieve the current property serialized path)
 	 * @param Property the property about to be serialized, can be null
+	 * @param DataPtr to the memory of that property
 	 * @return the overridden property operation */
-	static EOverriddenPropertyOperation GetOverriddenPropertyOperation(const FArchive& Ar, FProperty* Property = nullptr);
+	static EOverriddenPropertyOperation GetOverriddenPropertyOperation(const FArchive& Ar, FProperty* Property = nullptr, uint8* DataPtr = nullptr);
 
 private:
 	static thread_local bool bUseOverridableSerialization;
