Don't run VerifyNoUnreachable tests when performing the exit purge

diff --git a/Engine/Source/Runtime/CoreUObject/Private/UObject/GarbageCollection.cpp b/Engine/Source/Runtime/CoreUObject/Private/UObject/GarbageCollection.cpp
--- a/Engine/Source/Runtime/CoreUObject/Private/UObject/GarbageCollection.cpp
+++ b/Engine/Source/Runtime/CoreUObject/Private/UObject/GarbageCollection.cpp
@@ -4973,7 +4973,7 @@ bool GatherUnreachableObjects(UE::GC::EGatherOptions Options, double TimeLimit /
 		}
 
 #if VERIFY_DISREGARD_GC_ASSUMPTIONS
-		if (GVerifyNoUnreachableObjects != 0 && GUnreachableObjects.Num() > 0)
+		if (GVerifyNoUnreachableObjects != 0 && GUnreachableObjects.Num() > 0 && !GExitPurge)
 		{
 			DECLARE_SCOPE_CYCLE_COUNTER(TEXT("CollectGarbageInternal.VerifyNoUnreachableObjects"), STAT_CollectGarbageInternal_VerifyNoUnreachableObjects, STATGROUP_GC);
 			VerifyNoUnreachableObjects();
