Fix the StaticFindObject AutoRTFM bug.

diff --git a/Engine/Source/Runtime/CoreUObject/Private/UObject/UObjectHash.cpp b/Engine/Source/Runtime/CoreUObject/Private/UObject/UObjectHash.cpp
--- a/Engine/Source/Runtime/CoreUObject/Private/UObject/UObjectHash.cpp
+++ b/Engine/Source/Runtime/CoreUObject/Private/UObject/UObjectHash.cpp
@@ -15,6 +15,7 @@
 #include "HAL/LowLevelMemStats.h"
 #include "UObject/AnyPackagePrivate.h"
 #include "ProfilingDebugging/AssetMetadataTrace.h"
+#include "AutoRTFM/AutoRTFM.h"
 
 #include <atomic>
 
@@ -843,25 +844,39 @@ UObject* StaticFindObjectFastInternalThreadSafe(FUObjectHashTables& ThreadHash,
 
 UObject* StaticFindObjectFastInternal(const UClass* ObjectClass, const UObject* ObjectPackage, FName ObjectName, bool bExactClass, bool bAnyPackage, EObjectFlags ExcludeFlags, EInternalObjectFlags ExclusiveInternalFlags)
 {
-	INC_DWORD_STAT(STAT_FindObjectFast);
+	UObject* Result = nullptr;
 
-	check(ObjectPackage != ANY_PACKAGE_DEPRECATED); // this could never have returned anything but nullptr
+	// Transactionally, a static find operation has to occur in the open as it touches shared state.
+	UE_AUTORTFM_OPEN(
+	{
+		INC_DWORD_STAT(STAT_FindObjectFast);
+
+		check(ObjectPackage != ANY_PACKAGE_DEPRECATED); // this could never have returned anything but nullptr
+
+		// If they specified an outer use that during the hashing
+		FUObjectHashTables & ThreadHash = FUObjectHashTables::Get();
+		Result = StaticFindObjectFastInternalThreadSafe(ThreadHash, ObjectClass, ObjectPackage, ObjectName, bExactClass, bAnyPackage, ExcludeFlags, ExclusiveInternalFlags);
+	});
 
-	// If they specified an outer use that during the hashing
-	FUObjectHashTables& ThreadHash = FUObjectHashTables::Get();
-	UObject* Result = StaticFindObjectFastInternalThreadSafe(ThreadHash, ObjectClass, ObjectPackage, ObjectName, bExactClass, bAnyPackage, ExcludeFlags, ExclusiveInternalFlags);
 	return Result;
 }
 
 UObject* StaticFindObjectFastInternal(const UClass* ObjectClass, const UObject* ObjectPackage, FName ObjectName, bool bExactClass, EObjectFlags ExcludeFlags, EInternalObjectFlags ExclusiveInternalFlags)
 {
-	INC_DWORD_STAT(STAT_FindObjectFast);
+	UObject* Result = nullptr;
 
-	check(ObjectPackage != ANY_PACKAGE_DEPRECATED); // this could never have returned anything but nullptr
+	// Transactionally, a static find operation has to occur in the open as it touches shared state.
+	UE_AUTORTFM_OPEN(
+	{
+		INC_DWORD_STAT(STAT_FindObjectFast);
+
+		check(ObjectPackage != ANY_PACKAGE_DEPRECATED); // this could never have returned anything but nullptr
+
+		// If they specified an outer use that during the hashing
+		FUObjectHashTables& ThreadHash = FUObjectHashTables::Get();
+		Result = StaticFindObjectFastInternalThreadSafe(ThreadHash, ObjectClass, ObjectPackage, ObjectName, bExactClass, /*bAnyPackage =*/ false, ExcludeFlags, ExclusiveInternalFlags);
+	});
 
-	// If they specified an outer use that during the hashing
-	FUObjectHashTables& ThreadHash = FUObjectHashTables::Get();
-	UObject* Result = StaticFindObjectFastInternalThreadSafe(ThreadHash, ObjectClass, ObjectPackage, ObjectName, bExactClass, /*bAnyPackage =*/ false, ExcludeFlags, ExclusiveInternalFlags);
 	return Result;
 }
 
