EX_Let handles attempted assignments to null objects when the property has a setter

diff --git a/Engine/Source/Runtime/CoreUObject/Private/UObject/ScriptCore.cpp b/Engine/Source/Runtime/CoreUObject/Private/UObject/ScriptCore.cpp
--- a/Engine/Source/Runtime/CoreUObject/Private/UObject/ScriptCore.cpp
+++ b/Engine/Source/Runtime/CoreUObject/Private/UObject/ScriptCore.cpp
@@ -2826,9 +2826,9 @@ DEFINE_FUNCTION(UObject::execLet)
 
 	if (LocallyKnownProperty)
 	{
-		if (LocallyKnownProperty->HasSetter())
+		// LocalPropertyContainer will be nullptr if we raised LetAccessNone above
+		if (LocallyKnownProperty->HasSetter() && LocalPropertyContainer)
 		{
-			checkf(LocalPropertyContainer != nullptr, TEXT("execLet attempted to use Setter to write to property %s but had no container address"), *LocallyKnownProperty->GetPathName());
 			LocallyKnownProperty->SetValue_InContainer(LocalPropertyContainer, LocalTempResult);
 			Stack.MostRecentPropertyAddress = PreviousPropertyAddress;
 		}
