Added FArchiveFromStructuredArchiveImpl::operator<< overload for FObjectPtr - fixes crash when serializing TSubclassOf.
Removed deprecated SA_FIELD_NAME macro.

diff --git a/Engine/Source/Runtime/Core/Public/Serialization/StructuredArchiveNameHelpers.h b/Engine/Source/Runtime/Core/Public/Serialization/StructuredArchiveNameHelpers.h
--- a/Engine/Source/Runtime/Core/Public/Serialization/StructuredArchiveNameHelpers.h
+++ b/Engine/Source/Runtime/Core/Public/Serialization/StructuredArchiveNameHelpers.h
@@ -73,8 +73,6 @@ namespace UE::StructuredArchive::Private
 	}
 }
 
-#define SA_FIELD_NAME(x) /*DEPRECATED_MACRO(5.1, "Field names no longer need to be wrapped in the SA_FIELD_NAME macro.")*/ FArchiveFieldName(x)
-
 /** Construct a TNamedValue given an ANSI string and value reference. */
 #define SA_VALUE(Name, Value) UE::StructuredArchive::Private::MakeNamedValue(FArchiveFieldName(Name), Value)
 
diff --git a/Engine/Source/Runtime/CoreUObject/Private/Serialization/StructuredArchiveTest.cpp b/Engine/Source/Runtime/CoreUObject/Private/Serialization/StructuredArchiveTest.cpp
new file mode 100644
--- /dev/null
+++ b/Engine/Source/Runtime/CoreUObject/Private/Serialization/StructuredArchiveTest.cpp
@@ -0,0 +1,47 @@
+// Copyright Epic Games, Inc. All Rights Reserved.
+
+#if WITH_LOW_LEVEL_TESTS
+
+#include "Serialization/StructuredArchive.h"
+#include "Serialization/StructuredArchiveSlots.h"
+#include "Serialization/MemoryWriter.h"
+#include "Serialization/Formatters/JsonArchiveOutputFormatter.h"
+#include "Serialization/ObjectAndNameAsStringProxyArchive.h"
+#include "Templates/SubclassOf.h"
+#include "TestHarness.h"
+
+TEST_CASE("UE::CoreUObject::FStructuredArchive::TSubclassOf", "[Smoke]")
+{
+	// Get a UClass* as a TSubclassOf, and add it to a map that we can use as an indices map
+	TSubclassOf<UObject> Obj = UObject::StaticClass();
+	TMap<TObjectPtr<UObject>, FPackageIndex> Map;
+	Map.Add(Obj.Get());
+
+	// Create an output array and memory archive for testing
+	TArray<uint8> Output;
+	FMemoryWriter Ar(Output);
+
+	// Create a proxy archive which converts UObject pointers
+	FObjectAndNameAsStringProxyArchive Proxy(Ar, false);
+
+	// Create a JSON formatter to write to the proxy and set the map on it to allow UObjects to be mapped to indices
+	FJsonArchiveOutputFormatter Fmt(Proxy);
+	Fmt.SetObjectIndicesMap(&Map);
+
+	// Create structured archive and open the root object as a record
+	FStructuredArchive StrAr(Fmt);
+	FStructuredArchiveSlot Root = StrAr.Open();
+	FStructuredArchiveRecord Record = Root.EnterRecord();
+
+	// Serialize the TSubclassOf as a field of the record
+	Record << SA_VALUE(TEXT("Class"), Obj);
+
+	// Flush to memory
+	StrAr.Close();
+
+	// Test that we wrote out what we expected
+	TArray<uint8> Result(ANSITEXTVIEW("{" LINE_TERMINATOR_ANSI "\t\"Class\": \"0\"" LINE_TERMINATOR_ANSI "}"));
+	CHECK_EQUAL(Output, Result);
+}
+
+#endif // WITH_LOW_LEVEL_TESTS
diff --git a/Engine/Source/Runtime/CoreUObject/Public/Templates/SubclassOf.h b/Engine/Source/Runtime/CoreUObject/Public/Templates/SubclassOf.h
--- a/Engine/Source/Runtime/CoreUObject/Public/Templates/SubclassOf.h
+++ b/Engine/Source/Runtime/CoreUObject/Public/Templates/SubclassOf.h
@@ -6,6 +6,8 @@
 
 #include <type_traits>
 
+class FStructuredArchiveSlot;
+
 template <typename T>
 class TSubclassOf;
 
@@ -152,6 +154,11 @@ public:
 		Ar << Class;
 	}
 
+	FORCEINLINE void Serialize(FStructuredArchiveSlot& Slot)
+	{
+		Slot << Class;
+	}
+
 	friend uint32 GetTypeHash(const TSubclassOf& SubclassOf)
 	{
 		return GetTypeHash(SubclassOf.Class);
@@ -188,6 +195,12 @@ FArchive& operator<<(FArchive& Ar, TSubclassOf<T>& SubclassOf)
 	return Ar;
 }
 
+template <typename T>
+void operator<<(FStructuredArchiveSlot Slot, TSubclassOf<T>& SubclassOf)
+{
+	SubclassOf.Serialize(Slot);
+}
+
 #if UE_ENABLE_INCLUDE_ORDER_DEPRECATED_IN_5_2
 #include "CoreMinimal.h"
 #endif
