CIS fox.

#rnx

diff --git a/Engine/Source/Runtime/CoreUObject/Private/Serialization/StructuredArchiveTest.cpp b/Engine/Source/Runtime/CoreUObject/Private/Serialization/StructuredArchiveTest.cpp
--- a/Engine/Source/Runtime/CoreUObject/Private/Serialization/StructuredArchiveTest.cpp
+++ b/Engine/Source/Runtime/CoreUObject/Private/Serialization/StructuredArchiveTest.cpp
@@ -2,6 +2,7 @@
 
 #if WITH_LOW_LEVEL_TESTS
 
+#include "Containers/StringView.h"
 #include "Serialization/StructuredArchive.h"
 #include "Serialization/StructuredArchiveSlots.h"
 #include "Serialization/MemoryWriter.h"
@@ -40,8 +41,8 @@ TEST_CASE("UE::CoreUObject::FStructuredArchive::TSubclassOf", "[Smoke]")
 	StrAr.Close();
 
 	// Test that we wrote out what we expected
-	TArray<uint8> Result(ANSITEXTVIEW("{" LINE_TERMINATOR_ANSI "\t\"Class\": \"0\"" LINE_TERMINATOR_ANSI "}"));
-	CHECK_EQUAL(Output, Result);
+	FAnsiStringView Result = ANSITEXTVIEW("{" LINE_TERMINATOR_ANSI "\t\"Class\": \"0\"" LINE_TERMINATOR_ANSI "}");
+	CHECK_EQUAL(Output, TArray<uint8>((const uint8*)Result.GetData(), Result.Len()));
 }
 
 #endif // WITH_LOW_LEVEL_TESTS
