Restored handling of a null Enum in FEnumProperty

A null Enum can occur in a Blueprint when, for example, disabling a plugin that contained an enum type that the Blueprint depended on.

#rnx

diff --git a/Engine/Source/Runtime/CoreUObject/Private/UObject/EnumProperty.cpp b/Engine/Source/Runtime/CoreUObject/Private/UObject/EnumProperty.cpp
--- a/Engine/Source/Runtime/CoreUObject/Private/UObject/EnumProperty.cpp
+++ b/Engine/Source/Runtime/CoreUObject/Private/UObject/EnumProperty.cpp
@@ -571,17 +571,17 @@ void FEnumProperty::SaveToTag(FPropertyTag& Tag)
 {
 	Super::SaveToTag(Tag);
 
-	const UEnum* LocalEnum = Enum;
-	check(LocalEnum);
-
-	// RobM: Ugly hack so that we can avoid content changes in most of the packages
-	// Update FByteProperty when making changes here.
-	if (LocalEnum->GetPackage()->HasAnyPackageFlags(PKG_CompiledIn))
-	{				
-		Tag.EnumName = LocalEnum->GetFName();
-	}
-	else
+	if (const UEnum* LocalEnum = Enum)
 	{
-		Tag.EnumName = FName(*LocalEnum->GetPathName());
+		// RobM: Ugly hack so that we can avoid content changes in most of the packages
+		// Update FByteProperty when making changes here.
+		if (LocalEnum->GetPackage()->HasAnyPackageFlags(PKG_CompiledIn))
+		{				
+			Tag.EnumName = LocalEnum->GetFName();
+		}
+		else
+		{
+			Tag.EnumName = FName(*LocalEnum->GetPathName());
+		}
 	}
 }
