Missing file from CL 30412410 (GetRestrictedEnumValues)
#rnx

diff --git a/Engine/Source/Runtime/CoreUObject/Public/UObject/ObjectMacros.h b/Engine/Source/Runtime/CoreUObject/Public/UObject/ObjectMacros.h
--- a/Engine/Source/Runtime/CoreUObject/Public/UObject/ObjectMacros.h
+++ b/Engine/Source/Runtime/CoreUObject/Public/UObject/ObjectMacros.h
@@ -1494,11 +1494,11 @@ namespace UM
 		/// [PropertyMetadata] Used for enum properties to define the subset of invalid values as a comma-separated string; values within this subset are not made available.
 		InvalidEnumValues,
 
-		/// [PropertyMetadata] Used for enum properties to define the subset of valid values through a UFunction, e.g. meta=(GetValidEnumValues="FuncName"). 
+		/// [PropertyMetadata] Used for enum properties to define the subset of restricted values through a UFunction, e.g. meta=(GetRestrictedEnumValues="FuncName"). 
 		///
 		/// UFUNCTION()
 		/// TArray<FString> FuncName() const;
-		GetValidEnumValues,
+		GetRestrictedEnumValues,
 		
 		/// [PropertyMetadata] Causes assets to be filtered through a UFunction, e.g. meta=(GetAssetFilter="FuncName"). The UFunction should return true to exclude the asset.
 		///
