Don't try to reconstruct a template when constructing non-nullable's object property default value.

Doing so may result in crashes when trying to re-construct an object on top of itself while it hasn't been fully loaded yet

diff --git a/Engine/Source/Runtime/CoreUObject/Private/UObject/PropertyBaseObject.cpp b/Engine/Source/Runtime/CoreUObject/Private/UObject/PropertyBaseObject.cpp
--- a/Engine/Source/Runtime/CoreUObject/Private/UObject/PropertyBaseObject.cpp
+++ b/Engine/Source/Runtime/CoreUObject/Private/UObject/PropertyBaseObject.cpp
@@ -644,15 +644,15 @@ UObject* FObjectPropertyBase::ConstructDefaultObjectValueIfNecessary(UObject* Ex
 		// Sanity check to make sure the existing value class matches the property class
 		if (ExistingValueClass && (ExistingValueClass->IsChildOf(PropertyClass) || ExistingValueClass->GetAuthoritativeClass()->IsChildOf(PropertyClass)))
 		{
-			if (ExistingValue->IsTemplate())
+			if (ExistingValue->IsTemplate() && 	// Existing value is a template so we can construct a new value with it as the archetype
+				ExistingValue->GetOuter() != Outer) // Unless the template's Outer is the same as the new Outer in which case the template (ExistingValue) IS the object we can reuse
 			{
-				// Existing value is a template so we can construct a new value with it as the archetype
 				// We probably got here because an object value failed to load (missing import class) and the property is left with a template of default subobject
 				NewDefaultObjectValue = NewObject<UObject>(Outer, ExistingValue->GetClass(), ExistingValue->GetFName(), RF_NoFlags, ExistingValue);
 			}
 			else
 			{
-				// Existing value is not a template so we can use it directly
+				// Existing value is not a template or a template is what this property was pointing to so we can use it directly
 				// Similar to the above condition but the property was not referencing an instanced value in which case it's ok to leave the CDO default here
 				NewDefaultObjectValue = ExistingValue;
 			}
