[Bugfix] Fixed bugs related to loose properties in array elements

diff --git a/Engine/Source/Runtime/CoreUObject/Private/UObject/ArchetypeUtils.cpp b/Engine/Source/Runtime/CoreUObject/Private/UObject/ArchetypeUtils.cpp
--- a/Engine/Source/Runtime/CoreUObject/Private/UObject/ArchetypeUtils.cpp
+++ b/Engine/Source/Runtime/CoreUObject/Private/UObject/ArchetypeUtils.cpp
@@ -235,10 +235,10 @@ namespace UE
     }
 
 	static UStruct* CreatePropertyBagArchetypeStructRec(const UClass* StructClass, UStruct* TemplateStruct,
-		UObject* Outer, const TMap<FWildcardPropertyPathName, TArray<const FProperty*>>& LooseProperties, FWildcardPropertyPathName& Path);
+		UObject* Outer, const TMap<FWildcardPropertyPathName, TMap<FName, const FProperty*>>& LooseProperties, FWildcardPropertyPathName& Path);
 	template <typename TStructType>
 	TStructType* CreatePropertyBagArchetypeStructRec(UStruct* TemplateStruct, UObject* Outer,
-		const TMap<FWildcardPropertyPathName, TArray<const FProperty*>>& LooseProperties, FWildcardPropertyPathName& Path)
+		const TMap<FWildcardPropertyPathName, TMap<FName, const FProperty*>>& LooseProperties, FWildcardPropertyPathName& Path)
 	{
 		return CastChecked<TStructType>(CreatePropertyBagArchetypeStructRec(TStructType::StaticClass(), TemplateStruct, Outer, LooseProperties, Path));
 	}
@@ -254,7 +254,7 @@ namespace UE
 
 	// recursively re-instances all structs contained by this property to include loose properties
 	static void ConvertToArchetypeProperty(FProperty* Property, UObject* Outer,
-		const TMap<FWildcardPropertyPathName, TArray<const FProperty*>>& LooseProperties, FWildcardPropertyPathName& Path)
+		const TMap<FWildcardPropertyPathName, TMap<FName, const FProperty*>>& LooseProperties, FWildcardPropertyPathName& Path)
 	{
 		if (FStructProperty* AsStructProperty = CastField<FStructProperty>(Property))
 		{
@@ -262,22 +262,20 @@ namespace UE
 			AsStructProperty->Struct = CreatePropertyBagArchetypeStructRec<UScriptStruct>(AsStructProperty->Struct, Outer, LooseProperties, Path);
 #if WITH_EDITORONLY_DATA
 			AsStructProperty->SetMetaData("TPSOverrideStructName", TPSOverrideStructName.ToString());
+			AsStructProperty->SetMetaData("OriginalType", TPSOverrideStructName.ToString());
 #endif
 		}
 		else if (const FArrayProperty* AsArrayProperty = CastField<FArrayProperty>(Property))
 		{
-			Path.Push(CreateSegmentFromProperty(AsArrayProperty->Inner));
 			ConvertToArchetypeProperty(AsArrayProperty->Inner, Outer, LooseProperties, Path);
-			Path.Pop();
 		}
 		else if (const FSetProperty* AsSetProperty = CastField<FSetProperty>(Property))
 		{
-			Path.Push(CreateSegmentFromProperty(AsSetProperty->ElementProp));
 			ConvertToArchetypeProperty(AsSetProperty->ElementProp, Outer, LooseProperties, Path);
-			Path.Pop();
 		}
 		else if (const FMapProperty* AsMapProperty = CastField<FMapProperty>(Property))
 		{
+			// todo: This will likely need revisiting once devin has maps working
 			Path.Push(CreateSegmentFromProperty(AsMapProperty->KeyProp));
 			ConvertToArchetypeProperty(AsMapProperty->KeyProp, Outer, LooseProperties, Path);
 			Path.Pop();
@@ -290,7 +288,7 @@ namespace UE
 	
 	// copy template property then convert it into an archetype property by adding loose properties
 	static FProperty* CreateArchetypeProperty(const FProperty* TemplateProperty, UObject* Outer,
-		const TMap<FWildcardPropertyPathName, TArray<const FProperty*>>& LooseProperties, FWildcardPropertyPathName& Path)
+		const TMap<FWildcardPropertyPathName, TMap<FName, const FProperty*>>& LooseProperties, FWildcardPropertyPathName& Path)
 	{
 		FProperty* ArchetypeProperty = CastFieldChecked<FProperty>(FField::Duplicate(TemplateProperty, Outer));
 #if WITH_EDITORONLY_DATA
@@ -315,9 +313,9 @@ namespace UE
 	}
 
 	// recursively add all the wildcard paths of both Property and all it's sub-Properties to OutLooseProperties
-	static void AddWildcardedProperties(TMap<FWildcardPropertyPathName, TArray<const FProperty*>>& OutProperties, FWildcardPropertyPathName& ParentPath, const FProperty* Property)
+	static void AddWildcardedProperties(TMap<FWildcardPropertyPathName, TMap<FName, const FProperty*>>& OutProperties, FWildcardPropertyPathName& ParentPath, const FProperty* Property)
 	{
-		OutProperties.FindOrAdd(ParentPath).Add(Property);
+		OutProperties.FindOrAdd(ParentPath).Add(Property->GetFName(), Property);
 		
 		ParentPath.Push(CreateSegmentFromProperty(Property));
 		if (const FStructProperty* AsStructProperty = CastField<FStructProperty>(Property))
@@ -344,10 +342,10 @@ namespace UE
 	}
 
 	// construct a map that keys a parent struct by it's wildcard path and returns an array of all it's loose properties
-	static TMap<FWildcardPropertyPathName, TArray<const FProperty*>> GetWildcardedLooseProperties(const FPropertyBag* PropertyBag)
+	static TMap<FWildcardPropertyPathName, TMap<FName, const FProperty*>> GetWildcardedLooseProperties(const FPropertyBag* PropertyBag)
 	{
 		
-		TMap<FWildcardPropertyPathName, TArray<const FProperty*>> LooseProperties;
+		TMap<FWildcardPropertyPathName, TMap<FName, const FProperty*>> LooseProperties;
 		if (PropertyBag)
 		{
 			for (FPropertyBag::FConstIterator Itr = PropertyBag->CreateConstIterator(); Itr; ++Itr)
@@ -386,18 +384,22 @@ namespace UE
 
 	// constructs an archetype struct by merging the properties in 
 	static UStruct* CreatePropertyBagArchetypeStructRec(const UClass* StructClass, UStruct* TemplateStruct,
-		UObject* Outer, const TMap<FWildcardPropertyPathName, TArray<const FProperty*>>& LooseProperties, FWildcardPropertyPathName& Path)
+		UObject* Outer, const TMap<FWildcardPropertyPathName, TMap<FName, const FProperty*>>& LooseProperties, FWildcardPropertyPathName& Path)
 	{
 		UStruct* Super = nullptr;
 
-		const TArray<const FProperty*>* BagProperties = LooseProperties.Find(Path);
+		const TMap<FName, const FProperty*>* BagProperties = LooseProperties.Find(Path);
 
 		auto MatchesBagProperty = [&BagProperties](const FProperty* Property)
 		{
-			return BagProperties && BagProperties->ContainsByPredicate([Property](const FProperty* Found)
+			if (BagProperties)
 			{
-				return Found->SameType(Property) && Found->GetFName() == Property->GetFName();
-			});
+				if (const FProperty* const* Found = BagProperties->Find(Property->GetFName()))
+				{
+					return (*Found)->SameType(Property);
+				}
+			}
+			return false;
 		};
 		
 		if (TemplateStruct)
@@ -448,10 +450,10 @@ namespace UE
 		TArray<FProperty*> LooseArchetypeProperties;
 		if (BagProperties)
 		{
-			for (const FProperty* BagProperty : *BagProperties)
+			for (const TPair<FName, const FProperty*>& BagProperty : *BagProperties)
 			{
-				Path.Push(CreateSegmentFromProperty(BagProperty));
-				FProperty* LooseProperty = CreateArchetypeProperty(BagProperty, Result, LooseProperties, Path);
+				Path.Push(CreateSegmentFromProperty(BagProperty.Value));
+				FProperty* LooseProperty = CreateArchetypeProperty(BagProperty.Value, Result, LooseProperties, Path);
 				Path.Pop();
 				
 				MarkPropertyAsLoose(LooseProperty);
@@ -482,7 +484,7 @@ namespace UE
 	
 	UClass* CreatePropertyBagArchetypeClass(const FPropertyBag* PropertyBag, UStruct* TemplateStruct, UObject* Outer)
 	{
-		const TMap<FWildcardPropertyPathName, TArray<const FProperty*>> LooseProperties = GetWildcardedLooseProperties(PropertyBag);
+		const TMap<FWildcardPropertyPathName, TMap<FName, const FProperty*>> LooseProperties = GetWildcardedLooseProperties(PropertyBag);
 		FWildcardPropertyPathName ParentPath;
 		return CreatePropertyBagArchetypeStructRec<UClass>(TemplateStruct, Outer, LooseProperties, ParentPath);
 	}
