FoundationTests: Fix compilation error on non-windows platforms from StructuredArchiveTest.cpp

diff --git a/Engine/Source/Runtime/CoreUObject/Private/Serialization/StructuredArchiveTest.cpp b/Engine/Source/Runtime/CoreUObject/Private/Serialization/StructuredArchiveTest.cpp
--- a/Engine/Source/Runtime/CoreUObject/Private/Serialization/StructuredArchiveTest.cpp
+++ b/Engine/Source/Runtime/CoreUObject/Private/Serialization/StructuredArchiveTest.cpp
@@ -1,6 +1,6 @@
 // Copyright Epic Games, Inc. All Rights Reserved.
 
-#if WITH_LOW_LEVEL_TESTS
+#if WITH_LOW_LEVEL_TESTS && WITH_TEXT_ARCHIVE_SUPPORT
 
 #include "Containers/StringView.h"
 #include "Serialization/StructuredArchive.h"
@@ -45,4 +45,4 @@ TEST_CASE("UE::CoreUObject::FStructuredArchive::TSubclassOf", "[Smoke]")
 	CHECK_EQUAL(Output, TArray<uint8>((const uint8*)Result.GetData(), Result.Len()));
 }
 
-#endif // WITH_LOW_LEVEL_TESTS
+#endif // WITH_LOW_LEVEL_TESTS && WITH_TEXT_ARCHIVE_SUPPORT
