Fixing property reporting when dumping GC references

diff --git a/Engine/Source/Runtime/CoreUObject/Private/UObject/GarbageCollection.cpp b/Engine/Source/Runtime/CoreUObject/Private/UObject/GarbageCollection.cpp
--- a/Engine/Source/Runtime/CoreUObject/Private/UObject/GarbageCollection.cpp
+++ b/Engine/Source/Runtime/CoreUObject/Private/UObject/GarbageCollection.cpp
@@ -2981,6 +2981,8 @@ struct TBatchDispatcher
 	static constexpr bool bBatching = true;
 	static constexpr bool bParallel = IsParallel(ProcessorType::Options);
 
+	typedef FDebugSchemaStackNoOpScope SchemaStackScopeType;
+
 	FWorkerContext& Context;
 	const FPermanentObjectPoolExtents PermanentPool;
 	FReferenceCollector& Collector;
@@ -3166,6 +3168,10 @@ struct TBatchDispatcher
 		FlushQueuedReferences();
 		Context.IncrementalStructs = StructBatcher.Suspend();
 	}
+
+	void SetDebugSchemaStackMemberId(FMemberId Member)
+	{
+	}
 };
 
 //////////////////////////////////////////////////////////////////////////
@@ -3577,6 +3583,26 @@ public:
 
 #endif // !UE_BUILD_SHIPPING
 
+FString FDebugSchemaStackNode::ToString() const
+{
+	FString Result;
+#if !UE_BUILD_SHIPPING
+	for (const FDebugSchemaStackNode* Node = this; Node; Node = Node->Prev)
+	{
+		FMemberInfo MemberInfo = GetMemberDebugInfo(Node->Schema, Node->Member);
+		FString MemberInfoString = FString::Printf(TEXT("%s (0x%x)"), *MemberInfo.Name.GetPlainNameString(), MemberInfo.Offset);
+		if (!Result.IsEmpty())
+		{
+			MemberInfoString += TEXT(" -> ");
+		}
+		Result = MemberInfoString + Result;
+	}
+#else
+	Result += TEXT("Unknown");
+#endif
+	return Result;
+}
+
 FORCEINLINE static void CheckReference(UObject*& Object, const UObject* ReferencingObject, const FProperty* ReferencingProperty)
 {
 #if ENABLE_GC_OBJECT_CHECKS
@@ -3655,6 +3681,7 @@ void TFastReferenceCollector<ProcessorType, CollectorType>::ProcessStructs(Dispa
 				check(AoS.Stride == AoS.Schema.GetStructStride());
 				uint8* StructIt = AoS.Data;
 				uint8* StructEnd = AoS.Data + AoS.Num * AoS.Stride;
+				typename DispatcherType::SchemaStackScopeType SchemaStack(Dispatcher.Context, AoS.Schema);
 				do
 				{
 					Private::VisitMembers(Dispatcher, AoS.Schema, StructIt); 
diff --git a/Engine/Source/Runtime/CoreUObject/Private/UObject/GarbageCollectionVerification.cpp b/Engine/Source/Runtime/CoreUObject/Private/UObject/GarbageCollectionVerification.cpp
--- a/Engine/Source/Runtime/CoreUObject/Private/UObject/GarbageCollectionVerification.cpp
+++ b/Engine/Source/Runtime/CoreUObject/Private/UObject/GarbageCollectionVerification.cpp
@@ -90,9 +90,8 @@ public:
 				FString DebugInfo;
 				if (UClass *Class = (ReferencingObject ? ReferencingObject->GetClass() : nullptr))
 				{
-					UE::GC::FMemberInfo Member = GetMemberDebugInfo(Class->ReferenceSchema.Get(), MemberId);
-					DebugInfo = FString::Printf(TEXT("ReferencingObjectClass: %s, Property Name: %s, Offset: %d"),
-						*Class->GetFullName(), *Member.Name.GetPlainNameString(), Member.Offset);
+					DebugInfo = FString::Printf(TEXT("ReferencingObjectClass: %s, Property: %s"),
+						*Class->GetFullName(), *ObjectsToSerializeStruct.SchemaStack->ToString());
 				}
 				else
 				{
@@ -212,9 +211,8 @@ public:
 				FString DebugInfo;
 				if (UClass *Class = (ReferencingObject ? ReferencingObject->GetClass() : nullptr))
 				{
-					UE::GC::FMemberInfo Member = GetMemberDebugInfo(Class->ReferenceSchema.Get(), MemberId);
-					DebugInfo = FString::Printf(TEXT("ReferencingObjectClass: %s, Property Name: %s, Offset: %d"),
-						*Class->GetFullName(), *Member.Name.GetPlainNameString(), Member.Offset);
+					DebugInfo = FString::Printf(TEXT("ReferencingObjectClass: %s, Property: %s"),
+						*Class->GetFullName(), *ObjectsToSerializeStruct.SchemaStack->ToString());
 				}
 				else
 				{
@@ -456,9 +454,8 @@ public:
 				FString DebugInfo;
 				if (UClass* Class = (ReferencingObject ? ReferencingObject->GetClass() : nullptr))
 				{
-					UE::GC::FMemberInfo Member = GetMemberDebugInfo(Class->ReferenceSchema.Get(), MemberId);
-					DebugInfo = FString::Printf(TEXT("property: %s, offset: %d"),
-						*Member.Name.GetPlainNameString(), Member.Offset);
+					DebugInfo = FString::Printf(TEXT("property: %s"),
+						*ObjectsToSerializeStruct.SchemaStack->ToString());
 				}
 				else
 				{
diff --git a/Engine/Source/Runtime/CoreUObject/Public/UObject/FastReferenceCollector.h b/Engine/Source/Runtime/CoreUObject/Public/UObject/FastReferenceCollector.h
--- a/Engine/Source/Runtime/CoreUObject/Public/UObject/FastReferenceCollector.h
+++ b/Engine/Source/Runtime/CoreUObject/Public/UObject/FastReferenceCollector.h
@@ -290,6 +290,37 @@ struct FWeakReferenceInfo
 	UObject* ReferenceOwner = nullptr;
 };
 
+/** Maintains a stack of schemas currently processed by reachability analysis for debugging referencing property names */
+struct FDebugSchemaStackNode
+{
+#if !UE_BUILD_SHIPPING
+	FMemberId Member;
+	FSchemaView Schema;
+	FDebugSchemaStackNode* Prev;
+	
+	FDebugSchemaStackNode()
+		: Member(0)
+		, Prev(nullptr)
+	{
+	}
+	FDebugSchemaStackNode(FSchemaView InSchema, FDebugSchemaStackNode* PrevNode)
+		: Member(0)
+		, Schema(InSchema)
+		, Prev(PrevNode)
+	{
+	}
+#endif // !UE_BUILD_SHIPPING
+
+	FORCEINLINE void SetMemberId(FMemberId MemberId)
+	{
+#if !UE_BUILD_SHIPPING
+		Member = MemberId;
+#endif
+	}
+
+	COREUOBJECT_API FString ToString() const;
+};
+
 /** Thread-local context containing initial objects and references to collect */
 struct alignas(PLATFORM_CACHE_LINE_SIZE) FWorkerContext
 {
@@ -323,6 +354,8 @@ public:
 	bool bIsSuspended = false;
 	bool bDidWork = false;
 
+	FDebugSchemaStackNode* SchemaStack = nullptr;
+
 	FORCEINLINE UObject* GetReferencingObject()	{ return ReferencingObject;	}
 
 	TConstArrayView<UObject*> GetInitialObjects() { return InitialObjects; }
@@ -355,6 +388,40 @@ public:
 
 //////////////////////////////////////////////////////////////////////////
 
+struct FDebugSchemaStackScope
+{
+#if !UE_BUILD_SHIPPING
+	FWorkerContext& Context;
+	FDebugSchemaStackNode Node;
+#endif
+
+	FDebugSchemaStackScope(FWorkerContext& InContext, FSchemaView Schema)
+#if !UE_BUILD_SHIPPING
+		: Context(InContext)
+		, Node(Schema, InContext.SchemaStack)
+#endif
+	{
+#if !UE_BUILD_SHIPPING
+		InContext.SchemaStack = &Node;
+#endif
+	}
+	~FDebugSchemaStackScope()
+	{
+#if !UE_BUILD_SHIPPING
+		Context.SchemaStack = Node.Prev;
+#endif
+	}
+};
+
+struct FDebugSchemaStackNoOpScope
+{
+	FDebugSchemaStackNoOpScope(FWorkerContext& InContext, FSchemaView Schema)
+	{
+	}
+};
+
+//////////////////////////////////////////////////////////////////////////
+
 namespace Private {
 
 struct FMemberUnpacked
@@ -440,6 +507,7 @@ FORCEINLINE_DEBUGGABLE void VisitStructs(DispatcherType& Dispatcher, FSchemaView
 template<class DispatcherType, class ArrayType>
 FORCEINLINE_DEBUGGABLE void VisitStructArray(DispatcherType& Dispatcher, FSchemaView StructSchema, ArrayType& Array)
 {
+	typename DispatcherType::SchemaStackScopeType SchemaStack(Dispatcher.Context, StructSchema);
 	VisitStructs(Dispatcher, StructSchema, (uint8*)Array.GetData(), Array.Num());
 }
 
@@ -459,6 +527,7 @@ FORCEINLINE_DEBUGGABLE void VisitSparseStructArray(DispatcherType& Dispatcher, F
 		{
 			if (Array.IsAllocated(Idx))
 			{
+				typename DispatcherType::SchemaStackScopeType SchemaStack(Dispatcher.Context, StructSchema);
 				VisitNestedStructMembers(Dispatcher, StructSchema, It);
 			}
 		}
@@ -496,6 +565,7 @@ FORCEINLINE_DEBUGGABLE void VisitOptional(DispatcherType& Dispatcher, FSchemaVie
 	check(!StructSchema.IsEmpty());
 	uint32 ValueSize = StructSchema.GetStructStride();
 	bool bIsSet = *(bool*)(Instance + ValueSize);
+	typename DispatcherType::SchemaStackScopeType SchemaStack(Dispatcher.Context, StructSchema);
 	VisitStructs(Dispatcher, StructSchema, Instance, bIsSet);
 }
 
@@ -567,6 +637,7 @@ FORCEINLINE_DEBUGGABLE void VisitMembers(DispatcherType& Dispatcher, FSchemaView
 		for (FMemberUnpacked Member : Quad.Members)
 		{
 			uint8* MemberPtr = (uint8*)(InstanceCursor + Member.WordOffset);
+			Dispatcher.SetDebugSchemaStackMemberId(FMemberId(DebugIdx));
 
 			switch (Member.Type)
 			{
@@ -635,6 +706,7 @@ struct TDirectDispatcher
 	static constexpr bool bBatching = false;
 	static constexpr bool bParallel = IsParallel(ProcessorType::Options);
 
+	typedef FDebugSchemaStackScope SchemaStackScopeType;
 	ProcessorType& Processor;
 	FWorkerContext& Context;
 	FReferenceCollector& Collector;
@@ -721,6 +793,11 @@ struct TDirectDispatcher
 	void Suspend()
 	{
 	}
+
+	void SetDebugSchemaStackMemberId(FMemberId Member)
+	{
+		Context.SchemaStack->SetMemberId(Member);
+	}
 };
 
 // Default implementation is to create new direct dispatcher
@@ -839,6 +916,7 @@ StoleARO:
 		} // while (true)
 		
 		Processor.LogDetailedStatsSummary();
+		checkf(Context.SchemaStack == nullptr, TEXT("Debug Schema Stack is not empty after processing all objects"));
 	}
 
 private:
@@ -874,6 +952,7 @@ private:
 #endif
 			if (!Schema.IsEmpty())
 			{
+				typename DispatcherType::SchemaStackScopeType SchemaStack(Dispatcher.Context, Schema);
 				Processor.BeginTimingObject(CurrentObject);
 				Private::VisitMembers(Dispatcher, Schema, CurrentObject);
 				Processor.UpdateDetailedStats(CurrentObject);
