Added bTrackSerializedPropertyPath to FUObjectSerializeContext to disable property path tracking and its overhead by default

#rnx

diff --git a/Engine/Plugins/Editor/ArchetypeFixupTool/Source/Private/Tests/ArchetypeTests.cpp b/Engine/Plugins/Editor/ArchetypeFixupTool/Source/Private/Tests/ArchetypeTests.cpp
--- a/Engine/Plugins/Editor/ArchetypeFixupTool/Source/Private/Tests/ArchetypeTests.cpp
+++ b/Engine/Plugins/Editor/ArchetypeFixupTool/Source/Private/Tests/ArchetypeTests.cpp
@@ -230,13 +230,16 @@ static UObject* GenerateTestArchetype(UObject* ObjectOld, UClass* NewClass)
 	UObject* ObjectNew = NewObject<UObject>(GetTransientPackage(), NewClass);
 	{
 		FUObjectSerializeContext* LoadContext = FUObjectThreadContext::Get().GetSerializeContext();
+		const bool bRevertTrackSerializedPropertyPath = LoadContext->bTrackSerializedPropertyPath;
 		const bool bRevertSerializeUnknownProperty = LoadContext->bSerializeUnknownProperty;
 		UObject* RevertSerializedObject = LoadContext->SerializedObject;
+		LoadContext->bTrackSerializedPropertyPath = true;
 		LoadContext->bSerializeUnknownProperty = true;
 		LoadContext->SerializedObject = ObjectNew;
 		
 		FObjectReader(ObjectNew, SerializedData);
 		
+		LoadContext->bTrackSerializedPropertyPath = bRevertTrackSerializedPropertyPath;
 		LoadContext->bSerializeUnknownProperty = bRevertSerializeUnknownProperty;
 		LoadContext->SerializedObject = RevertSerializedObject;
 	}
diff --git a/Engine/Source/Runtime/CoreUObject/Private/UObject/Class.cpp b/Engine/Source/Runtime/CoreUObject/Private/UObject/Class.cpp
--- a/Engine/Source/Runtime/CoreUObject/Private/UObject/Class.cpp
+++ b/Engine/Source/Runtime/CoreUObject/Private/UObject/Class.cpp
@@ -1569,7 +1569,7 @@ void UStruct::SerializeVersionedTaggedProperties(FStructuredArchive::FSlot Slot,
 					Property = CustomFindProperty(Tag.Name);
 				}
 
-				if (SerializeContext)
+				if (SerializeContext && SerializeContext->bTrackSerializedPropertyPath)
 				{
 					const FName Name = Property ? Property->GetFName() : Tag.Name;
 					const int32 Index = Tag.ArrayIndex > 0 || (Property && Property->ArrayDim > 1) ? Tag.ArrayIndex : INDEX_NONE;
@@ -1713,7 +1713,7 @@ void UStruct::SerializeVersionedTaggedProperties(FStructuredArchive::FSlot Slot,
 					PropertyBag->LoadPropertyByTag(SerializeContext->SerializedPropertyPath, Tag, ValueSlot);
 				}
 
-				if (SerializeContext)
+				if (SerializeContext && SerializeContext->bTrackSerializedPropertyPath)
 				{
 					SerializeContext->SerializedPropertyPath.Pop();
 				}
diff --git a/Engine/Source/Runtime/CoreUObject/Private/UObject/PropertyArray.cpp b/Engine/Source/Runtime/CoreUObject/Private/UObject/PropertyArray.cpp
--- a/Engine/Source/Runtime/CoreUObject/Private/UObject/PropertyArray.cpp
+++ b/Engine/Source/Runtime/CoreUObject/Private/UObject/PropertyArray.cpp
@@ -616,7 +616,7 @@ void FArrayProperty::SerializeItem(FStructuredArchive::FSlot Slot, void* Value,
 	}
 
 	FUObjectSerializeContext* Context = FUObjectThreadContext::Get().GetSerializeContext();
-	if (Context && MaybeInnerTag)
+	if (Context && Context->bTrackSerializedPropertyPath && MaybeInnerTag)
 	{
 		Context->SerializedPropertyPath.PushType(MaybeInnerTag->StructName);
 	}
@@ -658,7 +658,7 @@ void FArrayProperty::SerializeItem(FStructuredArchive::FSlot Slot, void* Value,
 
 				// Serialize the item at this array index
 				i = PropertyNode->ArrayIndex;
-				if (Context)
+				if (Context && Context->bTrackSerializedPropertyPath)
 				{
 					Context->SerializedPropertyPath.SetIndex(i);
 				}
@@ -686,7 +686,7 @@ void FArrayProperty::SerializeItem(FStructuredArchive::FSlot Slot, void* Value,
 			NAME_UArraySerializeCount.SetNumber(i);
 			FArchive::FScopeAddDebugData P(UnderlyingArchive, NAME_UArraySerializeCount);
 #endif
-			if (Context)
+			if (Context && Context->bTrackSerializedPropertyPath)
 			{
 				Context->SerializedPropertyPath.SetIndex(i);
 			}
@@ -718,7 +718,7 @@ void FArrayProperty::SerializeItem(FStructuredArchive::FSlot Slot, void* Value,
 		}
 	}
 
-	if (Context)
+	if (Context && Context->bTrackSerializedPropertyPath)
 	{
 		Context->SerializedPropertyPath.SetIndex(INDEX_NONE);
 	}
diff --git a/Engine/Source/Runtime/CoreUObject/Private/UObject/PropertyMap.cpp b/Engine/Source/Runtime/CoreUObject/Private/UObject/PropertyMap.cpp
--- a/Engine/Source/Runtime/CoreUObject/Private/UObject/PropertyMap.cpp
+++ b/Engine/Source/Runtime/CoreUObject/Private/UObject/PropertyMap.cpp
@@ -728,7 +728,7 @@ void FMapProperty::SerializeItem(FStructuredArchive::FSlot Slot, void* Value, co
 			{
 				FStructuredArchive::FRecord EntryRecord = EntriesArray.EnterElement().EnterRecord();
 				int32 Index = MapHelper.AddDefaultValue_Invalid_NeedsRehash();
-				if (Context)
+				if (Context && Context->bTrackSerializedPropertyPath)
 				{
 					Context->SerializedPropertyPath.SetIndex(Index);
 				}
diff --git a/Engine/Source/Runtime/CoreUObject/Private/UObject/PropertySet.cpp b/Engine/Source/Runtime/CoreUObject/Private/UObject/PropertySet.cpp
--- a/Engine/Source/Runtime/CoreUObject/Private/UObject/PropertySet.cpp
+++ b/Engine/Source/Runtime/CoreUObject/Private/UObject/PropertySet.cpp
@@ -304,7 +304,7 @@ void FSetProperty::SerializeItem(FStructuredArchive::FSlot Slot, void* Value, co
 			for (; Num; --Num)
 			{
 				int32 Index = SetHelper.AddDefaultValue_Invalid_NeedsRehash();
-				if (Context)
+				if (Context && Context->bTrackSerializedPropertyPath)
 				{
 					Context->SerializedPropertyPath.SetIndex(Index);
 				}
@@ -387,7 +387,7 @@ void FSetProperty::SerializeItem(FStructuredArchive::FSlot Slot, void* Value, co
 			}
 		}
 
-		if (Context)
+		if (Context && Context->bTrackSerializedPropertyPath)
 		{
 			Context->SerializedPropertyPath.SetIndex(INDEX_NONE);
 		}
diff --git a/Engine/Source/Runtime/CoreUObject/Private/UObject/UObjectThreadContext.cpp b/Engine/Source/Runtime/CoreUObject/Private/UObject/UObjectThreadContext.cpp
--- a/Engine/Source/Runtime/CoreUObject/Private/UObject/UObjectThreadContext.cpp
+++ b/Engine/Source/Runtime/CoreUObject/Private/UObject/UObjectThreadContext.cpp
@@ -44,6 +44,7 @@ FUObjectSerializeContext::FUObjectSerializeContext()
 	, SerializedImportLinker(nullptr)
 	, SerializedExportIndex(0)
 	, SerializedExportLinker(nullptr)
+	, bTrackSerializedPropertyPath(false)
 	, bSerializeUnknownProperty(false)
 {}
 
diff --git a/Engine/Source/Runtime/CoreUObject/Public/UObject/UObjectThreadContext.h b/Engine/Source/Runtime/CoreUObject/Public/UObject/UObjectThreadContext.h
--- a/Engine/Source/Runtime/CoreUObject/Public/UObject/UObjectThreadContext.h
+++ b/Engine/Source/Runtime/CoreUObject/Public/UObject/UObjectThreadContext.h
@@ -169,9 +169,11 @@ public:
 	/** Points to the most recently used Linker for serialization by CreateExport() */
 	FLinkerLoad* SerializedExportLinker;
 	/** Path to the property currently being serialized */
-	UE::FPropertyPathName SerializedPropertyPath;
+	UE_INTERNAL UE::FPropertyPathName SerializedPropertyPath;
+	/** True when SerializedPropertyPath is being tracked during serialization. */
+	UE_INTERNAL bool bTrackSerializedPropertyPath;
 	/** True when unknown properties will be serialized to or from a property bag for the serialized object. */
-	bool bSerializeUnknownProperty;
+	UE_INTERNAL bool bSerializeUnknownProperty;
 
 	/** Adds a new loaded object */
 	COREUOBJECT_API void AddLoadedObject(UObject* InObject);
