Tweak the stack walk in TryFindTypeSlow to dump address/module if we don't find symbols (and ignore the stack trace method itself).

diff --git a/Engine/Source/Runtime/CoreUObject/Private/UObject/Class.cpp b/Engine/Source/Runtime/CoreUObject/Private/UObject/Class.cpp
--- a/Engine/Source/Runtime/CoreUObject/Private/UObject/Class.cpp
+++ b/Engine/Source/Runtime/CoreUObject/Private/UObject/Class.cpp
@@ -6475,19 +6475,23 @@ UField* UClass::TryFindTypeSlow(UClass* TypeClass, const FString& InPathNameOrSh
 			ANSICHAR Buffer[1024];			
 			uint64 StackFrames[10];
 			uint32 NumStackFrames = FPlatformStackWalk::CaptureStackBackTrace(StackFrames, UE_ARRAY_COUNT(StackFrames));
-			for (uint32 Idx = 0; Idx < NumStackFrames && Idx < UE_ARRAY_COUNT(StackFrames); Idx++)
+			const uint32 IgnoreStackCount = 1; // Ignore the call to CaptureStackBackTrace itself
+			for (uint32 Idx = IgnoreStackCount; Idx < NumStackFrames && Idx < UE_ARRAY_COUNT(StackFrames); Idx++)
 			{
 				Buffer[0] = '\0';
-				FPlatformStackWalk::ProgramCounterToHumanReadableString(Idx, StackFrames[Idx], Buffer, sizeof(Buffer));
-				ANSICHAR* TrimmedBuffer = FCStringAnsi::Strstr(Buffer, "!");
-				if (!TrimmedBuffer)
-				{
-					TrimmedBuffer = Buffer;
-				}
-				else
+				const ANSICHAR* TrimmedBuffer = Buffer;
+
+				// Trim the address/module only if we resolve the symbol
+				const bool bFoundSymbol = FPlatformStackWalk::ProgramCounterToHumanReadableString(Idx, StackFrames[Idx], Buffer, sizeof(Buffer));
+				if (bFoundSymbol)
 				{
-					TrimmedBuffer++;
+					const ANSICHAR* BufferAfterModuleAndAddress = FCStringAnsi::Strstr(Buffer, "!");
+					if (BufferAfterModuleAndAddress)
+					{
+						TrimmedBuffer = BufferAfterModuleAndAddress + 1;
+					}
 				}
+
 				Callstack.Append(TrimmedBuffer);
 				Callstack.Append(TEXT("\r\n"));
 			}
