Use unversioned property serialization when duplicating objects - this speeds up blueprint compilation by about 10% and will also improve other prolific duplicators

diff --git a/Engine/Source/Runtime/CoreUObject/Private/Serialization/DuplicateDataReader.cpp b/Engine/Source/Runtime/CoreUObject/Private/Serialization/DuplicateDataReader.cpp
--- a/Engine/Source/Runtime/CoreUObject/Private/Serialization/DuplicateDataReader.cpp
+++ b/Engine/Source/Runtime/CoreUObject/Private/Serialization/DuplicateDataReader.cpp
@@ -27,6 +27,7 @@ FDuplicateDataReader::FDuplicateDataReader( class FUObjectAnnotationSparse<FDupl
 {
 	this->SetIsLoading(true);
 	this->SetIsPersistent(true);
+	this->SetUseUnversionedPropertySerialization(true);
 	this->ArNoIntraPropertyDelta = true;
 	ArPortFlags |= PPF_Duplicate | InPortFlags;
 
diff --git a/Engine/Source/Runtime/CoreUObject/Private/Serialization/DuplicateDataWriter.cpp b/Engine/Source/Runtime/CoreUObject/Private/Serialization/DuplicateDataWriter.cpp
--- a/Engine/Source/Runtime/CoreUObject/Private/Serialization/DuplicateDataWriter.cpp
+++ b/Engine/Source/Runtime/CoreUObject/Private/Serialization/DuplicateDataWriter.cpp
@@ -38,6 +38,7 @@ FDuplicateDataWriter::FDuplicateDataWriter( FUObjectAnnotationSparse<FDuplicated
 {
 	this->SetIsSaving(true);
 	this->SetIsPersistent(true);
+	this->SetUseUnversionedPropertySerialization(true);
 	this->ArNoIntraPropertyDelta = true;
 	ArAllowLazyLoading	= false;
 	ArPortFlags |= PPF_Duplicate | InPortFlags;
