Removing a check that may run after GC contexts have already been deallocated.

diff --git a/Engine/Source/Runtime/CoreUObject/Public/UObject/FastReferenceCollector.h b/Engine/Source/Runtime/CoreUObject/Public/UObject/FastReferenceCollector.h
--- a/Engine/Source/Runtime/CoreUObject/Public/UObject/FastReferenceCollector.h
+++ b/Engine/Source/Runtime/CoreUObject/Public/UObject/FastReferenceCollector.h
@@ -916,7 +916,6 @@ StoleARO:
 		} // while (true)
 		
 		Processor.LogDetailedStatsSummary();
-		checkf(Context.SchemaStack == nullptr, TEXT("Debug Schema Stack is not empty after processing all objects"));
 	}
 
 private:
