[ZenLoader]
  - Get rid of old packages before trying to find one from memory to get an opportunity to update the import table instead of always trying to reload.
  - Always go throught the export table for deferred postloads to prevent potential missing postload when objects in the export table are reloaded.

diff --git a/Engine/Source/Runtime/CoreUObject/Private/Serialization/AsyncLoading2.cpp b/Engine/Source/Runtime/CoreUObject/Private/Serialization/AsyncLoading2.cpp
--- a/Engine/Source/Runtime/CoreUObject/Private/Serialization/AsyncLoading2.cpp
+++ b/Engine/Source/Runtime/CoreUObject/Private/Serialization/AsyncLoading2.cpp
@@ -1070,6 +1070,15 @@ public:
 		// is this the first reference to a package that already exists?
 		if (PackageRef.RefCount == 0)
 		{
+			// Remove stale package before searching below as its possible a UEDPIE package got trashed and replaced by a new one
+			// and its important that we find the one that replaced it so we don't try to load it if its a PKG_InMemoryOnly package.
+			if (UPackage* Package = PackageRef.GetPackage())
+			{
+				if (Package->IsUnreachable() || PackageRef.GetOriginalPackageName() != Package->GetFName())
+				{
+					RemoveUnreferencedObsoletePackage(PackageRef);
+				}
+			}
 #if WITH_EDITOR
 			if (!PackageRef.HasPackage() && !PackageNameIfKnown.IsNone())
 			{
@@ -5826,11 +5835,16 @@ EEventLoadNodeExecutionResult FAsyncPackage2::ExecuteDeferredPostLoadLinkerLoadP
 	// We can't return timeout during a flush as we're expected to be able to finish
 	const bool bIsReadyForAsyncPostLoadAllowed = ThreadState.SyncLoadContextStack.IsEmpty();
 
-	const int32 ObjectCount = ConstructedObjects.Num();
+	// Go through both ConstructedObjects and export table as its possible to reload objects in the export table
+	// without them being constructed and that would lead to missing postloads.
+	const int32 ConstructedObjectsCount = ConstructedObjects.Num();
+	const int32 ObjectCount = ConstructedObjectsCount + Data.Exports.Num();
 	while (LinkerLoadState->PostLoadExportIndex < ObjectCount)
 	{
 		const int32 ObjectIndex = LinkerLoadState->PostLoadExportIndex++;
-		if (UObject* Object = ConstructedObjects[ObjectIndex])
+
+		UObject* Object = ObjectIndex < ConstructedObjectsCount ? ConstructedObjects[ObjectIndex] : Data.Exports[ObjectIndex - ConstructedObjectsCount].Object;
+		if (Object && Object->HasAnyFlags(RF_NeedPostLoad))
 		{
 			// Only allow to wait when there is no flush waiting on us
 			if (bIsReadyForAsyncPostLoadAllowed && !Object->IsReadyForAsyncPostLoad())
