Improve message on !IsUnreachable() assertion failure in UObject::ProcessEvent().

diff --git a/Engine/Source/Runtime/CoreUObject/Private/UObject/ScriptCore.cpp b/Engine/Source/Runtime/CoreUObject/Private/UObject/ScriptCore.cpp
--- a/Engine/Source/Runtime/CoreUObject/Private/UObject/ScriptCore.cpp
+++ b/Engine/Source/Runtime/CoreUObject/Private/UObject/ScriptCore.cpp
@@ -1969,7 +1969,10 @@ thread_local int32 ProcessEventCounter = 0;
 
 void UObject::ProcessEvent( UFunction* Function, void* Parms )
 {
-	checkf(!IsUnreachable(),TEXT("%s  Function: '%s'"), *GetFullName(), *Function->GetPathName());
+	// Unreachable objects are either about to be destroyed by garbage collection or temporarily marked as unreachable during GC reachability analysis on the GameThread.
+	// UObject functions are unsafe to call off-GameThread unless precautions are taken not to coincide with garbage collection (analysis) on the GameThread.
+	checkf(!IsUnreachable(), TEXT("Function '%s' called on Object '%s' that was marked unreachable. Object is possibly about to be garbage collected due to not being referenced. %s"),
+		*Function->GetPathName(), *GetFullName(), !IsInGameThread() ? TEXT("Alternatively, this function was called from a non-GameThread which is unsafe.") : TEXT(""));
 	checkf(!FUObjectThreadContext::Get().IsRoutingPostLoad, TEXT("Cannot call UnrealScript (%s - %s) while PostLoading objects"), *GetFullName(), *Function->GetFullName());
 
 #if TOTAL_OVERHEAD_SCRIPT_STATS
