Fix a bounds check in FScriptSetHelper/FScriptMapHelper::FindInternalIndex.

diff --git a/Engine/Source/Runtime/CoreUObject/Public/UObject/UnrealType.h b/Engine/Source/Runtime/CoreUObject/Public/UObject/UnrealType.h
--- a/Engine/Source/Runtime/CoreUObject/Public/UObject/UnrealType.h
+++ b/Engine/Source/Runtime/CoreUObject/Public/UObject/UnrealType.h
@@ -4915,7 +4915,7 @@ public:
 		return WithScriptMap([this, LogicalIdx](auto* Map) -> int32
 		{
 			int32 LocalLogicalIdx = LogicalIdx;
-			if (LocalLogicalIdx < 0 && LocalLogicalIdx > Map->Num())
+			if (LocalLogicalIdx < 0 || LocalLogicalIdx >= Map->Num())
 			{
 				return INDEX_NONE;
 			}
@@ -5706,7 +5706,7 @@ public:
 	 */
 	int32 FindInternalIndex(int32 LogicalIdx) const
 	{
-		if (LogicalIdx < 0 && LogicalIdx > Num())
+		if (LogicalIdx < 0 || LogicalIdx >= Num())
 		{
 			return INDEX_NONE;
 		}
