Add a feature dubbed "data impersonator" that allows to serialize the TPS stream of one object as if it were the TPS stream of another.
This will be used for verse instance which will use a "data object" to serialize themselves.
The feature can also be used to tests backward compatibility and upgrading alongside property bags.

diff --git a/Engine/Source/Runtime/CoreUObject/Private/UObject/Class.cpp b/Engine/Source/Runtime/CoreUObject/Private/UObject/Class.cpp
--- a/Engine/Source/Runtime/CoreUObject/Private/UObject/Class.cpp
+++ b/Engine/Source/Runtime/CoreUObject/Private/UObject/Class.cpp
@@ -1595,7 +1595,7 @@ void UStruct::SerializeVersionedTaggedProperties(FStructuredArchive::FSlot Slot,
 
 				if (Property)
 				{
-					Tag.Prop = Property;
+					Property->AssignToTag(Tag);
 
 					FName PropID = Property->GetID();
 
diff --git a/Engine/Source/Runtime/CoreUObject/Private/UObject/Obj.cpp b/Engine/Source/Runtime/CoreUObject/Private/UObject/Obj.cpp
--- a/Engine/Source/Runtime/CoreUObject/Private/UObject/Obj.cpp
+++ b/Engine/Source/Runtime/CoreUObject/Private/UObject/Obj.cpp
@@ -1799,6 +1799,27 @@ void UObject::SerializeScriptProperties(FArchive& Ar) const
 	SerializeScriptProperties(FStructuredArchiveFromArchive(Ar).GetSlot());
 }
 
+namespace UE::Private
+{
+	/**
+	 *  Query if an object has another object used in its stead for TPS serialization
+	 *  This can be useful for backward compatibility testing or delaying version upgrading of data
+	 */
+	const UObject* GetDataImpersonator(const UObject* ThisObject)
+	{
+		if (ThisObject == nullptr)
+		{
+			return nullptr;
+		}
+		if (FUObjectSerializeContext* SerializeContext = FUObjectThreadContext::Get().GetSerializeContext();
+			SerializeContext && !SerializeContext->bImpersonateProperties)
+		{
+			return nullptr;
+		}
+		return FPropertyBagRepository::Get().FindArchetype(ThisObject);
+	}
+}
+
 void UObject::SerializeScriptProperties( FStructuredArchive::FSlot Slot ) const
 {
 	FArchive& UnderlyingArchive = Slot.GetUnderlyingArchive();
@@ -1819,19 +1840,35 @@ void UObject::SerializeScriptProperties( FStructuredArchive::FSlot Slot ) const
 		{
 			DiffObject = GetArchetype();
 		}
+		UStruct* DiffClass = HasAnyFlags(RF_ClassDefaultObject) ? ObjClass->GetSuperClass() : ObjClass;
+
+		// Query if this object data is being impersonated 
+		const UObject* ThisObject = this;
+		if (const UObject* Impersonator = UE::Private::GetDataImpersonator(ThisObject))
+		{
+			ensureAlwaysMsgf(!HasAnyFlags(RF_ClassDefaultObject), TEXT("CDO '%s' shoudn't be impersonated"), *ThisObject->GetPathName());
+			ThisObject = Impersonator;
+			ObjClass = ThisObject->GetClass();
+
+			//@todo FH: Support for Default Object needs to be added when using impersonators
+			const UObject* DiffImpersonator = UE::Private::GetDataImpersonator(DiffObject);
+			DiffObject = DiffImpersonator ? const_cast<UObject*>(DiffImpersonator) : ObjClass->GetDefaultObject(false);
+			DiffClass = DiffObject->GetClass();
+			ensureAlwaysMsgf(DiffClass == ObjClass, TEXT("Impersonation of '%s' using a different default class not appropriately supported at the moment. Class: '%s', DefaultClass: '%s'")
+				, *ThisObject->GetPathName(), *ObjClass->GetPathName(), *DiffClass->GetPathName());
+		}
+
 #if WITH_EDITOR
 		static const FBoolConfigValueHelper BreakSerializationRecursion(TEXT("StructSerialization"), TEXT("BreakSerializationRecursion"));
 		const bool bBreakSerializationRecursion = BreakSerializationRecursion && UnderlyingArchive.IsLoading() && UnderlyingArchive.GetLinker();
-#else 
-		const bool bBreakSerializationRecursion = false;
-#endif
-#if WITH_EDITOR
+
 		static const FName NAME_SerializeScriptProperties = FName(TEXT("SerializeScriptProperties"));
 		FArchive::FScopeAddDebugData P(UnderlyingArchive, NAME_SerializeScriptProperties);
 		FArchive::FScopeAddDebugData S(UnderlyingArchive, ObjClass->GetFName());
+#else 
+		const bool bBreakSerializationRecursion = false;
 #endif
-
-		ObjClass->SerializeTaggedProperties(Slot, (uint8*)this, HasAnyFlags(RF_ClassDefaultObject) ? ObjClass->GetSuperClass() : ObjClass, (uint8*)DiffObject, bBreakSerializationRecursion ? this : nullptr);
+		ObjClass->SerializeTaggedProperties(Slot, (uint8*)ThisObject, DiffClass, (uint8*)DiffObject, bBreakSerializationRecursion ? ThisObject : nullptr);
 	}
 	else if (UnderlyingArchive.GetPortFlags() != 0 && !UnderlyingArchive.ArUseCustomPropertyList )
 	{
@@ -1841,7 +1878,7 @@ void UObject::SerializeScriptProperties( FStructuredArchive::FSlot Slot ) const
 		{
 			DiffObject = GetArchetype();
 		}
-		ObjClass->SerializeBinEx(Slot, const_cast<UObject *>(this), DiffObject, DiffObject ? DiffObject->GetClass() : NULL);
+		ObjClass->SerializeBinEx(Slot, const_cast<UObject *>(this), DiffObject, DiffObject ? DiffObject->GetClass() : nullptr);
 	}
 	else
 	{
diff --git a/Engine/Source/Runtime/CoreUObject/Private/UObject/Property.cpp b/Engine/Source/Runtime/CoreUObject/Private/UObject/Property.cpp
--- a/Engine/Source/Runtime/CoreUObject/Private/UObject/Property.cpp
+++ b/Engine/Source/Runtime/CoreUObject/Private/UObject/Property.cpp
@@ -2117,9 +2117,14 @@ void FProperty::SaveToTag(FPropertyTag& Tag)
 	Tag.Name = GetFName();
 }
 
+void FProperty::AssignToTag(FPropertyTag& Tag)
+{
+	Tag.Prop = this;
+}
+
 FProperty* UStruct::FindPropertyByName(FName InName) const
 {
-	for (FProperty* Property = PropertyLink; Property != NULL; Property = Property->PropertyLinkNext)
+	for (FProperty* Property = PropertyLink; Property != nullptr; Property = Property->PropertyLinkNext)
 	{
 		if (Property->GetFName() == InName)
 		{
@@ -2127,5 +2132,5 @@ FProperty* UStruct::FindPropertyByName(FName InName) const
 		}
 	}
 
-	return NULL;
+	return nullptr;
 }
diff --git a/Engine/Source/Runtime/CoreUObject/Private/UObject/PropertyArray.cpp b/Engine/Source/Runtime/CoreUObject/Private/UObject/PropertyArray.cpp
--- a/Engine/Source/Runtime/CoreUObject/Private/UObject/PropertyArray.cpp
+++ b/Engine/Source/Runtime/CoreUObject/Private/UObject/PropertyArray.cpp
@@ -553,6 +553,7 @@ void FArrayProperty::SerializeItem(FStructuredArchive::FSlot Slot, void* Value,
 		{
 			MaybeInnerTag.Emplace(UnderlyingArchive, Inner, 0, (uint8*)Value, (uint8*)Defaults);
 			UnderlyingArchive << MaybeInnerTag.GetValue();
+			Inner->AssignToTag(MaybeInnerTag.GetValue());
 		}
 
 		FPropertyTag& InnerTag = MaybeInnerTag.GetValue();
diff --git a/Engine/Source/Runtime/CoreUObject/Private/UObject/PropertyStruct.cpp b/Engine/Source/Runtime/CoreUObject/Private/UObject/PropertyStruct.cpp
--- a/Engine/Source/Runtime/CoreUObject/Private/UObject/PropertyStruct.cpp
+++ b/Engine/Source/Runtime/CoreUObject/Private/UObject/PropertyStruct.cpp
@@ -8,6 +8,7 @@
 #include "UObject/UnrealTypePrivate.h"
 #include "UObject/PropertyHelper.h"
 #include "UObject/LinkerPlaceholderBase.h"
+#include "UObject/UObjectThreadContext.h"
 #include "Serialization/ArchiveUObjectFromStructuredArchive.h"
 #include "Hash/Blake3.h"
 #include "IO/IoHash.h"
@@ -499,12 +500,52 @@ bool FStructProperty::LoadFromTag(const FPropertyTag& Tag)
 	return false;
 }
 
+static const FName NAME_StructOriginalType(TEXT("OriginalType"));
+
 void FStructProperty::SaveToTag(FPropertyTag& Tag)
 {
 	Super::SaveToTag(Tag);
 
+	//@note FH: revisit this once UE-197352 lands to use the new property type name
+	FUObjectSerializeContext* Context = FUObjectThreadContext::Get().GetSerializeContext();
+	const FString* OriginalType = nullptr;
+	if (Context && Context->bImpersonateProperties)
+	{
+#if WITH_EDITORONLY_DATA
+		OriginalType = FindMetaData(NAME_StructOriginalType);
+		//@note: To support metadata defined on array of struct in UPROPERTY for testing purposes
+		FField* OwnerField = OriginalType == nullptr ? Owner.ToField() : nullptr;
+		OriginalType = OwnerField ? OwnerField->FindMetaData(NAME_StructOriginalType) : OriginalType;
+#endif
+	}
+
 	const UScriptStruct* LocalStruct = Struct;
 	check(LocalStruct);
-	Tag.StructName = LocalStruct->GetFName();
+	Tag.StructName = OriginalType ? FName(**OriginalType) : LocalStruct->GetFName();
 	Tag.StructGuid = LocalStruct->GetCustomGuid();
 }
+
+void FStructProperty::AssignToTag(FPropertyTag& Tag)
+{
+	Super::AssignToTag(Tag);
+
+	//@note FH: revisit this once UE-197352 lands to use the new property type name
+	FUObjectSerializeContext* Context = FUObjectThreadContext::Get().GetSerializeContext();
+	const FString* OriginalType = nullptr;
+	if (Context && Context->bImpersonateProperties)
+	{
+#if WITH_EDITORONLY_DATA
+		OriginalType = FindMetaData(NAME_StructOriginalType);
+		//@note: To support metadata defined on array of struct in UPROPERTY
+		FField* OwnerField = OriginalType == nullptr ? Owner.ToField() : nullptr;
+		OriginalType = OwnerField ? OwnerField->FindMetaData(NAME_StructOriginalType) : OriginalType;
+#endif
+	}
+
+	const UScriptStruct* LocalStruct = Struct;
+	check(LocalStruct);
+	if (OriginalType && FName(**OriginalType) == Tag.StructName)
+	{
+		Tag.StructName = LocalStruct->GetFName();
+	}
+}
diff --git a/Engine/Source/Runtime/CoreUObject/Private/UObject/UObjectThreadContext.cpp b/Engine/Source/Runtime/CoreUObject/Private/UObject/UObjectThreadContext.cpp
--- a/Engine/Source/Runtime/CoreUObject/Private/UObject/UObjectThreadContext.cpp
+++ b/Engine/Source/Runtime/CoreUObject/Private/UObject/UObjectThreadContext.cpp
@@ -46,6 +46,7 @@ FUObjectSerializeContext::FUObjectSerializeContext()
 	, SerializedExportLinker(nullptr)
 	, bTrackSerializedPropertyPath(false)
 	, bSerializeUnknownProperty(false)
+	, bImpersonateProperties(false)
 {}
 
 FUObjectSerializeContext::~FUObjectSerializeContext()
diff --git a/Engine/Source/Runtime/CoreUObject/Public/UObject/UObjectThreadContext.h b/Engine/Source/Runtime/CoreUObject/Public/UObject/UObjectThreadContext.h
--- a/Engine/Source/Runtime/CoreUObject/Public/UObject/UObjectThreadContext.h
+++ b/Engine/Source/Runtime/CoreUObject/Public/UObject/UObjectThreadContext.h
@@ -174,6 +174,9 @@ public:
 	UE_INTERNAL bool bTrackSerializedPropertyPath;
 	/** True when unknown properties will be serialized to or from a property bag for the serialized object. */
 	UE_INTERNAL bool bSerializeUnknownProperty;
+	/** True when the SerializedObject properties are being impersonated. */
+	UE_INTERNAL bool bImpersonateProperties;
+
 
 	/** Adds a new loaded object */
 	COREUOBJECT_API void AddLoadedObject(UObject* InObject);
diff --git a/Engine/Source/Runtime/CoreUObject/Public/UObject/UnrealType.h b/Engine/Source/Runtime/CoreUObject/Public/UObject/UnrealType.h
--- a/Engine/Source/Runtime/CoreUObject/Public/UObject/UnrealType.h
+++ b/Engine/Source/Runtime/CoreUObject/Public/UObject/UnrealType.h
@@ -1065,6 +1065,11 @@ public:
 	 */
 	COREUOBJECT_API virtual void SaveToTag(FPropertyTag& Tag);
 
+	/**
+	 * Assign the property to tag and optionally modify other fields
+	 */
+	COREUOBJECT_API virtual void AssignToTag(FPropertyTag& Tag);
+
 	/**
 	 * Returns the first FProperty in this property's Outer chain that does not have a FProperty for an Outer
 	 */
@@ -6058,6 +6063,7 @@ public:
 #endif
 	virtual bool LoadFromTag(const FPropertyTag& Tag) override;
 	virtual void SaveToTag(FPropertyTag& Tag) override;
+	virtual void AssignToTag(FPropertyTag& Tag) override;
 	// End of FProperty interface
 
 	bool UseBinaryOrNativeSerialization(const FArchive& Ar) const;
