[ZenLoader]
  - Implement package tracking for OnlyEditorOnly to work during cooking
  - Remove the package tracking scope in UObjectGlobals for async loading since it was redundant and caused wrong package to be logged as instigator

diff --git a/Engine/Source/Runtime/CoreUObject/Private/Serialization/AsyncLoading2.cpp b/Engine/Source/Runtime/CoreUObject/Private/Serialization/AsyncLoading2.cpp
--- a/Engine/Source/Runtime/CoreUObject/Private/Serialization/AsyncLoading2.cpp
+++ b/Engine/Source/Runtime/CoreUObject/Private/Serialization/AsyncLoading2.cpp
@@ -19,6 +19,7 @@
 #include "Stats/StatsMisc.h"
 #include "Misc/CoreStats.h"
 #include "HAL/IConsoleManager.h"
+#include "Misc/PackageAccessTrackingOps.h"
 #include "Misc/CoreDelegates.h"
 #include "Misc/CommandLine.h"
 #include "Misc/App.h"
@@ -491,6 +492,18 @@ struct FExportObject
 
 class FAsyncLoadingSyncLoadContext;
 
+struct FPackageReferencer
+{
+#if UE_WITH_PACKAGE_ACCESS_TRACKING
+	FName ReferencerPackageName;
+	FName ReferencerPackageOp;
+#endif
+
+#if WITH_EDITOR
+	ECookLoadType CookLoadType = ECookLoadType::Unexpected;
+#endif
+};
+
 struct FPackageRequest
 {
 	int32 RequestId = -1;
@@ -508,6 +521,7 @@ struct FPackageRequest
 	TUniquePtr<FLoadPackageAsyncDelegate> PackageLoadedDelegate;
 	TUniquePtr<FLoadPackageAsyncProgressDelegate> PackageProgressDelegate;
 	FPackageRequest* Next = nullptr;
+	FPackageReferencer PackageReferencer;
 
 	FLinkerInstancingContext* GetInstancingContext()
 	{
@@ -518,7 +532,7 @@ struct FPackageRequest
 #endif
 	}
 
-	static FPackageRequest Create(int32 RequestId, EPackageFlags PackageFlags, uint32 LoadFlags, int32 PIEInstanceID, int32 Priority, const FLinkerInstancingContext* InstancingContext, const FPackagePath& PackagePath, FName CustomName, TUniquePtr<FLoadPackageAsyncDelegate> PackageLoadedDelegate, TUniquePtr<FLoadPackageAsyncProgressDelegate> PackageProgressDelegate)
+	static FPackageRequest Create(int32 RequestId, EPackageFlags PackageFlags, uint32 LoadFlags, int32 PIEInstanceID, int32 Priority, const FLinkerInstancingContext* InstancingContext, const FPackagePath& PackagePath, FName CustomName, TUniquePtr<FLoadPackageAsyncDelegate> PackageLoadedDelegate, TUniquePtr<FLoadPackageAsyncProgressDelegate> PackageProgressDelegate, FPackageReferencer PackageReferencer)
 	{
 		return FPackageRequest
 		{
@@ -536,7 +550,8 @@ struct FPackageRequest
 			PackagePath,
 			MoveTemp(PackageLoadedDelegate),
 			MoveTemp(PackageProgressDelegate),
-			nullptr
+			nullptr,
+			PackageReferencer
 		};
 	}
 };
@@ -565,6 +580,8 @@ struct FAsyncPackageDesc2
 	// The package path of the package being loaded from disk
 	// Set to none for imported packages up until the package summary has been serialized
 	FPackagePath PackagePathToLoad;
+	// Package referencer
+	FPackageReferencer PackageReferencer;
 	// Packages with a a custom name can't be imported
 	bool bCanBeImported;
 
@@ -586,6 +603,7 @@ struct FAsyncPackageDesc2
 			PackageIdToLoad,
 			UPackageName,
 			MoveTemp(Request.PackagePath),
+			Request.PackageReferencer,
 #if WITH_EDITOR
 			true
 #else
@@ -614,6 +632,7 @@ struct FAsyncPackageDesc2
 			PackageIdToLoad,
 			UPackageName,
 			MoveTemp(PackagePathToLoad),
+			ImportingPackageDesc.PackageReferencer,
 			true
 		};
 	}
@@ -8837,6 +8856,11 @@ void FAsyncPackage2::CreateUPackage()
 	if (!LinkerRoot)
 	{
 		TRACE_CPUPROFILER_EVENT_SCOPE(UPackageCreate);
+		UE_TRACK_REFERENCING_PACKAGE_SCOPED(Desc.PackageReferencer.ReferencerPackageName, Desc.PackageReferencer.ReferencerPackageOp);
+
+#if WITH_EDITOR
+		FCookLoadScope CookLoadScope(Desc.PackageReferencer.CookLoadType);
+#endif
 		LinkerRoot = NewObject<UPackage>(/*Outer*/nullptr, Desc.UPackageName);
 		bCreatedLinkerRoot = true;
 	}
@@ -9006,7 +9030,11 @@ void FAsyncLoadingThread2::ConditionalProcessEditorCallbacks()
 		// In editor builds, call the asset load callback. This happens in both editor and standalone to match EndLoad
 		for (UObject* LoadedObject : LocalEditorLoadedAssets)
 		{
-			FCoreUObjectDelegates::OnAssetLoaded.Broadcast(LoadedObject);
+			if (LoadedObject)
+			{
+				UE_TRACK_REFERENCING_PACKAGE_SCOPED(LoadedObject->GetPackage()->GetFName(), PackageAccessTrackingOps::NAME_Load);
+				FCoreUObjectDelegates::OnAssetLoaded.Broadcast(LoadedObject);
+			}
 		}
 	}
 }
@@ -9055,7 +9083,19 @@ PRAGMA_ENABLE_DEPRECATION_WARNINGS
 	TRACE_LOADTIME_BEGIN_REQUEST(RequestId);
 	AddPendingRequest(RequestId);
 
-	PackageRequestQueue.Enqueue(FPackageRequest::Create(RequestId, InPackageFlags, InLoadFlags, InPIEInstanceID, InPackagePriority, InInstancingContext, InPackagePath, InCustomName, MoveTemp(InCompletionDelegate), MoveTemp(InProgressDelegate)));
+	FPackageReferencer PackageReferencer;
+#if UE_WITH_PACKAGE_ACCESS_TRACKING
+	PackageAccessTracking_Private::FTrackedData* AccumulatedScopeData = PackageAccessTracking_Private::FPackageAccessRefScope::GetCurrentThreadAccumulatedData();
+	if (AccumulatedScopeData)
+	{
+		PackageReferencer.ReferencerPackageName = AccumulatedScopeData->PackageName;
+		PackageReferencer.ReferencerPackageOp = AccumulatedScopeData->OpName;
+	}
+#endif
+#if WITH_EDITOR
+	PackageReferencer.CookLoadType = FCookLoadScope::GetCurrentValue();
+#endif
+	PackageRequestQueue.Enqueue(FPackageRequest::Create(RequestId, InPackageFlags, InLoadFlags, InPIEInstanceID, InPackagePriority, InInstancingContext, InPackagePath, InCustomName, MoveTemp(InCompletionDelegate), MoveTemp(InProgressDelegate), PackageReferencer));
 	++QueuedPackagesCounter;
 	++PackagesWithRemainingWorkCounter;
 
diff --git a/Engine/Source/Runtime/CoreUObject/Private/UObject/UObjectGlobals.cpp b/Engine/Source/Runtime/CoreUObject/Private/UObject/UObjectGlobals.cpp
--- a/Engine/Source/Runtime/CoreUObject/Private/UObject/UObjectGlobals.cpp
+++ b/Engine/Source/Runtime/CoreUObject/Private/UObject/UObjectGlobals.cpp
@@ -1639,8 +1639,6 @@ UPackage* LoadPackageInternal(UPackage* InOuter, const FPackagePath& PackagePath
 			return nullptr;
 		}
 
-		UE_TRACK_REFERENCING_PACKAGE_SCOPED(PackageName, PackageAccessTrackingOps::NAME_Load);
-
 		// This delegate is not thread-safe and the subscribers are mostly interested by sync loads
 		// that might stall the game thread anyway. So for now, do not broadcast when sync loading
 		// from the loading thread.
