Changed ArchiveCrc32 operator<<(FArchive&)  template generalization for structs to only apply for direct serialization to ArchiveCrc32 and if there is no better conversion available.

Details in Jira:

diff --git a/Engine/Source/Runtime/CoreUObject/Public/Serialization/ArchiveCrc32.h b/Engine/Source/Runtime/CoreUObject/Public/Serialization/ArchiveCrc32.h
--- a/Engine/Source/Runtime/CoreUObject/Public/Serialization/ArchiveCrc32.h
+++ b/Engine/Source/Runtime/CoreUObject/Public/Serialization/ArchiveCrc32.h
@@ -37,21 +37,22 @@ public:
 	virtual FString GetArchiveName() const { return TEXT("FArchiveCrc32"); }
 	//~ End FArchive Interface
 
+	/**
+		* Serializes a USTRUCT value into an ArchiveCrc32.
+		*
+		* @param Ar The archive to serialize from or to.
+		* @param Value The value to serialize.
+		* NOTE: This should be deprecated and removed as soon as possible JIRA: UE-203804. It must be kept as a hidden friend operator to avoid conflict with specializations found in StructuredArchiveAdapter.
+		*/
+	template <typename StructType UE_REQUIRES(TModels_V<CStaticStructProvider, StructType>)>
+	FORCEINLINE friend FArchiveCrc32& operator<<(FArchiveCrc32& Ar, const StructType& Value)
+	{
+		StructType* MutableValue = const_cast<StructType*>(&Value);
+		StructType::StaticStruct()->SerializeItem(Ar, MutableValue, nullptr);
+		return Ar;
+	}
+
 private:
 	uint32 CRC;
 	UObject* RootObject;
 };
-
-/**
-	* Serializes a USTRUCT value from or into an archive.
-	*
-	* @param Ar The archive to serialize from or to.
-	* @param Value The value to serialize.
-	*/
-template <typename StructType>
-FORCEINLINE typename TEnableIf<TModels_V<CStaticStructProvider, StructType>, FArchive&>::Type operator <<(FArchive& Ar, const StructType& Value)
-{
-	StructType* MutableValue = const_cast<StructType*>(&Value);
-	StructType::StaticStruct()->SerializeItem(Ar, MutableValue, nullptr);
-	return Ar;
-}
diff --git a/Engine/Source/Runtime/Engine/Classes/Components/InstancedStaticMeshComponent.h b/Engine/Source/Runtime/Engine/Classes/Components/InstancedStaticMeshComponent.h
--- a/Engine/Source/Runtime/Engine/Classes/Components/InstancedStaticMeshComponent.h
+++ b/Engine/Source/Runtime/Engine/Classes/Components/InstancedStaticMeshComponent.h
@@ -84,6 +84,13 @@ struct FInstancedStaticMeshRandomSeed
 
 	UPROPERTY()
 	int32 RandomSeed = 0;
+
+	friend FArchive& operator<<(FArchive& Ar, FInstancedStaticMeshRandomSeed& InstanceData)
+	{
+		Ar << InstanceData.StartInstanceIndex;
+		Ar << InstanceData.RandomSeed;
+		return Ar;
+	}
 };
 
 /** A component that efficiently renders multiple instances of the same StaticMesh. */
