Cooker: SkipOnlyEditorOnly fixes:
Mark PaperSprite's load of its SourceTexture as an EditorOnly load.
Fix AsyncLoading1's implementation of the call to OnLoadedAsset to have the PackageReferencerScope that is expected around the call (this contract was established by the implementation of OnLoadedAsset that is done in the legacy LinkerLoad loading path).
#rnx

diff --git a/Engine/Plugins/2D/Paper2D/Source/Paper2D/Private/PaperSprite.cpp b/Engine/Plugins/2D/Paper2D/Source/Paper2D/Private/PaperSprite.cpp
--- a/Engine/Plugins/2D/Paper2D/Source/Paper2D/Private/PaperSprite.cpp
+++ b/Engine/Plugins/2D/Paper2D/Source/Paper2D/Private/PaperSprite.cpp
@@ -17,6 +17,7 @@
 #include "PaperFlipbook.h"
 #include "Paper2DModule.h"
 #include "Paper2DPrivate.h"
+#include "UObject/ICookInfo.h"
 #include "UObject/LinkerLoad.h"
 
 #include UE_INLINE_GENERATED_CPP_BY_NAME(PaperSprite)
@@ -1449,6 +1450,9 @@ UTexture2D* UPaperSprite::GetSourceTexture() const
 
 	if (SourceTexturePtr == nullptr)
 	{
+		// The SourceTexture is editoronly and does not need to be cooked, but we sometimes load it during PostLoad due
+		// to testing NeedRescaleSpriteData
+		FCookLoadScope CookLoadScope(ECookLoadType::EditorOnly);
 		SourceTexturePtr = SourceTexture.LoadSynchronous();
 	}
 
diff --git a/Engine/Source/Runtime/CoreUObject/Private/Serialization/AsyncLoading.cpp b/Engine/Source/Runtime/CoreUObject/Private/Serialization/AsyncLoading.cpp
--- a/Engine/Source/Runtime/CoreUObject/Private/Serialization/AsyncLoading.cpp
+++ b/Engine/Source/Runtime/CoreUObject/Private/Serialization/AsyncLoading.cpp
@@ -4950,6 +4950,9 @@ EAsyncPackageState::Type FAsyncLoadingThread::ProcessLoadedPackages(bool bUseTim
 			// It may have been unloaded/marked pending kill since being added, ignore those cases
 			if (UObject* LoadedAsset = WeakAsset.Get())
 			{
+				// Current contract for UE_TRACK_REFERENCING_PACKAGE_SCOPED and AssetLoaded is that the caller of
+				// OnAssetLoaded sets the scope to the LoadedAsset's package.
+				UE_TRACK_REFERENCING_PACKAGE_SCOPED(LoadedAsset->GetPackage(), PackageAccessTrackingOps::NAME_Load);
 				FCoreUObjectDelegates::OnAssetLoaded.Broadcast(LoadedAsset);
 			}
 		}
