Added FPropertyTypeName to represent the full type name of a property

#rnx

diff --git a/Engine/Source/Runtime/CoreUObject/Private/UObject/PropertyTypeName.cpp b/Engine/Source/Runtime/CoreUObject/Private/UObject/PropertyTypeName.cpp
new file mode 100644
--- /dev/null
+++ b/Engine/Source/Runtime/CoreUObject/Private/UObject/PropertyTypeName.cpp
@@ -0,0 +1,339 @@
+// Copyright Epic Games, Inc. All Rights Reserved.
+
+#include "UObject/PropertyTypeName.h"
+
+#include "Async/Mutex.h"
+#include "Async/UniqueLock.h"
+#include "Containers/Map.h"
+#include "HAL/CriticalSection.h"
+#include "HAL/LowLevelMemTracker.h"
+#include "Misc/ScopeRWLock.h"
+
+namespace UE
+{
+
+static_assert(sizeof(FPropertyTypeNameNode) <= 16);
+
+constexpr static int32 GPropertyTypeNameBlockBits = 13;
+constexpr static int32 GPropertyTypeNameBlockOffsetBits = 13;
+
+constexpr static int32 GPropertyTypeNameBlockOffsetCount = 1 << GPropertyTypeNameBlockOffsetBits;
+constexpr static int32 GPropertyTypeNameBlockOffsetMask = (1 << GPropertyTypeNameBlockOffsetBits) - 1;
+
+constexpr static int32 GPropertyTypeNameBlockSize = GPropertyTypeNameBlockOffsetCount * sizeof(FPropertyTypeNameNode);
+constexpr static int32 GPropertyTypeNameBlockCount = 1 << GPropertyTypeNameBlockBits;
+
+LLM_DEFINE_TAG(FPropertyTypeName);
+
+///////////////////////////////////////////////////////////////////////////////////////////////////////////////////////
+
+inline uint32 GetTypeHash(const FPropertyTypeNameNode& Node)
+{
+	return HashCombineFast(GetTypeHashHelper(Node.Name), GetTypeHashHelper(Node.InnerCount));
+}
+
+inline bool operator==(const FPropertyTypeNameNode& Lhs, const FPropertyTypeNameNode& Rhs)
+{
+	return Lhs.Name == Rhs.Name && Lhs.InnerCount == Rhs.InnerCount;
+}
+
+inline FArchive& operator<<(FArchive& Ar, FPropertyTypeNameNode& Node)
+{
+	return Ar << Node.Name << Node.InnerCount;
+}
+
+inline const FPropertyTypeNameNode* AppendNode(FStringBuilderBase& Builder, const FPropertyTypeNameNode* Node)
+{
+	Builder << Node->Name;
+
+	if (int32 Remaining = Node++->InnerCount)
+	{
+		Builder.AppendChar(TEXT('<'));
+		for (; Remaining > 0; --Remaining)
+		{
+			Node = AppendNode(Builder, Node);
+			Builder.AppendChar(TEXT(','));
+		}
+		Builder.RemoveSuffix(1);
+		Builder.AppendChar(TEXT('>'));
+	}
+
+	return Node;
+}
+
+///////////////////////////////////////////////////////////////////////////////////////////////////////////////////////
+
+struct FPropertyTypeNameNodeProxy
+{
+	const FPropertyTypeNameNode* First = nullptr;
+
+	inline void HashAndMeasure(uint32& OutHash, int32& OutCount) const
+	{
+		uint32 Hash = 0;
+		const FPropertyTypeNameNode* Node = First;
+		for (int32 Remaining = 1; Remaining > 0; --Remaining, ++Node)
+		{
+			Hash = HashCombineFast(Hash, GetTypeHash(*Node));
+			Remaining += Node->InnerCount;
+		}
+		OutHash = Hash;
+		OutCount = UE_PTRDIFF_TO_INT32(Node - First);
+	}
+
+	friend inline uint32 GetTypeHash(const FPropertyTypeNameNodeProxy& Proxy)
+	{
+		uint32 Hash = 0;
+		int32 Count = 0;
+		Proxy.HashAndMeasure(Hash, Count);
+		return Hash;
+	}
+
+	friend inline bool operator==(const FPropertyTypeNameNodeProxy& Lhs, const FPropertyTypeNameNodeProxy& Rhs)
+	{
+		const FPropertyTypeNameNode* LhsNode = Lhs.First;
+		const FPropertyTypeNameNode* RhsNode = Rhs.First;
+		for (int32 Remaining = 1; Remaining > 0; --Remaining, ++LhsNode, ++RhsNode)
+		{
+			if (*LhsNode == *RhsNode)
+			{
+				Remaining += LhsNode->InnerCount;
+				continue;
+			}
+			return false;
+		}
+		return true;
+	}
+};
+
+///////////////////////////////////////////////////////////////////////////////////////////////////////////////////////
+
+class FPropertyTypeNameTable
+{
+public:
+	FPropertyTypeNameTable();
+
+	int32 FindOrAddByName(const FPropertyTypeNameNode* First);
+
+	const FPropertyTypeNameNode* ResolveByIndex(int32 Index) const;
+
+private:
+	int32 StoreByIndex(const FPropertyTypeNameNode* Nodes, int32 Count);
+
+	FPropertyTypeNameNode* AllocateBlock(int32 BlockIndex);
+
+	// TODO: Shard the hash table and stop using TMap.
+	mutable FRWLock NameMapMutex;
+	TMap<FPropertyTypeNameNodeProxy, int32> NameMap;
+
+	std::atomic<int32> NextIndex = 1; // 1 leaves a "None" entry at 0
+	std::atomic<FPropertyTypeNameNode*> Blocks[GPropertyTypeNameBlockCount]{};
+	FMutex BlockAllocationMutex;
+};
+
+FPropertyTypeNameTable GPropertyTypeNameTable;
+
+FPropertyTypeNameTable::FPropertyTypeNameTable()
+{
+	AllocateBlock(0);
+}
+
+int32 FPropertyTypeNameTable::FindOrAddByName(const FPropertyTypeNameNode* First)
+{
+	const FPropertyTypeNameNodeProxy Proxy{First};
+
+	uint32 Hash;
+	int32 Count;
+	Proxy.HashAndMeasure(Hash, Count);
+
+	if (FReadScopeLock Lock(NameMapMutex); const int32* Index = NameMap.FindByHash(Hash, Proxy))
+	{
+		return *Index;
+	}
+
+	LLM_SCOPE_BYTAG(FPropertyTypeName);
+	FWriteScopeLock Lock(NameMapMutex);
+	const int32 Index = StoreByIndex(First, Count);
+	NameMap.AddByHash(Hash, {ResolveByIndex(Index)}, Index);
+	return Index;
+}
+
+const FPropertyTypeNameNode* FPropertyTypeNameTable::ResolveByIndex(int32 Index) const
+{
+	const int32 BlockIndex = Index >> GPropertyTypeNameBlockOffsetBits;
+	const int32 BlockOffset = Index & GPropertyTypeNameBlockOffsetMask;
+	return Blocks[BlockIndex].load(std::memory_order_relaxed) + BlockOffset;
+}
+
+int32 FPropertyTypeNameTable::StoreByIndex(const FPropertyTypeNameNode* Nodes, int32 Count)
+{
+	if (Count == 1 && Nodes->Name.IsNone())
+	{
+		return 0;
+	}
+
+	int32 Index;
+	int32 BlockIndex;
+	int32 BlockOffset;
+	do
+	{
+		Index = NextIndex.fetch_add(Count, std::memory_order_relaxed);
+		BlockIndex = Index >> GPropertyTypeNameBlockOffsetBits;
+		BlockOffset = Index & GPropertyTypeNameBlockOffsetMask;
+	}
+	while (UNLIKELY(BlockOffset + Count >= GPropertyTypeNameBlockOffsetCount));
+
+	check(Index + Count < GPropertyTypeNameBlockCount * GPropertyTypeNameBlockOffsetCount);
+
+	FPropertyTypeNameNode* Block = Blocks[BlockIndex].load(std::memory_order_acquire);
+	if (UNLIKELY(!Block))
+	{
+		Block = AllocateBlock(BlockIndex);
+	}
+	for (FPropertyTypeNameNode* Target = Block + BlockOffset; Count > 0; --Count)
+	{
+		*Target++ = *Nodes++;
+	}
+	return Index;
+}
+
+FPropertyTypeNameNode* FPropertyTypeNameTable::AllocateBlock(int32 BlockIndex)
+{
+	TUniqueLock Lock(BlockAllocationMutex);
+	FPropertyTypeNameNode* Block = Blocks[BlockIndex].load(std::memory_order_acquire);
+	if (!Block)
+	{
+		Block = (FPropertyTypeNameNode*)FMemory::MallocZeroed(GPropertyTypeNameBlockSize, alignof(FPropertyTypeNameNode));
+		Blocks[BlockIndex].store(Block, std::memory_order_release);
+	}
+	return Block;
+}
+
+///////////////////////////////////////////////////////////////////////////////////////////////////////////////////////
+
+FName FPropertyTypeName::GetTypeName() const
+{
+	const FPropertyTypeNameNode* First = GPropertyTypeNameTable.ResolveByIndex(Index);
+	return First->Name;
+}
+
+int32 FPropertyTypeName::GetTypeParameterCount() const
+{
+	const FPropertyTypeNameNode* First = GPropertyTypeNameTable.ResolveByIndex(Index);
+	return First->InnerCount;
+}
+
+FPropertyTypeName FPropertyTypeName::GetTypeParameter(int32 ParamIndex) const
+{
+	const FPropertyTypeNameNode* First = GPropertyTypeNameTable.ResolveByIndex(Index);
+	if (UNLIKELY(ParamIndex < 0 || ParamIndex >= First->InnerCount))
+	{
+		return {};
+	}
+	const FPropertyTypeNameNode* Param = First + 1;
+	for (int32 Skip = ParamIndex; Skip > 0; --Skip, ++Param)
+	{
+		Skip += Param->InnerCount;
+	}
+	FPropertyTypeName Type = *this;
+	Type.Index += UE_PTRDIFF_TO_INT32(Param - First);
+	return Type;
+}
+
+uint32 GetTypeHash(const FPropertyTypeName& TypeName)
+{
+	const FPropertyTypeNameNode* First = GPropertyTypeNameTable.ResolveByIndex(TypeName.Index);
+	return GetTypeHash(FPropertyTypeNameNodeProxy{First});
+}
+
+bool operator==(const FPropertyTypeName& Lhs, const FPropertyTypeName& Rhs)
+{
+	if (Lhs.Index == Rhs.Index)
+	{
+		return true;
+	}
+
+	const FPropertyTypeNameNode* LhsNode = GPropertyTypeNameTable.ResolveByIndex(Lhs.Index);
+	const FPropertyTypeNameNode* RhsNode = GPropertyTypeNameTable.ResolveByIndex(Rhs.Index);
+	return FPropertyTypeNameNodeProxy{LhsNode} == FPropertyTypeNameNodeProxy{RhsNode};
+}
+
+FArchive& operator<<(FArchive& Ar, FPropertyTypeName& TypeName)
+{
+	if (!Ar.IsPersistent())
+	{
+		Ar << TypeName.Index;
+	}
+	else if (Ar.IsLoading())
+	{
+		TArray<FPropertyTypeNameNode, TInlineAllocator<16>> Nodes;
+		int32 Remaining = 1;
+		do 
+		{
+			FPropertyTypeNameNode& Node = Nodes.AddDefaulted_GetRef();
+			Ar << Node;
+			Remaining += Node.InnerCount - 1;
+		}
+		while (Remaining > 0);
+		TypeName.Index = GPropertyTypeNameTable.FindOrAddByName(Nodes.GetData());
+	}
+	else if (Ar.IsSaving())
+	{
+		const FPropertyTypeNameNode* Node = GPropertyTypeNameTable.ResolveByIndex(TypeName.Index);
+		for (int32 Remaining = 1; Remaining > 0; --Remaining, ++Node)
+		{
+			FPropertyTypeNameNode NodeCopy = *Node;
+			Ar << NodeCopy;
+			Remaining += NodeCopy.InnerCount;
+		}
+	}
+	return Ar;
+}
+
+FStringBuilderBase& operator<<(FStringBuilderBase& Builder, const FPropertyTypeName& TypeName)
+{
+	const FPropertyTypeNameNode* Node = GPropertyTypeNameTable.ResolveByIndex(TypeName.Index);
+	AppendNode(Builder, Node);
+	return Builder;
+}
+
+///////////////////////////////////////////////////////////////////////////////////////////////////////////////////////
+
+void FPropertyTypeNameBuilder::BeginTypeParameters()
+{
+	checkf(!Nodes.IsEmpty(), TEXT("A type name must be added prior to setting type parameters."));
+	ActiveIndex = Nodes.Num() - 1;
+}
+
+void FPropertyTypeNameBuilder::EndTypeParameters()
+{
+	ActiveIndex = OuterNodeIndex[ActiveIndex];
+}
+
+void FPropertyTypeNameBuilder::AddTypeName(FName Name)
+{
+	FPropertyTypeNameNode& Node = Nodes.AddDefaulted_GetRef();
+	Node.Name = Name;
+	Node.InnerCount = 0;
+
+	OuterNodeIndex.Add(ActiveIndex);
+	if (ActiveIndex >= 0)
+	{
+		++Nodes[ActiveIndex].InnerCount;
+	}
+}
+
+FPropertyTypeName FPropertyTypeNameBuilder::Build() const
+{
+	FPropertyTypeName Type;
+	Type.Index = GPropertyTypeNameTable.FindOrAddByName(Nodes.GetData());
+	return Type;
+}
+
+void FPropertyTypeNameBuilder::Reset()
+{
+	Nodes.Reset();
+	OuterNodeIndex.Reset();
+	ActiveIndex = INDEX_NONE;
+}
+
+} // UE
diff --git a/Engine/Source/Runtime/CoreUObject/Public/UObject/PropertyTypeName.h b/Engine/Source/Runtime/CoreUObject/Public/UObject/PropertyTypeName.h
new file mode 100644
--- /dev/null
+++ b/Engine/Source/Runtime/CoreUObject/Public/UObject/PropertyTypeName.h
@@ -0,0 +1,135 @@
+// Copyright Epic Games, Inc. All Rights Reserved.
+
+#pragma once
+
+#include "CoreTypes.h"
+#include "Containers/Array.h"
+#include "Containers/StringFwd.h"
+#include "UObject/NameTypes.h"
+
+#define UE_API COREUOBJECT_API
+
+class FArchive;
+
+namespace UE { class FPropertyTypeNameBuilder; }
+
+namespace UE
+{
+
+struct UE_INTERNAL FPropertyTypeNameNode
+{
+	FName Name;
+	int32 InnerCount = 0;
+};
+
+/**
+ * Represents the type name of a property, including any containers and underlying types.
+ *
+ * Examples:
+ * - int32 -> IntProperty
+ * - TArray<int32> -> ArrayProperty<IntProperty>
+ * - TArray<FStructType> -> ArrayProperty<StructProperty<StructType>>
+ * - TMap<FKeyStruct, EByteEnum> -> MapProperty<StructProperty<KeyStruct>,EnumProperty<ByteProperty<ByteEnum>>>
+ */
+class FPropertyTypeName
+{
+public:
+	inline bool IsEmpty() const
+	{
+		return Index == 0;
+	}
+
+	/**
+	 * Returns the type at the root of this property type name.
+	 *
+	 * Example: MapProperty<StructProperty<KeyStruct>,EnumProperty<ByteProperty<ByteEnum>>>
+	 * - GetTypeName() -> MapProperty
+	 */
+	UE_API FName GetTypeName() const;
+
+	/**
+	 * Returns the number of type parameters under the root of this property type name.
+	 *
+	 * Example: MapProperty<StructProperty<KeyStruct>,EnumProperty<ByteProperty<ByteEnum>>>
+	 * - GetTypeParameterCount() -> 2
+	 */
+	UE_API int32 GetTypeParameterCount() const;
+
+	/**
+	 * Returns the indexed parameter under the root of this property type name.
+	 *
+	 * Example: MapProperty<StructProperty<KeyStruct>,EnumProperty<ByteProperty<ByteEnum>>>
+	 * - GetTypeParameter(0) -> StructProperty<KeyStruct>
+	 * - GetTypeParameter(1) -> EnumProperty<ByteProperty<ByteEnum>>
+	 */
+	UE_API FPropertyTypeName GetTypeParameter(int32 ParamIndex) const;
+
+	/**
+	 * Resets this to an empty type name.
+	 */
+	inline void Reset()
+	{
+		Index = 0;
+	}
+
+private:
+	UE_API friend uint32 GetTypeHash(const FPropertyTypeName& TypeName);
+
+	UE_API friend bool operator==(const FPropertyTypeName& Lhs, const FPropertyTypeName& Rhs);
+
+	UE_API friend FArchive& operator<<(FArchive& Ar, FPropertyTypeName& TypeName);
+
+	UE_API friend FStringBuilderBase& operator<<(FStringBuilderBase& Builder, const FPropertyTypeName& TypeName);
+
+	int32 Index = 0;
+
+	friend FPropertyTypeNameBuilder;
+};
+
+/**
+ * Builder for FPropertyTypeName.
+ *
+ * Example: MapProperty<StructProperty<KeyStruct>,EnumProperty<ByteProperty<ByteEnum>>>
+ *
+ * FPropertyTypeNameBuilder Builder;
+ * Builder.AddTypeName(TEXT("MapProperty"));
+ * Builder.BeginTypeParameters();
+ *   Builder.AddTypeName(TEXT("StructProperty"));
+ *   Builder.BeginTypeParameters();
+ *     Builder.AddTypeName(TEXT("KeyStruct"));
+ *   Builder.EndTypeParameters();
+ *   Builder.AddTypeName(TEXT("EnumProperty"));
+ *   Builder.BeginTypeParameters();
+ *     Builder.AddTypeName(TEXT("ByteProperty"));
+ *     Builder.BeginTypeParameters();
+ *       Builder.AddTypeName(TEXT("ByteEnum"));
+ *     Builder.EndTypeParameters();
+ *   Builder.EndTypeParameters();
+ * Builder.EndTypeParameters();
+ * FPropertyTypeName Name = Builder.Build();
+ */
+class FPropertyTypeNameBuilder
+{
+public:
+	UE_API void AddTypeName(FName Name);
+
+	/** Mark the beginning of the type parameters for the last added type. */
+	UE_API void BeginTypeParameters();
+
+	/** Mark the end of the type parameters for the matching begin. */
+	UE_API void EndTypeParameters();
+
+	UE_API FPropertyTypeName Build() const;
+
+	/** Resets the builder to allow it to be reused. */
+	UE_API void Reset();
+
+private:
+	TArray<FPropertyTypeNameNode, TInlineAllocator<8>> Nodes;
+	TArray<int32, TInlineAllocator<8>> OuterNodeIndex;
+	int32 ActiveIndex = INDEX_NONE;
+};
+
+} // UE
+
+#undef UE_API
diff --git a/Engine/Source/Runtime/CoreUObject/Tests/UObject/PropertyTypeNameTest.cpp b/Engine/Source/Runtime/CoreUObject/Tests/UObject/PropertyTypeNameTest.cpp
new file mode 100644
--- /dev/null
+++ b/Engine/Source/Runtime/CoreUObject/Tests/UObject/PropertyTypeNameTest.cpp
@@ -0,0 +1,312 @@
+// Copyright Epic Games, Inc. All Rights Reserved.
+
+#include "UObject/PropertyTypeName.h"
+
+#if WITH_TESTS
+
+#include "Async/ParallelFor.h"
+#include "Misc/StringBuilder.h"
+#include "Serialization/MemoryReader.h"
+#include "Serialization/MemoryWriter.h"
+#include "Tests/TestHarnessAdapter.h"
+
+namespace UE
+{
+
+namespace PropertyTypeNameTest
+{
+
+static FPropertyTypeName CreateInt()
+{
+	FPropertyTypeNameBuilder Builder;
+	Builder.AddTypeName(NAME_IntProperty);
+	return Builder.Build();
+}
+
+static FPropertyTypeName CreateVectorArray()
+{
+	FPropertyTypeNameBuilder Builder;
+	Builder.AddTypeName(NAME_ArrayProperty);
+	{
+		Builder.BeginTypeParameters();
+		Builder.AddTypeName(NAME_StructProperty);
+		{
+			Builder.BeginTypeParameters();
+			Builder.AddTypeName(NAME_Vector);
+			Builder.EndTypeParameters();
+		}
+		Builder.EndTypeParameters();
+	}
+	return Builder.Build();
+}
+
+static FPropertyTypeName CreateEnumMap()
+{
+	FPropertyTypeNameBuilder Builder;
+	Builder.AddTypeName(NAME_MapProperty);
+	{
+		Builder.BeginTypeParameters();
+		Builder.AddTypeName(NAME_EnumProperty);
+		{
+			Builder.BeginTypeParameters();
+			Builder.AddTypeName(NAME_ByteProperty);
+			{
+				Builder.BeginTypeParameters();
+				Builder.AddTypeName(TEXT("Key"));
+				Builder.EndTypeParameters();
+			}
+			Builder.EndTypeParameters();
+		}
+		Builder.AddTypeName(NAME_EnumProperty);
+		{
+			Builder.BeginTypeParameters();
+			Builder.AddTypeName(NAME_ByteProperty);
+			{
+				Builder.BeginTypeParameters();
+				Builder.AddTypeName(TEXT("Value"));
+				Builder.EndTypeParameters();
+			}
+			Builder.EndTypeParameters();
+		}
+		Builder.EndTypeParameters();
+	}
+	return Builder.Build();
+}
+
+} // PropertyTypeNameTest
+
+TEST_CASE_NAMED(FPropertyTypeNameSmokeTest, "CoreUObject::PropertyTypeName::Smoke", "[Core][UObject][SmokeFilter]")
+{
+	SECTION("Empty")
+	{
+		FPropertyTypeName TypeName;
+		CHECK(TypeName.IsEmpty());
+		CHECK(TypeName.GetTypeName().IsNone());
+		CHECK(TypeName.GetTypeParameterCount() == 0);
+		CHECK(TypeName.GetTypeParameter(0).IsEmpty());
+	}
+
+	SECTION("Numeric")
+	{
+		FPropertyTypeName TypeName = PropertyTypeNameTest::CreateInt();
+		CHECK(!TypeName.IsEmpty());
+		CHECK(TypeName.GetTypeName() == NAME_IntProperty);
+		CHECK(TypeName.GetTypeParameterCount() == 0);
+		CHECK(TypeName.GetTypeParameter(0).IsEmpty());
+	}
+
+	SECTION("NumericArray")
+	{
+		FPropertyTypeNameBuilder Builder;
+		Builder.AddTypeName(NAME_ArrayProperty);
+		{
+			Builder.BeginTypeParameters();
+			Builder.AddTypeName(NAME_IntProperty);
+			Builder.EndTypeParameters();
+		}
+
+		FPropertyTypeName TypeName = Builder.Build();
+		CHECK(!TypeName.IsEmpty());
+		CHECK(TypeName.GetTypeName() == NAME_ArrayProperty);
+		CHECK(TypeName.GetTypeParameterCount() == 1);
+		CHECK(TypeName.GetTypeParameter(0).GetTypeName() == NAME_IntProperty);
+		CHECK(TypeName.GetTypeParameter(0).GetTypeParameterCount() == 0);
+		CHECK(TypeName.GetTypeParameter(0).GetTypeParameter(0).IsEmpty());
+		CHECK(TypeName.GetTypeParameter(1).IsEmpty());
+	}
+
+	SECTION("StructArray")
+	{
+		FPropertyTypeName TypeName = PropertyTypeNameTest::CreateVectorArray();
+		CHECK(!TypeName.IsEmpty());
+		CHECK(TypeName.GetTypeName() == NAME_ArrayProperty);
+		CHECK(TypeName.GetTypeParameterCount() == 1);
+		CHECK(TypeName.GetTypeParameter(0).GetTypeName() == NAME_StructProperty);
+		CHECK(TypeName.GetTypeParameter(0).GetTypeParameterCount() == 1);
+		CHECK(TypeName.GetTypeParameter(0).GetTypeParameter(0).GetTypeName() == NAME_Vector);
+		CHECK(TypeName.GetTypeParameter(0).GetTypeParameter(0).GetTypeParameterCount() == 0);
+		CHECK(TypeName.GetTypeParameter(0).GetTypeParameter(0).GetTypeParameter(0).IsEmpty());
+		CHECK(TypeName.GetTypeParameter(0).GetTypeParameter(1).IsEmpty());
+		CHECK(TypeName.GetTypeParameter(1).IsEmpty());
+	}
+
+	SECTION("NumericMap")
+	{
+		FPropertyTypeNameBuilder Builder;
+		Builder.AddTypeName(NAME_MapProperty);
+		{
+			Builder.BeginTypeParameters();
+			Builder.AddTypeName(NAME_Int32Property);
+			Builder.AddTypeName(NAME_UInt32Property);
+			Builder.EndTypeParameters();
+		}
+
+		FPropertyTypeName TypeName = Builder.Build();
+		CHECK(!TypeName.IsEmpty());
+		CHECK(TypeName.GetTypeName() == NAME_MapProperty);
+		CHECK(TypeName.GetTypeParameterCount() == 2);
+		CHECK(TypeName.GetTypeParameter(0).GetTypeName() == NAME_Int32Property);
+		CHECK(TypeName.GetTypeParameter(0).GetTypeParameter(0).IsEmpty());
+		CHECK(TypeName.GetTypeParameter(1).GetTypeName() == NAME_UInt32Property);
+		CHECK(TypeName.GetTypeParameter(1).GetTypeParameter(0).IsEmpty());
+		CHECK(TypeName.GetTypeParameter(2).IsEmpty());
+	}
+
+	SECTION("EnumMap")
+	{
+		FPropertyTypeName TypeName = PropertyTypeNameTest::CreateEnumMap();
+		CHECK(!TypeName.IsEmpty());
+		CHECK(TypeName.GetTypeName() == NAME_MapProperty);
+		CHECK(TypeName.GetTypeParameterCount() == 2);
+		CHECK(TypeName.GetTypeParameter(0).GetTypeName() == NAME_EnumProperty);
+		CHECK(TypeName.GetTypeParameter(0).GetTypeParameter(0).GetTypeName() == NAME_ByteProperty);
+		CHECK(TypeName.GetTypeParameter(0).GetTypeParameter(0).GetTypeParameter(0).GetTypeName() == TEXT("Key"));
+		CHECK(TypeName.GetTypeParameter(0).GetTypeParameter(0).GetTypeParameter(0).GetTypeParameter(0).IsEmpty());
+		CHECK(TypeName.GetTypeParameter(0).GetTypeParameter(0).GetTypeParameter(1).IsEmpty());
+		CHECK(TypeName.GetTypeParameter(0).GetTypeParameter(1).IsEmpty());
+		CHECK(TypeName.GetTypeParameter(1).GetTypeName() == NAME_EnumProperty);
+		CHECK(TypeName.GetTypeParameter(1).GetTypeParameter(0).GetTypeName() == NAME_ByteProperty);
+		CHECK(TypeName.GetTypeParameter(1).GetTypeParameter(0).GetTypeParameter(0).GetTypeName() == TEXT("Value"));
+		CHECK(TypeName.GetTypeParameter(1).GetTypeParameter(0).GetTypeParameter(0).GetTypeParameter(0).IsEmpty());
+		CHECK(TypeName.GetTypeParameter(1).GetTypeParameter(0).GetTypeParameter(1).IsEmpty());
+		CHECK(TypeName.GetTypeParameter(1).GetTypeParameter(1).IsEmpty());
+		CHECK(TypeName.GetTypeParameter(2).IsEmpty());
+	}
+
+	SECTION("Equals+GetTypeHash")
+	{
+		const FPropertyTypeName Int = PropertyTypeNameTest::CreateInt();
+		const FPropertyTypeName VectorArray = PropertyTypeNameTest::CreateVectorArray();
+
+		FPropertyTypeNameBuilder Builder;
+		Builder.AddTypeName(NAME_MapProperty);
+		Builder.BeginTypeParameters();
+		Builder.AddTypeName(NAME_IntProperty);
+		Builder.AddTypeName(NAME_IntProperty);
+		Builder.EndTypeParameters();
+		const FPropertyTypeName MapIntToInt = Builder.Build();
+
+		Builder.Reset();
+		Builder.AddTypeName(NAME_MapProperty);
+		Builder.BeginTypeParameters();
+		Builder.AddTypeName(NAME_IntProperty);
+		Builder.AddTypeName(NAME_MapProperty);
+		{
+			Builder.BeginTypeParameters();
+			Builder.AddTypeName(NAME_IntProperty);
+			Builder.AddTypeName(NAME_IntProperty);
+			Builder.EndTypeParameters();
+		}
+		Builder.EndTypeParameters();
+		const FPropertyTypeName MapIntToIntToInt = Builder.Build();
+
+		CHECK(FPropertyTypeName() == FPropertyTypeName());
+		CHECK(GetTypeHash(FPropertyTypeName()) == GetTypeHash(FPropertyTypeName()));
+
+		CHECK(Int == CopyTemp(Int));
+		CHECK_FALSE(Int == VectorArray);
+		CHECK_FALSE(Int == FPropertyTypeName());
+		CHECK(GetTypeHash(Int) == GetTypeHash(CopyTemp(Int)));
+		CHECK_FALSE(GetTypeHash(Int) == GetTypeHash(VectorArray));
+		CHECK_FALSE(GetTypeHash(Int) == GetTypeHash(FPropertyTypeName()));
+
+		CHECK(Int == MapIntToInt.GetTypeParameter(0));
+		CHECK_FALSE(Int == MapIntToInt);
+		CHECK_FALSE(MapIntToInt == Int);
+		CHECK(GetTypeHash(Int) == GetTypeHash(MapIntToInt.GetTypeParameter(0)));
+		CHECK_FALSE(GetTypeHash(Int) == GetTypeHash(MapIntToInt));
+		CHECK_FALSE(GetTypeHash(MapIntToInt) == GetTypeHash(Int));
+
+		CHECK(MapIntToInt == MapIntToIntToInt.GetTypeParameter(1));
+		CHECK_FALSE(MapIntToInt == MapIntToIntToInt);
+		CHECK(GetTypeHash(MapIntToInt) == GetTypeHash(MapIntToIntToInt.GetTypeParameter(1)));
+		CHECK_FALSE(GetTypeHash(MapIntToInt) == GetTypeHash(MapIntToIntToInt));
+	}
+
+	SECTION("Archive")
+	{
+		FPropertyTypeName Empty;
+		FPropertyTypeName Int = PropertyTypeNameTest::CreateInt();
+		FPropertyTypeName EnumMap = PropertyTypeNameTest::CreateEnumMap();
+		FPropertyTypeName VectorArray = PropertyTypeNameTest::CreateVectorArray();
+
+		TArray<uint8> PersistentData;
+		FMemoryWriter PersistentSaveAr(PersistentData, /*bIsPersistent*/ true);
+		PersistentSaveAr << Empty << Int << EnumMap << VectorArray;
+
+		TArray<uint8> MemoryData;
+		FMemoryWriter MemorySaveAr(MemoryData, /*bIsPersistent*/ false);
+		MemorySaveAr << Empty << Int << EnumMap << VectorArray;
+
+		FPropertyTypeName EmptyCopy = Int; // needs to be non-empty
+		FPropertyTypeName IntCopy;
+		FPropertyTypeName EnumMapCopy;
+		FPropertyTypeName VectorArrayCopy;
+
+		FMemoryReader PersistentReader(PersistentData, /*bIsPersistent*/ true);
+		PersistentReader << EmptyCopy << IntCopy << EnumMapCopy << VectorArrayCopy;
+		CHECK(Empty == EmptyCopy);
+		CHECK(Int == IntCopy);
+		CHECK(EnumMap == EnumMapCopy);
+		CHECK(VectorArray == VectorArrayCopy);
+
+		EmptyCopy = Int; // needs to be non-empty
+		IntCopy.Reset();
+		EnumMapCopy.Reset();
+		VectorArrayCopy.Reset();
+
+		FMemoryReader MemoryReader(MemoryData, /*bIsPersistent*/ false);
+		MemoryReader << EmptyCopy << IntCopy << EnumMapCopy << VectorArrayCopy;
+		CHECK(Empty == EmptyCopy);
+		CHECK(Int == IntCopy);
+		CHECK(EnumMap == EnumMapCopy);
+		CHECK(VectorArray == VectorArrayCopy);
+	}
+
+	SECTION("String")
+	{
+		CHECK(TEXTVIEW("None").Equals(WriteToString<128>(FPropertyTypeName())));
+		CHECK(TEXTVIEW("IntProperty").Equals(WriteToString<128>(PropertyTypeNameTest::CreateInt())));
+		CHECK(TEXTVIEW("MapProperty<EnumProperty<ByteProperty<Key>>,EnumProperty<ByteProperty<Value>>>").Equals(WriteToString<128>(PropertyTypeNameTest::CreateEnumMap())));
+		CHECK(TEXTVIEW("ArrayProperty<StructProperty<Vector>>").Equals(WriteToString<128>(PropertyTypeNameTest::CreateVectorArray())));
+	}
+}
+
+TEST_CASE_NAMED(FPropertyTypeNameConcurrencyTest, "CoreUObject::PropertyTypeName::Concurrency", "[Core][UObject][EngineFilter]")
+{
+	ParallelFor(TEXT("PropertyTypeNameConcurrencyTest"), 64 * 1024, 1, [](int32 Index)
+	{
+		FPropertyTypeNameBuilder Builder;
+
+		if (Index % 8 == 0)
+		{
+			Builder.AddTypeName(NAME_OptionalProperty);
+			Builder.BeginTypeParameters();
+		}
+		if (Index % 4 == 0)
+		{
+			Builder.AddTypeName(NAME_ArrayProperty);
+			Builder.BeginTypeParameters();
+		}
+
+		Builder.AddTypeName(NAME_StructProperty);
+		Builder.BeginTypeParameters();
+		Builder.AddTypeName(FName(WriteToString<32>(TEXTVIEW("TestStruct"), Index)));
+		Builder.EndTypeParameters();
+
+		if (Index % 4 == 0)
+		{
+			Builder.EndTypeParameters();
+		}
+		if (Index % 8 == 0)
+		{
+			Builder.EndTypeParameters();
+		}
+
+		Builder.Build();
+	});
+}
+
+} // UE
+
+#endif // WITH_TESTS
