[Bugfix] Accidentally submitted a temp debug change. Reverting

diff --git a/Engine/Source/Runtime/CoreUObject/Private/UObject/PropertyArray.cpp b/Engine/Source/Runtime/CoreUObject/Private/UObject/PropertyArray.cpp
--- a/Engine/Source/Runtime/CoreUObject/Private/UObject/PropertyArray.cpp
+++ b/Engine/Source/Runtime/CoreUObject/Private/UObject/PropertyArray.cpp
@@ -1255,7 +1255,7 @@ bool FArrayProperty::LoadFromTag(const FPropertyTag& Tag)
 		// Skip property types that are missing the name of the inner type.
 		// Structs have their name in a tag in the serialized data, but we cannot
 		// proceed safely unless we know if the struct used native serialization.
-		if (/*!Property->IsA<FStructProperty>() &&*/
+		if (!Property->IsA<FStructProperty>() &&
 			!Property->IsA<FByteProperty>() &&
 			!Property->IsA<FEnumProperty>() &&
 			Property->LoadFromTag(InnerTag))
