Fast MarkObjectsAsUnreachable. Slightly faster on the client AND SIGNIFICANTLY FASTER on the server where everything runs single-threaded (down from up to 20ms to < 1ms so no need to make it incremental).

- Introduced UE::GC::GReachableObjectFlag which makes flipping all objects from Reachable -> MaybeUnreachable a single Swap() function call.
- Roots and clustered objects are then marked back as Reachable (instead of MaybeUnreachable) but there's significantly fewer of these than all objects that previously needed to be marked as MaybeUnreachable.
- ReachabilityAnalysis now clears MaybeUnreachable AND marks as Reachable (no observable perf regression because of that)

TestClient ReplayRun results (MOAU - MarkObjectsAsUnreachable, RA - ReachabilityAnalysis)

Before:

Min MOAU Time: 0.367427
Max MOAU Time: 1.21097
Avg MOAU Time: 0.709261134146342

Min RA Time: 3.46
Max RA Time: 16.03
Avg RA Time: 11.2565957446808

After:

Min MOAU Time: 0.250281
Max MOAU Time: 1.280585
Avg MOAU Time: 0.692970346153846

Min RA Time: 3.72
Max RA Time: 13.97
Avg RA Time: 11.0717777777778

diff --git a/Engine/Source/Runtime/CoreUObject/Private/UObject/GarbageCollection.cpp b/Engine/Source/Runtime/CoreUObject/Private/UObject/GarbageCollection.cpp
--- a/Engine/Source/Runtime/CoreUObject/Private/UObject/GarbageCollection.cpp
+++ b/Engine/Source/Runtime/CoreUObject/Private/UObject/GarbageCollection.cpp
@@ -544,13 +544,16 @@ namespace UE::GC
 	}
 
 	EGCOptions GetReferenceCollectorOptions(bool bPerformFullPurge);
-	EGatherOptions GetObjectGatherOptions(bool bPerformFullPurge);
+	EGatherOptions GetObjectGatherOptions();
+
+	/** EInternalObjectFlags value representing a reachable object */
+	EInternalObjectFlags GReachableObjectFlag = EInternalObjectFlags::ReachabilityFlag0;
 
 	/** EInternalObjectFlag value representing an unreachable object */
-	EInternalObjectFlags GUnreachableObjectFlag = EInternalObjectFlags::ReachabilityFlag0;
+	EInternalObjectFlags GUnreachableObjectFlag = EInternalObjectFlags::ReachabilityFlag1;
 
 	/** EInternalObjectFlag value representing a maybe unreachable object */
-	EInternalObjectFlags GMaybeUnreachableObjectFlag = EInternalObjectFlags::ReachabilityFlag1;
+	EInternalObjectFlags GMaybeUnreachableObjectFlag = EInternalObjectFlags::ReachabilityFlag2;
 
 	bool GIsIncrementalReachabilityPending = false;
 } // namespace UE::GC
@@ -562,12 +565,93 @@ namespace UE::GC::Private
 	/** List of FUObjectItems representing cluster root objects marker as reachable by GC barrier (see UObject::MarkAsReachable()) */
 	static TExpandingChunkedList<FUObjectItem*> GReachableClusters;
 
-	typedef TThreadedGather<FUObjectItem*> FGatherUnreachableObjectsState;
+	using FGatherUnreachableObjectsState = TThreadedGather<TArray<FUObjectItem*>>;
 	static FGatherUnreachableObjectsState GGatherUnreachableObjectsState;
+
+	static TSet<int32> GRoots;
+	static FCriticalSection GRootsCritical;
 }
 
 static bool GatherUnreachableObjects(UE::GC::EGatherOptions Options, double TimeLimit = 0.0);
 
+bool FUObjectItem::SetRootFlags(EInternalObjectFlags FlagsToSet)
+{
+	using namespace UE::GC;
+	using namespace UE::GC::Private;
+	check(!GIsGarbageCollecting);
+	constexpr int32 RootFlags = (int32)EInternalObjectFlags_RootFlags;
+	bool bIChangedIt = false;
+	{
+		FScopeLock RootsLock(&GRootsCritical);
+		const int32 OldFlags = GetFlagsInternal();
+		if ((OldFlags & RootFlags) == 0)
+		{
+			if (!(GUObjectArray.IsOpenForDisregardForGC() & GUObjectArray.DisregardForGCEnabled())) //-V792
+			{
+				GRoots.Add(GUObjectArray.ObjectToIndex(Object));
+			}
+		}
+		bIChangedIt = ThisThreadAtomicallySetFlag_ForGC(FlagsToSet);
+	}
+	if (bIChangedIt & GIsIncrementalReachabilityPending) //-V792
+	{
+		// Setting any of the root flags on an object during incremental reachability requires a GC barrier
+		// to make sure an object with root flags does not get Garbage Collected
+		checkf(Object, TEXT("Setting an internal object flag on a null object entry"));
+		Object->MarkAsReachable();
+	}
+	return bIChangedIt;
+}
+
+bool FUObjectItem::ClearRootFlags(EInternalObjectFlags FlagsToClear)
+{
+	using namespace UE::GC::Private;
+
+	constexpr int32 RootFlags = (int32)EInternalObjectFlags_RootFlags;
+	FScopeLock RootsLock(&GRootsCritical);
+	const int32 OldFlags = GetFlagsInternal();
+	if ((OldFlags & RootFlags) != 0 && ((OldFlags & ~(int32)FlagsToClear) & RootFlags) == 0)
+	{
+		GRoots.Remove(GUObjectArray.ObjectToIndex(Object));
+	}
+	return ThisThreadAtomicallyClearedFlag_ForGC((EInternalObjectFlags)FlagsToClear);
+
+}
+
+void OnDisregardForGCSetDisabled(int32 NumObjects)
+{
+	using namespace UE::GC;
+	using namespace UE::GC::Private;
+	using FMarkMarkDisregardState = TThreadedGather<TArray<UObject*>>;
+
+	FMarkMarkDisregardState MarkDisregardState;
+	
+	MarkDisregardState.Start(EGatherOptions::Parallel, NumObjects);
+	FMarkMarkDisregardState::FThreadIterators& ThreadIterators = MarkDisregardState.GetThreadIterators();
+
+	// Objects in the disregard for GC set do not have any of the reachability flags set (this way they never become (Maybe)Unreachable) nor are they added to the GRoots array 
+	FScopeLock RootsLock(&GRootsCritical);
+	ParallelFor(TEXT("GC.OnDisregardForGCSetDisabled"), MarkDisregardState.NumWorkerThreads(), 1, [&ThreadIterators](int32 ThreadIndex)
+	{
+		TRACE_CPUPROFILER_EVENT_SCOPE(OnDisregardForGCSetDisabledTask);
+		FMarkMarkDisregardState::FIterator& ThreadState = ThreadIterators[ThreadIndex];
+
+		while (ThreadState.Index <= ThreadState.LastIndex)
+		{
+			int32 ObjectIndex = ThreadState.Index++;
+			FUObjectItem* RootItem = &GUObjectArray.GetObjectItemArrayUnsafe()[ObjectIndex];	
+			if (RootItem->Object)
+			{
+				RootItem->ThisThreadAtomicallySetFlag_ForGC(GReachableObjectFlag);
+				if (RootItem->HasAnyFlags(EInternalObjectFlags_RootFlags))
+				{
+					GRoots.Add(ObjectIndex);
+				}
+			}
+		}
+	}, (MarkDisregardState.NumWorkerThreads() == 1) ? EParallelForFlags::ForceSingleThread : EParallelForFlags::None);
+}
+
 /**
  * Helper class for destroying UObjects on a worker thread
  */
@@ -997,7 +1081,7 @@ static bool MarkClusterMutableObjectsAsReachable(FUObjectCluster& Cluster, Conta
 			{
 				if (ReferencedMutableObjectItem->IsMaybeUnreachable())
 				{
-					if (ReferencedMutableObjectItem->ThisThreadAtomicallyClearedMaybeUnreachable())
+					if (ReferencedMutableObjectItem->MarkAsReachableInterlocked_ForGC())
 					{
 						// Needs doing because this is either a normal unclustered object (clustered objects are never unreachable) or a cluster root
 						ObjectsToSerialize.Add(static_cast<UObject*>(ReferencedMutableObjectItem->Object));
@@ -1019,7 +1103,7 @@ static bool MarkClusterMutableObjectsAsReachable(FUObjectCluster& Cluster, Conta
 						if (ReferencedMutableObjectsClusterRootItem->IsMaybeUnreachable())
 						{
 							// The root is also maybe unreachable so process it and all the referenced clusters
-							if (ReferencedMutableObjectsClusterRootItem->ThisThreadAtomicallyClearedMaybeUnreachable())
+							if (ReferencedMutableObjectsClusterRootItem->MarkAsReachableInterlocked_ForGC())
 							{
 								MarkReferencedClustersAsReachable<Options>(ReferencedMutableObjectsClusterRootItem->GetClusterIndex(), ObjectsToSerialize);
 							}
@@ -1078,10 +1162,7 @@ static FORCENOINLINE void MarkReferencedClustersAsReachable(int32 ClusterIndex,
 			FUObjectItem* ReferencedClusterRootObjectItem = GUObjectArray.IndexToObjectUnsafeForGC(ReferncedClusterIndex);
 			if (!ReferencedClusterRootObjectItem->HasAnyFlags(EInternalObjectFlags::Garbage))
 			{
-				if (ReferencedClusterRootObjectItem->IsMaybeUnreachable())
-				{
-					ReferencedClusterRootObjectItem->ClearMaybeUnreachable();
-				}
+				ReferencedClusterRootObjectItem->FastMarkAsReachableInterlocked_ForGC();
 			}
 			else
 			{
@@ -1625,6 +1706,7 @@ private:
 		// shares cacheline with the raw UObject*. Not adding an offset leads to a bit faster and more compact code.
 		static constexpr uint32 InternalIndexPrefetchOffset = 0;
 		static constexpr uint32 PrefetchAhead = ValidatedPrefetchAhead;
+		static constexpr ::size_t OffsetOfFlags = FUObjectItem::OffsetOfFlags();
 
 		uint32 ObjectIndices[ValidatedBatchSize];
 		if (Num > PrefetchAhead)
@@ -1664,19 +1746,19 @@ private:
 
 			for (uint32 Idx = 0; Idx < PrefetchAhead; ++Idx)
 			{
-				FPlatformMisc::Prefetch(Metadatas[Idx].ObjectItem, offsetof(FUObjectItem, Flags));
+				FPlatformMisc::Prefetch(Metadatas[Idx].ObjectItem, OffsetOfFlags);
 			}
 			for (uint32 Idx = 0; Idx < Num; ++Idx)
 			{
 				Metadatas[Idx].Flags = Metadatas[Idx].ObjectItem->GetFlags();
-				FPlatformMisc::Prefetch(Metadatas[Idx + PrefetchAhead].ObjectItem, offsetof(FUObjectItem, Flags));
+				FPlatformMisc::Prefetch(Metadatas[Idx + PrefetchAhead].ObjectItem, OffsetOfFlags);
 			}
 		}
 		else
 		{
 			for (uint32 Idx = 0; Idx < Num; ++Idx)
 			{
-				FPlatformMisc::Prefetch(Metadatas[Idx].ObjectItem, offsetof(FUObjectItem, Flags));
+				FPlatformMisc::Prefetch(Metadatas[Idx].ObjectItem, OffsetOfFlags);
 			}
 			for (uint32 Idx = 0; Idx < Num; ++Idx)
 			{
@@ -2818,19 +2900,6 @@ constexpr FORCEINLINE EKillable MayKill(EOrigin Origin, bool bAllowKill)
 	return (bAllowKill & (IsEliminatingGarbage(Options) || Origin == EOrigin::Blueprint)) ? EKillable::Yes : EKillable::No;
 }
 
-// Return whether flag was cleared. Only thread-safe for concurrent clear, not concurrent set+clear. Don't use during mark phase.
-FORCEINLINE static bool ClearMaybeUnreachableInterlocked(int32& Flags)
-{
-	const int32 FlagToClear = int32(UE::GC::GMaybeUnreachableObjectFlag);
-	if (FPlatformAtomics::AtomicRead_Relaxed(&Flags) & FlagToClear)
-	{
-		int32 Old = FPlatformAtomics::InterlockedAnd(&Flags, ~FlagToClear);
-		return Old & FlagToClear;
-	}
-
-	return false;
-}
-
 FReachabilityAnalysisState GReachabilityState;
 
 template <EGCOptions InOptions>
@@ -2931,7 +3000,7 @@ public:
 
 	FORCEINLINE static bool HandleValidReference(FWorkerContext& Context, FImmutableReference Reference, FReferenceMetadata Metadata)
 	{
-		if (ClearMaybeUnreachableInterlocked(Metadata.ObjectItem->Flags))
+		if (Metadata.ObjectItem->MarkAsReachableInterlocked_ForGC())
 		{
 			// Objects that are part of a GC cluster should never have the unreachable flag set!
 			checkSlow(Metadata.ObjectItem->GetOwnerIndex() <= 0);
@@ -2959,7 +3028,7 @@ public:
 
 				if (Metadata.ObjectItem->ThisThreadAtomicallySetFlag_ForGC(EInternalObjectFlags::ReachableInCluster))
 				{
-					if (ClearMaybeUnreachableInterlocked(RootObjectItem->Flags))
+					if (RootObjectItem->MarkAsReachableInterlocked_ForGC())
 					{
 						// Make sure all referenced clusters are marked as reachable too
 						MarkReferencedClustersAsReachableThunk<Options>(RootObjectItem->GetClusterIndex(), Context.ObjectsToSerialize);
@@ -3857,11 +3926,8 @@ namespace UE::GC
 
 class FRealtimeGC : public FGarbageCollectionTracer
 {
-	typedef void(FRealtimeGC::*MarkObjectsFn)(EGCOptions, EObjectFlags);
 	typedef void(FRealtimeGC::*ReachabilityAnalysisFn)(FWorkerContext&);
 
-	/** Pointers to functions used for Marking objects as unreachable */
-	MarkObjectsFn MarkObjectsFunctions[2];
 	/** Pointers to functions used for Reachability Analysis */
 	ReachabilityAnalysisFn ReachabilityAnalysisFunctions[8];
 
@@ -3937,20 +4003,10 @@ class FRealtimeGC : public FGarbageCollectionTracer
 			(!!(InOptions & EGCOptions::IncrementalReachability) << 2);
 	}
 
-	static FORCEINLINE int32 GetMarkFunctionIndex(EGCOptions InOptions)
-	{
-		int32 Index = !!(InOptions & EGCOptions::EliminateGarbage);
-		checkf(Index >= 0 && Index <= 1, TEXT("Invalid MarkObjectsAsUnreachable function index (%d)"), Index);
-		return Index;
-	}
-
 public:
 	/** Default constructor, initializing all members. */
 	FRealtimeGC()
 	{
-		MarkObjectsFunctions[GetMarkFunctionIndex(EGCOptions::None)] = &FRealtimeGC::MarkObjectsAsUnreachable<false>;
-		MarkObjectsFunctions[GetMarkFunctionIndex(EGCOptions::EliminateGarbage)] = &FRealtimeGC::MarkObjectsAsUnreachable<true>;
-
 		ReachabilityAnalysisFunctions[GetGCFunctionIndex(EGCOptions::None)] = &FRealtimeGC::PerformReachabilityAnalysisOnObjectsInternal<EGCOptions::None | EGCOptions::None>;
 		ReachabilityAnalysisFunctions[GetGCFunctionIndex(EGCOptions::Parallel | EGCOptions::None)] = &FRealtimeGC::PerformReachabilityAnalysisOnObjectsInternal<EGCOptions::Parallel | EGCOptions::None>;
 
@@ -3966,200 +4022,218 @@ public:
 		FGCObject::StaticInit();
 	}
 
-	/** 
-	 * Marks all objects that don't have KeepFlags and EInternalObjectFlags_GarbageCollectionKeepFlags as unreachable
-	 * This function is a template to speed up the case where we don't need to assemble the token stream (saves about 6ms on PS4)
-	 */
-	template <bool bWithGarbageElimination>
-	void MarkObjectsAsUnreachable(const EGCOptions Options, const EObjectFlags KeepFlags)
-	{
-		TRACE_CPUPROFILER_EVENT_SCOPE(MarkObjectsAsUnreachable);
-		const int32 MaxNumberOfObjects = GUObjectArray.GetObjectArrayNum() - GUObjectArray.GetFirstGCIndex();
-		const int32 NumThreads = FMath::Max(1, FTaskGraphInterface::Get().GetNumWorkerThreads());
-		const int32 NumberOfObjectsPerThread = (MaxNumberOfObjects / NumThreads) + 1;		
-
-		TLockFreePointerListFIFO<FUObjectItem, PLATFORM_CACHE_LINE_SIZE> ClustersToDissolveList;
-		TLockFreePointerListFIFO<FUObjectItem, PLATFORM_CACHE_LINE_SIZE> KeepClusterRefsList;
-
-		TArray<TArray<UObject*>, TInlineAllocator<32>> ObjectsToSerializeArrays;
-		ObjectsToSerializeArrays.SetNum(NumThreads);
-
-		// Iterate over all objects. Note that we iterate over the UObjectArray and usually check only internal flags which
-		// are part of the array so we don't suffer from cache misses as much as we would if we were to check ObjectFlags.
-		ParallelFor( TEXT("GC.MarkUnreachable"),NumThreads,1,
-			[&ObjectsToSerializeArrays, &ClustersToDissolveList, &KeepClusterRefsList,
-			 KeepFlags, NumberOfObjectsPerThread, NumThreads, MaxNumberOfObjects, bIsRerun = Stats.bFoundGarbageRef] (int32 ThreadIndex)
-		{
-			TRACE_CPUPROFILER_EVENT_SCOPE(MarkObjectsAsUnreachableTask);
-			constexpr EInternalObjectFlags FastKeepFlags = EInternalObjectFlags_GarbageCollectionKeepFlags;
-			int32 FirstObjectIndex = ThreadIndex * NumberOfObjectsPerThread + GUObjectArray.GetFirstGCIndex();
-			int32 NumObjects = (ThreadIndex < (NumThreads - 1)) ? NumberOfObjectsPerThread : (MaxNumberOfObjects - (NumThreads - 1) * NumberOfObjectsPerThread);
-			int32 LastObjectIndex = FMath::Min(GUObjectArray.GetObjectArrayNum() - 1, FirstObjectIndex + NumObjects - 1);
-			int32 ObjectCountDuringMarkPhase = 0;
-			TArray<UObject*>& LocalObjectsToSerialize = ObjectsToSerializeArrays[ThreadIndex];
-
-			for (int32 ObjectIndex = FirstObjectIndex; ObjectIndex <= LastObjectIndex; ++ObjectIndex)
+	struct FMarkClustersArrays
+	{
+		TArray<FUObjectItem*> KeepClusters;
+		TArray<FUObjectItem*> ClustersToDissolve;
+
+		inline static bool Reserve(const TArray<TGatherIterator<FMarkClustersArrays>, TInlineAllocator<32>>& InGatherResults, FMarkClustersArrays& OutCombinedResults)
+		{
+			int32 NumKeepClusters = 0;
+			int32 NumClustersToDissolve = 0;
+			for (const TGatherIterator<FMarkClustersArrays>& It : InGatherResults)
 			{
-				FUObjectItem* ObjectItem = &GUObjectArray.GetObjectItemArrayUnsafe()[ObjectIndex];
-				if (ObjectItem->Object)
+				NumKeepClusters += It.Payload.KeepClusters.Num();
+				NumClustersToDissolve += It.Payload.ClustersToDissolve.Num();
+			}
+			OutCombinedResults.KeepClusters.Reserve(OutCombinedResults.KeepClusters.Num() + NumKeepClusters);
+			OutCombinedResults.ClustersToDissolve.Reserve(OutCombinedResults.ClustersToDissolve.Num() + NumClustersToDissolve);
+			return (NumKeepClusters + NumClustersToDissolve) > 0;
+		}
+
+		inline static int32 Num(const FMarkClustersArrays& InArrays)
+		{
+			return InArrays.KeepClusters.Num() + InArrays.ClustersToDissolve.Num();
+		}
+
+		inline static void Append(const FMarkClustersArrays& InSource, FMarkClustersArrays& OutDest)
+		{
+			OutDest.KeepClusters += InSource.KeepClusters;
+			OutDest.ClustersToDissolve += InSource.ClustersToDissolve;
+		}
+	};
+	
+	FORCENOINLINE void MarkClusteredObjectsAsReachable(const EGatherOptions Options, TArray<UObject*>& OutRootObjects)
+	{
+		using namespace UE::GC;
+		using namespace UE::GC::Private;
+		using FMarkClustersState = TThreadedGather<FMarkClustersArrays, FMarkClustersArrays>;
+
+		std::atomic<int32> TotalClusteredObjects = 0;
+		FMarkClustersState GatherClustersState;
+		TArray<FUObjectCluster>& ClusterArray = GUObjectClusters.GetClustersUnsafe();
+
+		// StartGathering calculates the number of threads based on the number of objects but here the objects are actually clusters
+		// that contain many more objects than the number of clusters so we want to be able to process at least two clusters per thread	
+		const int32 NumThreads = !!(Options & EGatherOptions::Parallel) ? FMath::Min(GetNumCollectReferenceWorkers(), (ClusterArray.Num() + 1) / 2) : 1;
+		GatherClustersState.Start(Options, ClusterArray.Num(), /* FirstIndex = */ 0, NumThreads);
+		FMarkClustersState::FThreadIterators& ThreadIterators = GatherClustersState.GetThreadIterators();
+
+		ParallelFor(TEXT("GC.MarkClusteredObjectsAsReachable"), GatherClustersState.NumWorkerThreads(), 1, [&ThreadIterators, &ClusterArray, &TotalClusteredObjects](int32 ThreadIndex)
+		{
+			TRACE_CPUPROFILER_EVENT_SCOPE(MarkClusteredObjectsAsReachableTask);
+			FMarkClustersState::FIterator& ThreadState = ThreadIterators[ThreadIndex];
+			int32 ThisThreadClusteredObjects = 0;
+
+			while (ThreadState.Index <= ThreadState.LastIndex)
+			{
+				int32 ClusterIndex = ThreadState.Index++;
+				FUObjectCluster& Cluster = ClusterArray[ClusterIndex];
+				if (Cluster.RootIndex >= 0)
 				{
-					UObject* Object = (UObject*)ObjectItem->Object;
-
-					// We can't collect garbage during an async load operation and by now all unreachable objects should've been purged.
-					checkf(	bIsRerun ||
-							!ObjectItem->HasAnyFlags(UE::GC::GUnreachableObjectFlag | EInternalObjectFlags::PendingConstruction),
-							TEXT("Object: '%s' with ObjectFlags=0x%08x and InternalObjectFlags=0x%08x. ")
-							TEXT("State: IsEngineExitRequested=%d, GIsCriticalError=%d, GExitPurge=%d, GObjPurgeIsRequired=%d, GObjIncrementalPurgeIsInProgress=%d, GObjFinishDestroyHasBeenRoutedToAllObjects=%d, GGCObjectsPendingDestructionCount=%d"),
-							*Object->GetFullName(),
-							Object->GetFlags(),
-							Object->GetInternalFlags(),
-							IsEngineExitRequested(),
-							GIsCriticalError,
-							GExitPurge,
-							GObjPurgeIsRequired,
-							GObjIncrementalPurgeIsInProgress,
-							GObjFinishDestroyHasBeenRoutedToAllObjects,
-							GGCObjectsPendingDestructionCount);
-
-					// Objects may still be marked as MaybeUnreachable if we forced a full purge in the middle of incremental GC
-					checkf(UE::GC::GIsIncrementalReachabilityPending || !ObjectItem->IsMaybeUnreachable() || bIsRerun, TEXT("%s"), *Object->GetFullName());
-
-					// Keep track of how many objects are around.
-					ObjectCountDuringMarkPhase++;
-					
-					ObjectItem->ClearFlags(EInternalObjectFlags::ReachableInCluster);
-					
-					// Special case handling for objects that are part of the root set.
-					if (ObjectItem->IsRootSet())
+					FUObjectItem* RootItem = &GUObjectArray.GetObjectItemArrayUnsafe()[Cluster.RootIndex];
+					if (!RootItem->IsGarbage())
 					{
-						// IsValidLowLevel is extremely slow in this loop so only do it in debug
-						checkSlow(Object->IsValidLowLevel());
-						// We cannot use RF_PendingKill on objects that are part of the root set.
-#if DO_GUARD_SLOW
-						checkCode(if (ObjectItem->IsGarbage()) { UE_LOG(LogGarbage, Fatal, TEXT("Object %s is part of root set though has been marked as Garbage!"), *Object->GetFullName()); });
-#endif
-						if (ObjectItem->HasAnyFlags(EInternalObjectFlags::ClusterRoot) || ObjectItem->GetOwnerIndex() > 0)
+						bool bKeepCluster = RootItem->HasAnyFlags(EInternalObjectFlags_RootFlags);
+						if (bKeepCluster)
 						{
-							KeepClusterRefsList.Push(ObjectItem);
+							RootItem->FastMarkAsReachableInterlocked_ForGC();
+							ThreadState.Payload.KeepClusters.Add(RootItem);
 						}
 
-						LocalObjectsToSerialize.Add(Object);
-					}
-					// Cluster objects 
-					else if (ObjectItem->GetOwnerIndex() > 0)
-					{
-						// treat cluster objects with FastKeepFlags the same way as if they are in the root set
-						if (ObjectItem->HasAnyFlags(FastKeepFlags))
+						for (int32 ObjectIndex : Cluster.Objects)
 						{
-							KeepClusterRefsList.Push(ObjectItem);
-							LocalObjectsToSerialize.Add(Object);
-						}
-					}
-					// Regular objects or cluster root objects
-					else
-					{
-						bool bMarkAsUnreachable = true;
-						// Internal flags are super fast to check and is used by async loading and must have higher precedence than Garbage
-						if (ObjectItem->HasAnyFlags(FastKeepFlags))
-						{
-							bMarkAsUnreachable = false;
-						}
-						// If KeepFlags is non zero this is going to be very slow due to cache misses
-						else if (!(bWithGarbageElimination && ObjectItem->IsGarbage()) && KeepFlags != RF_NoFlags && Object->HasAnyFlags(KeepFlags))
-						{
-							bMarkAsUnreachable = false;
-						}
-						else if (ObjectItem->HasAllFlags(EInternalObjectFlags::Garbage | EInternalObjectFlags::ClusterRoot))
-						{
-							ClustersToDissolveList.Push(ObjectItem);
-						}
+							FUObjectItem* ClusteredItem = &GUObjectArray.GetObjectItemArrayUnsafe()[ObjectIndex];
 
-						// Mark objects as unreachable unless they have any of the passed in KeepFlags set and it's not marked for elimination..
-						if (!bMarkAsUnreachable)
-						{
-							// IsValidLowLevel is extremely slow in this loop so only do it in debug
-							checkSlow(Object->IsValidLowLevel());
-							LocalObjectsToSerialize.Add(Object);
+							ClusteredItem->FastMarkAsReachableAndClearReachaleInClusterInterlocked_ForGC();
 
-							if (ObjectItem->HasAnyFlags(EInternalObjectFlags::ClusterRoot))
+							if (!bKeepCluster && ClusteredItem->HasAnyFlags(EInternalObjectFlags_RootFlags))
 							{
-								KeepClusterRefsList.Push(ObjectItem);
+								ThreadState.Payload.KeepClusters.Add(RootItem);
+								bKeepCluster = true;
 							}
 						}
-						else
-						{
-							ObjectItem->SetMaybeUnreachable();
-						}
-					}					
+					}
+					else
+					{
+						ThreadState.Payload.ClustersToDissolve.Add(RootItem);
+					}
 				}
 			}
+			TotalClusteredObjects += ThisThreadClusteredObjects;
+		}, (GatherClustersState.NumWorkerThreads() == 1) ? EParallelForFlags::ForceSingleThread : EParallelForFlags::None);
+	
+		FMarkClustersArrays MarkClustersResults;
+		GatherClustersState.Finish(MarkClustersResults);
 
-			GObjectCountDuringLastMarkPhase.Add(ObjectCountDuringMarkPhase);
-		}, !(Options & EGCOptions::Parallel) ? EParallelForFlags::ForceSingleThread : EParallelForFlags::None);
-		
-		// Collect all objects to serialize from all threads and put them into a single array
+		for (FUObjectItem* ObjectItem : MarkClustersResults.ClustersToDissolve)
 		{
-			int32 NumTotal = 0;
-			for (TArray<UObject*>& Objects : ObjectsToSerializeArrays)
-			{
-				NumTotal += Objects.Num();
-			}
-			InitialObjects.Reserve(InitialObjects.Num() + NumTotal + UE::GC::ObjectLookahead);
-			for (TArray<UObject*>& Objects : ObjectsToSerializeArrays)
+			// Check if the object is still a cluster root - it's possible one of the previous
+			// DissolveClusterAndMarkObjectsAsUnreachable calls already dissolved its cluster
+			if (ObjectItem->HasAnyFlags(EInternalObjectFlags::ClusterRoot))
 			{
-				InitialObjects.Append(Objects);
+				GUObjectClusters.DissolveClusterAndMarkObjectsAsUnreachable(ObjectItem);
+				GUObjectClusters.SetClustersNeedDissolving();
 			}
+		}
 
-			ObjectsToSerializeArrays.Empty();
+		for (FUObjectItem* ObjectItem : MarkClustersResults.KeepClusters)
+		{
+			checkSlow(ObjectItem->HasAnyFlags(EInternalObjectFlags::ClusterRoot));
+			// this thing is definitely not marked unreachable, so don't test it here
+			// Make sure all referenced clusters are marked as reachable too
+			MarkReferencedClustersAsReachable<EGCOptions::None>(ObjectItem->GetClusterIndex(), OutRootObjects);
 		}
+	}
+
+	FORCENOINLINE void MarkRootObjectsAsReachable(const EGatherOptions Options, const EObjectFlags KeepFlags, TArray<UObject*>& OutRootObjects)
+	{
+		using namespace UE::GC;
+		using namespace UE::GC::Private;		
+		using FMarkRootsState = TThreadedGather<TArray<UObject*>>;
+
+		FMarkRootsState MarkRootsState;		
 
 		{
-			TArray<FUObjectItem*> ClustersToDissolve;
-			ClustersToDissolveList.PopAll(ClustersToDissolve);
-			for (FUObjectItem* ObjectItem : ClustersToDissolve)
+			TArray<int32> RootsArray(GRoots.Array());				
+			MarkRootsState.Start(Options, RootsArray.Num());
+			FMarkRootsState::FThreadIterators& ThreadIterators = MarkRootsState.GetThreadIterators();
+
+			ParallelFor(TEXT("GC.MarkRootObjectsAsReachable"), MarkRootsState.NumWorkerThreads(), 1, [&ThreadIterators, &RootsArray](int32 ThreadIndex)
 			{
-				// Check if the object is still a cluster root - it's possible one of the previous
-				// DissolveClusterAndMarkObjectsAsUnreachable calls already dissolved its cluster
-				if (ObjectItem->HasAnyFlags(EInternalObjectFlags::ClusterRoot))
+				TRACE_CPUPROFILER_EVENT_SCOPE(MarkClusteredObjectsAsReachableTask);
+				FMarkRootsState::FIterator& ThreadState = ThreadIterators[ThreadIndex];
+
+				while (ThreadState.Index <= ThreadState.LastIndex)
 				{
-					GUObjectClusters.DissolveClusterAndMarkObjectsAsUnreachable(ObjectItem);
-					GUObjectClusters.SetClustersNeedDissolving();
+					FUObjectItem* RootItem = &GUObjectArray.GetObjectItemArrayUnsafe()[RootsArray[ThreadState.Index++]];
+					UObject* Object = static_cast<UObject*>(RootItem->Object);
+
+					// IsValidLowLevel is extremely slow in this loop so only do it in debug
+					checkSlow(Object->IsValidLowLevel());					
+#if DO_GUARD_SLOW
+					// We cannot mark Root objects as Garbage.
+					checkCode(if (ObjectItem->HasAllFlags(EInternalObjectFlags::Garbage | EInternalObjectFlags::RootSet)) { UE_LOG(LogGarbage, Fatal, TEXT("Object %s is part of root set though has been marked as Garbage!"), *Object->GetFullName()); });
+#endif
+
+					RootItem->FastMarkAsReachableInterlocked_ForGC();
+					ThreadState.Payload.Add(Object);
 				}
-			}
-		
-			TArray<FUObjectItem*> KeepClusterRefs;
-			KeepClusterRefsList.PopAll(KeepClusterRefs);
-			for (FUObjectItem* ObjectItem : KeepClusterRefs)
+			}, (MarkRootsState.NumWorkerThreads() == 1) ? EParallelForFlags::ForceSingleThread : EParallelForFlags::None);			
+		}
+
+		using FMarkObjectsState = TThreadedGather<TArray<UObject*>>;
+		FMarkObjectsState MarkObjectsState;
+
+		// This is super slow as we need to look through all existing UObjects and access their memory to check EObjectFlags
+		if (KeepFlags != RF_NoFlags)
+		{
+			MarkObjectsState.Start(Options, GUObjectArray.GetObjectArrayNum(), GUObjectArray.GetFirstGCIndex());
+
+			FMarkObjectsState::FThreadIterators& ThreadIterators = MarkObjectsState.GetThreadIterators();
+			ParallelFor(TEXT("GC.SlowMarkObjectAsReachable"), MarkObjectsState.NumWorkerThreads(), 1, [&ThreadIterators, &KeepFlags](int32 ThreadIndex)
 			{
-				if (ObjectItem->GetOwnerIndex() > 0)
+				TRACE_CPUPROFILER_EVENT_SCOPE(MarkClusteredObjectsAsReachableTask);
+				FMarkObjectsState::FIterator& ThreadState = ThreadIterators[ThreadIndex];
+				const bool bWithGarbageElimination = UObject::IsGarbageEliminationEnabled();
+
+				while (ThreadState.Index <= ThreadState.LastIndex)
 				{
-					checkSlow(!ObjectItem->HasAnyFlags(EInternalObjectFlags::ClusterRoot));
-					bool bNeedsDoing = !ObjectItem->HasAnyFlags(EInternalObjectFlags::ReachableInCluster);
-					if (bNeedsDoing)
+					FUObjectItem* ObjectItem = &GUObjectArray.GetObjectItemArrayUnsafe()[ThreadState.Index++];
+					UObject* Object = static_cast<UObject*>(ObjectItem->Object);
+					if (Object &&
+						!ObjectItem->HasAnyFlags(EInternalObjectFlags_RootFlags) && // It may be counter intuitive to reject roots but these are tracked with GRoots and have already been marked and added
+						!(bWithGarbageElimination && ObjectItem->IsGarbage()) && Object->HasAnyFlags(KeepFlags)) // Garbage elimination works regardless of KeepFlags
 					{
-						ObjectItem->ThisThreadAtomicallySetFlag_ForGC(EInternalObjectFlags::ReachableInCluster);
-						// Make sure cluster root object is reachable too
-						const int32 OwnerIndex = ObjectItem->GetOwnerIndex();
-						FUObjectItem* RootObjectItem = GUObjectArray.IndexToObjectUnsafeForGC(OwnerIndex);
-						checkSlow(RootObjectItem->HasAnyFlags(EInternalObjectFlags::ClusterRoot));
-						// if it is reachable via keep flags we will do this below (or maybe already have)
-						if (RootObjectItem->IsMaybeUnreachable()) 
-						{
-							RootObjectItem->ClearMaybeUnreachable();
-							// Make sure all referenced clusters are marked as reachable too
-							MarkReferencedClustersAsReachable<EGCOptions::None>(RootObjectItem->GetClusterIndex(), InitialObjects);
-						}
+						// IsValidLowLevel is extremely slow in this loop so only do it in debug
+						checkSlow(Object->IsValidLowLevel());
+
+						ObjectItem->FastMarkAsReachableInterlocked_ForGC();
+						ThreadState.Payload.Add(Object);
 					}
 				}
-				else
-				{
-					checkSlow(ObjectItem->HasAnyFlags(EInternalObjectFlags::ClusterRoot));
-					// this thing is definitely not marked unreachable, so don't test it here
-					// Make sure all referenced clusters are marked as reachable too
-					MarkReferencedClustersAsReachable<EGCOptions::None>(ObjectItem->GetClusterIndex(), InitialObjects);
-				}
-			}
+			}, (MarkObjectsState.NumWorkerThreads() == 1) ? EParallelForFlags::ForceSingleThread : EParallelForFlags::None);
+		}
+
+		// Preallocate the resulting array taking both MarkRootsState and MarkObjectsState results into account to avoild reallocating OutRootObjects in each of the Finish() calls.
+		OutRootObjects.Reserve(OutRootObjects.Num() + MarkRootsState.NumGathered() + MarkObjectsState.NumGathered() + ObjectLookahead);
+		MarkRootsState.Finish(OutRootObjects);
+		MarkObjectsState.Finish(OutRootObjects);
+	}
+
+	/**
+	 * Marks all objects that don't have KeepFlags and EInternalObjectFlags_GarbageCollectionKeepFlags as MaybeUnreachable
+	 */
+	FORCENOINLINE void MarkObjectsAsUnreachable(const EObjectFlags KeepFlags)
+	{
+		using namespace UE::GC;
+
+		// Don't swap the flags if we're re-entering this function to track garbage references
+		if (const bool bInitialMark = !Stats.bFoundGarbageRef)
+		{
+			// This marks all UObjects as MaybeUnreachable
+			Swap(GReachableObjectFlag, GMaybeUnreachableObjectFlag);
 		}
+
+		// Not counting the disregard for GC set to preserve legacy behavior
+		GObjectCountDuringLastMarkPhase.Set(GUObjectArray.GetObjectArrayNumMinusAvailable() - GUObjectArray.GetFirstGCIndex());
+
+		EGatherOptions GatherOptions = GetObjectGatherOptions();
+
+		// Now make sure all clustered objects and root objects are marked as Reachable. 
+		// This could be considered as initial part of reachability analysis and could be made incremental.
+		MarkClusteredObjectsAsReachable(GatherOptions, InitialObjects);
+		MarkRootObjectsAsReachable(GatherOptions, KeepFlags, InitialObjects);
 	}
 
 private:
@@ -4197,7 +4271,7 @@ private:
 
 		{
 			const double StartTime = FPlatformTime::Seconds();
-			(this->*MarkObjectsFunctions[GetMarkFunctionIndex(Options)])(Options, KeepFlags);
+			MarkObjectsAsUnreachable(KeepFlags);
 			UE_LOG(LogGarbage, Verbose, TEXT("%f ms for MarkObjectsAsUnreachable Phase (%d Objects To Serialize)"), (FPlatformTime::Seconds() - StartTime) * 1000, InitialObjects.Num());
 		}
 	}
@@ -4858,7 +4932,7 @@ void DissolveUnreachableClusters(UE::GC::EGatherOptions Options)
 
 	using namespace UE::GC;
 	using namespace UE::GC::Private;
-	typedef TThreadedGather<int32> FGatherClustersState;
+	using FGatherClustersState = TThreadedGather<TArray<int32>>;
 
 	std::atomic<int32> TotalClusteredObjects = 0;
 	FGatherClustersState GatherClustersState;
@@ -4999,10 +5073,10 @@ bool GatherUnreachableObjects(UE::GC::EGatherOptions Options, double TimeLimit /
 		}
 
 #if VERIFY_DISREGARD_GC_ASSUMPTIONS
-		if (GVerifyNoUnreachableObjects != 0 && GUnreachableObjects.Num() > 0 && !GExitPurge)
+		if (GVerifyNoUnreachableObjects != 0 && !GExitPurge)
 		{
 			DECLARE_SCOPE_CYCLE_COUNTER(TEXT("CollectGarbageInternal.VerifyNoUnreachableObjects"), STAT_CollectGarbageInternal_VerifyNoUnreachableObjects, STATGROUP_GC);
-			VerifyNoUnreachableObjects();
+			VerifyNoUnreachableObjects(GUnreachableObjects.Num());
 		}
 #endif // VERIFY_DISREGARD_GC_ASSUMPTIONS
 	}
@@ -5311,7 +5385,7 @@ void PreCollectGarbageImpl(EObjectFlags KeepFlags)
 				{
 					VerifyClustersAssumptions();
 				}
-				VerifyObjectFlagMirroring();
+				VerifyObjectFlags();
 				UE_LOG(LogGarbage, Log, TEXT("%.2f ms for Verify GC Assumptions"), (FPlatformTime::Seconds() - StartTime) * 1000);
 			}
 #endif
@@ -5873,7 +5947,7 @@ FORCEINLINE static void MarkObjectItemAsReachable(FUObjectItem* ObjectItem)
 	{
 		checkf(GIsIncrementalReachabilityPending, TEXT("%s is marked as MaybeUnreachable but Incremental Reachability Analysis is not in progress"), *static_cast<UObject*>(ObjectItem->Object)->GetFullName());
 	}
-	if (ObjectItem->ThisThreadAtomicallyClearedMaybeUnreachable())
+	if (ObjectItem->MarkAsReachableInterlocked_ForGC())
 	{
 		if (ObjectItem->GetOwnerIndex() >= 0)
 		{
diff --git a/Engine/Source/Runtime/CoreUObject/Private/UObject/GarbageCollectionVerification.cpp b/Engine/Source/Runtime/CoreUObject/Private/UObject/GarbageCollectionVerification.cpp
--- a/Engine/Source/Runtime/CoreUObject/Private/UObject/GarbageCollectionVerification.cpp
+++ b/Engine/Source/Runtime/CoreUObject/Private/UObject/GarbageCollectionVerification.cpp
@@ -384,14 +384,14 @@ void VerifyClustersAssumptions()
 	UE_CLOG(NumErrors.GetValue() > 0, LogGarbage, Fatal, TEXT("Encountered %d object(s) breaking GC Clusters assumptions. Please check log for details."), NumErrors.GetValue());
 }
 
-void VerifyObjectFlagMirroring()
+void VerifyObjectFlags()
 {
 	int32 MaxNumberOfObjects = GUObjectArray.GetObjectArrayNum();
 	int32 NumThreads = FMath::Max(1, FTaskGraphInterface::Get().GetNumWorkerThreads());
 	int32 NumberOfObjectsPerThread = (MaxNumberOfObjects / NumThreads) + 1;
 	std::atomic<uint32> NumErrors(0);
 
-	ParallelFor( TEXT("GC.VerifyFlagMirroring"),NumThreads,1, [&NumErrors, NumberOfObjectsPerThread, NumThreads, MaxNumberOfObjects](int32 ThreadIndex)
+	ParallelFor( TEXT("GC.VerifyObjectFlags"),NumThreads,1, [&NumErrors, NumberOfObjectsPerThread, NumThreads, MaxNumberOfObjects](int32 ThreadIndex)
 	{
 		int32 FirstObjectIndex = ThreadIndex * NumberOfObjectsPerThread;
 		int32 NumObjects = (ThreadIndex < (NumThreads - 1)) ? NumberOfObjectsPerThread : (MaxNumberOfObjects - (NumThreads - 1) * NumberOfObjectsPerThread);
@@ -401,22 +401,42 @@ void VerifyObjectFlagMirroring()
 			FUObjectItem& ObjectItem = GUObjectArray.GetObjectItemArrayUnsafe()[FirstObjectIndex + ObjectIndex];
 			if (ObjectItem.Object)
 			{
-				UObjectBaseUtility* Object = (UObjectBaseUtility*)ObjectItem.Object;
+				UObject* Object = (UObject*)ObjectItem.Object;
 				bool bHasObjectFlag = Object->HasAnyFlags(RF_MirroredGarbage);
 				bool bHasInternalFlag = ObjectItem.HasAnyFlags(EInternalObjectFlags::Garbage);
 				if (bHasObjectFlag != bHasInternalFlag)
 				{
-					UE_LOG(LogGarbage, Warning, TEXT("RF_Garbage (%d) and EInternalObjectFlags::Garbage (%d) flag mismatch on %s"),
+					UE_LOG(LogGarbage, Warning, TEXT("RF_Garbage (%d) and EInternalObjectFlags::Garbage (%d) flag mismatch on %s%s"),
 						(int32)bHasObjectFlag,
 						(int32)bHasInternalFlag,
+						*FReferenceChainSearch::GetObjectFlags(FGCObjectInfo(Object)),
 						*Object->GetFullName());
 
 					++NumErrors;
 				}
+
+				if (!ObjectItem.HasAnyFlags(UE::GC::GReachableObjectFlag) && !GUObjectArray.IsDisregardForGC(Object))
+				{
+					UE_LOG(LogGarbage, Warning, TEXT("Object %s%s is NOT marked as Reachable at the beginning of GC"),
+						*FReferenceChainSearch::GetObjectFlags(FGCObjectInfo(Object)),
+						*Object->GetFullName());
+
+					++NumErrors;
+				}
+
+				if (ObjectItem.HasAnyFlags(UE::GC::GUnreachableObjectFlag| UE::GC::GMaybeUnreachableObjectFlag))
+				{
+					UE_LOG(LogGarbage, Warning, TEXT("Object %s%s is marked with at least one of the unreachable flags at the beginning of GC"),
+						*FReferenceChainSearch::GetObjectFlags(FGCObjectInfo(Object)),
+						*Object->GetFullName());
+
+					++NumErrors;
+				}
+			}
 		}
-	}});
+	});
 
-	UE_CLOG(NumErrors > 0, LogGarbage, Fatal, TEXT("Encountered %d object(s) breaking Object and Internal flag mirroring assumptions. Please check log for details."), (uint32)NumErrors);
+	UE_CLOG(NumErrors > 0, LogGarbage, Fatal, TEXT("Encountered %d object(s) breaking Object and Internal flag assumptions. Please check log for details."), (uint32)NumErrors);
 }
 
 /**
@@ -476,7 +496,8 @@ public:
 					ReferencingObjectName = GetFullNameSafe(ReferencingObject);
 				}
 
-				UE_LOG(LogGarbage, Warning, TEXT("Unreachable object %s is being referenced by reachable object %s through %s"),
+				UE_LOG(LogGarbage, Warning, TEXT("Object %s%s is being referenced by reachable object %s through %s"),
+					*FReferenceChainSearch::GetObjectFlags(FGCObjectInfo(Object)),
 					*Object->GetFullName(),
 					*ReferencingObjectName,
 					*DebugInfo);
@@ -487,39 +508,90 @@ public:
 	}
 };
 
-void VerifyNoUnreachableObjects()
+void VerifyNoUnreachableObjects(int32 NumUnreachable)
 {
 	const double StartTime = FPlatformTime::Seconds();
 	const int32 MaxNumberOfReachableObjects = GUObjectArray.GetObjectArrayNum();
 	const int32 NumThreads = GetNumCollectReferenceWorkers();
 	const int32 NumberOfObjectsPerThread = (MaxNumberOfReachableObjects / NumThreads) + 1;
 	std::atomic<uint32> NumErrors(0);
+	std::atomic<int32> VerifiedNumUnreachable(0);
 
-	ParallelFor(TEXT("GC.VerifyNoUnreachableObjects"), NumThreads, 1, [&NumErrors, NumberOfObjectsPerThread, NumThreads, MaxNumberOfReachableObjects](int32 ThreadIndex)
+	ParallelFor(TEXT("GC.VerifyNoUnreachableObjects"), NumThreads, 1, [&NumErrors, &VerifiedNumUnreachable, NumberOfObjectsPerThread, NumThreads, MaxNumberOfReachableObjects, NumUnreachable](int32 ThreadIndex)
 	{
 		int32 FirstObjectIndex = ThreadIndex * NumberOfObjectsPerThread;
 		int32 NumObjects = (ThreadIndex < (NumThreads - 1)) ? NumberOfObjectsPerThread : (MaxNumberOfReachableObjects - (NumThreads - 1)*NumberOfObjectsPerThread);
 		TArray<UObject*> ObjectsToSerialize;
 		ObjectsToSerialize.Reserve(NumberOfObjectsPerThread);		
 
+		int32 ThisThreadUnrachableObjectsNum = 0;
+
 		for (int32 ObjectIndex = 0; ObjectIndex < NumObjects && (FirstObjectIndex + ObjectIndex) < GUObjectArray.GetObjectArrayNum(); ++ObjectIndex)
 		{
 			FUObjectItem& ObjectItem = GUObjectArray.GetObjectItemArrayUnsafe()[FirstObjectIndex + ObjectIndex];
-			if (ObjectItem.Object && !ObjectItem.HasAnyFlags(UE::GC::GMaybeUnreachableObjectFlag | UE::GC::GUnreachableObjectFlag))
+			if (ObjectItem.Object)
 			{
-				ObjectsToSerialize.Add(static_cast<UObject*>(ObjectItem.Object));
+				UObject* Object = static_cast<UObject*>(ObjectItem.Object);
+				if (!ObjectItem.HasAnyFlags(UE::GC::GUnreachableObjectFlag))
+				{
+					if (ObjectItem.HasAnyFlags(UE::GC::GMaybeUnreachableObjectFlag))
+					{
+						UE_LOG(LogGarbage, Warning, TEXT("Object %s%s is still marked as MaybeUnreachable after Reachability Analysis is complete."), 
+							*FReferenceChainSearch::GetObjectFlags(FGCObjectInfo(Object)),
+							*Object->GetFullName());
+						NumErrors++;
+					}
+					if (!ObjectItem.HasAnyFlags(UE::GC::GReachableObjectFlag) && !GUObjectArray.IsDisregardForGC(Object))
+					{
+						UE_LOG(LogGarbage, Warning, TEXT("Object %s%s is NOT marked as Unreachable and NOT marked as Reachable."),
+							*FReferenceChainSearch::GetObjectFlags(FGCObjectInfo(Object)),
+							*Object->GetFullName());
+						NumErrors++;
+					}					
+					ObjectsToSerialize.Add(Object);
+				}
+				else
+				{
+					ThisThreadUnrachableObjectsNum++;
+
+					if (ObjectItem.HasAnyFlags(UE::GC::GMaybeUnreachableObjectFlag))
+					{
+						UE_LOG(LogGarbage, Warning, TEXT("Object %s%s is still marked as MaybeUnreachable after Reachability Analysis is complete."),
+							*FReferenceChainSearch::GetObjectFlags(FGCObjectInfo(Object)),
+							*Object->GetFullName());
+						NumErrors++;
+					}
+					if (ObjectItem.HasAnyFlags(UE::GC::GReachableObjectFlag))
+					{
+						UE_LOG(LogGarbage, Warning, TEXT("Object %s%s is still marked as Reachable."),
+							*FReferenceChainSearch::GetObjectFlags(FGCObjectInfo(Object)),
+							*Object->GetFullName());
+						NumErrors++;
+					}
+				}
 			}
 		}
 
-		FUnreachableReferenceProcessor Processor;
-		UE::GC::FWorkerContext Context;
-		Context.SetInitialObjectsUnpadded(ObjectsToSerialize);
-		CollectReferences(Processor, Context);
-		NumErrors.fetch_add(Processor.GetErrorCount(), std::memory_order_acq_rel);
+		if (NumUnreachable > 0)
+		{
+			// No need to scan for unreachable obejcts if there's no unreachable objects. We only care about flag checks above in this case
+			FUnreachableReferenceProcessor Processor;
+			UE::GC::FWorkerContext Context;
+			Context.SetInitialObjectsUnpadded(ObjectsToSerialize);
+			CollectReferences(Processor, Context);
+			NumErrors.fetch_add(Processor.GetErrorCount(), std::memory_order_acq_rel);
+		}
+
+		VerifiedNumUnreachable.fetch_add(ThisThreadUnrachableObjectsNum, std::memory_order_acq_rel);
 	});
 
+	if (VerifiedNumUnreachable.load() != NumUnreachable)
+	{
+		NumErrors++;
+		UE_LOG(LogGarbage, Warning, TEXT("The actual number of objects with the Unreachable flag (%d) does not match the number of gathered unreachable objects (%d)."), VerifiedNumUnreachable.load(), NumUnreachable);
+	}
 
-	UE_CLOG(NumErrors > 0, LogGarbage, Fatal, TEXT("Encountered %d unreachable object(s) that are still reachable."), NumErrors.load());
+	UE_CLOG(NumErrors > 0, LogGarbage, Fatal, TEXT("Detected %d case(s) of breaking reachability assumptions."), NumErrors.load());
 
 	UE_LOG(LogGarbage, Log, TEXT("%f ms for VerifyNoUnreachableObjects"), (FPlatformTime::Seconds() - StartTime) * 1000);
 }
diff --git a/Engine/Source/Runtime/CoreUObject/Private/UObject/GarbageCollectionVerification.h b/Engine/Source/Runtime/CoreUObject/Private/UObject/GarbageCollectionVerification.h
--- a/Engine/Source/Runtime/CoreUObject/Private/UObject/GarbageCollectionVerification.h
+++ b/Engine/Source/Runtime/CoreUObject/Private/UObject/GarbageCollectionVerification.h
@@ -23,10 +23,10 @@ void VerifyGCAssumptions();
 void VerifyClustersAssumptions();
 
 /** Verifies Object and Internal flag mirroring assumptions */
-void VerifyObjectFlagMirroring();
+void VerifyObjectFlags();
 
 /** Verifies Reachability assumptions */
-void VerifyNoUnreachableObjects();
+void VerifyNoUnreachableObjects(int32 NumUnreachable);
 
 #endif // VERIFY_DISREGARD_GC_ASSUMPTIONS
 
diff --git a/Engine/Source/Runtime/CoreUObject/Private/UObject/ReachabilityAnalysisState.h b/Engine/Source/Runtime/CoreUObject/Private/UObject/ReachabilityAnalysisState.h
--- a/Engine/Source/Runtime/CoreUObject/Private/UObject/ReachabilityAnalysisState.h
+++ b/Engine/Source/Runtime/CoreUObject/Private/UObject/ReachabilityAnalysisState.h
@@ -165,13 +165,38 @@ struct TGatherIterator
 	PayloadType Payload;
 };
 
+template <typename PayloadType>
+struct TDefaultPayloadOps
+{
+	inline static int32 Num(const PayloadType& InPayload)
+	{
+		return InPayload.Num();
+	}
+
+	inline static bool Reserve(const TArray<TGatherIterator<PayloadType>, TInlineAllocator<32>>& ThreadIterators, PayloadType& OutPayload)
+	{
+		int32 SizeToReserve = 0;
+		for (const TGatherIterator<PayloadType>& It : ThreadIterators)
+		{
+			SizeToReserve += TDefaultPayloadOps<PayloadType>::Num(It.Payload);
+		}
+		OutPayload.Reserve(TDefaultPayloadOps<PayloadType>::Num(OutPayload) + SizeToReserve);
+		return SizeToReserve > 0;
+	}
+
+	inline static void Append(const PayloadType& InSource, PayloadType& OutDest)
+	{
+		OutDest += InSource;
+	}
+};
+
 /** Helper class that holds the current state of Object Gathering phases */
-template <typename ObjectType>
+template <typename PayloadType, typename PayloadOps = TDefaultPayloadOps<PayloadType>>
 class TThreadedGather
 {
 public:
 
-	typedef TGatherIterator<TArray<ObjectType>> FIterator;
+	typedef TGatherIterator<PayloadType> FIterator;
 	typedef TArray<FIterator, TInlineAllocator<32>> FThreadIterators;
 
 	FORCEINLINE FThreadIterators& GetThreadIterators()
@@ -212,22 +237,17 @@ public:
 		}
 	}
 
-	FORCENOINLINE void Finish(TArray<ObjectType>& OutGatheredObjects)
+	FORCENOINLINE void Finish(PayloadType& OutGatheredObjects)
 	{
-		const int32 NumGatheredObjects = NumGathered();
-		if (NumGatheredObjects)
+		if (ThreadIterators.Num() == 1 && PayloadOps::Num(OutGatheredObjects) == 0)
 		{
-			OutGatheredObjects.Reserve(NumGatheredObjects);
-			if (ThreadIterators.Num() == 1)
-			{
-				OutGatheredObjects = MoveTemp(ThreadIterators[0].Payload);
-			}
-			else
+			OutGatheredObjects = MoveTemp(ThreadIterators[0].Payload);
+		}
+		else if (PayloadOps::Reserve(ThreadIterators, OutGatheredObjects))
+		{
+			for (FIterator& Iterator : ThreadIterators)
 			{
-				for (FIterator& Iterator : ThreadIterators)
-				{
-					OutGatheredObjects += Iterator.Payload;
-				}
+				PayloadOps::Append(Iterator.Payload, OutGatheredObjects);
 			}
 		}
 		ThreadIterators.Reset();
@@ -254,7 +274,7 @@ public:
 		int32 NumGathered = 0;
 		for (const FIterator& ThreadState : ThreadIterators)
 		{
-			NumGathered += ThreadState.Payload.Num();
+			NumGathered += PayloadOps::Num(ThreadState.Payload);
 		}
 		return NumGathered;
 	}
diff --git a/Engine/Source/Runtime/CoreUObject/Private/UObject/ReferenceChainSearch.cpp b/Engine/Source/Runtime/CoreUObject/Private/UObject/ReferenceChainSearch.cpp
--- a/Engine/Source/Runtime/CoreUObject/Private/UObject/ReferenceChainSearch.cpp
+++ b/Engine/Source/Runtime/CoreUObject/Private/UObject/ReferenceChainSearch.cpp
@@ -77,7 +77,7 @@ namespace UE::ReferenceChainSearch
 		bool IsRoot(FVertex Vertex, EReferenceChainSearchMode SearchMode) const
 		{
 			const UObject* Object = VertexToObject(Vertex);
-			return Object->HasAnyInternalFlags(EInternalObjectFlags_GarbageCollectionKeepFlags | EInternalObjectFlags::RootSet)
+			return Object->HasAnyInternalFlags(EInternalObjectFlags_RootFlags)
 				|| (GARBAGE_COLLECTION_KEEPFLAGS != RF_NoFlags && Object->HasAnyFlags(GARBAGE_COLLECTION_KEEPFLAGS)
 					&& !(SearchMode & EReferenceChainSearchMode::FullChain));
 		}
@@ -187,7 +187,7 @@ namespace UE::ReferenceChainSearch
 		bool IsRoot(FVertex Vertex, EReferenceChainSearchMode SearchMode) const
 		{
 			ObjectType Object = VertexToObject(Vertex);
-			return Object->HasAnyInternalFlags(EInternalObjectFlags_GarbageCollectionKeepFlags | EInternalObjectFlags::RootSet)
+			return Object->HasAnyInternalFlags(EInternalObjectFlags_RootFlags)
 				|| (GARBAGE_COLLECTION_KEEPFLAGS != RF_NoFlags && Object->HasAnyFlags(GARBAGE_COLLECTION_KEEPFLAGS)
 					&& !(SearchMode & EReferenceChainSearchMode::FullChain));
 		}
@@ -1263,61 +1263,79 @@ namespace UE::ReferenceChainSearch
 	}
 } // namespace UE::ReferenceChainSearch
 
-FString FReferenceChainSearch::GetObjectFlags(FGCObjectInfo* InObject)
+FString FReferenceChainSearch::GetObjectFlags(const FGCObjectInfo& InObject)
 {
 	FString Flags;
-	if (InObject->IsRooted())
+
+	if (!InObject.IsDisregardForGC())
+	{
+		if (!InObject.HasAnyInternalFlags(UE::GC::GReachableObjectFlag | UE::GC::GMaybeUnreachableObjectFlag | UE::GC::GUnreachableObjectFlag))
+		{
+			Flags += TEXT("(Error: No reachability flag) ");
+		}
+	}
+	else if (InObject.HasAnyInternalFlags(UE::GC::GReachableObjectFlag))
+	{
+		Flags += TEXT("(Error: Reachable but NeverGCed) ");
+	}
+
+	if (InObject.HasAnyInternalFlags(UE::GC::GMaybeUnreachableObjectFlag))
+	{
+		Flags += FString::Printf(TEXT("(MaybeUnreachable<%d>) "), (int32)UE::GC::GMaybeUnreachableObjectFlag);
+	}
+
+	if (InObject.HasAnyInternalFlags(UE::GC::GUnreachableObjectFlag))
+	{
+		Flags += FString::Printf(TEXT("(Unreachable<%d>) "), (int32)UE::GC::GUnreachableObjectFlag);
+	}
+
+	if (InObject.IsRooted())
 	{
 		Flags += TEXT("(root) ");
 	}
 
 	CA_SUPPRESS(6011)
-	if (InObject->IsNative())
+	if (InObject.IsNative())
 	{
 		Flags += TEXT("(native) ");
 	}
 
-	if (InObject->HasAnyInternalFlags(EInternalObjectFlags::Garbage))
+	if (InObject.HasAnyInternalFlags(EInternalObjectFlags::Garbage))
 	{
 		Flags += TEXT("(Garbage) ");
 	}
 
-	if (InObject->HasAnyFlags(RF_Standalone))
+	if (InObject.HasAnyFlags(RF_Standalone))
 	{
 		Flags += TEXT("(standalone) ");
 	}
 
-	if (InObject->HasAnyInternalFlags(EInternalObjectFlags::Async))
+	if (InObject.HasAnyInternalFlags(EInternalObjectFlags::Async))
 	{
 		Flags += TEXT("(async) ");
 	}
 
-	if (InObject->HasAnyInternalFlags(EInternalObjectFlags::AsyncLoading))
+	if (InObject.HasAnyInternalFlags(EInternalObjectFlags::AsyncLoading))
 	{
 		Flags += TEXT("(asyncloading) ");
 	}
 
-	if (InObject->HasAnyInternalFlags(EInternalObjectFlags::LoaderImport))
+	if (InObject.HasAnyInternalFlags(EInternalObjectFlags::LoaderImport))
 	{
 		Flags += TEXT("(loaderimport) ");
 	}
 
-	if (InObject->IsDisregardForGC())
+	if (InObject.IsDisregardForGC())
 	{
 		Flags += TEXT("(NeverGCed) ");
 	}
 
-	if (InObject->HasAnyInternalFlags(EInternalObjectFlags::ClusterRoot))
+	if (InObject.HasAnyInternalFlags(EInternalObjectFlags::ClusterRoot))
 	{
 		Flags += TEXT("(ClusterRoot) ");
 	}
 
-	if (InObject->HasAnyInternalFlags(UE::GC::GMaybeUnreachableObjectFlag))
-	{
-		Flags += TEXT("(MaybeUnreachable) ");
-	}
-
-	if (InObject->GetOwnerIndex() > 0)
+	if (InObject.GetOwnerIndex() > 0)
 	{
 		Flags += TEXT("(Clustered) ");
 	}
@@ -1385,7 +1403,7 @@ void FReferenceChainSearch::DumpChain(FReferenceChainSearch::FReferenceChain* Ch
 
 			Out.Logf(ELogVerbosity::Log, TEXT("%s%s %s"),
 				FCString::Spc(Params.Indent),
-				*GetObjectFlags(ReferencerObject),
+				*GetObjectFlags(*ReferencerObject),
 				*ReferencerObject->GetFullName());
 
 			bPostCallbackContinue = ReferenceCallback(Params);
@@ -1448,7 +1466,7 @@ void FReferenceChainSearch::DumpChain(FReferenceChainSearch::FReferenceChain* Ch
 				Out.Logf(ELogVerbosity::Log, TEXT("%s-> %s = %s %s"),
 					FCString::Spc(Params.Indent),
 					*ReferencingPropertyName,
-					*GetObjectFlags(Object),
+					*GetObjectFlags(*Object),
 					*Object->GetFullName());
 			}
 			else if (ReferenceInfo && ReferenceInfo->Type == EReferenceType::AddReferencedObjects)
@@ -1475,7 +1493,7 @@ void FReferenceChainSearch::DumpChain(FReferenceChainSearch::FReferenceChain* Ch
 				Out.Logf(ELogVerbosity::Log, TEXT("%s-> %s::AddReferencedObjects(%s %s)"),
 					FCString::Spc(Params.Indent),
 					*UObjectOrGCObjectName,
-					*GetObjectFlags(Object),
+					*GetObjectFlags(*Object),
 					*Object->GetFullName());
 
 				if (ReferenceInfo->StackFrames.Num())
@@ -1492,7 +1510,7 @@ void FReferenceChainSearch::DumpChain(FReferenceChainSearch::FReferenceChain* Ch
 				Out.Logf(ELogVerbosity::Log, TEXT("%s-> %s = %s %s"),
 					FCString::Spc(Params.Indent),
 					TEXT("Outer Chain"),
-					*GetObjectFlags(Object),
+					*GetObjectFlags(*Object),
 					*Object->GetFullName());
 			}
 			else
@@ -1500,7 +1518,7 @@ void FReferenceChainSearch::DumpChain(FReferenceChainSearch::FReferenceChain* Ch
 				Out.Logf(ELogVerbosity::Log, TEXT("%s-> %s = %s %s"),
 					FCString::Spc(Params.Indent),
 					TEXT("UNKNOWN"),
-					*GetObjectFlags(Object),
+					*GetObjectFlags(*Object),
 					*Object->GetFullName());
 			}
 
@@ -1685,31 +1703,31 @@ int32 FReferenceChainSearch::PrintResults(TFunctionRef<bool(FCallbackParams& Par
 
 	if (NumPrintedChains == 0)
 	{
-		auto LogUnreachableObject = [this](FGCObjectInfo* ObjInfo) {
-			if (ObjInfo->HasAnyInternalFlags(EInternalObjectFlags_GarbageCollectionKeepFlags))
+		auto LogUnreachableObject = [this](const FGCObjectInfo& ObjInfo) {
+			if (ObjInfo.HasAnyInternalFlags(EInternalObjectFlags_RootFlags))
 			{
-				UE_LOG(LogReferenceChain, Log, TEXT("%s%s is not currently reachable but it does have some of EInternalObjectFlags_GarbageCollectionKeepFlags set."), *GetObjectFlags(ObjInfo), *ObjInfo->GetFullName());
+				UE_LOG(LogReferenceChain, Log, TEXT("%s%s is not currently reachable but it does have some of EInternalObjectFlags_RootFlags set."), *GetObjectFlags(ObjInfo), *ObjInfo.GetFullName());
 			}
-			else if (ObjInfo->HasAnyFlags(GARBAGE_COLLECTION_KEEPFLAGS))
+			else if (ObjInfo.HasAnyFlags(GARBAGE_COLLECTION_KEEPFLAGS))
 			{
-				UE_LOG(LogReferenceChain, Log, TEXT("%s%s is not currently reachable but it does have some of GARBAGE_COLLECTION_KEEPFLAGS set."), *GetObjectFlags(ObjInfo), *ObjInfo->GetFullName());
+				UE_LOG(LogReferenceChain, Log, TEXT("%s%s is not currently reachable but it does have some of GARBAGE_COLLECTION_KEEPFLAGS set."), *GetObjectFlags(ObjInfo), *ObjInfo.GetFullName());
 			}
 			else
 			{
-				UE_LOG(LogReferenceChain, Log, TEXT("%s%s is not currently reachable. Try using GC history to debug transient leaks with 'gc.historysize 1'"), *GetObjectFlags(ObjInfo), *ObjInfo->GetFullName());
+				UE_LOG(LogReferenceChain, Log, TEXT("%s%s is not currently reachable. Try using GC history to debug transient leaks with 'gc.historysize 1'"), *GetObjectFlags(ObjInfo), *ObjInfo.GetFullName());
 			}
 		};
 		if (TargetObject)
 		{
 			FGCObjectInfo* ObjInfo = FGCObjectInfo::FindOrAddInfoHelper(TargetObject, const_cast<TMap<const UObject*, FGCObjectInfo*>&>(ObjectToInfoMap));
 			check(ObjInfo);
-			LogUnreachableObject(ObjInfo);
+			LogUnreachableObject(*ObjInfo);
 		}
 		else
 		{
 			for (FGCObjectInfo* ObjInfo : ObjectInfosToFindReferencesTo)
 			{
-				LogUnreachableObject(ObjInfo);
+				LogUnreachableObject(*ObjInfo);
 			}
 		}
 	}
@@ -1749,7 +1767,7 @@ FString FReferenceChainSearch::GetRootPath(TFunctionRef<bool(FCallbackParams& Pa
 		{
 			FGCObjectInfo ObjectInfo(TargetObject);
 			return FString::Printf(TEXT("%s%s is not currently reachable."),
-				*GetObjectFlags(&ObjectInfo),
+				*GetObjectFlags(ObjectInfo),
 				*ObjectInfo.GetFullName()
 			);
 		}
@@ -1784,7 +1802,7 @@ void FReferenceChainSearch::Cleanup()
 
 static FORCEINLINE bool HasGarbageCollectionKeepFlags(FGCObjectInfo* ObjectInfo)
 {
-	return ObjectInfo && (ObjectInfo->HasAnyFlags(GARBAGE_COLLECTION_KEEPFLAGS) || ObjectInfo->HasAnyInternalFlags(EInternalObjectFlags_GarbageCollectionKeepFlags | EInternalObjectFlags::RootSet));
+	return ObjectInfo && (ObjectInfo->HasAnyFlags(GARBAGE_COLLECTION_KEEPFLAGS) || ObjectInfo->HasAnyInternalFlags(EInternalObjectFlags_RootFlags));
 }
 
 static bool PrintStaleReferenceChainsAndFindReferencingObjects(UObject* ObjectToFindReferencesTo, FReferenceChainSearch& RefChainSearch, FGCObjectInfo*& OutGarbageObject, FGCObjectInfo*& OutReferencingObject, ELogVerbosity::Type Verbosity)
diff --git a/Engine/Source/Runtime/CoreUObject/Private/UObject/UObjectArray.cpp b/Engine/Source/Runtime/CoreUObject/Private/UObject/UObjectArray.cpp
--- a/Engine/Source/Runtime/CoreUObject/Private/UObject/UObjectArray.cpp
+++ b/Engine/Source/Runtime/CoreUObject/Private/UObject/UObjectArray.cpp
@@ -12,6 +12,7 @@
 #include "UObject/UObjectAllocator.h"
 #include "UObject/Class.h"
 #include "UObject/UObjectIterator.h"
+#include "UObject/ReachabilityAnalysisState.h"
 
 DEFINE_LOG_CATEGORY_STATIC(LogUObjectArray, Log, All);
 
@@ -181,6 +182,17 @@ void FUObjectArray::CloseDisregardForGC()
 
 void FUObjectArray::DisableDisregardForGC()
 {
+	if (!GExitPurge && (ObjFirstGCIndex > 0 || DisregardForGCEnabled()))
+	{
+		void OnDisregardForGCSetDisabled(int32 NumObjects);
+		// If disregard for GC was already closed then ObjFirstGCIndex is the number of objects we need to scan, otherwise disregard for GC is still open and we need to scan all objects
+		int32 NumDisregardForGCObjects = ObjFirstGCIndex > 0 ? ObjFirstGCIndex : GetObjectArrayNum();
+		if (NumDisregardForGCObjects > 0)
+		{
+			OnDisregardForGCSetDisabled(NumDisregardForGCObjects);
+		}
+	}
+
 	MaxObjectsNotConsideredByGC = 0;
 	ObjFirstGCIndex = 0;
 	if (IsOpenForDisregardForGC())
@@ -242,11 +254,21 @@ void FUObjectArray::AllocateUObjectIndex(UObjectBase* Object, EInternalObjectFla
 	UE_CLOG(ObjectItem->Object != nullptr, LogUObjectArray, Fatal, TEXT("Attempting to add %s at index %d but another object (0x%016llx) exists at that index!"), *Object->GetFName().ToString(), Index, (int64)(PTRINT)ObjectItem->Object);
 	ObjectItem->Object = Object;
 	// At this point all not-compiled-in objects are not fully constructed yet and this is the earliest we can mark them as such
-	ObjectItem->Flags = (int32)(EInternalObjectFlags::PendingConstruction | InitialFlags);
+	ObjectItem->Flags = (int32)EInternalObjectFlags::PendingConstruction;
+	if (!(IsOpenForDisregardForGC() & GUObjectArray.DisregardForGCEnabled())) //-V792
+	{
+		ObjectItem->Flags |= (int32)UE::GC::GReachableObjectFlag;
+	}
 	ObjectItem->ClusterRootIndex = 0;
 	ObjectItem->SerialNumber = SerialNumber;
 	Object->InternalIndex = Index;
 
+	// This needs to happen after the InternalIndex is set because setting root flags may result in the object being added to UE::GC::Priate::GRoots array 
+	if (InitialFlags != EInternalObjectFlags::None)
+	{
+		ObjectItem->ThisThreadAtomicallySetFlag(InitialFlags);
+	}
+
 	UnlockInternalArray();
 
 	//  @todo: threading: lock UObjectCreateListeners
@@ -305,6 +327,13 @@ void FUObjectArray::FreeUObjectIndex(UObjectBase* Object)
 	int32 Index = Object->InternalIndex;
 	FUObjectItem* ObjectItem = IndexToObject(Index);
 	UE_CLOG(ObjectItem->Object != Object, LogUObjectArray, Fatal, TEXT("Removing object (0x%016llx) at index %d but the index points to a different object (0x%016llx)!"), (int64)(PTRINT)Object, Index, (int64)(PTRINT)ObjectItem->Object);
+
+	// Clear root flags to remove this object's index from UE::GC::Private::GRoots array 
+	if ((ObjectItem->Flags & (int32)EInternalObjectFlags_RootFlags) != 0)
+	{
+		ObjectItem->ThisThreadAtomicallyClearedFlag(EInternalObjectFlags_RootFlags);
+	}
+
 	ObjectItem->Object = nullptr;
 	ObjectItem->Flags = 0;
 	ObjectItem->ClusterRootIndex = 0;
diff --git a/Engine/Source/Runtime/CoreUObject/Public/UObject/GarbageCollectionGlobals.h b/Engine/Source/Runtime/CoreUObject/Public/UObject/GarbageCollectionGlobals.h
--- a/Engine/Source/Runtime/CoreUObject/Public/UObject/GarbageCollectionGlobals.h
+++ b/Engine/Source/Runtime/CoreUObject/Public/UObject/GarbageCollectionGlobals.h
@@ -10,6 +10,9 @@
 
 namespace UE::GC
 {
+	/** Current EInternalObjectFlags value representing a reachable object */
+	extern COREUOBJECT_API EInternalObjectFlags GReachableObjectFlag;
+
 	/** Current EInternalObjectFlags value representing an unreachable object */
 	extern COREUOBJECT_API EInternalObjectFlags GUnreachableObjectFlag;
 
diff --git a/Engine/Source/Runtime/CoreUObject/Public/UObject/ObjectMacros.h b/Engine/Source/Runtime/CoreUObject/Public/UObject/ObjectMacros.h
--- a/Engine/Source/Runtime/CoreUObject/Public/UObject/ObjectMacros.h
+++ b/Engine/Source/Runtime/CoreUObject/Public/UObject/ObjectMacros.h
@@ -608,8 +608,9 @@ enum class EInternalObjectFlags : int32
 {
 	None = 0,
 
-	ReachabilityFlag0 = 1 << 0, ///< One of the flags used by Garbage Collector to determine UObject's reachability state 
-	ReachabilityFlag1 = 1 << 1, ///< One of the flags used by Garbage Collector to determine UObject's reachability state 
+	ReachabilityFlag0 = 1 << 0, ///< One of the flags used by Garbage Collector to determine UObject's reachability state
+	ReachabilityFlag1 = 1 << 1, ///< One of the flags used by Garbage Collector to determine UObject's reachability state
+	ReachabilityFlag2 = 1 << 2, ///< One of the flags used by Garbage Collector to determine UObject's reachability state
 
 	MaybeUnreachable UE_DEPRECATED(5.4, "MaybeUnreachable flag should no longer be used. Use FUObjectItem::IsMaybeUnreachable() or UE::GC::GMaybeUnreachableObjectFlag flag instead.") = 1 << 19, ///< Flag set on all non-root objects at the beginning of Reachability Analysis
 	LoaderImport = 1 << 20, ///< Object is ready to be imported by another package during loading
@@ -640,7 +641,7 @@ ENUM_CLASS_FLAGS(EInternalObjectFlags);
 
 //~ Make sure these macros are up to date!
 #define EInternalObjectFlags_GarbageCollectionKeepFlags (EInternalObjectFlags::Native | EInternalObjectFlags::Async | EInternalObjectFlags::AsyncLoading | EInternalObjectFlags::LoaderImport)
-#define EInternalObjectFlags_AllFlags (EInternalObjectFlags::ReachabilityFlag0 | EInternalObjectFlags::ReachabilityFlag1 | EInternalObjectFlags::LoaderImport | EInternalObjectFlags::Garbage | EInternalObjectFlags::ReachableInCluster | EInternalObjectFlags::ClusterRoot | EInternalObjectFlags::Native | EInternalObjectFlags::Async | EInternalObjectFlags::AsyncLoading | EInternalObjectFlags::RootSet | EInternalObjectFlags::PendingConstruction | (EInternalObjectFlags)(1 << 19 /*MaybeUnreachable*/) | (EInternalObjectFlags)(1 << 28 /*Unreachable*/))
+#define EInternalObjectFlags_AllFlags (EInternalObjectFlags::ReachabilityFlag0 | EInternalObjectFlags::ReachabilityFlag1 | EInternalObjectFlags::ReachabilityFlag2 | EInternalObjectFlags::LoaderImport | EInternalObjectFlags::Garbage | EInternalObjectFlags::ReachableInCluster | EInternalObjectFlags::ClusterRoot | EInternalObjectFlags::Native | EInternalObjectFlags::Async | EInternalObjectFlags::AsyncLoading | EInternalObjectFlags::RootSet | EInternalObjectFlags::PendingConstruction | (EInternalObjectFlags)(1 << 19 /*MaybeUnreachable*/) | (EInternalObjectFlags)(1 << 28 /*Unreachable*/))
 #define EInternalObjectFlags_RootFlags (EInternalObjectFlags::RootSet | EInternalObjectFlags_GarbageCollectionKeepFlags)
 
 /**
diff --git a/Engine/Source/Runtime/CoreUObject/Public/UObject/ReferenceChainSearch.h b/Engine/Source/Runtime/CoreUObject/Public/UObject/ReferenceChainSearch.h
--- a/Engine/Source/Runtime/CoreUObject/Public/UObject/ReferenceChainSearch.h
+++ b/Engine/Source/Runtime/CoreUObject/Public/UObject/ReferenceChainSearch.h
@@ -373,9 +373,10 @@ private:
 		FName PropertyName,
 		TConstArrayView<uint64> StackFrames = {});
 
+public:
 	/** Returns a string with all flags (we care about) set on an object */
-	static FString GetObjectFlags(FGCObjectInfo* InObject);
-	static FString GetObjectFlags(EInternalObjectFlags InternalFlags, EObjectFlags Flags);
+	static FString GetObjectFlags(const FGCObjectInfo& InObject);
+private:
 
 	/** Dumps a reference chain to log */
 	static void DumpChain(FReferenceChainSearch::FReferenceChain* Chain, TFunctionRef<bool(FCallbackParams& Params)> ReferenceCallback,
diff --git a/Engine/Source/Runtime/CoreUObject/Public/UObject/UObjectArray.h b/Engine/Source/Runtime/CoreUObject/Public/UObject/UObjectArray.h
--- a/Engine/Source/Runtime/CoreUObject/Public/UObject/UObjectArray.h
+++ b/Engine/Source/Runtime/CoreUObject/Public/UObject/UObjectArray.h
@@ -36,10 +36,14 @@ struct
 #endif
 	FUObjectItem
 {
+	friend class FUObjectArray;
+
 	// Pointer to the allocated object
 	class UObjectBase* Object;
-	// Internal flags
+private:
+	// Internal flags. These can only be changed via Set* and Clear* functions
 	int32 Flags;
+public:
 	// UObject Owner Cluster Index
 	int32 ClusterRootIndex;	
 	// Weak Object Pointer Serial number associated with the object
@@ -124,11 +128,11 @@ struct
 	}
 
 	/**
-	 * Uses atomics to clear the specified flag(s).
+	 * Uses atomics to clear the specified flag(s). GC internal version
 	 * @param FlagsToClear
 	 * @return True if this call cleared the flag, false if it has been cleared by another thread.
 	 */
-	FORCEINLINE bool ThisThreadAtomicallyClearedFlag(EInternalObjectFlags FlagToClear)
+	FORCEINLINE bool ThisThreadAtomicallyClearedFlag_ForGC(EInternalObjectFlags FlagToClear)
 	{
 		static_assert(sizeof(int32) == sizeof(Flags), "Flags must be 32-bit for atomics.");
 		bool bIChangedIt = false;
@@ -149,6 +153,24 @@ struct
 		return bIChangedIt;
 	}
 
+	/**
+	 * Uses atomics to clear the specified flag(s).
+	 * @param FlagsToClear
+	 * @return True if this call cleared the flag, false if it has been cleared by another thread.
+	 */
+	FORCEINLINE bool ThisThreadAtomicallyClearedFlag(EInternalObjectFlags FlagToClear)
+	{
+		FlagToClear &= ~UE::GC::GReachableObjectFlag; // reachability bit can only be cleared by GC through *_ForGC functions
+		if (!!(FlagToClear & EInternalObjectFlags_RootFlags))
+		{
+			return ClearRootFlags(FlagToClear);
+		}
+		else
+		{
+			return ThisThreadAtomicallyClearedFlag_ForGC(FlagToClear);
+		}
+	}
+
 	/**
 	 * Uses atomics to set the specified flag(s). GC internal version.
 	 * @param FlagToSet
@@ -181,17 +203,15 @@ struct
 	 * @return True if this call set the flag, false if it has been set by another thread.
 	 */
 	FORCEINLINE bool ThisThreadAtomicallySetFlag(EInternalObjectFlags FlagToSet)
-	{
-		bool bMarkAsReachable = UE::GC::GIsIncrementalReachabilityPending & !!(FlagToSet & EInternalObjectFlags_RootFlags); //-V792
-		bool bIChangedIt = ThisThreadAtomicallySetFlag_ForGC(FlagToSet);
-		if (bIChangedIt & bMarkAsReachable) //-V792
+	{		
+		if (!!(FlagToSet & EInternalObjectFlags_RootFlags))
 		{
-			// Setting any of the root flags on an object during incremental reachability requires a GC barrier
-			// to make sure an object with root flags does not get Garbage Collected
-			checkf(Object, TEXT("Setting an internal object flag on a null object entry"));
-			Object->MarkAsReachable();
+			return SetRootFlags(FlagToSet);
+		}
+		else
+		{
+			return ThisThreadAtomicallySetFlag_ForGC(FlagToSet);
 		}
-		return bIChangedIt;
 	}
 
 	FORCEINLINE bool HasAnyFlags(EInternalObjectFlags InFlags) const
@@ -206,19 +226,17 @@ struct
 
 	FORCEINLINE void SetUnreachable()
 	{
+		ThisThreadAtomicallyClearedFlag_ForGC(UE::GC::GReachableObjectFlag);
 		ThisThreadAtomicallySetFlag_ForGC(UE::GC::GUnreachableObjectFlag);
 	}
 	FORCEINLINE void SetMaybeUnreachable()
 	{
+		ThisThreadAtomicallyClearedFlag_ForGC(UE::GC::GReachableObjectFlag);
 		ThisThreadAtomicallySetFlag_ForGC(UE::GC::GMaybeUnreachableObjectFlag);
 	}
 	FORCEINLINE void ClearUnreachable()
 	{
-		ThisThreadAtomicallyClearedFlag(UE::GC::GUnreachableObjectFlag);
-	}
-	FORCEINLINE void ClearMaybeUnreachable()
-	{
-		ThisThreadAtomicallyClearedFlag(UE::GC::GMaybeUnreachableObjectFlag);
+		ThisThreadAtomicallyClearedRFUnreachable();
 	}
 	FORCEINLINE bool IsUnreachable() const
 	{
@@ -230,11 +248,12 @@ struct
 	}
 	FORCEINLINE bool ThisThreadAtomicallyClearedRFUnreachable()
 	{
-		return ThisThreadAtomicallyClearedFlag(UE::GC::GUnreachableObjectFlag);
-	}
-	FORCEINLINE bool ThisThreadAtomicallyClearedMaybeUnreachable()
-	{
-		return ThisThreadAtomicallyClearedFlag(UE::GC::GMaybeUnreachableObjectFlag);
+		if (ThisThreadAtomicallyClearedFlag_ForGC(UE::GC::GUnreachableObjectFlag))
+		{
+			ThisThreadAtomicallySetFlag_ForGC(UE::GC::GReachableObjectFlag);
+			return true;
+		}
+		return false;
 	}
 	FORCEINLINE void SetGarbage()
 	{
@@ -242,7 +261,7 @@ struct
 	}
 	FORCEINLINE void ClearGarbage()
 	{
-		ThisThreadAtomicallyClearedFlag(EInternalObjectFlags::Garbage);
+		ThisThreadAtomicallyClearedFlag_ForGC(EInternalObjectFlags::Garbage);
 	}
 	FORCEINLINE bool IsGarbage() const
 	{
@@ -281,11 +300,52 @@ struct
 #if STATS || ENABLE_STATNAMEDEVENTS_UOBJECT
 	COREUOBJECT_API void CreateStatID() const;
 #endif
+
+	// Mark this object item as Reachable and clear MaybeUnreachable flag. For GC use only.
+	FORCEINLINE void FastMarkAsReachableInterlocked_ForGC()
+	{
+		using namespace UE::GC;
+		FPlatformAtomics::InterlockedAnd(&Flags, ~int32(GMaybeUnreachableObjectFlag));
+		FPlatformAtomics::InterlockedOr(&Flags, int32(GReachableObjectFlag));
+	}
+
+	// Mark this object item as Reachable and clear ReachableInCluster and MaybeUnreachable flags. For GC use only.
+	FORCEINLINE void FastMarkAsReachableAndClearReachaleInClusterInterlocked_ForGC()
+	{
+		using namespace UE::GC;
+		FPlatformAtomics::InterlockedAnd(&Flags, ~int32(GMaybeUnreachableObjectFlag | EInternalObjectFlags::ReachableInCluster));
+		FPlatformAtomics::InterlockedOr(&Flags, int32(GReachableObjectFlag));
+	}
+
+	/**
+	 * Mark this object item as Reachable and clear MaybeUnreachable flag. Only thread-safe for concurrent clear, not concurrent set+clear. Don't use during mark phase. For GC use only.
+	 * @return True if this call cleared MaybeUnreachable flag, false if it has been cleared by another thread.
+	 */
+	FORCEINLINE bool MarkAsReachableInterlocked_ForGC()
+	{
+		using namespace UE::GC;
+		const int32 FlagToClear = int32(UE::GC::GMaybeUnreachableObjectFlag);
+		if (FPlatformAtomics::AtomicRead_Relaxed(&Flags) & FlagToClear)
+		{
+			int32 Old = FPlatformAtomics::InterlockedAnd(&Flags, ~FlagToClear);
+			FPlatformAtomics::InterlockedOr(&Flags, int32(GReachableObjectFlag));
+			return Old & FlagToClear;
+		}
+		return false;
+	}
+
+	FORCEINLINE static constexpr ::size_t OffsetOfFlags()
+	{
+		return offsetof(FUObjectItem, Flags);
+	}
+
 private:
 	FORCEINLINE int32 GetFlagsInternal() const
 	{
 		return FPlatformAtomics::AtomicRead_Relaxed((int32*)&Flags);
 	}
+	COREUOBJECT_API bool SetRootFlags(EInternalObjectFlags FlagsToSet);
+	COREUOBJECT_API bool ClearRootFlags(EInternalObjectFlags FlagsToClear);
 };
 
 namespace UE::UObjectArrayPrivate
diff --git a/Engine/Source/Runtime/CoreUObject/Public/UObject/UObjectBaseUtility.h b/Engine/Source/Runtime/CoreUObject/Public/UObject/UObjectBaseUtility.h
--- a/Engine/Source/Runtime/CoreUObject/Public/UObject/UObjectBaseUtility.h
+++ b/Engine/Source/Runtime/CoreUObject/Public/UObject/UObjectBaseUtility.h
@@ -177,7 +177,7 @@ public:
 	{
 		check(!IsRooted());
 		AtomicallySetFlags(RF_MirroredGarbage);
-		GUObjectArray.IndexToObject(InternalIndex)->ThisThreadAtomicallySetFlag(EInternalObjectFlags::Garbage);
+		GUObjectArray.IndexToObject(InternalIndex)->SetGarbage();
 	}
 
 	/**
@@ -186,7 +186,7 @@ public:
 	FORCEINLINE void ClearGarbage()
 	{
 		AtomicallyClearFlags(RF_MirroredGarbage);
-		GUObjectArray.IndexToObject(InternalIndex)->ThisThreadAtomicallyClearedFlag(EInternalObjectFlags::Garbage);
+		GUObjectArray.IndexToObject(InternalIndex)->ClearGarbage();
 	}
 
 	/**
