Fixing CIS

diff --git a/Engine/Source/Runtime/CoreUObject/Private/UObject/GarbageCollection.cpp b/Engine/Source/Runtime/CoreUObject/Private/UObject/GarbageCollection.cpp
--- a/Engine/Source/Runtime/CoreUObject/Private/UObject/GarbageCollection.cpp
+++ b/Engine/Source/Runtime/CoreUObject/Private/UObject/GarbageCollection.cpp
@@ -4163,7 +4163,7 @@ public:
 					checkSlow(Object->IsValidLowLevel());					
 #if DO_GUARD_SLOW
 					// We cannot mark Root objects as Garbage.
-					checkCode(if (ObjectItem->HasAllFlags(EInternalObjectFlags::Garbage | EInternalObjectFlags::RootSet)) { UE_LOG(LogGarbage, Fatal, TEXT("Object %s is part of root set though has been marked as Garbage!"), *Object->GetFullName()); });
+					checkCode(if (RootItem->HasAllFlags(EInternalObjectFlags::Garbage | EInternalObjectFlags::RootSet)) { UE_LOG(LogGarbage, Fatal, TEXT("Object %s is part of root set though has been marked as Garbage!"), *Object->GetFullName()); });
 #endif
 
 					RootItem->FastMarkAsReachableInterlocked_ForGC();
