Fix for subobjects overrides not always be consistent

diff --git a/Engine/Source/Runtime/CoreUObject/Private/UObject/OverridableManager.cpp b/Engine/Source/Runtime/CoreUObject/Private/UObject/OverridableManager.cpp
--- a/Engine/Source/Runtime/CoreUObject/Private/UObject/OverridableManager.cpp
+++ b/Engine/Source/Runtime/CoreUObject/Private/UObject/OverridableManager.cpp
@@ -48,6 +48,11 @@ FOverriddenPropertySet* FOverridableManager::GetOverriddenProperties(UObject& Ob
 	return OverriddenObjectAnnotations.Find(Object);
 }
 
+const FOverriddenPropertySet* FOverridableManager::GetOverriddenProperties(const UObject& Object)
+{
+	return OverriddenObjectAnnotations.Find(Object);
+}
+
 FOverriddenPropertySet& FOverridableManager::SetOverriddenProperties(UObject& Object, EOverriddenPropertyOperation Operation)
 {
 	FOverriddenPropertySet& ObjectOverriddenProperties = OverriddenObjectAnnotations.FindOrAdd(Object);
@@ -56,19 +61,6 @@ FOverriddenPropertySet& FOverridableManager::SetOverriddenProperties(UObject& Ob
 	return ObjectOverriddenProperties;
 }
 
-void FOverridableManager::CopyOverriddenProperties(UObject& DestObject, const UObject& SourceObject)
-{
-	if (const FOverriddenPropertySet* SourceObjectOverriddenProperties = OverriddenObjectAnnotations.Find(SourceObject))
-	{
-		FOverriddenPropertySet& DestObjectOverriddenProperties = OverriddenObjectAnnotations.FindOrAdd(DestObject);
-		DestObjectOverriddenProperties = *SourceObjectOverriddenProperties;
-	}
-	else
-	{
-		OverriddenObjectAnnotations.RemoveAnnotation(&DestObject);
-	}
-}
-
 EOverriddenState FOverridableManager::GetOverriddenState(UObject& Object)
 {
 	if(FOverriddenPropertySet* OverriddenProperties = GetOverriddenProperties(Object))
@@ -151,7 +143,10 @@ void FOverridableManager::PostOverrideProperty(UObject& Object, const FPropertyC
 
 void FOverridableManager::NotifyPropertyChange(const EPropertyNotificationType Notification, UObject& Object, const FPropertyChangedEvent& PropertyEvent, const FEditPropertyChain::TDoubleLinkedListNode* PropertyNode)
 {
-	OverriddenObjectAnnotations.FindChecked(Object).NotifyPropertyChange(Notification, PropertyEvent, PropertyNode, &Object);
+	if (FOverriddenPropertySet* ThisObjectOverriddenProperties = OverriddenObjectAnnotations.Find(Object))
+	{
+		ThisObjectOverriddenProperties->NotifyPropertyChange(Notification, PropertyEvent, PropertyNode, &Object);
+	}
 }
 
 void FOverridableManager::ClearOverrides(UObject& Object)
diff --git a/Engine/Source/Runtime/CoreUObject/Private/UObject/OverriddenPropertySet.cpp b/Engine/Source/Runtime/CoreUObject/Private/UObject/OverriddenPropertySet.cpp
--- a/Engine/Source/Runtime/CoreUObject/Private/UObject/OverriddenPropertySet.cpp
+++ b/Engine/Source/Runtime/CoreUObject/Private/UObject/OverriddenPropertySet.cpp
@@ -145,29 +145,28 @@ FOverriddenPropertyNode& FOverriddenPropertySet::FindOrAddNode(FOverriddenProper
 	return OverriddenPropertyNodes.Get(NewID);
 }
 
-void FOverriddenPropertySet::NotifyPropertyChange(FOverriddenPropertyNode& ParentPropertyNode, const EPropertyNotificationType Notification, const FPropertyChangedEvent& PropertyEvent, const FEditPropertyChain::TDoubleLinkedListNode* PropertyNode, const void* Data)
+void FOverriddenPropertySet::NotifyPropertyChange(FOverriddenPropertyNode* ParentPropertyNode, const EPropertyNotificationType Notification, const FPropertyChangedEvent& PropertyEvent, const FEditPropertyChain::TDoubleLinkedListNode* PropertyNode, const void* Data)
 {
 	checkf(IsValid(Owner), TEXT("Expecting a valid overridable owner"));
 
 	FOverridableManager& OverridableManager = FOverridableManager::Get();
 	if (!PropertyNode)
 	{
-		if (Notification == EPropertyNotificationType::PostEdit)
+		if (ParentPropertyNode && Notification == EPropertyNotificationType::PostEdit)
 		{
 			// Replacing this entire property
-			ParentPropertyNode.Operation = EOverriddenPropertyOperation::Replace;
+			ParentPropertyNode->Operation = EOverriddenPropertyOperation::Replace;
 
 			// Sub-property overrides are not needed from now on, so clear them
-			RemoveOverriddenSubProperties(ParentPropertyNode);
+			RemoveOverriddenSubProperties(*ParentPropertyNode);
 
 			// If we are overriding the root node, need to propagate the overrides to all instanced sub object
 			const FOverriddenPropertyNode* RootNode = OverriddenPropertyNodes.FindByHash(GetTypeHash(RootNodeID), RootNodeID);
 			checkf(RootNode, TEXT("Expecting to always have a "));
-			if (RootNode == &ParentPropertyNode)
+			if (RootNode == ParentPropertyNode)
 			{
 				OverridableManager.PropagateOverrideToInstancedSubObjects(*Owner);
 			}
-
 		}
 		return;
 	}
@@ -178,11 +177,11 @@ void FOverriddenPropertySet::NotifyPropertyChange(FOverriddenPropertyNode& Paren
 	const FOverriddenPropertyNodeID PropID(Property->GetFName());
 	const void* SubValuePtr = Property->ContainerPtrToValuePtr<void>(Data, 0); //@todo support static arrays
 
-	FOverriddenPropertyNode& SubPropertyNode = FindOrAddNode(ParentPropertyNode, PropID);
-	if (SubPropertyNode.Operation == EOverriddenPropertyOperation::Replace)
+	FOverriddenPropertyNode* SubPropertyNode = nullptr;
+	if (ParentPropertyNode)
 	{
-		// No need to continue further, everything is replaced for here down anyway
-		return;
+		FOverriddenPropertyNode& SubPropertyNodeRef = FindOrAddNode(*ParentPropertyNode, PropID);
+		SubPropertyNode = SubPropertyNodeRef.Operation != EOverriddenPropertyOperation::Replace ? &SubPropertyNodeRef : nullptr;
 	}
 
 	if (const FArrayProperty* ArrayProperty = CastField<FArrayProperty>(Property))
@@ -212,8 +211,11 @@ void FOverriddenPropertySet::NotifyPropertyChange(FOverriddenPropertyNode& Paren
 
 					auto ArrayReplace = [&]
 					{
-						// Overriding all entry in the array
-						SubPropertyNode.Operation = EOverriddenPropertyOperation::Replace;
+						if (SubPropertyNode)
+						{
+							// Overriding all entry in the array
+							SubPropertyNode->Operation = EOverriddenPropertyOperation::Replace;
+						}
 
 						// This is a case of the entire array is overridden
 						// Need to loop through every sub object and setup them up as overridden
@@ -221,9 +223,12 @@ void FOverriddenPropertySet::NotifyPropertyChange(FOverriddenPropertyNode& Paren
 						{
 							if (UObject* SubObject = InnerObjectProperty->GetObjectPropertyValue(ArrayHelper.GetElementPtr(i)))
 							{
-								const FOverriddenPropertyNodeID SubObjectID(*SubObject);
-								FOverriddenPropertyNode& SubObjectNode = FindOrAddNode(SubPropertyNode, SubObjectID);
-								SubObjectNode.Operation = EOverriddenPropertyOperation::Replace;
+								if(SubPropertyNode)
+								{
+									const FOverriddenPropertyNodeID SubObjectID(*SubObject);
+									FOverriddenPropertyNode& SubObjectNode = FindOrAddNode(*SubPropertyNode, SubObjectID);
+									SubObjectNode.Operation = EOverriddenPropertyOperation::Replace;
+								}
 
 								OverridableManager.OverrideInstancedSubObject(*Owner, *SubObject);
 							}
@@ -235,19 +240,18 @@ void FOverriddenPropertySet::NotifyPropertyChange(FOverriddenPropertyNode& Paren
 						checkf(ArrayHelper.IsValidIndex(ArrayIndex), TEXT("ArrayAdd change type expected to have an valid index"));
 						if (UObject* AddedSubObject = InnerObjectProperty->GetObjectPropertyValue(ArrayHelper.GetElementPtr(ArrayIndex)))
 						{
-							const FOverriddenPropertyNodeID  AddedSubObjectID(*AddedSubObject);
-							FOverriddenPropertyNode& AddedSubObjectNode = FindOrAddNode(SubPropertyNode, AddedSubObjectID);
-							AddedSubObjectNode.Operation = EOverriddenPropertyOperation::Add;
-
-							// Check if this could be a readd
-							UObject* RemovedSubObjectArchetype = AddedSubObject->GetArchetype();
-							if(RemovedSubObjectArchetype && !RemovedSubObjectArchetype->HasAnyFlags(RF_ClassDefaultObject))
+							if(SubPropertyNode)
 							{
-								const FOverriddenPropertyNodeID RemovedSubObjectID(*RemovedSubObjectArchetype);
-								if (SubPropertyNode.SubPropertyNodeKeys.Remove(RemovedSubObjectID))
+								const FOverriddenPropertyNodeID  AddedSubObjectID(*AddedSubObject);
+								FOverriddenPropertyNode& AddedSubObjectNode = FindOrAddNode(*SubPropertyNode, AddedSubObjectID);
+								AddedSubObjectNode.Operation = EOverriddenPropertyOperation::Add;
+
+								// Check if this could be a readd
+								UObject* RemovedSubObjectArchetype = AddedSubObject->GetArchetype();
+								if(RemovedSubObjectArchetype && !RemovedSubObjectArchetype->HasAnyFlags(RF_ClassDefaultObject))
 								{
-									// This is a readd, make sure all the subobjects values are overridden
-									OverridableManager.OverrideInstancedSubObject(*Owner, *AddedSubObject);
+									const FOverriddenPropertyNodeID RemovedSubObjectID(*RemovedSubObjectArchetype);
+									SubPropertyNode->SubPropertyNodeKeys.Remove(RemovedSubObjectID);
 								}
 							}
 						}
@@ -258,19 +262,24 @@ void FOverriddenPropertySet::NotifyPropertyChange(FOverriddenPropertyNode& Paren
 						checkf(PreEditSubObjectsArrayHelper.IsValidIndex(ArrayIndex), TEXT("ArrayRemove change type expected to have an valid index"));
 						if (UObject* RemovedSubObject = InnerObjectProperty->GetObjectPropertyValue(PreEditSubObjectsArrayHelper.GetElementPtr(ArrayIndex)))
 						{
-							UObject* RemovedSubObjectArchetype = RemovedSubObject->GetArchetype();
-							const FOverriddenPropertyNodeID RemovedSubObjectID (!RemovedSubObjectArchetype || RemovedSubObjectArchetype->HasAnyFlags(RF_ClassDefaultObject) ? *RemovedSubObject : *RemovedSubObjectArchetype);
-							FOverriddenPropertyNode& RemovedSubObjectNode = FindOrAddNode(SubPropertyNode, RemovedSubObjectID);
-							OverridableManager.ClearInstancedSubObjectOverrides(*Owner, *RemovedSubObject);
-							if (RemovedSubObjectNode.Operation == EOverriddenPropertyOperation::Add)
+							if(SubPropertyNode)
 							{
-								// An add then a remove becomes no opt
-								SubPropertyNode.SubPropertyNodeKeys.Remove(RemovedSubObjectID);
-							}
-							else
-							{
-								RemovedSubObjectNode.Operation = EOverriddenPropertyOperation::Remove;
+								UObject* RemovedSubObjectArchetype = RemovedSubObject->GetArchetype();
+								const FOverriddenPropertyNodeID RemovedSubObjectID (!RemovedSubObjectArchetype || RemovedSubObjectArchetype->HasAnyFlags(RF_ClassDefaultObject) ? *RemovedSubObject : *RemovedSubObjectArchetype);
+								FOverriddenPropertyNode& RemovedSubObjectNode = FindOrAddNode(*SubPropertyNode, RemovedSubObjectID);
+
+								if (RemovedSubObjectNode.Operation == EOverriddenPropertyOperation::Add)
+								{
+									// An add then a remove becomes no opt
+									SubPropertyNode->SubPropertyNodeKeys.Remove(RemovedSubObjectID);
+								}
+								else
+								{
+									RemovedSubObjectNode.Operation = EOverriddenPropertyOperation::Remove;
+								}
 							}
+
+							OverridableManager.ClearInstancedSubObjectOverrides(*Owner, *RemovedSubObject);
 						}
 					};
 
@@ -282,10 +291,10 @@ void FOverriddenPropertySet::NotifyPropertyChange(FOverriddenPropertyNode& Paren
 							// Overriding all entry in the array + override instanced sub obejects
 							ArrayReplace();
 						}
-						else
+						else if (SubPropertyNode)
 						{
 							// Overriding all entry in the array
-							SubPropertyNode.Operation = EOverriddenPropertyOperation::Replace;
+							SubPropertyNode->Operation = EOverriddenPropertyOperation::Replace;
 						}
 						return;
 					}
@@ -419,25 +428,31 @@ void FOverriddenPropertySet::NotifyPropertyChange(FOverriddenPropertyNode& Paren
 			auto MapReplace = [&]()
 			{
 				// Overriding a all entries in the map
-				SubPropertyNode.Operation = EOverriddenPropertyOperation::Replace;
+				if (SubPropertyNode)
+				{
+					SubPropertyNode->Operation = EOverriddenPropertyOperation::Replace;
+				}
 
 				// This is a case of the entire array is overridden
 				// Need to loop through every sub object and setup them up as overridden
 				for (FScriptMapHelper::FIterator It(MapHelper); It; ++It)
 				{
-					FOverriddenPropertyNodeID OverriddenKeyID;
-					if (const UObject* OverriddenKeyObject = KeyObjectProperty ? KeyObjectProperty->GetObjectPropertyValue(MapHelper.GetKeyPtr(It.GetInternalIndex())) : nullptr)
-					{
-						OverriddenKeyID = FOverriddenPropertyNodeID(*OverriddenKeyObject);
-					}
-					else
+					if(SubPropertyNode)
 					{
-						FString AddedKey;
-						MapProperty->KeyProp->ExportTextItem_Direct(AddedKey, MapHelper.GetKeyPtr(It.GetInternalIndex()), nullptr, nullptr, PPF_None);
-						OverriddenKeyID = FName(AddedKey);
+						FOverriddenPropertyNodeID OverriddenKeyID;
+						if (const UObject* OverriddenKeyObject = KeyObjectProperty ? KeyObjectProperty->GetObjectPropertyValue(MapHelper.GetKeyPtr(It.GetInternalIndex())) : nullptr)
+						{
+							OverriddenKeyID = FOverriddenPropertyNodeID(*OverriddenKeyObject);
+						}
+						else
+						{
+							FString AddedKey;
+							MapProperty->KeyProp->ExportTextItem_Direct(AddedKey, MapHelper.GetKeyPtr(It.GetInternalIndex()), nullptr, nullptr, PPF_None);
+							OverriddenKeyID = FName(AddedKey);
+						}
+						FOverriddenPropertyNode& OverriddenKeyNode = FindOrAddNode(*SubPropertyNode, OverriddenKeyID);
+						OverriddenKeyNode.Operation = EOverriddenPropertyOperation::Replace;
 					}
-					FOverriddenPropertyNode& OverriddenKeyNode = FindOrAddNode(SubPropertyNode, OverriddenKeyID);
-					OverriddenKeyNode.Operation = EOverriddenPropertyOperation::Replace;
 
 					// @todo support instanced object as a key in maps
 					//if (UObject* KeySubObject = KeyInstancedObjectProperty ? KeyInstancedObjectProperty->GetObjectPropertyValue(MapHelper.GetKeyPtr(*It)) : nullptr)
@@ -454,62 +469,53 @@ void FOverriddenPropertySet::NotifyPropertyChange(FOverriddenPropertyNode& Paren
 			auto MapAddImpl = [&]()
 			{
 				checkf(MapHelper.IsValidIndex(InternalMapIndex), TEXT("ArrayAdd change type expected to have an valid index"));
-				FOverriddenPropertyNodeID AddedKeyID;
-				if (const UObject* AddedKeyObject = KeyObjectProperty ? KeyObjectProperty->GetObjectPropertyValue(MapHelper.GetKeyPtr(InternalMapIndex)) : nullptr)
-				{
-					AddedKeyID = FOverriddenPropertyNodeID(*AddedKeyObject);
-				}
-				else
-				{
-					FString AddedKey;
-					MapProperty->KeyProp->ExportTextItem_Direct(AddedKey, MapHelper.GetKeyPtr(InternalMapIndex), nullptr, nullptr, PPF_None);
-					AddedKeyID = FName(AddedKey);
-				}
 
-				FOverriddenPropertyNode& AddedKeyNode = FindOrAddNode(SubPropertyNode, AddedKeyID);
-				// @todo support instanced object as a key in maps
-				//if (UObject* AddedKeySubObject = KeyInstancedObjectProperty ? KeyInstancedObjectProperty->GetObjectPropertyValue(MapHelper.GetKeyPtr(InternalMapIndex)) : nullptr)
-				//{
-				//	// This is a readd, make sure all the subobjects values are overridden
-				//	if (AddedKeyNode.Operation == EOverriddenPropertyOperation::Remove)
-				//	{
-				//		OverridableManager.OverrideInstancedSubObject(*Owner, *AddedKeySubObject);
-				//	}
-				//}
-				if (UObject* AddedValueSubObject = ValueInstancedObjectProperty ? ValueInstancedObjectProperty->GetObjectPropertyValue(MapHelper.GetValuePtr(InternalMapIndex)) : nullptr)
+				if(SubPropertyNode)
 				{
-					// This is a readd, make sure all the subobjects values are overridden
-					if (AddedKeyNode.Operation == EOverriddenPropertyOperation::Remove)
+					FOverriddenPropertyNodeID AddedKeyID;
+					if (const UObject* AddedKeyObject = KeyObjectProperty ? KeyObjectProperty->GetObjectPropertyValue(MapHelper.GetKeyPtr(InternalMapIndex)) : nullptr)
 					{
-						OverridableManager.OverrideInstancedSubObject(*Owner, *AddedValueSubObject);
+						AddedKeyID = FOverriddenPropertyNodeID(*AddedKeyObject);
 					}
+					else
+					{
+						FString AddedKey;
+						MapProperty->KeyProp->ExportTextItem_Direct(AddedKey, MapHelper.GetKeyPtr(InternalMapIndex), nullptr, nullptr, PPF_None);
+						AddedKeyID = FName(AddedKey);
+					}
+
+					FOverriddenPropertyNode& AddedKeyNode = FindOrAddNode(*SubPropertyNode, AddedKeyID);
+					AddedKeyNode.Operation = EOverriddenPropertyOperation::Add;
 				}
-				AddedKeyNode.Operation = EOverriddenPropertyOperation::Add;
 			};
 
 			auto MapRemoveImpl = [&]()
 			{
 				checkf(PreEditMapHelper.IsValidIndex(InternalPreEditMapIndex), TEXT("ArrayRemove change type expected to have an valid index"));
-				FOverriddenPropertyNodeID RemovedKeyID;
-				if (const UObject* RemovedKeyObject = KeyObjectProperty ? KeyObjectProperty->GetObjectPropertyValue(PreEditMapHelper.GetKeyPtr(InternalPreEditMapIndex)) : nullptr)
-				{
-					RemovedKeyID = FOverriddenPropertyNodeID(*RemovedKeyObject);
-				}
-				else
-				{
-					FString RemovedKey;
-					MapProperty->KeyProp->ExportTextItem_Direct(RemovedKey, PreEditMapHelper.GetKeyPtr(InternalPreEditMapIndex), nullptr, nullptr, PPF_None);
-					RemovedKeyID = FName(RemovedKey);
-				}
-				FOverriddenPropertyNode& RemovedKeyNode = FindOrAddNode(SubPropertyNode, RemovedKeyID);
-				if (RemovedKeyNode.Operation == EOverriddenPropertyOperation::Add)
-				{
-					// @Todo support remove/add/remove
-					SubPropertyNode.SubPropertyNodeKeys.Remove(RemovedKeyID);
-				}
-				else
+
+				if(SubPropertyNode)
 				{
-					RemovedKeyNode.Operation = EOverriddenPropertyOperation::Remove;
+					FOverriddenPropertyNodeID RemovedKeyID;
+					if (const UObject* RemovedKeyObject = KeyObjectProperty ? KeyObjectProperty->GetObjectPropertyValue(PreEditMapHelper.GetKeyPtr(InternalPreEditMapIndex)) : nullptr)
+					{
+						RemovedKeyID = FOverriddenPropertyNodeID(*RemovedKeyObject);
+					}
+					else
+					{
+						FString RemovedKey;
+						MapProperty->KeyProp->ExportTextItem_Direct(RemovedKey, PreEditMapHelper.GetKeyPtr(InternalPreEditMapIndex), nullptr, nullptr, PPF_None);
+						RemovedKeyID = FName(RemovedKey);
+					}
+					FOverriddenPropertyNode& RemovedKeyNode = FindOrAddNode(*SubPropertyNode, RemovedKeyID);
+					if (RemovedKeyNode.Operation == EOverriddenPropertyOperation::Add)
+					{
+						// @Todo support remove/add/remove
+						SubPropertyNode->SubPropertyNodeKeys.Remove(RemovedKeyID);
+					}
+					else
+					{
+						RemovedKeyNode.Operation = EOverriddenPropertyOperation::Remove;
+					}
 				}
 
 				// @todo support instanced object as a key in maps
@@ -531,10 +537,10 @@ void FOverriddenPropertySet::NotifyPropertyChange(FOverriddenPropertyNode& Paren
 					// Overriding all entry in the array + override instanced sub obejects
 					MapReplace();
 				}
-				else
+				else if(SubPropertyNode)
 				{
 					// Overriding all entry in the array
-					SubPropertyNode.Operation = EOverriddenPropertyOperation::Replace;
+					SubPropertyNode->Operation = EOverriddenPropertyOperation::Replace;
 				}
 				return;
 			}
@@ -610,37 +616,59 @@ void FOverriddenPropertySet::NotifyPropertyChange(FOverriddenPropertyNode& Paren
 	}
 	else if (Property->IsA<FStructProperty>())
 	{
-		NotifyPropertyChange(SubPropertyNode, Notification, PropertyEvent, PropertyNode->GetNextNode(), SubValuePtr);
+		if (!PropertyNode->GetNextNode())
+		{
+			if (Notification == EPropertyNotificationType::PostEdit && SubPropertyNode)
+			{
+				SubPropertyNode->Operation = EOverriddenPropertyOperation::Replace;
+			}
+		}
+		else
+		{
+			NotifyPropertyChange(SubPropertyNode, Notification, PropertyEvent, PropertyNode->GetNextNode(), SubValuePtr);
+		}
 		return;
 	}
 	else if (const FObjectProperty* ObjectProperty = CastField<FObjectProperty>(Property))
 	{
-		if(Notification == EPropertyNotificationType::PostEdit && !PropertyNode->GetNextNode())
+		if (!PropertyNode->GetNextNode())
 		{
-			SubPropertyNode.Operation = EOverriddenPropertyOperation::Replace;
+			if (Notification == EPropertyNotificationType::PostEdit && SubPropertyNode)
+			{
+				SubPropertyNode->Operation = EOverriddenPropertyOperation::Replace;
+			}
 		}
-
-		// @TODO Need to handle non instanced object here.
-		if (UObject* SubObject = ObjectProperty->GetObjectPropertyValue(SubValuePtr))
+		else if (UObject* SubObject = ObjectProperty->GetObjectPropertyValue(SubValuePtr))
 		{
 			// This should not be needed in the property grid, as it should already been called on the subobject.
 			OverridableManager.NotifyPropertyChange(Notification, *SubObject, PropertyEvent, PropertyNode->GetNextNode());
-			return;
 		}
+		return;
 	}
 	else if (const FOptionalProperty* OptionalProperty = CastField<FOptionalProperty>(Property))
 	{
-		if (OptionalProperty->IsSet(Data))
+		if (!PropertyNode->GetNextNode())
+		{
+			if (Notification == EPropertyNotificationType::PostEdit && SubPropertyNode)
+			{
+				SubPropertyNode->Operation = EOverriddenPropertyOperation::Replace;
+			}
+		}
+		else if (OptionalProperty->IsSet(Data))
 		{
 			NotifyPropertyChange(SubPropertyNode, Notification, PropertyEvent, PropertyNode->GetNextNode(), OptionalProperty->GetValuePointerForRead(SubValuePtr));
-			return;
 		}
+		return;
 	}
 
 	UE_CLOG(PropertyNode->GetNextNode(), LogOverridableObject, Warning, TEXT("Unsupported property type(%s), fallback to overriding entire property"), *Property->GetName());
 	if (Notification == EPropertyNotificationType::PostEdit)
 	{
-		NotifyPropertyChange(SubPropertyNode, Notification, PropertyEvent, nullptr, SubValuePtr);
+		if (SubPropertyNode)
+		{
+			// Replacing this entire property
+			SubPropertyNode->Operation = EOverriddenPropertyOperation::Replace;
+		}
 	}
 }
 
@@ -658,7 +686,7 @@ void FOverriddenPropertySet::RemoveOverriddenSubProperties(FOverriddenPropertyNo
 
 void FOverriddenPropertySet::NotifyPropertyChange(const EPropertyNotificationType Notification, const FPropertyChangedEvent& PropertyEvent, const FEditPropertyChain::TDoubleLinkedListNode* PropertyNode, const void* Data)
 {
-	NotifyPropertyChange(OverriddenPropertyNodes.FindOrAddByHash(GetTypeHash(RootNodeID), RootNodeID), Notification, PropertyEvent, PropertyNode, Data);
+	NotifyPropertyChange(&OverriddenPropertyNodes.FindOrAddByHash(GetTypeHash(RootNodeID), RootNodeID), Notification, PropertyEvent, PropertyNode, Data);
 }
 
 EOverriddenPropertyOperation FOverriddenPropertySet::GetOverriddenPropertyOperation(const FArchiveSerializedPropertyChain* CurrentPropertyChain, FProperty* Property) const
diff --git a/Engine/Source/Runtime/CoreUObject/Public/UObject/OverridableManager.h b/Engine/Source/Runtime/CoreUObject/Public/UObject/OverridableManager.h
--- a/Engine/Source/Runtime/CoreUObject/Public/UObject/OverridableManager.h
+++ b/Engine/Source/Runtime/CoreUObject/Public/UObject/OverridableManager.h
@@ -111,6 +111,12 @@ public:
 	 * @return the overridden properties if the object have overridable serialization enabled */
 	COREUOBJECT_API FOverriddenPropertySet* GetOverriddenProperties(UObject& Object);
 
+	/**
+	 * Retrieve the overridden properties for the specified object
+	 * @param Object to fetch the overridden properties
+	 * @return the overridden properties if the object have overridable serialization enabled */
+	COREUOBJECT_API const FOverriddenPropertySet* GetOverriddenProperties(const UObject& Object);
+
 	/**
 	 * Set the override operation on this object and it will enable it if it wasn't already enabled
 	 * @param Object to set the override operation on
@@ -118,12 +124,6 @@ public:
 	 * @return the overridden properties of the object */
 	COREUOBJECT_API FOverriddenPropertySet& SetOverriddenProperties(UObject& Object, EOverriddenPropertyOperation Operation);
 
-	/**
-	 * Copies the overridden properties from one object to the other
-	 * @param DestObject copying the overridden properties to
-	 * @param SourceObject copying the overridden properties from */
-	COREUOBJECT_API void CopyOverriddenProperties(UObject& DestObject, const UObject& SourceObject);
-
 	/**
 	 * Retrieve the overridden state for the specified object
 	 * @param Object to fetch the overridden properties
diff --git a/Engine/Source/Runtime/CoreUObject/Public/UObject/OverriddenPropertySet.h b/Engine/Source/Runtime/CoreUObject/Public/UObject/OverriddenPropertySet.h
--- a/Engine/Source/Runtime/CoreUObject/Public/UObject/OverriddenPropertySet.h
+++ b/Engine/Source/Runtime/CoreUObject/Public/UObject/OverriddenPropertySet.h
@@ -263,7 +263,7 @@ public:
 
 protected:
 
-	void NotifyPropertyChange(FOverriddenPropertyNode& ParentPropertyNode, const EPropertyNotificationType Notification, const FPropertyChangedEvent& PropertyEvent, const FEditPropertyChain::TDoubleLinkedListNode* PropertyNode, const void* Data);
+	void NotifyPropertyChange(FOverriddenPropertyNode* ParentPropertyNode, const EPropertyNotificationType Notification, const FPropertyChangedEvent& PropertyEvent, const FEditPropertyChain::TDoubleLinkedListNode* PropertyNode, const void* Data);
 	EOverriddenPropertyOperation GetOverriddenPropertyOperation(const FOverriddenPropertyNode& ParentPropertyNode, const FArchiveSerializedPropertyChain* CurrentPropertyChain, FProperty* Property) const;
 	FOverriddenPropertyNode* SetOverriddenPropertyOperation(EOverriddenPropertyOperation Operation, FOverriddenPropertyNode& ParentPropertyNode, const FArchiveSerializedPropertyChain* CurrentPropertyChain, FProperty* Property);
 
