Fixing automated tests. Removing rigorous check agains GisGarbageCollecting which may get triggered during unhashing and purge phase of GC.

diff --git a/Engine/Source/Runtime/CoreUObject/Private/UObject/GarbageCollection.cpp b/Engine/Source/Runtime/CoreUObject/Private/UObject/GarbageCollection.cpp
--- a/Engine/Source/Runtime/CoreUObject/Private/UObject/GarbageCollection.cpp
+++ b/Engine/Source/Runtime/CoreUObject/Private/UObject/GarbageCollection.cpp
@@ -578,7 +578,7 @@ bool FUObjectItem::SetRootFlags(EInternalObjectFlags FlagsToSet)
 {
 	using namespace UE::GC;
 	using namespace UE::GC::Private;
-	check(!GIsGarbageCollecting);
+
 	constexpr int32 RootFlags = (int32)EInternalObjectFlags_RootFlags;
 	bool bIChangedIt = false;
 	{
