Guarding roots list against potential race conditions when accessing it when marking root objects as reachable in the initial phase of GC

diff --git a/Engine/Source/Runtime/CoreUObject/Private/UObject/GarbageCollection.cpp b/Engine/Source/Runtime/CoreUObject/Private/UObject/GarbageCollection.cpp
--- a/Engine/Source/Runtime/CoreUObject/Private/UObject/GarbageCollection.cpp
+++ b/Engine/Source/Runtime/CoreUObject/Private/UObject/GarbageCollection.cpp
@@ -4145,7 +4145,9 @@ public:
 		FMarkRootsState MarkRootsState;		
 
 		{
+			GRootsCritical.Lock();
 			TArray<int32> RootsArray(GRoots.Array());				
+			GRootsCritical.Unlock();
 			MarkRootsState.Start(Options, RootsArray.Num());
 			FMarkRootsState::FThreadIterators& ThreadIterators = MarkRootsState.GetThreadIterators();
 
