Query object property override state from FEditPropertyChain

diff --git a/Engine/Source/Runtime/CoreUObject/Private/UObject/OverridableManager.cpp b/Engine/Source/Runtime/CoreUObject/Private/UObject/OverridableManager.cpp
--- a/Engine/Source/Runtime/CoreUObject/Private/UObject/OverridableManager.cpp
+++ b/Engine/Source/Runtime/CoreUObject/Private/UObject/OverridableManager.cpp
@@ -149,6 +149,24 @@ void FOverridableManager::NotifyPropertyChange(const EPropertyNotificationType N
 	}
 }
 
+EOverriddenPropertyOperation FOverridableManager::GetOverriddenPropertyOperation(UObject& Object, const FEditPropertyChain& PropertyChain)
+{
+	if (FOverriddenPropertySet* ThisObjectOverriddenProperties = OverriddenObjectAnnotations.Find(Object))
+	{
+		return ThisObjectOverriddenProperties->GetOverriddenPropertyOperation(PropertyChain.GetHead());
+	}
+	return EOverriddenPropertyOperation::None;
+}
+
+bool FOverridableManager::SetOverriddenPropertyOperation(UObject& Object, const FEditPropertyChain& PropertyChain, EOverriddenPropertyOperation Operation)
+{
+	if (FOverriddenPropertySet* ThisObjectOverriddenProperties = OverriddenObjectAnnotations.Find(Object))
+	{
+		return ThisObjectOverriddenProperties->SetOverriddenPropertyOperation(Operation, PropertyChain.GetHead()) != nullptr;
+	}
+	return false;
+}
+
 void FOverridableManager::ClearOverrides(UObject& Object)
 {
 	if(FOverriddenPropertySet* ThisObjectOverriddenProperties = OverriddenObjectAnnotations.Find(Object))
diff --git a/Engine/Source/Runtime/CoreUObject/Private/UObject/OverriddenPropertySet.cpp b/Engine/Source/Runtime/CoreUObject/Private/UObject/OverriddenPropertySet.cpp
--- a/Engine/Source/Runtime/CoreUObject/Private/UObject/OverriddenPropertySet.cpp
+++ b/Engine/Source/Runtime/CoreUObject/Private/UObject/OverriddenPropertySet.cpp
@@ -145,6 +145,68 @@ FOverriddenPropertyNode& FOverriddenPropertySet::FindOrAddNode(FOverriddenProper
 	return OverriddenPropertyNodes.Get(NewID);
 }
 
+EOverriddenPropertyOperation FOverriddenPropertySet::GetOverriddenPropertyOperation(const FOverriddenPropertyNode& ParentPropertyNode, const FEditPropertyChain::TDoubleLinkedListNode* PropertyNode) const
+{
+	// No need to look further
+	// if it is the entire property is replaced or
+	// if it is the FOverriddenPropertySet struct which is always Overridden
+	if (ParentPropertyNode.Operation == EOverriddenPropertyOperation::Replace)
+	{
+		return EOverriddenPropertyOperation::Replace;
+	}
+
+	const FEditPropertyChain::TDoubleLinkedListNode* PropertyIterator = PropertyNode;
+	const FOverriddenPropertyNode* OverriddenPropertyNode = &ParentPropertyNode;
+	while (PropertyIterator && OverriddenPropertyNode && OverriddenPropertyNode->Operation != EOverriddenPropertyOperation::Replace)
+	{
+		const FProperty* CurrentProperty = PropertyIterator->GetValue();
+		const FName CurrentPropID = CurrentProperty->GetFName();
+		if (const FOverriddenPropertyNodeID* CurrentPropKey = OverriddenPropertyNode->SubPropertyNodeKeys.Find(CurrentPropID))
+		{
+			OverriddenPropertyNode = OverriddenPropertyNodes.FindByHash(GetTypeHash(*CurrentPropKey), *CurrentPropKey);
+			checkf(OverriddenPropertyNode, TEXT("Expecting a node"));
+		}
+		else
+		{
+			OverriddenPropertyNode = nullptr;
+			break;
+		}
+		PropertyIterator = PropertyIterator->GetNextNode();
+	}
+
+	return OverriddenPropertyNode ? OverriddenPropertyNode->Operation : EOverriddenPropertyOperation::None;
+}
+
+FOverriddenPropertyNode* FOverriddenPropertySet::SetOverriddenPropertyOperation(EOverriddenPropertyOperation Operation, FOverriddenPropertyNode& ParentPropertyNode, const FEditPropertyChain::TDoubleLinkedListNode* PropertyNode)
+{
+	// No need to look further
+	// if it is the entire property is replaced or
+	// if it is the FOverriddenPropertySet struct which is always Overridden
+	if (ParentPropertyNode.Operation == EOverriddenPropertyOperation::Replace)
+	{
+		return nullptr;
+	}
+
+	const FEditPropertyChain::TDoubleLinkedListNode* PropertyIterator = PropertyNode;
+	FOverriddenPropertyNode* OverriddenPropertyNode = &ParentPropertyNode;
+	while (PropertyIterator && OverriddenPropertyNode->Operation != EOverriddenPropertyOperation::Replace)
+	{
+		const FProperty* CurrentProperty = PropertyIterator->GetValue();
+		const FName CurrentPropID = CurrentProperty->GetFName();
+		OverriddenPropertyNode = &FindOrAddNode(*OverriddenPropertyNode, CurrentPropID);
+		PropertyIterator = PropertyIterator->GetNextNode();
+	}
+
+	// Might have stop before as one of the parent property was completely replaced.
+	if (!PropertyIterator)
+	{
+		OverriddenPropertyNode->Operation = Operation;
+		return OverriddenPropertyNode;
+	}
+
+	return nullptr;
+}
+
 void FOverriddenPropertySet::NotifyPropertyChange(FOverriddenPropertyNode* ParentPropertyNode, const EPropertyNotificationType Notification, const FPropertyChangedEvent& PropertyEvent, const FEditPropertyChain::TDoubleLinkedListNode* PropertyNode, const void* Data)
 {
 	checkf(IsValid(Owner), TEXT("Expecting a valid overridable owner"));
@@ -684,6 +746,20 @@ void FOverriddenPropertySet::RemoveOverriddenSubProperties(FOverriddenPropertyNo
 	PropertyNode.SubPropertyNodeKeys.Empty();
 }
 
+EOverriddenPropertyOperation FOverriddenPropertySet::GetOverriddenPropertyOperation(const FEditPropertyChain::TDoubleLinkedListNode* PropertyNode) const
+{
+	if (const FOverriddenPropertyNode* RootNode = OverriddenPropertyNodes.FindByHash(GetTypeHash(RootNodeID), RootNodeID))
+	{
+		return GetOverriddenPropertyOperation(*RootNode, PropertyNode);
+	}
+	return EOverriddenPropertyOperation::None;
+}
+
+FOverriddenPropertyNode* FOverriddenPropertySet::SetOverriddenPropertyOperation(EOverriddenPropertyOperation Operation, const FEditPropertyChain::TDoubleLinkedListNode* PropertyNode)
+{
+	return SetOverriddenPropertyOperation(Operation, OverriddenPropertyNodes.FindOrAddByHash(GetTypeHash(RootNodeID), RootNodeID), PropertyNode);
+}
+
 void FOverriddenPropertySet::NotifyPropertyChange(const EPropertyNotificationType Notification, const FPropertyChangedEvent& PropertyEvent, const FEditPropertyChain::TDoubleLinkedListNode* PropertyNode, const void* Data)
 {
 	NotifyPropertyChange(&OverriddenPropertyNodes.FindOrAddByHash(GetTypeHash(RootNodeID), RootNodeID), Notification, PropertyEvent, PropertyNode, Data);
diff --git a/Engine/Source/Runtime/CoreUObject/Public/UObject/OverridableManager.h b/Engine/Source/Runtime/CoreUObject/Public/UObject/OverridableManager.h
--- a/Engine/Source/Runtime/CoreUObject/Public/UObject/OverridableManager.h
+++ b/Engine/Source/Runtime/CoreUObject/Public/UObject/OverridableManager.h
@@ -195,6 +195,23 @@ public:
 	 */
 	COREUOBJECT_API void NotifyPropertyChange(const EPropertyNotificationType Notification, UObject& Object, const FPropertyChangedEvent& PropertyEvent, const FEditPropertyChain::TDoubleLinkedListNode* PropertyNode);
 
+	/**
+	 * Retrieve the overridable operation from the specified the edit property chain
+	 * @param Object owning the property
+	 * @param PropertyChain leading to the property the caller is interested in
+	 * @return the current type of override operation on the property
+	 */
+	COREUOBJECT_API EOverriddenPropertyOperation GetOverriddenPropertyOperation(UObject& Object, const FEditPropertyChain& PropertyChain);
+
+	/**
+	 * Set the overridable operation from the specified the edit property chain
+	 * @param Object owning the property
+	 * @param PropertyChain leading to the property the caller is interested in
+	 * @param Operation the override operation to set on the property
+	 * @return if the operation was successful
+	 */
+	COREUOBJECT_API bool SetOverriddenPropertyOperation(UObject& Object, const FEditPropertyChain& PropertyChain, EOverriddenPropertyOperation Operation);
+
 	/**
 	 * Serializes the overriden properties of the specified object into the record
 	 * @param Object to serialize the overriden property
diff --git a/Engine/Source/Runtime/CoreUObject/Public/UObject/OverriddenPropertySet.h b/Engine/Source/Runtime/CoreUObject/Public/UObject/OverriddenPropertySet.h
--- a/Engine/Source/Runtime/CoreUObject/Public/UObject/OverriddenPropertySet.h
+++ b/Engine/Source/Runtime/CoreUObject/Public/UObject/OverriddenPropertySet.h
@@ -213,6 +213,19 @@ public:
 		: Owner(&InOwner)
 	{}
 
+	/**
+	 * Retrieve the overridable operation from the specified the edit property chain and the specified property
+	 * @param PropertyNode in the linked list leading to the property the caller is interested in
+	 * @return the current type of override operation on the property */
+	EOverriddenPropertyOperation GetOverriddenPropertyOperation(const FEditPropertyChain::TDoubleLinkedListNode* PropertyNode) const;
+
+	/**
+	 * Setup the overridable operation of the current property from the serialized property chain and the specified property
+	 * @param Operation to set for this property
+	 * @param PropertyNode in the linked list leading to the property the caller is interested in
+	 * @return the node containing the information of the overridden property */
+	FOverriddenPropertyNode* SetOverriddenPropertyOperation(EOverriddenPropertyOperation Operation, const FEditPropertyChain::TDoubleLinkedListNode* PropertyNode);
+
 	/**
 	 * Handling and storing modification on a property of an object
 	 * @param Notification type either pre/post property overridden
@@ -230,7 +243,7 @@ public:
 
 	/**
 	 * Setup the overridable operation of the current property from the serialized property chain and the specified property
-	 * @param Operation to set on the proto set for this property
+	 * @param Operation to set for this property
 	 * @param CurrentPropertyChain leading to the property being serialized if any
 	 * @param Property being serialized if any, otherwise will fallback on the last property of the chain
 	 * @return the node containing the information of the overridden property */
@@ -263,15 +276,18 @@ public:
 
 protected:
 
+	FOverriddenPropertyNode& FindOrAddNode(FOverriddenPropertyNode& ParentPropertyNode, FOverriddenPropertyNodeID NodeID);
+
+	EOverriddenPropertyOperation GetOverriddenPropertyOperation(const FOverriddenPropertyNode& ParentPropertyNode, const FEditPropertyChain::TDoubleLinkedListNode* PropertyNode) const;
+	FOverriddenPropertyNode* SetOverriddenPropertyOperation(EOverriddenPropertyOperation Operation, FOverriddenPropertyNode& ParentPropertyNode, const FEditPropertyChain::TDoubleLinkedListNode* PropertyNode);
 	void NotifyPropertyChange(FOverriddenPropertyNode* ParentPropertyNode, const EPropertyNotificationType Notification, const FPropertyChangedEvent& PropertyEvent, const FEditPropertyChain::TDoubleLinkedListNode* PropertyNode, const void* Data);
+
 	EOverriddenPropertyOperation GetOverriddenPropertyOperation(const FOverriddenPropertyNode& ParentPropertyNode, const FArchiveSerializedPropertyChain* CurrentPropertyChain, FProperty* Property) const;
 	FOverriddenPropertyNode* SetOverriddenPropertyOperation(EOverriddenPropertyOperation Operation, FOverriddenPropertyNode& ParentPropertyNode, const FArchiveSerializedPropertyChain* CurrentPropertyChain, FProperty* Property);
-
 	const FOverriddenPropertyNode* GetOverriddenPropertyNode(const FOverriddenPropertyNode& ParentPropertyNode, const FArchiveSerializedPropertyChain* CurrentPropertyChain) const;
 
 	void RemoveOverriddenSubProperties(FOverriddenPropertyNode& PropertyNode);
 
-	FOverriddenPropertyNode& FindOrAddNode(FOverriddenPropertyNode& ParentPropertyNode, FOverriddenPropertyNodeID NodeID);
 
 private:
 	UPROPERTY()
