Updated FPropertyPathName to use FPropertyTypeName

#rnx

diff --git a/Engine/Source/Runtime/CoreUObject/Internal/UObject/PropertyBag.h b/Engine/Source/Runtime/CoreUObject/Internal/UObject/PropertyBag.h
--- a/Engine/Source/Runtime/CoreUObject/Internal/UObject/PropertyBag.h
+++ b/Engine/Source/Runtime/CoreUObject/Internal/UObject/PropertyBag.h
@@ -6,6 +6,7 @@
 #include "Containers/Map.h"
 #include "Misc/TVariant.h"
 #include "UObject/PropertyPathName.h"
+#include "UObject/PropertyTypeName.h"
 #include "UObject/PropertyTag.h"
 
 #define UE_API COREUOBJECT_API
@@ -43,7 +44,7 @@ class FPropertyBag
 
 	struct FNode
 	{
-		FName Type;
+		FPropertyTypeName Type;
 		FValue Value;
 		FNodeMap Nodes;
 	};
diff --git a/Engine/Source/Runtime/CoreUObject/Private/UObject/Class.cpp b/Engine/Source/Runtime/CoreUObject/Private/UObject/Class.cpp
--- a/Engine/Source/Runtime/CoreUObject/Private/UObject/Class.cpp
+++ b/Engine/Source/Runtime/CoreUObject/Private/UObject/Class.cpp
@@ -1576,26 +1576,35 @@ void UStruct::SerializeVersionedTaggedProperties(FStructuredArchive::FSlot Slot,
 
 				if (SerializeContext && SerializeContext->bTrackSerializedPropertyPath)
 				{
-					const FName Name = Property ? Property->GetFName() : Tag.Name;
-					const int32 Index = Tag.ArrayIndex > 0 || (Property && Property->ArrayDim > 1) ? Tag.ArrayIndex : INDEX_NONE;
-					SerializeContext->SerializedPropertyPath.Push({Name, Tag.Type, Index});
+					FPropertyTypeNameBuilder TypeBuilder;
+					TypeBuilder.AddTypeName(Tag.Type);
 
 					if (!Tag.StructName.IsNone())
 					{
-						SerializeContext->SerializedPropertyPath.PushType(Tag.StructName);
+						TypeBuilder.BeginTypeParameters();
+						TypeBuilder.AddTypeName(Tag.StructName);
+						TypeBuilder.EndTypeParameters();
 					}
 					else if (!Tag.EnumName.IsNone())
 					{
-						SerializeContext->SerializedPropertyPath.PushType(Tag.EnumName);
+						TypeBuilder.BeginTypeParameters();
+						TypeBuilder.AddTypeName(Tag.EnumName);
+						TypeBuilder.EndTypeParameters();
 					}
 					else if (!Tag.InnerType.IsNone())
 					{
-						SerializeContext->SerializedPropertyPath.PushType(Tag.InnerType);
+						TypeBuilder.BeginTypeParameters();
+						TypeBuilder.AddTypeName(Tag.InnerType);
 						if (!Tag.ValueType.IsNone())
 						{
-							SerializeContext->SerializedPropertyPath.PushType(Tag.ValueType);
+							TypeBuilder.AddTypeName(Tag.ValueType);
 						}
+						TypeBuilder.EndTypeParameters();
 					}
+
+					const FName Name = Property ? Property->GetFName() : Tag.Name;
+					const int32 Index = Tag.ArrayIndex > 0 || (Property && Property->ArrayDim > 1) ? Tag.ArrayIndex : INDEX_NONE;
+					SerializeContext->SerializedPropertyPath.Push({Name, TypeBuilder.Build(), Index});
 				}
 
 				if (Property)
diff --git a/Engine/Source/Runtime/CoreUObject/Private/UObject/InstanceDataObjectUtils.cpp b/Engine/Source/Runtime/CoreUObject/Private/UObject/InstanceDataObjectUtils.cpp
--- a/Engine/Source/Runtime/CoreUObject/Private/UObject/InstanceDataObjectUtils.cpp
+++ b/Engine/Source/Runtime/CoreUObject/Private/UObject/InstanceDataObjectUtils.cpp
@@ -24,81 +24,92 @@ namespace UE
 		bool bAddIfNeeded = false;
 	};
 
-	static void BuildSegmentTypeFromProperty(const FProperty* Property, TStringBuilder<256>& OutType)
+	static void BuildSegmentTypeFromProperty(const FProperty* Property, FPropertyTypeNameBuilder& OutType)
 	{
-		auto Push = [&OutType](FName TypeName)
-		{
-			if (OutType.Len() > 0)
-			{
-				OutType.AppendChar(' ');
-			}
-			FName::GetEntry(TypeName.GetDisplayIndex())->AppendNameToString(OutType);
-		};
-		
 #if false // TODO: @jordan.hoffmann when complete type info is finished use this branch instead
-		OutType.Add(Property->GetID());
+		OutType.AddTypeName(Property->GetID());
 		if (const FObjectProperty* AsObjectProperty = CastField<FObjectProperty>(Property))
 		{
-			Push(AsObjectProperty->GetID());
+			OutType.BeginTypeParameters();
+			OutType.AddTypeName(AsObjectProperty->GetID());
+			OutType.EndTypeParameters();
 		}
 		else if (const FEnumProperty* AsEnumProperty = CastField<FEnumProperty>(Property))
 		{
-			Push(AsEnumProperty->GetEnum()->GetFName());
+			OutType.BeginTypeParameters();
+			OutType.AddTypeName(AsEnumProperty->GetEnum()->GetFName());
+			OutType.EndTypeParameters();
 		}
 		else if (const FArrayProperty* AsArrayProperty = CastField<FArrayProperty>(Property))
 		{
+			OutType.BeginTypeParameters();
 			BuildSegmentTypeFromProperty(AsArrayProperty->Inner, OutType);
+			OutType.EndTypeParameters();
 		}
 		else if (const FSetProperty* AsSetProperty = CastField<FSetProperty>(Property))
 		{
+			OutType.BeginTypeParameters();
 			BuildSegmentTypeFromProperty(AsSetProperty->ElementProp, OutType);
+			OutType.EndTypeParameters();
 		}
 		else if (const FMapProperty* AsMapProperty = CastField<FMapProperty>(Property))
 		{
+			OutType.BeginTypeParameters();
 			BuildSegmentTypeFromProperty(AsMapProperty->KeyProp, OutType);
 			BuildSegmentTypeFromProperty(AsMapProperty->ValueProp, OutType);
+			OutType.EndTypeParameters();
 		}
 #else
-		Push(Property->GetID());
+		OutType.AddTypeName(Property->GetID());
 		if (const FStructProperty* AsStructProperty = CastField<FStructProperty>(Property))
 		{
-			Push(AsStructProperty->Struct->GetFName());
+			OutType.BeginTypeParameters();
+			OutType.AddTypeName(AsStructProperty->Struct->GetFName());
+			OutType.EndTypeParameters();
 		}
 		else if (const FEnumProperty* AsEnumProperty = CastField<FEnumProperty>(Property))
 		{
-			Push(AsEnumProperty->GetEnum()->GetFName());
+			OutType.BeginTypeParameters();
+			OutType.AddTypeName(AsEnumProperty->GetEnum()->GetFName());
+			OutType.EndTypeParameters();
 		}
 		else if (const FArrayProperty* AsArrayProperty = CastField<FArrayProperty>(Property))
 		{
+			OutType.BeginTypeParameters();
 			if (AsArrayProperty->Inner->IsA<FEnumProperty>())
 			{
 				// enum paths currently don't recurse when they're in containers
-				Push(AsArrayProperty->Inner->GetClass()->GetFName());
+				OutType.AddTypeName(AsArrayProperty->Inner->GetClass()->GetFName());
 			}
 			else
 			{
 				BuildSegmentTypeFromProperty(AsArrayProperty->Inner, OutType);
 			}
+			OutType.EndTypeParameters();
 		}
 		else if (const FSetProperty* AsSetProperty = CastField<FSetProperty>(Property))
 		{
 			// sets currently don't recurse their element property types
-			Push(AsSetProperty->ElementProp->GetClass()->GetFName());
+			OutType.BeginTypeParameters();
+			OutType.AddTypeName(AsSetProperty->ElementProp->GetClass()->GetFName());
+			OutType.EndTypeParameters();
 		}
 		else if (const FMapProperty* AsMapProperty = CastField<FMapProperty>(Property))
 		{
 			// maps currently don't recurse their key/value property types
-			Push(AsMapProperty->KeyProp->GetClass()->GetFName());
-			Push(AsMapProperty->ValueProp->GetClass()->GetFName());
+			OutType.BeginTypeParameters();
+			OutType.AddTypeName(AsMapProperty->KeyProp->GetClass()->GetFName());
+			OutType.AddTypeName(AsMapProperty->ValueProp->GetClass()->GetFName());
+			OutType.EndTypeParameters();
 		}
 #endif
 	}
 	
-	static FName GetSegmentTypeFromProperty(const FProperty* Property)
+	static FPropertyTypeName GetSegmentTypeFromProperty(const FProperty* Property)
 	{
-		TStringBuilder<256> Type;
-		BuildSegmentTypeFromProperty(Property, Type);
-		return FName(Type);
+		FPropertyTypeNameBuilder TypeBuilder;
+		BuildSegmentTypeFromProperty(Property, TypeBuilder);
+		return TypeBuilder.Build();
 	}
 	
 	static bool ResolvePropertyPathNameHelper(const UStruct* Struct, ResolvePropertyPathNameHelperParams& Params);
@@ -186,7 +197,7 @@ namespace UE
 		for (FProperty* Property : TFieldRange<FProperty>(Struct))
 		{
 			// find a property that matches the segment
-			const FName PropertySegmentType = GetSegmentTypeFromProperty(Property);
+			const FPropertyTypeName PropertySegmentType = GetSegmentTypeFromProperty(Property);
 			if (Segment.Name != Property->GetFName() || Segment.Type != PropertySegmentType)
 			{
 				continue;
diff --git a/Engine/Source/Runtime/CoreUObject/Private/UObject/PropertyArray.cpp b/Engine/Source/Runtime/CoreUObject/Private/UObject/PropertyArray.cpp
--- a/Engine/Source/Runtime/CoreUObject/Private/UObject/PropertyArray.cpp
+++ b/Engine/Source/Runtime/CoreUObject/Private/UObject/PropertyArray.cpp
@@ -619,7 +619,26 @@ void FArrayProperty::SerializeItem(FStructuredArchive::FSlot Slot, void* Value,
 	FUObjectSerializeContext* Context = FUObjectThreadContext::Get().GetSerializeContext();
 	if (Context && Context->bTrackSerializedPropertyPath && MaybeInnerTag)
 	{
-		Context->SerializedPropertyPath.PushType(MaybeInnerTag->StructName);
+		if (const int32 SegmentCount = Context->SerializedPropertyPath.GetSegmentCount(); SegmentCount > 0)
+		{
+			const UE::FPropertyTypeName ExistingArrayType = Context->SerializedPropertyPath.GetSegment(SegmentCount - 1).Type;
+			if (ExistingArrayType.GetTypeName() == NAME_ArrayProperty && ExistingArrayType.GetTypeParameterCount() == 1)
+			{
+				const UE::FPropertyTypeName ExistingStructType = ExistingArrayType.GetTypeParameter(0);
+				if (ExistingStructType.GetTypeName() == NAME_StructProperty && ExistingStructType.GetTypeParameterCount() == 0)
+				{
+					UE::FPropertyTypeNameBuilder NewTypeBuilder;
+					NewTypeBuilder.AddTypeName(NAME_ArrayProperty);
+					NewTypeBuilder.BeginTypeParameters();
+					NewTypeBuilder.AddTypeName(NAME_StructProperty);
+					NewTypeBuilder.BeginTypeParameters();
+					NewTypeBuilder.AddTypeName(MaybeInnerTag->StructName);
+					NewTypeBuilder.EndTypeParameters();
+					NewTypeBuilder.EndTypeParameters();
+					Context->SerializedPropertyPath.SetType(NewTypeBuilder.Build());
+				}
+			}
+		}
 	}
 
 	// need to know how much data this call to SerializeItem consumes, so mark where we are
diff --git a/Engine/Source/Runtime/CoreUObject/Private/UObject/PropertyPathName.cpp b/Engine/Source/Runtime/CoreUObject/Private/UObject/PropertyPathName.cpp
--- a/Engine/Source/Runtime/CoreUObject/Private/UObject/PropertyPathName.cpp
+++ b/Engine/Source/Runtime/CoreUObject/Private/UObject/PropertyPathName.cpp
@@ -24,7 +24,7 @@ inline int32 FPropertyPathName::FSegment::Compare(const FSegment& Segment) const
 	{
 		return CompareNameWithIndex;
 	}
-	return Type.CompareLexical(Segment.Type);
+	return Type < Segment.Type;
 }
 
 bool FPropertyPathName::operator==(const FPropertyPathName& Path) const
@@ -63,52 +63,6 @@ bool FPropertyPathName::operator<(const FPropertyPathName& Path) const
 	return SegmentCountA < SegmentCountB;
 }
 
-void FPropertyPathName::PushTypeInternal(FName Type)
-{
-	FNameEntryId& LastType = Segments.Last().Type;
-
-	TStringBuilder<256> CombinedTypes;
-	FName::GetEntry(LastType)->AppendNameToString(CombinedTypes);
-	CombinedTypes.AppendChar(' ');
-	Type.AppendString(CombinedTypes);
-
-	LastType = FName(CombinedTypes).GetDisplayIndex();
-}
-
-FName FPropertyPathName::PopType()
-{
-	if (Segments.IsEmpty())
-	{
-		return FName();
-	}
-
-	FNameEntryId& LastType = Segments.Last().Type;
-
-	if (LastType.IsNone())
-	{
-		return FName();
-	}
-
-	FName PoppedType;
-	TStringBuilder<256> CombinedTypes;
-	FName::GetEntry(LastType)->AppendNameToString(CombinedTypes);
-
-	const int32 Index = String::FindLastChar(CombinedTypes, TEXT(' '));
-	if (Index == INDEX_NONE)
-	{
-		PoppedType = FName::CreateFromDisplayId(LastType, NAME_NO_NUMBER_INTERNAL);
-		LastType = FNameEntryId();
-	}
-	else
-	{
-		const FStringView CombinedTypesView(CombinedTypes);
-		PoppedType = FName(CombinedTypesView.RightChop(Index + 1));
-		LastType = FName(CombinedTypesView.Left(Index)).GetDisplayIndex();
-	}
-
-	return PoppedType;
-}
-
 void FPropertyPathName::ToString(FStringBuilderBase& Out, FStringView Separator) const
 {
 	bool bFirst = true;
@@ -131,7 +85,7 @@ void FPropertyPathName::ToString(FStringBuilderBase& Out, FStringView Separator)
 			Out << TEXT('[') << Index << TEXT(']');
 		}
 
-		if (!UnpackedSegment.Type.IsNone())
+		if (!UnpackedSegment.Type.IsEmpty())
 		{
 			Out << TEXTVIEW(" (") << UnpackedSegment.Type << TEXT(')');
 		}
diff --git a/Engine/Source/Runtime/CoreUObject/Private/UObject/PropertyTypeName.cpp b/Engine/Source/Runtime/CoreUObject/Private/UObject/PropertyTypeName.cpp
--- a/Engine/Source/Runtime/CoreUObject/Private/UObject/PropertyTypeName.cpp
+++ b/Engine/Source/Runtime/CoreUObject/Private/UObject/PropertyTypeName.cpp
@@ -32,11 +32,6 @@ inline uint32 GetTypeHash(const FPropertyTypeNameNode& Node)
 	return HashCombineFast(GetTypeHashHelper(Node.Name), GetTypeHashHelper(Node.InnerCount));
 }
 
-inline bool operator==(const FPropertyTypeNameNode& Lhs, const FPropertyTypeNameNode& Rhs)
-{
-	return Lhs.Name == Rhs.Name && Lhs.InnerCount == Rhs.InnerCount;
-}
-
 inline FArchive& operator<<(FArchive& Ar, FPropertyTypeNameNode& Node)
 {
 	return Ar << Node.Name << Node.InnerCount;
@@ -94,7 +89,7 @@ struct FPropertyTypeNameNodeProxy
 		const FPropertyTypeNameNode* RhsNode = Rhs.First;
 		for (int32 Remaining = 1; Remaining > 0; --Remaining, ++LhsNode, ++RhsNode)
 		{
-			if (*LhsNode == *RhsNode)
+			if (LhsNode->Name == RhsNode->Name && LhsNode->InnerCount == RhsNode->InnerCount)
 			{
 				Remaining += LhsNode->InnerCount;
 				continue;
@@ -103,6 +98,25 @@ struct FPropertyTypeNameNodeProxy
 		}
 		return true;
 	}
+
+	friend inline bool operator<(const FPropertyTypeNameNodeProxy& Lhs, const FPropertyTypeNameNodeProxy& Rhs)
+	{
+		const FPropertyTypeNameNode* LhsNode = Lhs.First;
+		const FPropertyTypeNameNode* RhsNode = Rhs.First;
+		for (int32 Remaining = 1; Remaining > 0; --Remaining, ++LhsNode, ++RhsNode)
+		{
+			if (const int32 Compare = LhsNode->Name.Compare(RhsNode->Name))
+			{
+				return Compare < 0;
+			}
+			if (const int32 Compare = LhsNode->InnerCount - RhsNode->InnerCount)
+			{
+				return Compare < 0;
+			}
+			Remaining += LhsNode->InnerCount;
+		}
+		return true;
+	}
 };
 
 ///////////////////////////////////////////////////////////////////////////////////////////////////////////////////////
@@ -257,6 +271,18 @@ bool operator==(const FPropertyTypeName& Lhs, const FPropertyTypeName& Rhs)
 	return FPropertyTypeNameNodeProxy{LhsNode} == FPropertyTypeNameNodeProxy{RhsNode};
 }
 
+bool operator<(const FPropertyTypeName& Lhs, const FPropertyTypeName& Rhs)
+{
+	if (Lhs.Index == Rhs.Index)
+	{
+		return false;
+	}
+
+	const FPropertyTypeNameNode* LhsNode = GPropertyTypeNameTable.ResolveByIndex(Lhs.Index);
+	const FPropertyTypeNameNode* RhsNode = GPropertyTypeNameTable.ResolveByIndex(Rhs.Index);
+	return FPropertyTypeNameNodeProxy{LhsNode} < FPropertyTypeNameNodeProxy{RhsNode};
+}
+
 FArchive& operator<<(FArchive& Ar, FPropertyTypeName& TypeName)
 {
 	if (!Ar.IsPersistent())
diff --git a/Engine/Source/Runtime/CoreUObject/Public/UObject/PropertyPathName.h b/Engine/Source/Runtime/CoreUObject/Public/UObject/PropertyPathName.h
--- a/Engine/Source/Runtime/CoreUObject/Public/UObject/PropertyPathName.h
+++ b/Engine/Source/Runtime/CoreUObject/Public/UObject/PropertyPathName.h
@@ -6,6 +6,7 @@
 #include "Containers/Array.h"
 #include "Containers/StringFwd.h"
 #include "UObject/NameTypes.h"
+#include "UObject/PropertyTypeName.h"
 
 #define UE_API COREUOBJECT_API
 
@@ -13,14 +14,14 @@ namespace UE
 {
 
 /**
- * Represents one segment of a property path including type names and container index.
+ * Represents one segment of a property path including the property type and the container index.
  */
 struct FPropertyPathNameSegment
 {
 	/** Name of the property. */
 	FName Name;
-	/** Type name(s) of the property, separated with spaces, like "ArrayProperty StructProperty TypeOfStruct" */
-	FName Type;
+	/** Type of the property, if available. */
+	FPropertyTypeName Type;
 	/** Index within a container, or INDEX_NONE if not in a container. Uses ElementId for associative containers. */
 	int32 Index = INDEX_NONE;
 
@@ -34,14 +35,14 @@ struct FPropertyPathNameSegment
 		return {FName(NameWithIndex, NAME_NO_NUMBER_INTERNAL), Type, NAME_INTERNAL_TO_EXTERNAL(NameWithIndex.GetNumber())};
 	}
 
-	[[nodiscard]] inline FPropertyPathNameSegment SetType(FName NewType) const
+	[[nodiscard]] inline FPropertyPathNameSegment SetType(FPropertyTypeName NewType) const
 	{
 		return {Name, NewType, Index};
 	}
 };
 
 /**
- * Represents the path to a property, by name, including type names and the index in each container.
+ * Represents the path to a property, by name, including the property type and the container index.
  *
  * Sequenced containers use the index directly and associative containers use their ElementId.
  */
@@ -52,16 +53,16 @@ class FPropertyPathName
 	struct FSegment
 	{
 		FName NameWithIndex;
-		FNameEntryId Type;
+		FPropertyTypeName Type;
 
 		inline FPropertyPathNameSegment Unpack() const
 		{
-			return FPropertyPathNameSegment().SetNameWithIndex(NameWithIndex).SetType(FName::CreateFromDisplayId(Type, NAME_NO_NUMBER_INTERNAL));
+			return FPropertyPathNameSegment().SetNameWithIndex(NameWithIndex).SetType(Type);
 		}
 
 		inline static FSegment Pack(const FPropertyPathNameSegment& Segment)
 		{
-			return {Segment.PackNameWithIndex(), Segment.Type.GetDisplayIndex()};
+			return {Segment.PackNameWithIndex(), Segment.Type};
 		}
 
 		bool operator==(const FSegment& Segment) const;
@@ -100,26 +101,15 @@ public:
 		Segments.Pop(/*bAllowShrinking*/ false);
 	}
 
-	/** Pushes the type name on the end of the type of the last segment of the path. Ignored if the path is empty. */
-	inline void PushType(FName Type)
+	/** Sets the type of the last segment of the path. Ignored if the path is empty. */
+	inline void SetType(FPropertyTypeName Type)
 	{
 		if (!Segments.IsEmpty())
 		{
-			FNameEntryId& LastType = Segments.Last().Type;
-			if (LastType.IsNone())
-			{
-				LastType = Type.GetDisplayIndex();
-			}
-			else
-			{
-				PushTypeInternal(Type);
-			}
+			Segments.Last().Type = Type;
 		}
 	}
 
-	/** Pops a type name off the end of the type of the last segment of the path. Returns None if the path is empty. */
-	UE_API FName PopType();
-
 	/** Sets the index of the last segment of the path. Ignored if the path is empty. */
 	inline void SetIndex(int32 Index)
 	{
@@ -152,8 +142,6 @@ public:
 	UE_API friend uint32 GetTypeHash(const FPropertyPathName& Path);
 
 private:
-	UE_API void PushTypeInternal(FName Type);
-
 	TArray<FSegment> Segments;
 };
 
diff --git a/Engine/Source/Runtime/CoreUObject/Public/UObject/PropertyTypeName.h b/Engine/Source/Runtime/CoreUObject/Public/UObject/PropertyTypeName.h
--- a/Engine/Source/Runtime/CoreUObject/Public/UObject/PropertyTypeName.h
+++ b/Engine/Source/Runtime/CoreUObject/Public/UObject/PropertyTypeName.h
@@ -16,7 +16,7 @@ namespace UE { class FPropertyTypeNameBuilder; }
 namespace UE
 {
 
-struct UE_INTERNAL FPropertyTypeNameNode
+struct FPropertyTypeNameNode
 {
 	FName Name;
 	int32 InnerCount = 0;
@@ -76,6 +76,7 @@ private:
 	UE_API friend uint32 GetTypeHash(const FPropertyTypeName& TypeName);
 
 	UE_API friend bool operator==(const FPropertyTypeName& Lhs, const FPropertyTypeName& Rhs);
+	UE_API friend bool operator<(const FPropertyTypeName& Lhs, const FPropertyTypeName& Rhs);
 
 	UE_API friend FArchive& operator<<(FArchive& Ar, FPropertyTypeName& TypeName);
 
diff --git a/Engine/Source/Runtime/CoreUObject/Tests/UObject/PropertyPathNameTest.cpp b/Engine/Source/Runtime/CoreUObject/Tests/UObject/PropertyPathNameTest.cpp
new file mode 100644
--- /dev/null
+++ b/Engine/Source/Runtime/CoreUObject/Tests/UObject/PropertyPathNameTest.cpp
@@ -0,0 +1,169 @@
+// Copyright Epic Games, Inc. All Rights Reserved.
+
+#include "UObject/PropertyPathName.h"
+
+#if WITH_TESTS
+
+#include "Misc/StringBuilder.h"
+#include "Tests/TestHarnessAdapter.h"
+
+namespace UE
+{
+
+TEST_CASE_NAMED(FPropertyPathNameTest, "CoreUObject::PropertyPathName", "[Core][UObject][SmokeFilter]")
+{
+	const FName CountName(TEXTVIEW("Count"));
+
+	FPropertyTypeNameBuilder IntTypeBuilder;
+	IntTypeBuilder.AddTypeName(NAME_IntProperty);
+	const FPropertyTypeName IntType = IntTypeBuilder.Build();
+
+	SECTION("Empty")
+	{
+		FPropertyPathName PathName;
+		CHECK(PathName.IsEmpty());
+		CHECK(PathName.GetSegmentCount() == 0);
+		CHECK(WriteToString<16>(PathName).Len() == 0);
+	}
+
+	SECTION("Name")
+	{
+		FPropertyPathName PathName;
+		PathName.Push({CountName});
+		CHECK(!PathName.IsEmpty());
+		CHECK(PathName.GetSegmentCount() == 1);
+		CHECK(PathName.GetSegment(0).Name == CountName);
+		CHECK(PathName.GetSegment(0).Type.IsEmpty());
+		CHECK(PathName.GetSegment(0).Index == INDEX_NONE);
+		CHECK(TEXTVIEW("Count").Equals(WriteToString<16>(PathName)));
+	}
+
+	SECTION("NameType")
+	{
+		FPropertyPathName PathName;
+		PathName.Push({CountName, IntType});
+		CHECK(!PathName.IsEmpty());
+		CHECK(PathName.GetSegmentCount() == 1);
+		CHECK(PathName.GetSegment(0).Name == CountName);
+		CHECK(PathName.GetSegment(0).Type == IntType);
+		CHECK(PathName.GetSegment(0).Index == INDEX_NONE);
+		CHECK(TEXTVIEW("Count (IntProperty)").Equals(WriteToString<16>(PathName)));
+	}
+
+	SECTION("NameTypeIndex")
+	{
+		FPropertyPathName PathName;
+		PathName.Push({CountName, IntType, 7});
+		CHECK(!PathName.IsEmpty());
+		CHECK(PathName.GetSegmentCount() == 1);
+		CHECK(PathName.GetSegment(0).Name == CountName);
+		CHECK(PathName.GetSegment(0).Type == IntType);
+		CHECK(PathName.GetSegment(0).Index == 7);
+		CHECK(TEXTVIEW("Count[7] (IntProperty)").Equals(WriteToString<16>(PathName)));
+	}
+
+	SECTION("MultipleSegments")
+	{
+		FPropertyTypeNameBuilder ArrayTypeBuilder;
+		ArrayTypeBuilder.AddTypeName(NAME_ArrayProperty);
+		ArrayTypeBuilder.BeginTypeParameters();
+		ArrayTypeBuilder.AddTypeName(NAME_StructProperty);
+		ArrayTypeBuilder.BeginTypeParameters();
+		ArrayTypeBuilder.AddTypeName(TEXT("TestType"));
+		ArrayTypeBuilder.EndTypeParameters();
+		ArrayTypeBuilder.EndTypeParameters();
+
+		FPropertyTypeNameBuilder VectorTypeBuilder;
+		VectorTypeBuilder.AddTypeName(NAME_StructProperty);
+		VectorTypeBuilder.BeginTypeParameters();
+		VectorTypeBuilder.AddTypeName(TEXT("IntVector"));
+		VectorTypeBuilder.EndTypeParameters();
+
+		FPropertyPathName PathName;
+		PathName.Push({TEXT("TestArray"), ArrayTypeBuilder.Build(), 7});
+		PathName.Push({TEXT("Position"), VectorTypeBuilder.Build()});
+		PathName.Push({TEXT("X"), IntType});
+
+		CHECK(PathName.GetSegmentCount() == 3);
+		CHECK(TEXTVIEW("TestArray[7] (ArrayProperty<StructProperty<TestType>>) -> Position (StructProperty<IntVector>) -> X (IntProperty)").Equals(WriteToString<128>(PathName)));
+	}
+
+	SECTION("SetType/SetIndex")
+	{
+		FPropertyPathName PathName;
+		PathName.Push({CountName});
+		PathName.Push({CountName});
+
+		PathName.SetType(IntType);
+		CHECK(PathName.GetSegment(0).Type.IsEmpty());
+		CHECK(PathName.GetSegment(1).Type == IntType);
+
+		PathName.SetIndex(7);
+		CHECK(PathName.GetSegment(0).Index == INDEX_NONE);
+		CHECK(PathName.GetSegment(1).Index == 7);
+	}
+
+	SECTION("Pop/Empty/Reset")
+	{
+		FPropertyPathName PathName;
+		PathName.Push({CountName, IntType, 0});
+		PathName.Push({CountName, IntType, 1});
+		PathName.Push({CountName, IntType, 2});
+		CHECK(PathName.GetSegmentCount() == 3);
+		CHECK(PathName.GetSegment(2).Index == 2);
+
+		SECTION("Pop")
+		{
+			PathName.Pop();
+			CHECK(PathName.GetSegmentCount() == 2);
+			CHECK(PathName.GetSegment(1).Index == 1);
+			PathName.Pop();
+			CHECK(PathName.GetSegmentCount() == 1);
+			CHECK(PathName.GetSegment(0).Index == 0);
+			PathName.Pop();
+			CHECK(PathName.GetSegmentCount() == 0);
+			CHECK(PathName.IsEmpty());
+		}
+
+		SECTION("Empty")
+		{
+			PathName.Empty();
+			CHECK(PathName.GetSegmentCount() == 0);
+			CHECK(PathName.IsEmpty());
+		}
+
+		SECTION("Reset")
+		{
+			PathName.Reset();
+			CHECK(PathName.GetSegmentCount() == 0);
+			CHECK(PathName.IsEmpty());
+		}
+	}
+
+	SECTION("GetTypeHash")
+	{
+		FPropertyPathName Empty, Name, NameType, NameTypeIndex;
+		Name.Push({CountName});
+		NameType.Push({CountName, IntType});
+		NameTypeIndex.Push({CountName, IntType, 7});
+
+		const uint32 EmptyHash = GetTypeHash(Empty);
+		const uint32 NameHash = GetTypeHash(Name);
+		const uint32 NameTypeHash = GetTypeHash(NameType);
+		const uint32 NameTypeIndexHash = GetTypeHash(NameTypeIndex);
+
+		CHECK(EmptyHash == 0);
+		CHECK(EmptyHash != NameHash);
+		CHECK(EmptyHash != NameTypeHash);
+		CHECK(EmptyHash != NameTypeIndexHash);
+		CHECK(NameHash != NameTypeHash);
+		CHECK(NameTypeHash != NameTypeIndexHash);
+
+		NameTypeIndex.SetIndex(8);
+		CHECK(NameTypeIndexHash != GetTypeHash(NameTypeIndex));
+	}
+}
+
+} // UE
+
+#endif // WITH_TESTS
diff --git a/Engine/Source/Runtime/CoreUObject/Tests/UObject/PropertyTypeNameTest.cpp b/Engine/Source/Runtime/CoreUObject/Tests/UObject/PropertyTypeNameTest.cpp
--- a/Engine/Source/Runtime/CoreUObject/Tests/UObject/PropertyTypeNameTest.cpp
+++ b/Engine/Source/Runtime/CoreUObject/Tests/UObject/PropertyTypeNameTest.cpp
@@ -173,7 +173,7 @@ TEST_CASE_NAMED(FPropertyTypeNameSmokeTest, "CoreUObject::PropertyTypeName::Smok
 		CHECK(TypeName.GetTypeParameter(2).IsEmpty());
 	}
 
-	SECTION("Equals+GetTypeHash")
+	SECTION("Equals+Less+GetTypeHash")
 	{
 		const FPropertyTypeName Int = PropertyTypeNameTest::CreateInt();
 		const FPropertyTypeName VectorArray = PropertyTypeNameTest::CreateVectorArray();
@@ -204,21 +204,36 @@ TEST_CASE_NAMED(FPropertyTypeNameSmokeTest, "CoreUObject::PropertyTypeName::Smok
 		CHECK(GetTypeHash(FPropertyTypeName()) == GetTypeHash(FPropertyTypeName()));
 
 		CHECK(Int == CopyTemp(Int));
+		CHECK_FALSE(Int < CopyTemp(Int));
+		CHECK(GetTypeHash(Int) == GetTypeHash(CopyTemp(Int)));
+
 		CHECK_FALSE(Int == VectorArray);
 		CHECK_FALSE(Int == FPropertyTypeName());
-		CHECK(GetTypeHash(Int) == GetTypeHash(CopyTemp(Int)));
+		CHECK_FALSE(VectorArray < VectorArray);
+		CHECK_FALSE(Int < VectorArray);
+		CHECK(VectorArray < Int);
 		CHECK_FALSE(GetTypeHash(Int) == GetTypeHash(VectorArray));
 		CHECK_FALSE(GetTypeHash(Int) == GetTypeHash(FPropertyTypeName()));
 
 		CHECK(Int == MapIntToInt.GetTypeParameter(0));
 		CHECK_FALSE(Int == MapIntToInt);
 		CHECK_FALSE(MapIntToInt == Int);
+		CHECK_FALSE(MapIntToInt < MapIntToInt);
+		CHECK(Int < MapIntToInt);
+		CHECK_FALSE(MapIntToInt < Int);
+		CHECK_FALSE(MapIntToInt < VectorArray);
+		CHECK(VectorArray < MapIntToInt);
 		CHECK(GetTypeHash(Int) == GetTypeHash(MapIntToInt.GetTypeParameter(0)));
 		CHECK_FALSE(GetTypeHash(Int) == GetTypeHash(MapIntToInt));
 		CHECK_FALSE(GetTypeHash(MapIntToInt) == GetTypeHash(Int));
 
 		CHECK(MapIntToInt == MapIntToIntToInt.GetTypeParameter(1));
 		CHECK_FALSE(MapIntToInt == MapIntToIntToInt);
+		CHECK_FALSE(MapIntToIntToInt < MapIntToIntToInt);
+		CHECK(MapIntToInt < MapIntToIntToInt);
+		CHECK_FALSE(MapIntToIntToInt < MapIntToInt);
+		CHECK_FALSE(MapIntToIntToInt < VectorArray);
+		CHECK(VectorArray < MapIntToIntToInt);
 		CHECK(GetTypeHash(MapIntToInt) == GetTypeHash(MapIntToIntToInt.GetTypeParameter(1)));
 		CHECK_FALSE(GetTypeHash(MapIntToInt) == GetTypeHash(MapIntToIntToInt));
 	}
