Make the bulk data request counters atomic.

The counters are manipulated on concurrent tasks at the same time, and should be atomic.

diff --git a/Engine/Source/Runtime/CoreUObject/Private/Serialization/BulkDataStreaming.cpp b/Engine/Source/Runtime/CoreUObject/Private/Serialization/BulkDataStreaming.cpp
--- a/Engine/Source/Runtime/CoreUObject/Private/Serialization/BulkDataStreaming.cpp
+++ b/Engine/Source/Runtime/CoreUObject/Private/Serialization/BulkDataStreaming.cpp
@@ -17,8 +17,8 @@
 
 //////////////////////////////////////////////////////////////////////////////
 
-TRACE_DECLARE_INT_COUNTER(BulkDataBatchRequest_Count, TEXT("BulkData/BatchRequest/Count"));
-TRACE_DECLARE_INT_COUNTER(BulkDataBatchRequest_PendingCount, TEXT("BulkData/BatchRequest/Pending"));
+TRACE_DECLARE_ATOMIC_INT_COUNTER(BulkDataBatchRequest_Count, TEXT("BulkData/BatchRequest/Count"));
+TRACE_DECLARE_ATOMIC_INT_COUNTER(BulkDataBatchRequest_PendingCount, TEXT("BulkData/BatchRequest/Pending"));
 
 /**
  * When enabled calls to FChunkReadFileHandle::ReadRequest will validate that the request
