Made the JsonArchive structured formatters non-copyable so that their useless copy and move consturctors don't show up in IntelliSense.

diff --git a/Engine/Source/Runtime/CoreUObject/Private/Serialization/Formatters/JsonArchiveInputFormatter.cpp b/Engine/Source/Runtime/CoreUObject/Private/Serialization/Formatters/JsonArchiveInputFormatter.cpp
--- a/Engine/Source/Runtime/CoreUObject/Private/Serialization/Formatters/JsonArchiveInputFormatter.cpp
+++ b/Engine/Source/Runtime/CoreUObject/Private/Serialization/Formatters/JsonArchiveInputFormatter.cpp
@@ -45,11 +45,7 @@ FArchive& FJsonArchiveInputFormatter::GetUnderlyingArchive()
 
 FStructuredArchiveFormatter* FJsonArchiveInputFormatter::CreateSubtreeReader()
 {
-	FJsonArchiveInputFormatter* Cloned = new FJsonArchiveInputFormatter(*this);
-	Cloned->ObjectStack.Empty();
-	Cloned->ValueStack.Empty();
-	Cloned->MapIteratorStack.Empty();
-	Cloned->ArrayValuesRemainingStack.Empty();
+	FJsonArchiveInputFormatter* Cloned = new FJsonArchiveInputFormatter(this->Inner, this->ResolveObject);
 	Cloned->ValueStack.Push(ValueStack.Top());
 
 	return Cloned;
diff --git a/Engine/Source/Runtime/CoreUObject/Public/Serialization/Formatters/JsonArchiveInputFormatter.h b/Engine/Source/Runtime/CoreUObject/Public/Serialization/Formatters/JsonArchiveInputFormatter.h
--- a/Engine/Source/Runtime/CoreUObject/Public/Serialization/Formatters/JsonArchiveInputFormatter.h
+++ b/Engine/Source/Runtime/CoreUObject/Public/Serialization/Formatters/JsonArchiveInputFormatter.h
@@ -31,7 +31,13 @@ struct FWeakObjectPtr;
 class FJsonArchiveInputFormatter final : public FStructuredArchiveFormatter
 {
 public:
-	COREUOBJECT_API FJsonArchiveInputFormatter(FArchive& InInner, TFunction<UObject* (const FPackageIndex)> InResolveObject = nullptr);
+	// Noncopyable
+	FJsonArchiveInputFormatter(FJsonArchiveInputFormatter&&) = delete;
+	FJsonArchiveInputFormatter(const FJsonArchiveInputFormatter&) = delete;
+	FJsonArchiveInputFormatter& operator=(FJsonArchiveInputFormatter&&) = delete;
+	FJsonArchiveInputFormatter& operator=(const FJsonArchiveInputFormatter&) = delete;
+
+	COREUOBJECT_API explicit FJsonArchiveInputFormatter(FArchive& InInner, TFunction<UObject* (const FPackageIndex)> InResolveObject = nullptr);
 	COREUOBJECT_API virtual ~FJsonArchiveInputFormatter();
 
 	COREUOBJECT_API virtual FArchive& GetUnderlyingArchive() override;
diff --git a/Engine/Source/Runtime/CoreUObject/Public/Serialization/Formatters/JsonArchiveOutputFormatter.h b/Engine/Source/Runtime/CoreUObject/Public/Serialization/Formatters/JsonArchiveOutputFormatter.h
--- a/Engine/Source/Runtime/CoreUObject/Public/Serialization/Formatters/JsonArchiveOutputFormatter.h
+++ b/Engine/Source/Runtime/CoreUObject/Public/Serialization/Formatters/JsonArchiveOutputFormatter.h
@@ -27,7 +27,13 @@ struct FWeakObjectPtr;
 class FJsonArchiveOutputFormatter final : public FStructuredArchiveFormatter
 {
 public:
-	COREUOBJECT_API FJsonArchiveOutputFormatter(FArchive& InInner);
+	// Noncopyable
+	FJsonArchiveOutputFormatter(FJsonArchiveOutputFormatter&&) = delete;
+	FJsonArchiveOutputFormatter(const FJsonArchiveOutputFormatter&) = delete;
+	FJsonArchiveOutputFormatter& operator=(FJsonArchiveOutputFormatter&&) = delete;
+	FJsonArchiveOutputFormatter& operator=(const FJsonArchiveOutputFormatter&) = delete;
+
+	COREUOBJECT_API explicit FJsonArchiveOutputFormatter(FArchive& InInner);
 	COREUOBJECT_API virtual ~FJsonArchiveOutputFormatter();
 
 	COREUOBJECT_API virtual FArchive& GetUnderlyingArchive() override;
