[ZenLoader]
 - Fix cases where objects constructed during callbacks at end of load (i.e. flush reinstancing queue, etc...) were wrongly assigned to the package on the stack, leading to their premature postloading even before they were fully serialized.

diff --git a/Engine/Source/Runtime/CoreUObject/Private/Serialization/AsyncLoading2.cpp b/Engine/Source/Runtime/CoreUObject/Private/Serialization/AsyncLoading2.cpp
--- a/Engine/Source/Runtime/CoreUObject/Private/Serialization/AsyncLoading2.cpp
+++ b/Engine/Source/Runtime/CoreUObject/Private/Serialization/AsyncLoading2.cpp
@@ -2775,21 +2775,34 @@ struct FAsyncPackage2
 
 	void AddConstructedObject(UObject* Object, bool bSubObjectThatAlreadyExists)
 	{
+		// Skip the tracking when we reach the game-thread as we don't have to add the asyncloading flag anymore and since
+		// we're already in the postloading step, there is no value in accepting new objects and trying to postload them.
+		// Also, given that anything can create UObject outside the scope of loading once we reach postload and callbacks
+		// we could end up with objects that have nothing to do with the actual package currently being on the stack,
+		// and these package might not even have finished loading, so postloading any of these objects might be premature
+		// anyway.
+		if (AsyncPackageLoadingState >= EAsyncPackageLoadingState2::DeferredPostLoad)
+		{
+			return;
+		}
+
 		if (bSubObjectThatAlreadyExists)
 		{
 			ConstructedObjects.AddUnique(Object);
 		}
 		else
 		{
+			// Mark objects created during async loading process (e.g. from within PostLoad or CreateExport) as async loaded so they 
+			// cannot be found. This requires also keeping track of them so we can remove the async loading flag later one when we 
+			// finished routing PostLoad to all objects.
+			Object->SetInternalFlags(EInternalObjectFlags::AsyncLoading);
+
 			ConstructedObjects.Add(Object);
 		}
 	}
 
 	void ClearConstructedObjects();
 
-	/** Returns the UPackage wrapped by this, if it is valid */
-	UPackage* GetLoadedPackage();
-
 	/** Class specific callback for initializing non-native objects */
 	EAsyncPackageState::Type PostLoadInstances(FAsyncLoadingThreadState2& ThreadState);
 
@@ -8607,13 +8620,6 @@ void FAsyncLoadingThread2::NotifyConstructedDuringAsyncLoading(UObject* Object,
 		return;
 	}
 
-	// Mark objects created during async loading process (e.g. from within PostLoad or CreateExport) as async loaded so they 
-	// cannot be found. This requires also keeping track of them so we can remove the async loading flag later one when we 
-	// finished routing PostLoad to all objects.
-	if (!bSubObjectThatAlreadyExists)
-	{
-		Object->SetInternalFlags(EInternalObjectFlags::AsyncLoading);
-	}
 	FAsyncPackage2* AsyncPackage2 = (FAsyncPackage2*)ThreadContext.AsyncPackage;
 	AsyncPackage2->AddConstructedObject(Object, bSubObjectThatAlreadyExists);
 }
@@ -8861,6 +8867,8 @@ void FAsyncPackage2::CreateUPackage()
 #if WITH_EDITOR
 		FCookLoadScope CookLoadScope(Desc.PackageReferencer.CookLoadType);
 #endif
+		// Add scope so that this constructed object is assigned to ourself and not another package up the stack in ImportPackagesRecursive.
+		FAsyncPackageScope2 Scope(this);
 		LinkerRoot = NewObject<UPackage>(/*Outer*/nullptr, Desc.UPackageName);
 		bCreatedLinkerRoot = true;
 	}
