[ZenLoader]
 - Fix !bAllDone assert that could happen when a node timed out and was resumed during a recursive flush while still executing higher on the stack

diff --git a/Engine/Source/Runtime/CoreUObject/Private/Serialization/AsyncLoading2.cpp b/Engine/Source/Runtime/CoreUObject/Private/Serialization/AsyncLoading2.cpp
--- a/Engine/Source/Runtime/CoreUObject/Private/Serialization/AsyncLoading2.cpp
+++ b/Engine/Source/Runtime/CoreUObject/Private/Serialization/AsyncLoading2.cpp
@@ -4734,10 +4734,14 @@ bool FAsyncLoadEventQueue2::PopAndExecute(FAsyncLoadingThreadState2& ThreadState
 {
 	if (TimedOutEventNode)
 	{
-		EEventLoadNodeExecutionResult Result = TimedOutEventNode->Execute(ThreadState);
-		if (Result != EEventLoadNodeExecutionResult::Timeout)
+		// Backup and reset the node before executing it in case we end up with a recursive flush call, we don't want the same node run multiple time.
+		FEventLoadNode2* LocalTimedOutEventNode = TimedOutEventNode;
+		TimedOutEventNode = nullptr;
+
+		EEventLoadNodeExecutionResult Result = LocalTimedOutEventNode->Execute(ThreadState);
+		if (Result == EEventLoadNodeExecutionResult::Timeout)
 		{
-			TimedOutEventNode = nullptr;
+			TimedOutEventNode = LocalTimedOutEventNode;
 		}
 		return true;
 	}
@@ -4792,9 +4796,12 @@ bool FAsyncLoadEventQueue2::ExecuteSyncLoadEvents(FAsyncLoadingThreadState2& Thr
 	bool bDidSomething = false;
 	if (TimedOutEventNode && ShouldExecuteNode(*TimedOutEventNode))
 	{
-		EEventLoadNodeExecutionResult Result = TimedOutEventNode->Execute(ThreadState);
-		check(Result == EEventLoadNodeExecutionResult::Complete); // we can't timeout during a sync load operation
+		// Backup and reset the node before executing it in case we end up with a recursive flush call, we don't want the same node run multiple time.
+		FEventLoadNode2* LocalTimedOutEventNode = TimedOutEventNode;
 		TimedOutEventNode = nullptr;
+
+		EEventLoadNodeExecutionResult Result = LocalTimedOutEventNode->Execute(ThreadState);
+		check(Result == EEventLoadNodeExecutionResult::Complete); // we can't timeout during a sync load operation
 		bDidSomething = true;
 	}
 
