[ZenLoader/UObject]
 - Avoid automatically adding AsyncLoading flag during duplication as this can prevent GC. This flag should only be managed by the async loading code itself.

diff --git a/Engine/Source/Runtime/CoreUObject/Private/UObject/UObjectGlobals.cpp b/Engine/Source/Runtime/CoreUObject/Private/UObject/UObjectGlobals.cpp
--- a/Engine/Source/Runtime/CoreUObject/Private/UObject/UObjectGlobals.cpp
+++ b/Engine/Source/Runtime/CoreUObject/Private/UObject/UObjectGlobals.cpp
@@ -2977,8 +2977,8 @@ UObject* StaticDuplicateObjectEx( FObjectDuplicationParameters& Parameters )
 		*Parameters.DestClass->GetName(), Parameters.DestClass->GetPropertiesSize());
 	
 	UE_CLOG(FPlatformProperties::RequiresCookedData() && Parameters.SourceObject->HasAnyInternalFlags(EInternalObjectFlags::AsyncLoading), LogUObjectGlobals, Warning, TEXT("Duplicating object '%s' that's still being async loaded"), *Parameters.SourceObject->GetFullName());
-	// Make sure we're not duplicating the Async or LoaderImport internal flags, they will prevent the object from being gcd
-	Parameters.InternalFlagMask &= ~(EInternalObjectFlags::Async | EInternalObjectFlags::LoaderImport);
+	// Make sure we're not duplicating the AsyncLoading, Async or LoaderImport internal flags, they will prevent the object from being gcd.
+	Parameters.InternalFlagMask &= ~(EInternalObjectFlags::Async | EInternalObjectFlags::LoaderImport | EInternalObjectFlags::AsyncLoading);
 
 	if (!IsAsyncLoading() && Parameters.SourceObject->HasAnyFlags(RF_ClassDefaultObject))
 	{
