Remove text.h from script.h (script.h is included  by object.h)

diff --git a/Engine/Plugins/Animation/LiveLink/Source/LiveLink/Private/LiveLinkBlueprintLibrary.cpp b/Engine/Plugins/Animation/LiveLink/Source/LiveLink/Private/LiveLinkBlueprintLibrary.cpp
--- a/Engine/Plugins/Animation/LiveLink/Source/LiveLink/Private/LiveLinkBlueprintLibrary.cpp
+++ b/Engine/Plugins/Animation/LiveLink/Source/LiveLink/Private/LiveLinkBlueprintLibrary.cpp
@@ -1,6 +1,7 @@
 ﻿// Copyright Epic Games, Inc. All Rights Reserved.
 
 #include "LiveLinkBlueprintLibrary.h"
+#include "Blueprint/BlueprintExceptionInfo.h"
 #include "Features/IModularFeatures.h"
 #include "ILiveLinkClient.h"
 #include "ILiveLinkSource.h"
diff --git a/Engine/Plugins/Experimental/AnimNext/Source/AnimNext/Private/Component/AnimNextComponent.cpp b/Engine/Plugins/Experimental/AnimNext/Source/AnimNext/Private/Component/AnimNextComponent.cpp
--- a/Engine/Plugins/Experimental/AnimNext/Source/AnimNext/Private/Component/AnimNextComponent.cpp
+++ b/Engine/Plugins/Experimental/AnimNext/Source/AnimNext/Private/Component/AnimNextComponent.cpp
@@ -2,6 +2,7 @@
 
 #include "Component/AnimNextComponent.h"
 
+#include "Blueprint/BlueprintExceptionInfo.h"
 #include "Param/ParamUtils.h"
 #include "Scheduler/Scheduler.h"
 #include "Scheduler/ScheduleContext.h"
diff --git a/Engine/Plugins/Experimental/Chooser/Source/Chooser/Private/ChooserFunctionLibrary.cpp b/Engine/Plugins/Experimental/Chooser/Source/Chooser/Private/ChooserFunctionLibrary.cpp
--- a/Engine/Plugins/Experimental/Chooser/Source/Chooser/Private/ChooserFunctionLibrary.cpp
+++ b/Engine/Plugins/Experimental/Chooser/Source/Chooser/Private/ChooserFunctionLibrary.cpp
@@ -1,5 +1,6 @@
 // Copyright Epic Games, Inc. All Rights Reserved.
 #include "ChooserFunctionLibrary.h"
+#include "Blueprint/BlueprintExceptionInfo.h"
 #include "Chooser.h"
 
 #include UE_INLINE_GENERATED_CPP_BY_NAME(ChooserFunctionLibrary)
diff --git a/Engine/Plugins/Experimental/PCG/Source/PCG/Private/Helpers/PCGGraphParametersHelpers.cpp b/Engine/Plugins/Experimental/PCG/Source/PCG/Private/Helpers/PCGGraphParametersHelpers.cpp
--- a/Engine/Plugins/Experimental/PCG/Source/PCG/Private/Helpers/PCGGraphParametersHelpers.cpp
+++ b/Engine/Plugins/Experimental/PCG/Source/PCG/Private/Helpers/PCGGraphParametersHelpers.cpp
@@ -1,7 +1,7 @@
 // Copyright Epic Games, Inc. All Rights Reserved.
 
 #include "Helpers/PCGGraphParametersHelpers.h"
-
+#include "Blueprint/BlueprintExceptionInfo.h"
 #include "PCGGraph.h"
 #include "PropertyBag.h"
 
diff --git a/Engine/Plugins/Experimental/PCG/Source/PCG/Private/Metadata/PCGMetadataAccessor.cpp b/Engine/Plugins/Experimental/PCG/Source/PCG/Private/Metadata/PCGMetadataAccessor.cpp
--- a/Engine/Plugins/Experimental/PCG/Source/PCG/Private/Metadata/PCGMetadataAccessor.cpp
+++ b/Engine/Plugins/Experimental/PCG/Source/PCG/Private/Metadata/PCGMetadataAccessor.cpp
@@ -1,6 +1,7 @@
 // Copyright Epic Games, Inc. All Rights Reserved.
 
 #include "Metadata/PCGMetadataAccessor.h"
+#include "Blueprint/BlueprintExceptionInfo.h"
 #include "Metadata/PCGMetadata.h"
 #include "PCGPoint.h"
 
diff --git a/Engine/Plugins/Experimental/StructUtils/Source/StructUtilsEngine/Private/StructUtilsFunctionLibrary.cpp b/Engine/Plugins/Experimental/StructUtils/Source/StructUtilsEngine/Private/StructUtilsFunctionLibrary.cpp
--- a/Engine/Plugins/Experimental/StructUtils/Source/StructUtilsEngine/Private/StructUtilsFunctionLibrary.cpp
+++ b/Engine/Plugins/Experimental/StructUtils/Source/StructUtilsEngine/Private/StructUtilsFunctionLibrary.cpp
@@ -1,6 +1,7 @@
 // Copyright Epic Games, Inc. All Rights Reserved.
 
 #include "StructUtilsFunctionLibrary.h"
+#include "Blueprint/BlueprintExceptionInfo.h"
 #include "UObject/EnumProperty.h"
 
 #include UE_INLINE_GENERATED_CPP_BY_NAME(StructUtilsFunctionLibrary)
diff --git a/Engine/Plugins/JsonBlueprintUtilities/Source/JsonBlueprintUtilities/Private/JsonBlueprintFunctionLibrary.cpp b/Engine/Plugins/JsonBlueprintUtilities/Source/JsonBlueprintUtilities/Private/JsonBlueprintFunctionLibrary.cpp
--- a/Engine/Plugins/JsonBlueprintUtilities/Source/JsonBlueprintUtilities/Private/JsonBlueprintFunctionLibrary.cpp
+++ b/Engine/Plugins/JsonBlueprintUtilities/Source/JsonBlueprintUtilities/Private/JsonBlueprintFunctionLibrary.cpp
@@ -2,6 +2,7 @@
 
 #include "JsonBlueprintFunctionLibrary.h"
 
+#include "Blueprint/BlueprintExceptionInfo.h"
 #include "JsonObjectConverter.h"
 #include "JsonObjectWrapper.h"
 #include "Misc/FileHelper.h"
diff --git a/Engine/Plugins/Online/OnlineSubsystem/Source/Public/Interfaces/OnlineStoreInterfaceV2.h b/Engine/Plugins/Online/OnlineSubsystem/Source/Public/Interfaces/OnlineStoreInterfaceV2.h
--- a/Engine/Plugins/Online/OnlineSubsystem/Source/Public/Interfaces/OnlineStoreInterfaceV2.h
+++ b/Engine/Plugins/Online/OnlineSubsystem/Source/Public/Interfaces/OnlineStoreInterfaceV2.h
@@ -1,7 +1,7 @@
 // Copyright Epic Games, Inc. All Rights Reserved.
 
 #pragma once
-
+#include "Internationalization/Text.h"
 #include "Misc/DateTime.h"
 #include "OnlineDelegateMacros.h"
 #include "OnlineStoreInterfaceV2.generated.h"
diff --git a/Engine/Plugins/Runtime/ModelViewViewModel/Source/ModelViewViewModel/Private/MVVMViewModelBase.cpp b/Engine/Plugins/Runtime/ModelViewViewModel/Source/ModelViewViewModel/Private/MVVMViewModelBase.cpp
--- a/Engine/Plugins/Runtime/ModelViewViewModel/Source/ModelViewViewModel/Private/MVVMViewModelBase.cpp
+++ b/Engine/Plugins/Runtime/ModelViewViewModel/Source/ModelViewViewModel/Private/MVVMViewModelBase.cpp
@@ -1,7 +1,7 @@
 // Copyright Epic Games, Inc. All Rights Reserved.
 
 #include "MVVMViewModelBase.h"
-
+#include "Blueprint/BlueprintExceptionInfo.h"
 #include "ViewModel/MVVMViewModelBlueprintGeneratedClass.h"
 
 #include UE_INLINE_GENERATED_CPP_BY_NAME(MVVMViewModelBase)
diff --git a/Engine/Source/Editor/AnimationBlueprintEditor/Private/AnimationBlueprintEditor.cpp b/Engine/Source/Editor/AnimationBlueprintEditor/Private/AnimationBlueprintEditor.cpp
--- a/Engine/Source/Editor/AnimationBlueprintEditor/Private/AnimationBlueprintEditor.cpp
+++ b/Engine/Source/Editor/AnimationBlueprintEditor/Private/AnimationBlueprintEditor.cpp
@@ -50,6 +50,7 @@
 #include "AssetRegistry/AssetData.h"
 #include "BlendSpaceDocumentTabFactory.h"
 #include "BlendSpaceGraph.h"
+#include "Blueprint/BlueprintExceptionInfo.h"
 #include "Components/SkeletalMeshComponent.h"
 #include "Containers/EnumAsByte.h"
 #include "CoreGlobals.h"
diff --git a/Engine/Source/Editor/AnimationBlueprintLibrary/Private/AnimationAttributeBlueprintLibrary.cpp b/Engine/Source/Editor/AnimationBlueprintLibrary/Private/AnimationAttributeBlueprintLibrary.cpp
--- a/Engine/Source/Editor/AnimationBlueprintLibrary/Private/AnimationAttributeBlueprintLibrary.cpp
+++ b/Engine/Source/Editor/AnimationBlueprintLibrary/Private/AnimationAttributeBlueprintLibrary.cpp
@@ -6,6 +6,7 @@
 #include "Animation/AnimData/AttributeIdentifier.h"
 #include "Animation/AnimData/IAnimationDataController.h"
 #include "Animation/AttributeCurve.h"
+#include "Blueprint/BlueprintExceptionInfo.h"
 #include "Containers/ArrayView.h"
 #include "Curves/KeyHandle.h"
 #include "Internationalization/Internationalization.h"
diff --git a/Engine/Source/Editor/UnrealEd/Private/Kismet2/KismetDebugUtilities.cpp b/Engine/Source/Editor/UnrealEd/Private/Kismet2/KismetDebugUtilities.cpp
--- a/Engine/Source/Editor/UnrealEd/Private/Kismet2/KismetDebugUtilities.cpp
+++ b/Engine/Source/Editor/UnrealEd/Private/Kismet2/KismetDebugUtilities.cpp
@@ -44,6 +44,7 @@
 #include "UObject/UnrealType.h"
 #include "AnimationGraphSchema.h"
 #include "BlueprintEditorSettings.h"
+#include "Blueprint/BlueprintExceptionInfo.h"
 
 #define LOCTEXT_NAMESPACE "BlueprintDebugging"
 
diff --git a/Engine/Source/Runtime/CoreUObject/Private/UObject/ScriptCore.cpp b/Engine/Source/Runtime/CoreUObject/Private/UObject/ScriptCore.cpp
--- a/Engine/Source/Runtime/CoreUObject/Private/UObject/ScriptCore.cpp
+++ b/Engine/Source/Runtime/CoreUObject/Private/UObject/ScriptCore.cpp
@@ -28,6 +28,7 @@
 #include "UObject/Stack.h"
 #include "UObject/Reload.h"
 #include "Blueprint/BlueprintSupport.h"
+#include "Blueprint/BlueprintExceptionInfo.h"
 #include "UObject/ScriptMacros.h"
 #include "UObject/UObjectThreadContext.h"
 #include "HAL/IConsoleManager.h"
diff --git a/Engine/Source/Runtime/CoreUObject/Public/Blueprint/BlueprintExceptionInfo.h b/Engine/Source/Runtime/CoreUObject/Public/Blueprint/BlueprintExceptionInfo.h
new file mode 100644
--- /dev/null
+++ b/Engine/Source/Runtime/CoreUObject/Public/Blueprint/BlueprintExceptionInfo.h
@@ -0,0 +1,50 @@
+// Copyright Epic Games, Inc. All Rights Reserved.
+
+#pragma once
+
+#include "Internationalization/Text.h"
+
+// Kinds of Blueprint exceptions
+namespace EBlueprintExceptionType
+{
+	enum Type
+	{
+		Breakpoint,
+		Tracepoint,
+		WireTracepoint,
+		AccessViolation,
+		InfiniteLoop,
+		NonFatalError,
+		FatalError,
+		AbortExecution,
+	};
+}
+
+// Information about a blueprint exception
+struct FBlueprintExceptionInfo
+{
+public:
+	FBlueprintExceptionInfo(EBlueprintExceptionType::Type InEventType)
+		: EventType(InEventType)
+	{
+	}
+
+	FBlueprintExceptionInfo(EBlueprintExceptionType::Type InEventType, const FText& InDescription)
+		: EventType(InEventType)
+		, Description(InDescription)
+	{
+	}
+
+	EBlueprintExceptionType::Type GetType() const
+	{
+		return EventType;
+	}
+
+	const FText& GetDescription() const
+	{
+		return Description;
+	}
+protected:
+	EBlueprintExceptionType::Type EventType;
+	FText Description;
+};
diff --git a/Engine/Source/Runtime/CoreUObject/Public/UObject/Script.h b/Engine/Source/Runtime/CoreUObject/Public/UObject/Script.h
--- a/Engine/Source/Runtime/CoreUObject/Public/UObject/Script.h
+++ b/Engine/Source/Runtime/CoreUObject/Public/UObject/Script.h
@@ -7,13 +7,16 @@
 #pragma once
 
 #include "HAL/ThreadSingleton.h"
+#if UE_ENABLE_INCLUDE_ORDER_DEPRECATED_IN_5_4
 #include "Internationalization/Text.h"
+#endif
 #include "Stats/Stats.h"
 #include "Misc/EnumClassFlags.h"
 #include "Misc/CoreMisc.h"
 #include "Memory/VirtualStackAllocator.h"
 
 struct FFrame;
+struct FBlueprintExceptionInfo;
 
 // It's best to set only one of these, but strictly speaking you could set both.
 // The results will be confusing. Native time would be included only in a coarse 
@@ -335,22 +338,6 @@ enum class EBlueprintTextLiteralType : uint8
 	StringTableEntry,
 };
 
-// Kinds of Blueprint exceptions
-namespace EBlueprintExceptionType
-{
-	enum Type
-	{
-		Breakpoint,
-		Tracepoint,
-		WireTracepoint,
-		AccessViolation,
-		InfiniteLoop,
-		NonFatalError,
-		FatalError,
-		AbortExecution,
-	};
-}
-
 // Script instrumentation event types
 namespace EScriptInstrumentation
 {
@@ -376,35 +363,6 @@ namespace EScriptInstrumentation
 	};
 }
 
-// Information about a blueprint exception
-struct FBlueprintExceptionInfo
-{
-public:
-	FBlueprintExceptionInfo(EBlueprintExceptionType::Type InEventType)
-		: EventType(InEventType)
-	{
-	}
-
-	FBlueprintExceptionInfo(EBlueprintExceptionType::Type InEventType, const FText& InDescription)
-		: EventType(InEventType)
-		, Description(InDescription)
-	{
-	}
-
-	EBlueprintExceptionType::Type GetType() const
-	{
-		return EventType;
-	}
-
-	const FText& GetDescription() const
-	{
-		return Description;
-	}
-protected:
-	EBlueprintExceptionType::Type EventType;
-	FText Description;
-};
-
 // Information about a blueprint instrumentation signal
 struct FScriptInstrumentationSignal
 {
diff --git a/Engine/Source/Runtime/Engine/Classes/Kismet/DataTableFunctionLibrary.h b/Engine/Source/Runtime/Engine/Classes/Kismet/DataTableFunctionLibrary.h
--- a/Engine/Source/Runtime/Engine/Classes/Kismet/DataTableFunctionLibrary.h
+++ b/Engine/Source/Runtime/Engine/Classes/Kismet/DataTableFunctionLibrary.h
@@ -10,6 +10,7 @@
 #include "Kismet/BlueprintFunctionLibrary.h"
 #include "Engine/DataTable.h"
 #include "UObject/Class.h" // for FStructUtils
+#include "Blueprint/BlueprintExceptionInfo.h"
 #include "DataTableFunctionLibrary.generated.h"
 
 class UCurveTable;
diff --git a/Engine/Source/Runtime/Engine/Private/FieldNotification/FieldNotificationLibrary.cpp b/Engine/Source/Runtime/Engine/Private/FieldNotification/FieldNotificationLibrary.cpp
--- a/Engine/Source/Runtime/Engine/Private/FieldNotification/FieldNotificationLibrary.cpp
+++ b/Engine/Source/Runtime/Engine/Private/FieldNotification/FieldNotificationLibrary.cpp
@@ -6,6 +6,7 @@
 #include "GameFramework/Actor.h"
 #include "Net/NetPushModelHelpers.h"
 #include "UObject/Script.h"
+#include "Blueprint/BlueprintExceptionInfo.h"
 
 #define LOCTEXT_NAMESPACE "FieldNotificationLibrary"
 
diff --git a/Engine/Source/Runtime/Engine/Private/KismetSystemLibrary.cpp b/Engine/Source/Runtime/Engine/Private/KismetSystemLibrary.cpp
--- a/Engine/Source/Runtime/Engine/Private/KismetSystemLibrary.cpp
+++ b/Engine/Source/Runtime/Engine/Private/KismetSystemLibrary.cpp
@@ -3,6 +3,7 @@
 #include "Kismet/KismetSystemLibrary.h"
 #include "AssetRegistry/ARFilter.h"
 #include "Blueprint/BlueprintSupport.h"
+#include "Blueprint/BlueprintExceptionInfo.h"
 #include "Engine/AssetManagerTypes.h"
 #include "Engine/BlueprintGeneratedClass.h"
 #include "HAL/FileManager.h"
