Serialize Override Property data if using Unversioned serialization
* Allows resulting object from StaticDuplicateObject to get override information

diff --git a/Engine/Source/Runtime/CoreUObject/Private/UObject/Obj.cpp b/Engine/Source/Runtime/CoreUObject/Private/UObject/Obj.cpp
--- a/Engine/Source/Runtime/CoreUObject/Private/UObject/Obj.cpp
+++ b/Engine/Source/Runtime/CoreUObject/Private/UObject/Obj.cpp
@@ -1648,9 +1648,10 @@ void UObject::Serialize(FStructuredArchive::FRecord Record)
 		// Handle derived UClass objects (exact UClass objects are native only and shouldn't be touched)
 		if (ObjClass != UClass::StaticClass())
 		{
-			if (UnderlyingArchive.IsTransacting())
+			// Serializing the overriden properties for undo/redo
+			// of if using UPS (ie. object duplication)
+			if (UnderlyingArchive.IsTransacting() || UnderlyingArchive.UseUnversionedPropertySerialization())
 			{
-				// Serialize the overriden property inside the transaction buffer
 				FOverridableManager::Get().SerializeOverriddenProperties(*this, Record);
 			}
 
