Fix default CDO value not being serialized in cooked server content when using Overridable Serialization.

diff --git a/Engine/Source/Runtime/CoreUObject/Private/UObject/Class.cpp b/Engine/Source/Runtime/CoreUObject/Private/UObject/Class.cpp
--- a/Engine/Source/Runtime/CoreUObject/Private/UObject/Class.cpp
+++ b/Engine/Source/Runtime/CoreUObject/Private/UObject/Class.cpp
@@ -1791,7 +1791,7 @@ void UStruct::SerializeVersionedTaggedProperties(FStructuredArchive::FSlot Slot,
 					uint8* DataPtr      = Property->ContainerPtrToValuePtr           <uint8>(Data, Idx);
 					uint8* DefaultValue = Property->ContainerPtrToValuePtrForDefaults<uint8>(DefaultsStruct, Defaults, Idx);
 					if (StaticArrayContainer.IsSet() || CustomPropertyNode || !bDoDeltaSerialization ||
-						(FOverridableSerializationLogic::IsEnabled() && FOverridableSerializationLogic::GetOverriddenPropertyOperation(UnderlyingArchive, Property, DataPtr) != EOverriddenPropertyOperation::None) ||
+						(FOverridableSerializationLogic::IsEnabled() && FOverridableSerializationLogic::GetOverriddenPropertyOperation(UnderlyingArchive, Property, DataPtr, DefaultValue) != EOverriddenPropertyOperation::None) ||
 						(!FOverridableSerializationLogic::IsEnabled() && !Property->Identical(DataPtr, DefaultValue, UnderlyingArchive.GetPortFlags())))
 					{
 						if (bUseAtomicSerialization)
diff --git a/Engine/Source/Runtime/CoreUObject/Private/UObject/OverriddenPropertySet.cpp b/Engine/Source/Runtime/CoreUObject/Private/UObject/OverriddenPropertySet.cpp
--- a/Engine/Source/Runtime/CoreUObject/Private/UObject/OverriddenPropertySet.cpp
+++ b/Engine/Source/Runtime/CoreUObject/Private/UObject/OverriddenPropertySet.cpp
@@ -22,7 +22,7 @@ DEFINE_LOG_CATEGORY(LogOverridableObject);
 thread_local bool FOverridableSerializationLogic::bUseOverridableSerialization = false;
 thread_local FOverriddenPropertySet* FOverridableSerializationLogic::OverriddenProperties = nullptr;
 
-EOverriddenPropertyOperation FOverridableSerializationLogic::GetOverriddenPropertyOperation(const FArchive& Ar, FProperty* Property /*= nullptr*/, uint8* DataPtr /*= nullptr*/)
+EOverriddenPropertyOperation FOverridableSerializationLogic::GetOverriddenPropertyOperation(const FArchive& Ar, FProperty* Property /*= nullptr*/, uint8* DataPtr /*= nullptr*/, uint8* DefaultValue)
 {
 	checkf(bUseOverridableSerialization, TEXT("Nobody should use this method if it is not setup to use overridable serialization"));
 
@@ -40,6 +40,17 @@ EOverriddenPropertyOperation FOverridableSerializationLogic::GetOverriddenProper
 		{
 			return EOverriddenPropertyOperation::Replace;
 		}
+
+		// In the case of a CDO owning default value, we might need to serialize it to keep its value.
+		if (OverriddenProperties && OverriddenProperties->IsCDOOwningProperty(*CurrentProperty))
+		{
+			// Only need serialize this value if it is different from the default property value
+			if (!CurrentProperty->Identical(DataPtr, DefaultValue, Ar.GetPortFlags()))
+			{
+				return 	EOverriddenPropertyOperation::Replace;
+			}
+		}
+
 		// Here we should just use CurrentProperty->ContainsInstancedObjectProperty() but somehow this CPF_InstanceReference seems to always be there on verse classes
 		// This should probably be fix in verse at some point.
 		else if (CurrentProperty->HasAnyPropertyFlags(CPF_PersistentInstance))
@@ -877,6 +888,20 @@ FOverriddenPropertyNode* FOverriddenPropertySet::SetSubPropertyOperation(EOverri
 	return &OverriddenPropertyNode;
 }
 
+bool FOverriddenPropertySet::IsCDOOwningProperty(const FProperty& Property) const
+{
+	checkf(IsValid(Owner), TEXT("Expecting a valid overridable owner"));
+	if (!Owner->HasAnyFlags(RF_ClassDefaultObject))
+	{
+		return false;
+	}
+
+	// We need to serialize only if the property owner is the current CDO class
+	// Otherwise on derived class, this is done in parent CDO or it should be explicitly overridden if it is different than the parent value
+	// This is sort of like saying it overrides the default property initialization value.
+	return Property.GetOwnerClass() == Owner->GetClass();
+}
+
 void FOverriddenPropertySet::Reset()
 {
 	OverriddenPropertyNodes.Reset();
diff --git a/Engine/Source/Runtime/CoreUObject/Public/UObject/OverriddenPropertySet.h b/Engine/Source/Runtime/CoreUObject/Public/UObject/OverriddenPropertySet.h
--- a/Engine/Source/Runtime/CoreUObject/Public/UObject/OverriddenPropertySet.h
+++ b/Engine/Source/Runtime/CoreUObject/Public/UObject/OverriddenPropertySet.h
@@ -65,8 +65,9 @@ public:
 	 * @param Ar currently being used to serialize the current object (will be used to retrieve the current property serialized path)
 	 * @param Property the property about to be serialized, can be null
 	 * @param DataPtr to the memory of that property
+	 * @param DefaultValue memory pointer of that property
 	 * @return the overridden property operation */
-	static EOverriddenPropertyOperation GetOverriddenPropertyOperation(const FArchive& Ar, FProperty* Property = nullptr, uint8* DataPtr = nullptr);
+	static EOverriddenPropertyOperation GetOverriddenPropertyOperation(const FArchive& Ar, FProperty* Property = nullptr, uint8* DataPtr = nullptr, uint8* DefaultValue = nullptr);
 
 private:
 	static thread_local bool bUseOverridableSerialization;
@@ -269,6 +270,14 @@ public:
 	 * @return the node to the sub property */
 	FOverriddenPropertyNode* SetSubPropertyOperation(EOverriddenPropertyOperation Operation, FOverriddenPropertyNode& Node, FOverriddenPropertyNodeID NodeID);
 
+	/**
+	 * Check if this is an overridden property set of a CDO and that this property is owned by the class of this CDO
+	 * NOTE: this is used to know if a property should be serialized to keep its default CDO value.
+	 * @param Property 
+	 * @return 
+	 */
+	bool IsCDOOwningProperty(const FProperty& Property) const;
+
 	/**
 	 * Resets all overrides of the object */
 	void Reset();
