Adding support for multi-line entries on config set properties.
Duplicate values are filtered out during load.
Existing config values for set properties should continue to be loaded the same way as before.

Example syntax:
SetProperty=("Value1", "Value2")
+SetProperty="Value3"
+SetProperty=("Value4", "Value5")

diff --git a/Engine/Source/Runtime/CoreUObject/Private/UObject/Obj.cpp b/Engine/Source/Runtime/CoreUObject/Private/UObject/Obj.cpp
--- a/Engine/Source/Runtime/CoreUObject/Private/UObject/Obj.cpp
+++ b/Engine/Source/Runtime/CoreUObject/Private/UObject/Obj.cpp
@@ -2979,9 +2979,181 @@ void UObject::LoadConfig( UClass* ConfigClass/*=NULL*/, const TCHAR* InFilename/
 		const bool bIsPerPlatformConfig = GetClass()->HasAnyClassFlags(CLASS_PerPlatformConfig);
 #endif // #if WITH_EDITOR
 
+		// Track if we loaded this config value using special handling (e.g. array or set)
+		bool bProcessedProperty = false;
+
 		UE_LOG(LogConfig, Verbose, TEXT("   Loading value for %s from [%s]"), *Key, *ClassSection);
-		FArrayProperty* Array = CastField<FArrayProperty>( Property );
-		if( Array == NULL )
+
+		FArrayProperty* Array = CastField<FArrayProperty>(Property);
+		FSetProperty* SetProperty = CastField<FSetProperty>(Property);
+
+		if (Array || SetProperty)
+		{
+			const FConfigSection* Sec = GetConfigSection(*ClassSection, *PropFileName);
+			if (!Sec && bPerObject && ClassPathSection.Len())
+			{
+				Sec = GetConfigSection(*ClassPathSection, *PropFileName);
+			}
+
+#if !UE_BUILD_SHIPPING
+			if (!Sec && !FPlatformProperties::RequiresCookedData())
+			{
+				CheckMissingSection(ClassSection, PropFileName);
+			}
+#endif
+
+			if (Array)
+			{
+				FScriptArrayHelper_InContainer ArrayHelper(Array, this);
+
+				bProcessedProperty = true;
+#if WITH_EDITOR
+				// Empty out any array properties if this is a PerPlatformConfig class
+				// as we are replacing the values with the Platform's version when entering
+				// a new preview platform.
+				if (bIsPerPlatformConfig)
+				{
+					ArrayHelper.EmptyValues();
+				}
+#endif
+				if (Sec)
+				{
+					TArray<FConfigValue> List;
+					const FName KeyName(*Key, FNAME_Find);
+					Sec->MultiFind(KeyName, List);
+
+					const int32 Size = Array->Inner->ElementSize;
+
+					// Only override default properties if there is something to override them with.
+					if (!List.IsEmpty())
+					{
+						ArrayHelper.EmptyAndAddValues(List.Num());
+						for (int32 i = List.Num() - 1, c = 0; i >= 0; i--, c++)
+						{
+							Array->Inner->ImportText_Direct(*List[i].GetValue(), ArrayHelper.GetRawPtr(c), this, PortFlags);
+						}
+					}
+					else
+					{
+						int32 Index = 0;
+						const FConfigValue* ElementValue = nullptr;
+						do
+						{
+							// Add array index number to end of key
+							FString IndexedKey = FString::Printf(TEXT("%s[%i]"), *Key, Index);
+
+							// Try to find value of key
+							const FName IndexedName(*IndexedKey, FNAME_Find);
+							if (IndexedName == NAME_None)
+							{
+								break;
+							}
+							ElementValue = Sec->Find(IndexedName);
+
+							// If found, import the element
+							if (ElementValue != nullptr)
+							{
+								// expand the array if necessary so that Index is a valid element
+								ArrayHelper.ExpandForIndex(Index);
+								Array->Inner->ImportText_Direct(*ElementValue->GetValue(), ArrayHelper.GetRawPtr(Index), this, PortFlags);
+							}
+
+							Index++;
+						} while (ElementValue || Index < ArrayHelper.Num());
+					}
+				}
+			}
+			else if (SetProperty)
+			{
+				FScriptSetHelper_InContainer SetHelper(SetProperty, this);
+
+#if WITH_EDITOR
+				// Empty out any set properties if this is a PerPlatformConfig class
+				// as we are replacing the values with the Platform's version when entering
+				// a new preview platform.
+				if (bIsPerPlatformConfig)
+				{
+					SetHelper.EmptyElements();
+				}
+#endif
+
+				if (Sec)
+				{
+					TArray<FConfigValue> List;
+					const FName KeyName(*Key, FNAME_Find);
+					Sec->MultiFind(KeyName, List);
+
+					bool bSingleSetEntry = false;
+					if (List.Num() == 1)
+					{
+						const FString& SingleListValue = List[0].GetValue();
+						if (SingleListValue.Len() > 1
+							&& SingleListValue[0] == TEXT('(')
+							&& SingleListValue[SingleListValue.Len() - 1] == TEXT(')'))
+						{
+							// If we have a single entry in the set that is surrounded with parentheses, fall back to
+							// the old processing method
+							bSingleSetEntry = true;
+						}
+					}
+
+					// Only override default properties if there is something to override them with.
+					if (!bSingleSetEntry && !List.IsEmpty())
+					{
+						bProcessedProperty = true;
+						SetHelper.EmptyElements(List.Num());
+
+						// Each config value entry can possibly specify multiple set elements - create a temporary
+						// set here that we can import to and add its elements to the object's property
+						void* TempSet = FMemory::Malloc(SetProperty->GetSize(), SetProperty->GetMinAlignment());
+						SetProperty->InitializeValue(TempSet);
+
+						// Importing elements to the set inline does not check for duplicates - create a temporary
+						// element that we can import to add uniquely to the set
+						void* TempElement = FMemory::Malloc(SetProperty->ElementProp->GetSize(),
+							SetProperty->ElementProp->GetMinAlignment());
+						SetProperty->ElementProp->InitializeValue(TempElement);
+
+						for (const FConfigValue& ListValue : List)
+						{
+							const FString& ListString = ListValue.GetValue();
+
+							// Try to import the config value as an entire set first
+							const TCHAR* SetImportResult =
+								SetProperty->ImportText_Direct(*ListString, TempSet, this, PortFlags);
+
+							if (SetImportResult && SetImportResult != *ListString)
+							{
+								FScriptSetHelper TempSetHelper(SetProperty, TempSet);
+								for (FScriptSetHelper::FIterator Itr = TempSetHelper.CreateIterator(); Itr; ++Itr)
+								{
+									SetHelper.AddElement(TempSetHelper.GetElementPtr(Itr));
+								}
+							}
+							else
+							{
+								// If we failed to import the value as an entire set, try to import it as a single
+								// element
+								const TCHAR* ElementImportResult = 
+									SetProperty->ElementProp->ImportText_Direct(*ListString, TempElement, this, PortFlags);
+
+								if (ElementImportResult && ElementImportResult != *ListString)
+								{
+									SetHelper.AddElement(TempElement);
+								}
+							}
+
+						}
+
+						SetProperty->DestroyAndFreeValue(TempSet);
+						SetProperty->ElementProp->DestroyAndFreeValue(TempElement);
+						SetHelper.Rehash();
+					}
+				}
+			}
+		}
+		
+		if (!bProcessedProperty)
 		{
 			for( int32 i=0; i<Property->ArrayDim; i++ )
 			{
@@ -3015,78 +3187,6 @@ void UObject::LoadConfig( UClass* ConfigClass/*=NULL*/, const TCHAR* InFilename/
 #endif
 			}
 		}
-		else
-		{
-			FScriptArrayHelper_InContainer ArrayHelper(Array, this);
-
-#if WITH_EDITOR
-			// Empty out any array properties if this is a PerPlatformConfig class
-			// as we are replacing the values with the Platform's version when entering
-			// a new preview platform.
-			if (bIsPerPlatformConfig)
-			{
-				ArrayHelper.EmptyValues();
-			}
-#endif
-
-			const FConfigSection* Sec = GetConfigSection(*ClassSection, *PropFileName);
-			if (!Sec && bPerObject && ClassPathSection.Len())
-			{
-				Sec = GetConfigSection(*ClassPathSection, *PropFileName);
-			}
-			if( Sec )
-			{
-				TArray<FConfigValue> List;
-				const FName KeyName(*Key, FNAME_Find);
-				Sec->MultiFind(KeyName,List);
-
-				const int32 Size = Array->Inner->ElementSize;
-
-				// Only override default properties if there is something to override them with.
-				if ( List.Num() > 0 )
-				{
-					ArrayHelper.EmptyAndAddValues(List.Num());
-					for( int32 i=List.Num()-1,c=0; i>=0; i--,c++ )
-					{
-						Array->Inner->ImportText_Direct( *List[i].GetValue(), ArrayHelper.GetRawPtr(c), this, PortFlags );
-					}
-				}
-				else
-				{
-					int32 Index = 0;
-					const FConfigValue* ElementValue = nullptr;
-					do
-					{
-						// Add array index number to end of key
-						FString IndexedKey = FString::Printf(TEXT("%s[%i]"), *Key, Index);
-
-						// Try to find value of key
-						const FName IndexedName(*IndexedKey,FNAME_Find);
-						if (IndexedName == NAME_None)
-						{
-							break;
-						}
-						ElementValue = Sec->Find(IndexedName);
-
-						// If found, import the element
-						if ( ElementValue != nullptr )
-						{
-							// expand the array if necessary so that Index is a valid element
-							ArrayHelper.ExpandForIndex(Index);
-							Array->Inner->ImportText_Direct(*ElementValue->GetValue(), ArrayHelper.GetRawPtr(Index), this, PortFlags);
-						}
-
-						Index++;
-					} while( ElementValue || Index < ArrayHelper.Num() );
-				}
-			}
-#if !UE_BUILD_SHIPPING
-			else if (!FPlatformProperties::RequiresCookedData())
-			{
-				CheckMissingSection(ClassSection, PropFileName);
-			}
-#endif
-		}
 	}
 
 	// if we are reloading config data after the initial class load, fire the callback now
@@ -3197,7 +3297,8 @@ void UObject::SaveConfig( uint64 Flags, const TCHAR* InFilename, FConfigCacheIni
 			UObject* SuperClassDefaultObject = GetClass()->GetSuperClass()->GetDefaultObject();
 
 			FArrayProperty* Array   = CastField<FArrayProperty>( Property );
-			if( Array )
+			FSetProperty* SetProperty = CastField<FSetProperty>(Property);
+			if (Array || SetProperty)
 			{
 				const FConfigSection* Sec = Config->GetSection(*Section, 1, PropFileName);
 				// Default ini's require the array syntax to be applied to the property name
@@ -3211,17 +3312,36 @@ void UObject::SaveConfig( uint64 Flags, const TCHAR* InFilename, FConfigCacheIni
 				if (!bPropDeprecated && (!bShouldCheckIfIdenticalBeforeAdding || !Property->Identical_InContainer(this, SuperClassDefaultObject)))
 				{
 					check(Sec);
-					FScriptArrayHelper_InContainer ArrayHelper(Array, this);
-					for( int32 i=0; i<ArrayHelper.Num(); i++ )
+
+					if (Array)
 					{
-						FString	Buffer;
-						Array->Inner->ExportTextItem_Direct( Buffer, ArrayHelper.GetRawPtr(i), ArrayHelper.GetRawPtr(i), this, PortFlags );
-						Config->AddToSection(*Section, *CompleteKey, *Buffer, PropFileName);
+						FScriptArrayHelper_InContainer ArrayHelper(Array, this);
+						for( int32 i=0; i<ArrayHelper.Num(); i++ )
+						{
+							FString	Buffer;
+							Array->Inner->ExportTextItem_Direct( Buffer, ArrayHelper.GetRawPtr(i), ArrayHelper.GetRawPtr(i), this, PortFlags );
+							Config->AddToSection(*Section, *CompleteKey, *Buffer, PropFileName);
+						}
+						if (ArrayHelper.Num() == 0 && bIsADefaultIniWrite)
+						{
+							const FString EmptyKey = FString::Printf(TEXT("!%s"), *Key);
+							Config->AddToSection(*Section, *EmptyKey, TEXT("__ClearArray__"), PropFileName);
+						}
 					}
-					if (ArrayHelper.Num() == 0 && bIsADefaultIniWrite)
+					else if (SetProperty)
 					{
-						const FString EmptyKey = FString::Printf(TEXT("!%s"), *Key);
-						Config->AddToSection(*Section, *EmptyKey, TEXT("__ClearArray__"), PropFileName);
+						FScriptSetHelper_InContainer SetHelper(SetProperty, this);
+						for (int32 i = 0; i < SetHelper.Num(); i++)
+						{
+							FString	Buffer;
+							SetProperty->ElementProp->ExportTextItem_Direct(Buffer, SetHelper.GetElementPtr(i), SetHelper.GetElementPtr(i), this, PortFlags);
+							Config->AddToSection(*Section, *CompleteKey, *Buffer, PropFileName);
+						}
+						if (SetHelper.Num() == 0 && bIsADefaultIniWrite)
+						{
+							const FString EmptyKey = FString::Printf(TEXT("!%s"), *Key);
+							Config->AddToSection(*Section, *EmptyKey, TEXT("__ClearSet__"), PropFileName);
+						}
 					}
 				}
 			}
