[BugFix] CDOs weren't being copied when InstanceDataObjects were generated
[BugFix] mark properties as overridden when changed by fixup tool
[Temp] gave IDO properties without categories a category as a workaround to a details customization bug

diff --git a/Engine/Plugins/Editor/InstanceDataObjectFixupTool/Source/Private/InstanceDataObjectFixupPanel.cpp b/Engine/Plugins/Editor/InstanceDataObjectFixupTool/Source/Private/InstanceDataObjectFixupPanel.cpp
--- a/Engine/Plugins/Editor/InstanceDataObjectFixupTool/Source/Private/InstanceDataObjectFixupPanel.cpp
+++ b/Engine/Plugins/Editor/InstanceDataObjectFixupTool/Source/Private/InstanceDataObjectFixupPanel.cpp
@@ -11,6 +11,9 @@
 #include "Editor.h"
 #include "UObject/PropertyBagRepository.h"
 
+#include "UObject/OverriddenPropertySet.h"
+#include "UObject/OverridableManager.h"
+
 #define LOCTEXT_NAMESPACE "InstanceDataObjectFixupPanel"
 
 FRedirectedPropertyNode::FRedirectedPropertyNode(const FRedirectedPropertyNode& Other)
@@ -382,6 +385,29 @@ static void* ResolvePath(const FPropertyPath& Path, void* Value)
 	return Value;
 }
 
+static FPropertyChangedEvent ConstructChangeEventForRedirect(const FPropertyPath& Path, FEditPropertyChain& OutChain, TMap<FString, int32>& OutArrayIndices)
+{
+	FPropertyChangedEvent OutEvent = FPropertyChangedEvent(Path.GetLeafMostProperty().Property.Get(), EPropertyChangeType::ValueSet);
+	for (int32 I = 0; I < Path.GetNumProperties(); ++I)
+	{
+		const FPropertyInfo& Info = Path.GetPropertyInfo(I);
+		OutChain.AddTail(Info.Property.Get()); // only the head is used in OverrideProperty
+		if (Info.ArrayIndex != INDEX_NONE)
+		{
+			OutArrayIndices.Add(Info.Property->GetName(), Info.ArrayIndex);
+		}
+		if (Info.Property->IsA<FArrayProperty>() || Info.Property->IsA<FSetProperty>() || Info.Property->IsA<FMapProperty>())
+		{
+			if (++I < Path.GetNumProperties())
+			{
+				OutArrayIndices.Add(Info.Property->GetName(), Path.GetPropertyInfo(I).ArrayIndex);
+			}
+		}
+	}
+	OutEvent.SetArrayIndexPerObject(MakeArrayView(&OutArrayIndices, 1));
+	return OutEvent;
+}
+
 void FInstanceDataObjectFixupPanel::RedirectProperty(const FPropertyPath& From, const FPropertyPath& To)
 {
 	UInstanceDataObjectFixupUndoHandler* Snapshot = NewObject<UInstanceDataObjectFixupUndoHandler>();
@@ -473,7 +499,12 @@ void FInstanceDataObjectFixupPanel::RedirectProperty(const FPropertyPath& From,
 			continue;
 		}
 	
-		FPropertyChangedEvent ChangeEvent(To.GetRootProperty().Property.Get(), EPropertyChangeType::ValueSet);
+
+		// construct change event
+		FEditPropertyChain Chain;
+		TMap<FString, int32> ArrayIndices;
+		FPropertyChangedEvent ChangeEvent = ConstructChangeEventForRedirect(To, Chain, ArrayIndices);
+		FOverridableManager::Get().PreOverrideProperty(*Instance, Chain);
 		Instance->PreEditChange(ChangeEvent.Property);
 
 		if (ToRevertInfo)
@@ -503,6 +534,7 @@ void FInstanceDataObjectFixupPanel::RedirectProperty(const FPropertyPath& From,
 			FromRevertInfoItr += DestinationProperty->ArrayDim * DestinationProperty->ElementSize;
 		}
 		Instance->PostEditChangeProperty(ChangeEvent);
+		FOverridableManager::Get().PostOverrideProperty(*Instance, ChangeEvent, Chain);
 	}
 
 	GEditor->EndTransaction();
diff --git a/Engine/Source/Runtime/CoreUObject/Private/UObject/InstanceDataObjectUtils.cpp b/Engine/Source/Runtime/CoreUObject/Private/UObject/InstanceDataObjectUtils.cpp
--- a/Engine/Source/Runtime/CoreUObject/Private/UObject/InstanceDataObjectUtils.cpp
+++ b/Engine/Source/Runtime/CoreUObject/Private/UObject/InstanceDataObjectUtils.cpp
@@ -290,6 +290,13 @@ namespace UE
 			ConvertToInstanceDataObjectProperty(AsMapProperty->ValueProp, Outer, LooseProperties, Path);
 			Path.Pop();
 		}
+		
+#if WITH_EDITORONLY_DATA
+		if (!Property->HasMetaData(TEXT("category")))
+		{
+			Property->SetMetaData(TEXT("category"), TEXT("Verse"));
+		}
+#endif
 	}
 	
 	// copy template property then convert it into an InstanceDataObject property by adding loose properties
@@ -425,6 +432,13 @@ namespace UE
 				}
 				Path.Push(CreateSegmentFromProperty(TemplateProperty));
 				FProperty* SuperProperty = CreateInstanceDataObjectProperty(TemplateProperty, Super, LooseProperties, Path);
+				
+#if WITH_EDITORONLY_DATA
+				if (!SuperProperty->HasMetaData(TEXT("category")))
+				{
+					SuperProperty->SetMetaData(TEXT("category"), TEXT("Verse"));
+				}
+#endif
 				Path.Pop();
 				SuperProperties.Add(SuperProperty);
 			}
@@ -487,12 +501,35 @@ namespace UE
 		Result->StaticLink(/*RelinkExistingProperties*/true);
 		return Result;
 	}
+
+	void CopyCDO(const UStruct* SourceStruct, const void* SourceData, const UStruct* DestinationStruct, void* DestinationData)
+	{
+		for (const FProperty* SourceProperty : TFieldRange<FProperty>(SourceStruct))
+		{
+			if (const FProperty* DestinationProperty = DestinationStruct->FindPropertyByName(SourceProperty->GetFName()))
+			{
+				FString ValueText;
+				const void* SourceValue = SourceProperty->ContainerPtrToValuePtr<void>(SourceData);
+				void* DestinationValue = DestinationProperty->ContainerPtrToValuePtr<void>(DestinationData);
+				SourceProperty->ExportText_Direct(ValueText, SourceValue, SourceValue, nullptr, PPF_None);
+				DestinationProperty->ImportText_Direct(*ValueText, DestinationValue, nullptr, PPF_None);
+			}
+		}
+	}
 	
-	UClass* CreateInstanceDataObjectClass(const FPropertyBag* PropertyBag, UStruct* TemplateStruct, UObject* Outer)
+	UClass* CreateInstanceDataObjectClass(const FPropertyBag* PropertyBag, UClass* OwnerClass, UObject* Outer)
 	{
 		const TMap<FWildcardPropertyPathName, TMap<FName, const FProperty*>> LooseProperties = GetWildcardedLooseProperties(PropertyBag);
 		FWildcardPropertyPathName ParentPath;
-		return CreateInstanceDataObjectStructRec<UClass>(TemplateStruct, Outer, LooseProperties, ParentPath);
+		UClass* Result = CreateInstanceDataObjectStructRec<UClass>(OwnerClass, Outer, LooseProperties, ParentPath);
+
+		const UObject* OwnerCDO = OwnerClass->GetDefaultObject(true);
+		UObject* ResultCDO = Result->GetDefaultObject(true);
+		if (ensure(OwnerCDO && ResultCDO))
+		{
+			CopyCDO(OwnerClass, OwnerCDO, Result, ResultCDO);
+		}
+		return Result;
 	}
 
 	static void MarkPropertySetBySerialization(const UStruct* Struct, const void* StructData, const void* PropertyDataPtr)
diff --git a/Engine/Source/Runtime/CoreUObject/Private/UObject/InstanceDataObjectUtils.h b/Engine/Source/Runtime/CoreUObject/Private/UObject/InstanceDataObjectUtils.h
--- a/Engine/Source/Runtime/CoreUObject/Private/UObject/InstanceDataObjectUtils.h
+++ b/Engine/Source/Runtime/CoreUObject/Private/UObject/InstanceDataObjectUtils.h
@@ -14,7 +14,7 @@ namespace UE
 	class FPropertyPathName;
 
 	// generate a UClass that unions the properties of PropertyBag and TemplateStruct
-	UClass* CreateInstanceDataObjectClass(const FPropertyBag* PropertyBag, UStruct* TemplateStruct, UObject* Outer);
+	UClass* CreateInstanceDataObjectClass(const FPropertyBag* PropertyBag, UClass* OwnerClass, UObject* Outer);
 
 	// notify that a property in an object was set when the object was deserialized
 	void MarkPropertySetBySerialization(UObject* Object, const FPropertyPathName& Path);
