GC stats for analytics.

Added GarbageCollectComplete event that gets fired after GC is truly done (existing events were inconsitent in that manner)

diff --git a/Engine/Source/Runtime/CoreUObject/Private/UObject/GarbageCollection.cpp b/Engine/Source/Runtime/CoreUObject/Private/UObject/GarbageCollection.cpp
--- a/Engine/Source/Runtime/CoreUObject/Private/UObject/GarbageCollection.cpp
+++ b/Engine/Source/Runtime/CoreUObject/Private/UObject/GarbageCollection.cpp
@@ -103,6 +103,8 @@ static TArray<FUObjectItem*> GUnreachableObjects;
 static FCriticalSection GUnreachableObjectsCritical;
 static int32 GUnrechableObjectIndex = 0;
 
+static UE::GC::Private::FStats GGCStats;
+
 struct FGCTimingInfo
 {
 	double LastGCTime{};
@@ -327,10 +329,22 @@ static FAutoConsoleVariableRef CVarMaxFinishDestroyTimeoutObjectsToLog (
 	ECVF_Default
 );
 
+#if UE_BUILD_SHIPPING
+static constexpr int32 GDumpGCAnalyticsToLog = 0;
+#else
+static int32 GDumpGCAnalyticsToLog = 0;
+static FAutoConsoleVariableRef CVarDumpGCAnalyticsToLog(
+	TEXT("gc.DumpAnalyticsToLog"),
+	GDumpGCAnalyticsToLog,
+	TEXT("Dumps Garbage Collection analytics to log at the end of each GC."),
+	ECVF_Default
+);
+#endif
+
 namespace UE::GC
 {
 
-	struct FStats
+	struct FDetailedClassStats
 	{
 #if PERF_DETAILED_PER_CLASS_GC_STATS
 		/** Map from a UClass' FName to the number of objects that were purged during the last purge phase of this class.	*/
@@ -379,7 +393,7 @@ namespace UE::GC
 		void IncPurgeCount(UObject* Object);
 		void LogPurgeStats(int32 Total);
 	};
-	static FStats GStats;
+	static FDetailedClassStats GDetailedStats;
 
 	/**
 	 * Helper function to log the various class to count info maps.
@@ -389,7 +403,7 @@ namespace UE::GC
 	 * @param	NumItemsToList		Number of items to log
 	 * @param	TotalCount			Total count, if 0 will be calculated
 	 */
-	void FStats::LogClassCountInfo( const TCHAR* LogText, TMap<const FName, uint64>& ClassToCountMap, int32 NumItemsToLog, uint64 TotalCount )
+	void FDetailedClassStats::LogClassCountInfo( const TCHAR* LogText, TMap<const FName, uint64>& ClassToCountMap, int32 NumItemsToLog, uint64 TotalCount )
 	{
 #if PERF_DETAILED_PER_CLASS_GC_STATS
 		// Figure out whether we need to calculate the total count.
@@ -416,7 +430,7 @@ namespace UE::GC
 #endif
 	};
 
-	FORCEINLINE void FStats::IncClusterToObjectRefs(FUObjectItem* ObjectItem)
+	FORCEINLINE void FDetailedClassStats::IncClusterToObjectRefs(FUObjectItem* ObjectItem)
 	{
 #if PERF_DETAILED_PER_CLASS_GC_STATS
 		NumClusterToObjectRefs++;
@@ -427,28 +441,28 @@ namespace UE::GC
 #endif
 	}
 
-	FORCEINLINE void FStats::IncClusterToClusterRefs(int32 Count)
+	FORCEINLINE void FDetailedClassStats::IncClusterToClusterRefs(int32 Count)
 	{
 #if PERF_DETAILED_PER_CLASS_GC_STATS
 		NumClusterToClusterRefs += Count;
 #endif
 	}
 
-	FORCEINLINE void FStats::IncNumClustersTraversed()
+	FORCEINLINE void FDetailedClassStats::IncNumClustersTraversed()
 	{
 #if PERF_DETAILED_PER_CLASS_GC_STATS
 		NumClusterToObjectRefs++;
 #endif
 	}
 
-	void FStats::BeginTimingObject(UObject* CurrentObject)
+	void FDetailedClassStats::BeginTimingObject(UObject* CurrentObject)
 	{
 #if PERF_DETAILED_PER_CLASS_GC_STATS
 		CurrentObjectStartCyles = FPlatformTime::Cycles64();
 #endif
 	}
 
-	void FStats::UpdateDetailedStats(UObject* CurrentObject)
+	void FDetailedClassStats::UpdateDetailedStats(UObject* CurrentObject)
 	{
 #if PERF_DETAILED_PER_CLASS_GC_STATS
 		uint64 DeltaCycles = FPlatformTime::Cycles64() - CurrentObjectStartCyles;
@@ -471,7 +485,7 @@ namespace UE::GC
 #endif
 	}
 
-	void FStats::LogDetailedStatsSummary()
+	void FDetailedClassStats::LogDetailedStatsSummary()
 	{
 #if PERF_DETAILED_PER_CLASS_GC_STATS
 		LogClassCountInfo(TEXT("references to regular objects from"), ClassToRegularObjectRefsMap, 20, 0);
@@ -493,7 +507,7 @@ namespace UE::GC
 #endif
 	}
 
-	FORCEINLINE void FStats::IncreaseObjectRefStats(UObject* RefToObject)
+	FORCEINLINE void FDetailedClassStats::IncreaseObjectRefStats(UObject* RefToObject)
 	{
 #if PERF_DETAILED_PER_CLASS_GC_STATS
 		if (!RefToObject)
@@ -521,7 +535,7 @@ namespace UE::GC
 #endif
 	}
 
-	FORCEINLINE void FStats::IncPurgeCount(UObject* Object)
+	FORCEINLINE void FDetailedClassStats::IncPurgeCount(UObject* Object)
 	{
 #if PERF_DETAILED_PER_CLASS_GC_STATS
 		// Keep track of how many objects of a certain class we're purging.
@@ -531,7 +545,7 @@ namespace UE::GC
 #endif
 	}
 
-	FORCEINLINE void FStats::LogPurgeStats(int32 Total)
+	FORCEINLINE void FDetailedClassStats::LogPurgeStats(int32 Total)
 	{
 #if PERF_DETAILED_PER_CLASS_GC_STATS
 		LogClassCountInfo(TEXT("objects of"), ClassToPurgeCountMap, 10, Total);
@@ -1076,7 +1090,7 @@ static bool MarkClusterMutableObjectsAsReachable(FUObjectCluster& Cluster, Conta
 		if (ReferencedMutableObjectIndex >= 0) // Garbage Elimination support
 		{
 			FUObjectItem* ReferencedMutableObjectItem = GUObjectArray.IndexToObjectUnsafeForGC(ReferencedMutableObjectIndex);
-			UE::GC::GStats.IncClusterToObjectRefs(ReferencedMutableObjectItem);
+			UE::GC::GDetailedStats.IncClusterToObjectRefs(ReferencedMutableObjectItem);
 			if (!ReferencedMutableObjectItem->HasAnyFlags(EInternalObjectFlags::Garbage))
 			{
 				if (ReferencedMutableObjectItem->IsMaybeUnreachable())
@@ -1145,7 +1159,7 @@ FORCEINLINE static void MarkClusterMutableCellsAsReachable(FUObjectCluster& Clus
 template<EGCOptions Options, class ContainerType>
 static FORCENOINLINE void MarkReferencedClustersAsReachable(int32 ClusterIndex, ContainerType& ObjectsToSerialize)
 {
-	UE::GC::GStats.IncNumClustersTraversed();
+	UE::GC::GDetailedStats.IncNumClustersTraversed();
 
 	// If we run across some Garbage objects we need to add all objects from this cluster
 	// to ObjectsToSerialize so that we can properly null out all the references.
@@ -1153,7 +1167,7 @@ static FORCENOINLINE void MarkReferencedClustersAsReachable(int32 ClusterIndex,
 
 	bool bAddClusterObjectsToSerialize = false;
 	FUObjectCluster& Cluster = GUObjectClusters[ClusterIndex];
-	UE::GC::GStats.IncClusterToClusterRefs(Cluster.ReferencedClusters.Num());
+	UE::GC::GDetailedStats.IncClusterToClusterRefs(Cluster.ReferencedClusters.Num());
 	// Also mark all referenced objects from outside of the cluster as reachable
 	for (int32& ReferncedClusterIndex : Cluster.ReferencedClusters)
 	{
@@ -2908,17 +2922,17 @@ class TReachabilityProcessor
 public:
 	FORCEINLINE void BeginTimingObject(UObject* CurrentObject)
 	{
-		UE::GC::GStats.BeginTimingObject(CurrentObject);
+		UE::GC::GDetailedStats.BeginTimingObject(CurrentObject);
 	}
 
 	FORCEINLINE void UpdateDetailedStats(UObject* CurrentObject)
 	{
-		UE::GC::GStats.UpdateDetailedStats(CurrentObject);
+		UE::GC::GDetailedStats.UpdateDetailedStats(CurrentObject);
 	}
 
 	FORCEINLINE void LogDetailedStatsSummary()
 	{
-		UE::GC::GStats.LogDetailedStatsSummary();
+		UE::GC::GDetailedStats.LogDetailedStatsSummary();
 	}
 
 	FORCEINLINE bool IsTimeLimitExceeded() const
@@ -2957,7 +2971,7 @@ public:
 	template<EKillable Killable>
 	static FORCEINLINE_DEBUGGABLE void ProcessReferenceDirectly(FWorkerContext& Context, FPermanentObjectPoolExtents PermanentPool, const UObject* ReferencingObject, UObject*& Object, FMemberId MemberId)
 	{
-		UE::GC::GStats.IncreaseObjectRefStats(Object);
+		UE::GC::GDetailedStats.IncreaseObjectRefStats(Object);
 		if (ValidateReference(Object, PermanentPool, FReferenceToken(ReferencingObject), MemberId))
 		{
 			const int32 ObjectIndex = GUObjectArray.ObjectToIndex(Object);
@@ -2979,7 +2993,7 @@ public:
 
 	FORCEINLINE static void HandleBatchedReference(FWorkerContext& Context, FResolvedMutableReference Reference, FReferenceMetadata Metadata)
 	{
-		UE::GC::GStats.IncreaseObjectRefStats(GetObject(Reference));
+		UE::GC::GDetailedStats.IncreaseObjectRefStats(GetObject(Reference));
 		if (Metadata.Has(KillFlag))
 		{
 			checkSlow(Metadata.ObjectItem->GetOwnerIndex() <= 0);
@@ -2993,7 +3007,7 @@ public:
 
 	FORCEINLINE static void HandleBatchedReference(FWorkerContext& Context, FImmutableReference Reference, FReferenceMetadata Metadata)
 	{
-		UE::GC::GStats.IncreaseObjectRefStats(GetObject(Reference));
+		UE::GC::GDetailedStats.IncreaseObjectRefStats(GetObject(Reference));
 		DetectGarbageReference(Context, Metadata);
 		HandleValidReference(Context, Reference, Metadata);
 	}
@@ -3433,17 +3447,17 @@ class TDebugReachabilityProcessor
 public:
 	FORCEINLINE void BeginTimingObject(UObject* CurrentObject)
 	{
-		UE::GC::GStats.BeginTimingObject(CurrentObject);
+		UE::GC::GDetailedStats.BeginTimingObject(CurrentObject);
 	}
 
 	FORCEINLINE void UpdateDetailedStats(UObject* CurrentObject)
 	{
-		UE::GC::GStats.UpdateDetailedStats(CurrentObject);
+		UE::GC::GDetailedStats.UpdateDetailedStats(CurrentObject);
 	}
 
 	FORCEINLINE void LogDetailedStatsSummary()
 	{
-		UE::GC::GStats.LogDetailedStatsSummary();
+		UE::GC::GDetailedStats.LogDetailedStatsSummary();
 	}
 
 	FORCEINLINE bool IsTimeLimitExceeded() const
@@ -3471,7 +3485,7 @@ public:
 
 	FORCENOINLINE void HandleTokenStreamObjectReference(FWorkerContext& Context, const UObject* ReferencingObject, UObject*& Object, FMemberId MemberId, EOrigin Origin, bool bAllowReferenceElimination)
 	{
-		UE::GC::GStats.IncreaseObjectRefStats(Object);
+		UE::GC::GDetailedStats.IncreaseObjectRefStats(Object);
 		if (ValidateReference(Object, PermanentPool, FReferenceToken(ReferencingObject), MemberId))
 		{
 			FReferenceMetadata Metadata(GUObjectArray.ObjectToIndex(Object));
@@ -4081,6 +4095,8 @@ public:
 				FUObjectCluster& Cluster = ClusterArray[ClusterIndex];
 				if (Cluster.RootIndex >= 0)
 				{
+					ThisThreadClusteredObjects += Cluster.Objects.Num();
+
 					FUObjectItem* RootItem = &GUObjectArray.GetObjectItemArrayUnsafe()[Cluster.RootIndex];
 					if (!RootItem->IsGarbage())
 					{
@@ -4134,6 +4150,8 @@ public:
 			// Make sure all referenced clusters are marked as reachable too
 			MarkReferencedClustersAsReachable<EGCOptions::None>(ObjectItem->GetClusterIndex(), OutRootObjects);
 		}
+
+		GGCStats.NumClusteredObjects = TotalClusteredObjects.load(std::memory_order_acquire);
 	}
 
 	FORCENOINLINE void MarkRootObjectsAsReachable(const EGatherOptions Options, const EObjectFlags KeepFlags, TArray<UObject*>& OutRootObjects)
@@ -4211,6 +4229,8 @@ public:
 		OutRootObjects.Reserve(OutRootObjects.Num() + MarkRootsState.NumGathered() + MarkObjectsState.NumGathered() + ObjectLookahead);
 		MarkRootsState.Finish(OutRootObjects);
 		MarkObjectsState.Finish(OutRootObjects);
+
+		GGCStats.NumRoots = OutRootObjects.Num();
 	}
 
 	/**
@@ -4274,7 +4294,12 @@ private:
 		{
 			const double StartTime = FPlatformTime::Seconds();
 			MarkObjectsAsUnreachable(KeepFlags);
-			UE_LOG(LogGarbage, Verbose, TEXT("%f ms for MarkObjectsAsUnreachable Phase (%d Objects To Serialize)"), (FPlatformTime::Seconds() - StartTime) * 1000, InitialObjects.Num());
+			const double ElapsedTime = FPlatformTime::Seconds() - StartTime;
+			if (!Stats.bFoundGarbageRef)
+			{
+				GGCStats.MarkObjectsAsUnreachableTime = ElapsedTime;
+			}
+			UE_LOG(LogGarbage, Verbose, TEXT("%f ms for MarkObjectsAsUnreachable Phase (%d Objects To Serialize)"), ElapsedTime * 1000, InitialObjects.Num());
 		}
 	}
 
@@ -4304,6 +4329,7 @@ private:
 		{
 			// Add objects marked with the GC barrier to the inital set of objects for the next iteration of incremental reachability
 			Private::GReachableObjects.PopAllAndEmpty(InitialObjects);
+			GGCStats.NumBarrierObjects += InitialObjects.Num();
 			UE_LOG(LogGarbage, Verbose, TEXT("Adding %d object(s) marker by GC barrier to the list of objects to process"), InitialObjects.Num());
 			ConditionallyAddBarrierReferencesToHistory(*Context);
 		}
@@ -4353,6 +4379,8 @@ public:
 	{
 		LLM_SCOPE(ELLMTag::GC);
 
+		const bool bIsGarbageTracking = !GReachabilityState.IsSuspended() && Stats.bFoundGarbageRef;
+
 		if (!GReachabilityState.IsSuspended())
 		{
 			StartReachabilityAnalysis(KeepFlags, Options);
@@ -4372,7 +4400,12 @@ public:
 			// if Verse is no longer active.
 			} while ((VerseGCActive() || !Private::GReachableObjects.IsEmpty() || !Private::GReachableClusters.IsEmpty()) && !GReachabilityState.IsSuspended());
 
-			UE_LOG(LogGarbage, Verbose, TEXT("%f ms for Reachability Analysis"), (FPlatformTime::Seconds() - StartTime) * 1000);
+			const double ElapsedTime = FPlatformTime::Seconds() - StartTime;
+			if (!bIsGarbageTracking)
+			{
+				GGCStats.ReferenceCollectionTime += ElapsedTime;
+			}
+			UE_LOG(LogGarbage, Verbose, TEXT("%f ms for Reachability Analysis"), ElapsedTime * 1000);
 		}
 
 PRAGMA_DISABLE_DEPRECATION_WARNINGS
@@ -4551,6 +4584,23 @@ void IncrementalPurgeGarbage(bool bUseTimeLimit, double TimeLimit)
 
 		// when running incrementally using a time limit, add one last tick for the memory trim
 		bCompleted = bCompleted && !bUseTimeLimit;
+
+		if (bUseTimeLimit)
+		{
+			// Add total time only if we're using time limit otherwise purge phase time is included in PostGarbageCollect
+			GGCStats.TotalTime += FPlatformTime::Seconds() - GCStartTime;
+		}
+	}
+	GGCStats.bInProgress = !bCompleted;
+
+	if (bCompleted && bUseTimeLimit)
+	{
+		// If this was incremental purge then its completion marks the completion of the entire GC cycle (otherwise see PostCollectGarbageImpl)		
+		FCoreUObjectDelegates::GarbageCollectComplete.Broadcast();
+		if (GDumpGCAnalyticsToLog)
+		{
+			GGCStats.DumpToLog();
+		}
 	}
 }
 
@@ -4592,8 +4642,8 @@ bool IncrementalDestroyGarbage(bool bUseTimeLimit, double TimeLimit)
 	bool bTimeLimitReached = false;
 
 	// Keep track of time it took to destroy objects for stats
-	double IncrementalDestroyGarbageStartTime = FPlatformTime::Seconds();
-	double LastTimeoutWarningTime = FPlatformTime::Seconds();
+	const double IncrementalDestroyGarbageStartTime = FPlatformTime::Seconds();
+	double LastTimeoutWarningTime = IncrementalDestroyGarbageStartTime;	
 
 	// Depending on platform FPlatformTime::Seconds might take a noticeable amount of time if called thousands of times so we avoid
 	// enforcing the time limit too often, especially as neither Destroy nor actual deletion should take significant
@@ -4639,7 +4689,7 @@ bool IncrementalDestroyGarbage(bool bUseTimeLimit, double TimeLimit)
 					// Only proceed with destroying the object if the asynchronous cleanup started by BeginDestroy has finished.
 					if(Object->IsReadyForFinishDestroy())
 					{
-						UE::GC::GStats.IncPurgeCount(Object);
+						UE::GC::GDetailedStats.IncPurgeCount(Object);
 						// Send FinishDestroy message.
 						Object->ConditionalFinishDestroy();
 					}
@@ -4699,7 +4749,7 @@ bool IncrementalDestroyGarbage(bool bUseTimeLimit, double TimeLimit)
 					// Only proceed with destroying the object if the asynchronous cleanup started by BeginDestroy has finished.
 					if( Object->IsReadyForFinishDestroy() )
 					{
-						UE::GC::GStats.IncPurgeCount(Object);
+						UE::GC::GDetailedStats.IncPurgeCount(Object);
 						// Send FinishDestroy message.
 						Object->ConditionalFinishDestroy();
 
@@ -4865,15 +4915,19 @@ bool IncrementalDestroyGarbage(bool bUseTimeLimit, double TimeLimit)
 				GObjectCountDuringLastMarkPhase.GetValue(), 
 				GObjectCountDuringLastMarkPhase.GetValue() - PurgedObjectCountSinceLastMarkPhase,
 				(FPlatformTime::Seconds() - IncrementalDestroyGarbageStartTime) * 1000);
-			UE::GC::GStats.LogPurgeStats(PurgedObjectCountSinceLastMarkPhase);
+			UE::GC::GDetailedStats.LogPurgeStats(PurgedObjectCountSinceLastMarkPhase);
 			GAsyncPurge->ResetObjectsDestroyedSinceLastMarkPhase();
 		}
 	}
 
+	const double ElapsedTime = FPlatformTime::Seconds() - IncrementalDestroyGarbageStartTime;
+	GGCStats.DestroyGarbageTime += ElapsedTime;
+	GGCStats.DestroyGarbageTimeLimit = TimeLimit;
+
 	if (bUseTimeLimit && !bCompleted)
 	{
 		UE_LOG(LogGarbage, Log, TEXT("%.3f ms for incrementally purging unreachable objects (FinishDestroyed: %d, Destroyed: %d / %d)"),
-			(FPlatformTime::Seconds() - IncrementalDestroyGarbageStartTime) * 1000,
+			ElapsedTime * 1000,
 			GObjCurrentPurgeObjectIndex,
 			GAsyncPurge->GetObjectsDestroyedSinceLastMarkPhase(),
 			GUnreachableObjects.Num());
@@ -4989,11 +5043,15 @@ void DissolveUnreachableClusters(UE::GC::EGatherOptions Options)
 		GUObjectClusters.FreeCluster(ClusterIndex);
 	}
 
+	GGCStats.NumUnreachableClusteredObjects = TotalClusteredObjects.load(std::memory_order_relaxed);
+	GGCStats.NumDissolvedClusters = ClustersToDestroy.Num();
+	GGCStats.DissolveUnreachableClustersTime = FPlatformTime::Seconds() - StartTime;
+
 	UE_LOG(LogGarbage, Log, TEXT("%f ms for Dissolve Unreachable Clusters (%d/%d clusters dissolved containing %d cluster objects)"),
-		(FPlatformTime::Seconds() - StartTime) * 1000,
-		ClustersToDestroy.Num(),
+		GGCStats.DissolveUnreachableClustersTime * 1000,
+		GGCStats.NumDissolvedClusters,
 		InitialNumClusters,
-		TotalClusteredObjects.load(std::memory_order_relaxed));
+		GGCStats.NumUnreachableClusteredObjects);
 }
 
 bool GatherUnreachableObjects(UE::GC::EGatherOptions Options, double TimeLimit /*= 0.0*/)
@@ -5004,7 +5062,7 @@ bool GatherUnreachableObjects(UE::GC::EGatherOptions Options, double TimeLimit /
 	using namespace UE::GC;
 	using namespace UE::GC::Private;
 
-	const double StartTime = FPlatformTime::Seconds();
+	const double GatherStartTime = FPlatformTime::Seconds();
 
 	GUnreachableObjects.Reset();
 	GUnrechableObjectIndex = 0;
@@ -5020,11 +5078,11 @@ bool GatherUnreachableObjects(UE::GC::EGatherOptions Options, double TimeLimit /
 
 	// Iterate over all objects. Note that we iterate over the UObjectArray and usually check only internal flags which
 	// are part of the array so we don't suffer from cache misses as much as we would if we were to check ObjectFlags.
-	ParallelFor( TEXT("GC.GatherUnreachable"), GGatherUnreachableObjectsState.NumWorkerThreads(), 1, [&ThreadIterators, &StartTime, &TimeLimit, &TimeLimitExceededFlag](int32 ThreadIndex)
+	ParallelFor( TEXT("GC.GatherUnreachable"), GGatherUnreachableObjectsState.NumWorkerThreads(), 1, [&ThreadIterators, &GatherStartTime, &TimeLimit, &TimeLimitExceededFlag](int32 ThreadIndex)
 	{
 		TRACE_CPUPROFILER_EVENT_SCOPE(GatherUnreachableObjectsTask);
 		constexpr int32 TimeLimitPollInterval = 10;
-		FTimeSlicer Timer(TimeLimitPollInterval, ThreadIndex * TimeLimitPollInterval, StartTime, TimeLimit, TimeLimitExceededFlag);
+		FTimeSlicer Timer(TimeLimitPollInterval, ThreadIndex * TimeLimitPollInterval, GatherStartTime, TimeLimit, TimeLimitExceededFlag);
 		FGatherUnreachableObjectsState::FIterator& Iterator = ThreadIterators[ThreadIndex];
 		
 		while (Iterator.Index <= Iterator.LastIndex)
@@ -5053,14 +5111,20 @@ bool GatherUnreachableObjects(UE::GC::EGatherOptions Options, double TimeLimit /
 		GGatherUnreachableObjectsState.Finish(GUnreachableObjects);
 	}
 
+	const double GatherEndTime = FPlatformTime::Seconds();
+	const double GatherIterationTime = GatherEndTime - GatherStartTime;
+	GGCStats.GatherUnreachableTime += GatherIterationTime;
+
 	UE_LOG(LogGarbage, Log, TEXT("%f ms for Gather Unreachable Objects (%d objects collected / %d scanned with %d thread(s))"),
-		(FPlatformTime::Seconds() - StartTime) * 1000,
+		GatherIterationTime * 1000,
 		NumGathered,
 		NumScanned,
 		NumThreads);
 
 	if (bCompleted)
 	{
+		GGCStats.NumUnreachableObjects = GUnreachableObjects.Num();
+
 		// NotifyUnreachableObjects needs to be called from within a GC lock so that Async Loading Thread does not attempt to create new exports  
 		// before it gets a chance to process all unreachable objects
 		const bool bNeedsGCLock = !FGCCSyncObject::Get().IsGCLocked();
@@ -5074,11 +5138,15 @@ bool GatherUnreachableObjects(UE::GC::EGatherOptions Options, double TimeLimit /
 			ReleaseGCLock();
 		}
 
+		const double NotifyEndTime = FPlatformTime::Seconds();
+		GGCStats.NotifyUnreachableTime = NotifyEndTime - GatherEndTime;
+
 #if VERIFY_DISREGARD_GC_ASSUMPTIONS
 		if (GVerifyNoUnreachableObjects != 0 && !GExitPurge)
 		{
 			DECLARE_SCOPE_CYCLE_COUNTER(TEXT("CollectGarbageInternal.VerifyNoUnreachableObjects"), STAT_CollectGarbageInternal_VerifyNoUnreachableObjects, STATGROUP_GC);
 			VerifyNoUnreachableObjects(GUnreachableObjects.Num());
+			GGCStats.VerifyNoUnreachableTime = FPlatformTime::Seconds() - NotifyEndTime;
 		}
 #endif // VERIFY_DISREGARD_GC_ASSUMPTIONS
 	}
@@ -5127,6 +5195,7 @@ static void ClearWeakReferences(TConstArrayView<TUniquePtr<FWorkerContext>> Cont
 	TSet<UObject*> ObjectsThatNeedWeakReferenceClearing;
 	for (const TUniquePtr<FWorkerContext>& Context : Contexts)
 	{
+		GGCStats.NumWeakReferencesForClearing += Context->WeakReferences.Num();
 		for (FWeakReferenceInfo& ReferenceInfo : Context->WeakReferences)
 		{
 			UObject* ReferencedObject = ReferenceInfo.ReferencedObject;
@@ -5149,6 +5218,7 @@ static void ClearWeakReferences(TConstArrayView<TUniquePtr<FWorkerContext>> Cont
 	}
 	if constexpr (bGatheredWithIncrementalReachability)
 	{
+		GGCStats.NumObjectsThatNeedWeakReferenceClearing = ObjectsThatNeedWeakReferenceClearing.Num();
 		FWeakReferenceEliminator ReferenceEliminator;
 		for (UObject* Object : ObjectsThatNeedWeakReferenceClearing)
 		{
@@ -5303,6 +5373,8 @@ void PreCollectGarbageImpl(EObjectFlags KeepFlags)
 {
 	using namespace UE::GC::Private;
 
+	const double PreCollectStartTime = FPlatformTime::Seconds();
+
 	FGCCSyncObject::Get().ResetGCIsWaiting();
 
 #if defined(WITH_CODE_GUARD_HANDLER) && WITH_CODE_GUARD_HANDLER
@@ -5326,6 +5398,8 @@ void PreCollectGarbageImpl(EObjectFlags KeepFlags)
 		ReleaseGCLock();
 		FlushAsyncLoading();
 		AcquireGCLock();
+
+		GGCStats.bFlushedAsyncLoading = true;
 	}
 
 	// Route callbacks so we can ensure that we are e.g. not in the middle of loading something by flushing
@@ -5342,13 +5416,16 @@ void PreCollectGarbageImpl(EObjectFlags KeepFlags)
 
 		// Make sure previous incremental purge has finished or we do a full purge pass in case we haven't kicked one
 		// off yet since the last call to garbage collection.
-		if (GObjIncrementalPurgeIsInProgress || GObjPurgeIsRequired)
+		if (IsIncrementalPurgePending())
 		{
 			IncrementalPurgeGarbage(false);
 			if (!bPerformFullPurge)
 			{
 				FMemory::Trim();
 			}
+			GGCStats.bPurgedPreviousGCObjects = true;
+			// IncrementalPurgeGarbage sets bInProgress to false when all objects have been purged because IncrementalPurgeGarbage is also called by engine Tick
+			GGCStats.bInProgress = true;
 		}
 
 
@@ -5388,11 +5465,14 @@ void PreCollectGarbageImpl(EObjectFlags KeepFlags)
 					VerifyClustersAssumptions();
 				}
 				VerifyObjectFlags();
-				UE_LOG(LogGarbage, Log, TEXT("%.2f ms for Verify GC Assumptions"), (FPlatformTime::Seconds() - StartTime) * 1000);
+				GGCStats.VerifyTime = FPlatformTime::Seconds() - StartTime;
+				UE_LOG(LogGarbage, Log, TEXT("%.2f ms for Verify GC Assumptions"), GGCStats.VerifyTime * 1000);
 			}
 #endif
 		}
 	}
+
+	GGCStats.TotalTime += FPlatformTime::Seconds() - PreCollectStartTime;
 }
 
 template<bool bPerformFullPurge>
@@ -5413,6 +5493,8 @@ void CollectGarbageImpl(EObjectFlags KeepFlags)
 template<bool bPerformFullPurge>
 void PostCollectGarbageImpl(EObjectFlags KeepFlags)
 {
+	const double PostCollectStartTime = FPlatformTime::Seconds();
+
 	using namespace UE::GC;
 	using namespace UE::GC::Private;	
 
@@ -5516,8 +5598,20 @@ void PostCollectGarbageImpl(EObjectFlags KeepFlags)
 		}
 	}
 
-	GTimingInfo.LastGCTime = FPlatformTime::Seconds();
+	const double PostCollectEndTime = FPlatformTime::Seconds();
+	GTimingInfo.LastGCTime = PostCollectEndTime;
+	GGCStats.TotalTime += PostCollectEndTime - PostCollectStartTime;
 	STAT_ADD_CUSTOMMESSAGE_NAME(STAT_NamedMarker, TEXT("GarbageCollection - End"));
+
+	if (bPerformFullPurge)
+	{
+		// If this was a full purge then PostCollectGarbageImpl completion marks the completion of the entire GC cycle (otherwise see IncrementalPurgeGarbage)
+		FCoreUObjectDelegates::GarbageCollectComplete.Broadcast();
+		if (GDumpGCAnalyticsToLog)
+		{
+			GGCStats.DumpToLog();
+		}
+	}
 }
 
 void FReachabilityAnalysisState::CollectGarbage(EObjectFlags KeepFlags, bool bFullPurge)
@@ -5550,6 +5644,17 @@ void FReachabilityAnalysisState::PerformReachabilityAnalysisAndConditionallyPurg
 
 	LLM_SCOPE(ELLMTag::GC);
 
+	if (!GIsIncrementalReachabilityPending)
+	{
+		GGCStats = UE::GC::Private::FStats();
+		GGCStats.bInProgress = true;
+		GGCStats.bStartedAsFullPurge = bPerformFullPurge;
+		GGCStats.NumObjects = GUObjectArray.GetObjectArrayNumMinusAvailable() - GUObjectArray.GetFirstGCIndex();
+		GGCStats.NumClusters = GUObjectClusters.GetNumAllocatedClusters();
+		GGCStats.ReachabilityTimeLimit = GetReachabilityAnalysisTimeLimit();
+	}
+	GGCStats.bFinishedAsFullPurge = bPerformFullPurge;
+
 	if (bPerformFullPurge)
 	{
 		UE::GC::PreCollectGarbageImpl<true>(ObjectKeepFlags);
@@ -5559,7 +5664,7 @@ void FReachabilityAnalysisState::PerformReachabilityAnalysisAndConditionallyPurg
 		UE::GC::PreCollectGarbageImpl<false>(ObjectKeepFlags);
 	}	
 	
-	const bool bForceNonIncrementalReachability = 
+	const bool bForceNonIncrementalReachability =
 		!GIsIncrementalReachabilityPending &&
 		(bPerformFullPurge || !GAllowIncrementalReachability);
 
@@ -5603,6 +5708,9 @@ void FReachabilityAnalysisState::PerformReachabilityAnalysisAndConditionallyPurg
 		ReferenceProcessingTotalTime += ReferenceProcessingElapsedTime;
 		IncrementalMarkPhaseTotalTime += ElapsedTime;
 
+		GGCStats.ReachabilityTime += ReferenceProcessingElapsedTime;
+		GGCStats.TotalTime += ReferenceProcessingElapsedTime;
+
 		if (UE_LOG_ACTIVE(LogGarbage, Log))
 		{
 			if (GIsIncrementalReachabilityPending)
@@ -5647,7 +5755,10 @@ void FReachabilityAnalysisState::PerformReachabilityAnalysisAndConditionallyPurg
 			GC.Stats = Stats; // This is to pass Stats.bFoundGarbageRef to CG
 			GC.PerformReachabilityAnalysis(ObjectKeepFlags, GetReferenceCollectorOptions(bPerformFullPurge));
 		}
-		UE_LOG(LogGarbage, Log, TEXT("%.2f ms for GC rerun to track garbage references (gc.GarbageReferenceTrackingEnabled=%d)"), (FPlatformTime::Seconds() - StartTime) * 1000, GGarbageReferenceTrackingEnabled);
+		const double ElapsedTime = FPlatformTime::Seconds() - StartTime;
+		GGCStats.GarbageTrackingTime = ElapsedTime;
+		GGCStats.TotalTime += ElapsedTime;
+		UE_LOG(LogGarbage, Log, TEXT("%.2f ms for GC rerun to track garbage references (gc.GarbageReferenceTrackingEnabled=%d)"), ElapsedTime * 1000, GGarbageReferenceTrackingEnabled);
 	}
 	// Maybe purge garbage (if we're done with incremental reachability and there's still time left)
 
@@ -5735,6 +5846,14 @@ void FinalizeIncrementalReachabilityAnalysis()
 	}
 }
 
+namespace UE::GC::Private
+{
+	UE::GC::Private::FStats GetGarbageCollectionStats()
+	{
+		return GGCStats;
+	}
+}
+
 FString FGarbageReferenceInfo::GetReferencingObjectInfo() const
 {
 	if (bReferencerUObject)
@@ -5805,6 +5924,7 @@ bool UnhashUnreachableObjects(bool bUseTimeLimit, double TimeLimit)
 
 	// Unhash all unreachable objects.
 	const double StartTime = FPlatformTime::Seconds();
+	double LastPollTime = 0.0;
 	const int32 TimeLimitEnforcementGranularityForBeginDestroy = 10;
 	int32 TimePollCounter = 0;
 	const bool bFirstIteration = (GUnrechableObjectIndex == 0);
@@ -5824,21 +5944,29 @@ bool UnhashUnreachableObjects(bool bUseTimeLimit, double TimeLimit)
 			}
 
 			const bool bPollTimeLimit = ((TimePollCounter++) % TimeLimitEnforcementGranularityForBeginDestroy == 0);
-			if (bUseTimeLimit && bPollTimeLimit && ((FPlatformTime::Seconds() - StartTime) > TimeLimit))
+			if (bUseTimeLimit & bPollTimeLimit) //-V792
 			{
-				break;
+				LastPollTime = FPlatformTime::Seconds();
+				if ((LastPollTime - StartTime) > TimeLimit)
+				{
+					break;
+				}
 			}
 		}
 	}
 
 	bTimeLimitReached = (GUnrechableObjectIndex < GUnreachableObjects.Num());
 
+	const double ElapsedTime = (LastPollTime == 0.0 ? FPlatformTime::Seconds() : LastPollTime) - StartTime;
+	GGCStats.UnhashingTime += ElapsedTime;
+	GGCStats.UnhashingTimeLimit = TimeLimit;
+
 	if (!bUseTimeLimit)
 	{
 		UE_LOG(LogGarbage, Log, TEXT("%f ms for %sunhashing unreachable objects (%d objects unhashed)"),
-		(FPlatformTime::Seconds() - StartTime) * 1000,
-		bUseTimeLimit ? TEXT("incrementally ") : TEXT(""),
-		GUnreachableObjects.Num());
+			ElapsedTime * 1000,
+			bUseTimeLimit ? TEXT("incrementally ") : TEXT(""),
+			GUnreachableObjects.Num());
 	}
 	else if (!bTimeLimitReached)
 	{
diff --git a/Engine/Source/Runtime/CoreUObject/Private/UObject/ReachabilityAnalysisState.cpp b/Engine/Source/Runtime/CoreUObject/Private/UObject/ReachabilityAnalysisState.cpp
--- a/Engine/Source/Runtime/CoreUObject/Private/UObject/ReachabilityAnalysisState.cpp
+++ b/Engine/Source/Runtime/CoreUObject/Private/UObject/ReachabilityAnalysisState.cpp
@@ -5,10 +5,7 @@
 =============================================================================*/
 
 #include "UObject/ReachabilityAnalysisState.h"
-#include "HAL/Runnable.h"
-#include "HAL/RunnableThread.h"
-#include "Misc/App.h"
-#include "Misc/Fork.h"
+#include "UObject/ReachabilityAnalysis.h"
 
 namespace UE::GC
 {
@@ -52,3 +49,40 @@ bool FReachabilityAnalysisState::CheckIfAnyContextIsSuspended()
 
 } // namespce UE::GC
 
+///////////////////////////////////////////////////////////////////////////////////////////////////////////////////////////////////////
+
+namespace UE::GC::Private
+{
+
+FString FIterationTimerStat::ToString() const
+{
+	return FString::Printf(TEXT("%7.3fms, Iterations: %4d, Max: %7.3fms (%4d)"), Total * 1000, NumIterations, Max * 1000, SlowestIteration);
+}
+
+void FStats::DumpToLog()
+{	
+#define UE_TRUE_FALSE(b) (b ? TEXT("true") : TEXT("false"))
+	UE_LOG(LogGarbage, Log, TEXT("Garbage Collection Total  %7.3fms%s"), TotalTime * 1000, bInProgress ? TEXT(" (still in progress, results may be incomplete)") : TEXT(""));
+	UE_LOG(LogGarbage, Log, TEXT("  Verify                  %7.3fms"), VerifyTime * 1000);	
+	UE_LOG(LogGarbage, Log, TEXT("  Reachability            %s"), *ReachabilityTime.ToString());
+	UE_LOG(LogGarbage, Log, TEXT("    MarkAsUnreachable     %7.3fms"), MarkObjectsAsUnreachableTime * 1000);
+	UE_LOG(LogGarbage, Log, TEXT("    Reference Collection  %s"), *ReferenceCollectionTime.ToString());
+	UE_LOG(LogGarbage, Log, TEXT("  GarbageTracking         %7.3fms"), GarbageTrackingTime * 1000);
+	UE_LOG(LogGarbage, Log, TEXT("  DissolveClusters        %7.3fms"), DissolveUnreachableClustersTime * 1000);
+	UE_LOG(LogGarbage, Log, TEXT("  GatherUnreachable       %s"), *GatherUnreachableTime.ToString());
+	UE_LOG(LogGarbage, Log, TEXT("  NotifyUnreachable       %7.3fms"), NotifyUnreachableTime * 1000);
+	UE_LOG(LogGarbage, Log, TEXT("  VerifyNoUnreachable     %7.3fms"), VerifyNoUnreachableTime * 1000);	
+	UE_LOG(LogGarbage, Log, TEXT("  Unhashing               %s"), *UnhashingTime.ToString());
+	UE_LOG(LogGarbage, Log, TEXT("  DestroyGarbage          %s"), *DestroyGarbageTime.ToString());
+
+	UE_LOG(LogGarbage, Log, TEXT("Pre GC:     Objects: %7d, Roots : %7d, Clusters : %7d, Clustered Objects : %7d"), NumObjects, NumRoots, NumClusters, NumClusteredObjects);
+	UE_LOG(LogGarbage, Log, TEXT("Destroyed:  Objects: %7d, including        Clusters : %7d, Clustered Objects : %7d"), NumUnreachableObjects, NumDissolvedClusters, NumUnreachableClusteredObjects);
+	UE_LOG(LogGarbage, Log, TEXT("Number of barrier objects: %d"), NumBarrierObjects);
+	UE_LOG(LogGarbage, Log, TEXT("Number of weak references for clearing %d and objects that need weak reference clearing: %d"), NumWeakReferencesForClearing, NumObjectsThatNeedWeakReferenceClearing);
+	UE_LOG(LogGarbage, Log, TEXT("Started as full purge: %s, finished as full purge: %s"), UE_TRUE_FALSE(bStartedAsFullPurge), UE_TRUE_FALSE(bFinishedAsFullPurge));
+	UE_LOG(LogGarbage, Log, TEXT("Flushed async loading: %s"), UE_TRUE_FALSE(bFlushedAsyncLoading));
+	UE_LOG(LogGarbage, Log, TEXT("Purged previous GC objects: %s"), UE_TRUE_FALSE(bPurgedPreviousGCObjects));
+#undef UE_TRUE_FALSE
+}
+
+} // namespce UE::GC::Private
diff --git a/Engine/Source/Runtime/CoreUObject/Private/UObject/ReachabilityAnalysisState.h b/Engine/Source/Runtime/CoreUObject/Private/UObject/ReachabilityAnalysisState.h
--- a/Engine/Source/Runtime/CoreUObject/Private/UObject/ReachabilityAnalysisState.h
+++ b/Engine/Source/Runtime/CoreUObject/Private/UObject/ReachabilityAnalysisState.h
@@ -10,6 +10,7 @@
 #include "UObject/UObjectBase.h"
 #include "UObject/Object.h"
 #include "UObject/FastReferenceCollector.h"
+#include "UObject/ReachabilityAnalysis.h"
 
 namespace UE::GC
 {
@@ -328,4 +329,14 @@ public:
 	}
 };
 
-} // namespace UE::GC
\ No newline at end of file
+} // namespace UE::GC
+
+namespace UE::GC::Private
+{
+	/**
+	 * Returns statistics of the last Garbage Collecion cycle.
+	 *
+	 * @return	Statistics of the last Garabage Collection cycle.
+	 */
+	COREUOBJECT_API UE::GC::Private::FStats GetGarbageCollectionStats();
+}
\ No newline at end of file
diff --git a/Engine/Source/Runtime/CoreUObject/Private/UObject/UObjectGlobals.cpp b/Engine/Source/Runtime/CoreUObject/Private/UObject/UObjectGlobals.cpp
--- a/Engine/Source/Runtime/CoreUObject/Private/UObject/UObjectGlobals.cpp
+++ b/Engine/Source/Runtime/CoreUObject/Private/UObject/UObjectGlobals.cpp
@@ -219,6 +219,8 @@ FSimpleMulticastDelegate FCoreUObjectDelegates::PostReachabilityAnalysis;
 FSimpleMulticastDelegate FCoreUObjectDelegates::PreGarbageCollectConditionalBeginDestroy;
 FSimpleMulticastDelegate FCoreUObjectDelegates::PostGarbageCollectConditionalBeginDestroy;
 
+FSimpleMulticastDelegate FCoreUObjectDelegates::GarbageCollectComplete;
+
 FCoreUObjectDelegates::FPreLoadMapDelegate FCoreUObjectDelegates::PreLoadMap;
 FCoreUObjectDelegates::FPreLoadMapWithContextDelegate FCoreUObjectDelegates::PreLoadMapWithContext;
 FCoreUObjectDelegates::FPostLoadMapDelegate FCoreUObjectDelegates::PostLoadMapWithWorld;
diff --git a/Engine/Source/Runtime/CoreUObject/Public/UObject/ReachabilityAnalysis.h b/Engine/Source/Runtime/CoreUObject/Public/UObject/ReachabilityAnalysis.h
--- a/Engine/Source/Runtime/CoreUObject/Public/UObject/ReachabilityAnalysis.h
+++ b/Engine/Source/Runtime/CoreUObject/Public/UObject/ReachabilityAnalysis.h
@@ -35,3 +35,74 @@ COREUOBJECT_API void SetReachabilityAnalysisTimeLimit(float TimeLimitSeconds);
  * @return	time limit (in seconds) for incremental rachability analysis.
  */
 COREUOBJECT_API float GetReachabilityAnalysisTimeLimit();
+
+namespace UE::GC::Private
+{
+
+struct FIterationTimerStat
+{
+	double Total = 0.0;
+	double Max = 0.0;
+	int32 NumIterations = 0;
+	int32 SlowestIteration = 0;
+
+	double operator += (double ElapsedTime)
+	{
+		int32 IterationIndex = NumIterations++;
+		Total += ElapsedTime;
+		if (ElapsedTime > Max)
+		{
+			Max = ElapsedTime;
+			SlowestIteration = IterationIndex;
+		}
+		return Total;
+	}
+
+	FString ToString() const;
+};
+
+struct FStats
+{
+	double TotalTime = 0.0;
+
+	double ReachabilityTimeLimit = 0.0;
+	double UnhashingTimeLimit = 0.0;
+	double DestroyGarbageTimeLimit = 0.0;
+
+	double VerifyTime = 0.0;
+	double VerifyNoUnreachableTime = 0.0;
+	double GarbageTrackingTime = 0.0;
+
+	double MarkObjectsAsUnreachableTime = 0.0;
+	double NotifyUnreachableTime = 0.0;
+	double DissolveUnreachableClustersTime = 0.0;
+
+	FIterationTimerStat ReachabilityTime;
+	FIterationTimerStat ReferenceCollectionTime;
+	FIterationTimerStat GatherUnreachableTime;
+	FIterationTimerStat UnhashingTime;
+	FIterationTimerStat DestroyGarbageTime;
+
+	int32 NumObjects = 0;
+	int32 NumRoots = 0;
+	int32 NumClusters = 0;
+	int32 NumClusteredObjects = 0;
+	int32 NumUnreachableObjects = 0;
+	int32 NumDissolvedClusters = 0;
+	int32 NumUnreachableClusteredObjects = 0;
+	int32 NumBarrierObjects = 0;
+	int32 NumWeakReferencesForClearing = 0;
+	int32 NumObjectsThatNeedWeakReferenceClearing = 0;
+
+	bool bStartedAsFullPurge = false;
+	bool bFinishedAsFullPurge = false;
+	bool bFlushedAsyncLoading = false;
+	bool bPurgedPreviousGCObjects = false;
+	bool bInProgress = false;
+
+	void DumpToLog();
+};
+
+} // namespace UE::GC::Private
+
+
diff --git a/Engine/Source/Runtime/CoreUObject/Public/UObject/UObjectGlobals.h b/Engine/Source/Runtime/CoreUObject/Public/UObject/UObjectGlobals.h
--- a/Engine/Source/Runtime/CoreUObject/Public/UObject/UObjectGlobals.h
+++ b/Engine/Source/Runtime/CoreUObject/Public/UObject/UObjectGlobals.h
@@ -3157,7 +3157,7 @@ struct FCoreUObjectDelegates
 	/** Called after reachability analysis, before any purging */
 	static COREUOBJECT_API FSimpleMulticastDelegate PostReachabilityAnalysis;
 
-	/** Called after garbage collection */
+	/** Called after garbage collection (before purge phase if incremental purge is enabled and after purge phase if incremental purge is disabled) */
 	static COREUOBJECT_API FSimpleMulticastDelegate& GetPostGarbageCollect();
 
 	/** Called after purging unreachable objects during garbage collection */
@@ -3175,6 +3175,9 @@ struct FCoreUObjectDelegates
 	/** Called after ConditionalBeginDestroy phase of garbage collection */
 	static COREUOBJECT_API FSimpleMulticastDelegate PostGarbageCollectConditionalBeginDestroy;
 
+	/** Called after garbage collection is complete, all objects have been purged (regardless of whether incremental purge is enabled or not), memory has been trimmed and all other GC callbacks have been fired. */
+	static COREUOBJECT_API FSimpleMulticastDelegate GarbageCollectComplete;
+
 	/** Queries whether an object should be loaded on top ( replace ) an already existing one */
 	DECLARE_DELEGATE_RetVal_OneParam(bool, FOnLoadObjectsOnTop, const FString&);
 	static COREUOBJECT_API FOnLoadObjectsOnTop ShouldLoadOnTop;
