Fixing CIS

diff --git a/Engine/Source/Runtime/CoreUObject/Private/UObject/GarbageCollection.cpp b/Engine/Source/Runtime/CoreUObject/Private/UObject/GarbageCollection.cpp
--- a/Engine/Source/Runtime/CoreUObject/Private/UObject/GarbageCollection.cpp
+++ b/Engine/Source/Runtime/CoreUObject/Private/UObject/GarbageCollection.cpp
@@ -3329,7 +3329,7 @@ struct TVerseDebugReachabilityVisitor : public Verse::FAbstractVisitor
 
 	virtual void VisitNonNull(UObject*& InObject, const TCHAR* ElementName) override
 	{
-		UE::GC::GStats.IncreaseObjectRefStats(InObject);
+		UE::GC::GDetailedStats.IncreaseObjectRefStats(InObject);
 		Verse::FAbstractVisitor::FReferrerContext* VisitorContext = GetContext();
 		const Verse::VCell* Referencer = VisitorContext != nullptr && VisitorContext->GetReferrer().IsCell() ? VisitorContext->GetReferrer().AsCell() : nullptr;
 		if (ValidateReference(InObject, PermanentPool, FReferenceToken(Referencer), FMemberId(0)))
