Add new IsTemplateForSubobjects function that correctly handles querying for all types of subobject templates
Add comments warning against use of IsDefaultSubobject due to the historically inconsistent behavior, and deprecate CollectDefaultSubobjects for being redundant and incorrect
Improve general comments related to subobject functions to better match the current behavior

diff --git a/Engine/Source/Editor/UnrealEd/Private/Kismet2/ReloadUtilities.cpp b/Engine/Source/Editor/UnrealEd/Private/Kismet2/ReloadUtilities.cpp
--- a/Engine/Source/Editor/UnrealEd/Private/Kismet2/ReloadUtilities.cpp
+++ b/Engine/Source/Editor/UnrealEd/Private/Kismet2/ReloadUtilities.cpp
@@ -489,7 +489,7 @@ namespace UE::Reload::Private
 		const int32 DefaultSubobjectArrayCapacity = 16;
 		TArray<UObject*> DefaultSubobjectArray;
 		DefaultSubobjectArray.Empty(DefaultSubobjectArrayCapacity);
-		NewClass->GetDefaultObject()->CollectDefaultSubobjects(DefaultSubobjectArray);
+		NewClass->GetDefaultObject()->GetDefaultSubobjects(DefaultSubobjectArray);
 
 		TArray<FPropertyToUpdate> PropertiesToUpdate;
 		// Collect all properties that have actually changed
@@ -548,7 +548,7 @@ namespace UE::Reload::Private
 				}
 
 				DefaultSubobjectArray.Empty(DefaultSubobjectArrayCapacity);
-				ObjectPtr->CollectDefaultSubobjects(DefaultSubobjectArray);
+				ObjectPtr->GetDefaultSubobjects(DefaultSubobjectArray);
 
 				for (auto& PropertyToUpdate : PropertiesToUpdate)
 				{
diff --git a/Engine/Source/Runtime/CoreUObject/Private/UObject/Obj.cpp b/Engine/Source/Runtime/CoreUObject/Private/UObject/Obj.cpp
--- a/Engine/Source/Runtime/CoreUObject/Private/UObject/Obj.cpp
+++ b/Engine/Source/Runtime/CoreUObject/Private/UObject/Obj.cpp
@@ -1372,10 +1372,9 @@ void UObject::PostLoadSubobjects( FObjectInstancingGraph* OuterInstanceGraph/*=N
 				InstanceGraph = &CurrentInstanceGraph;
 			}
 
-			// this will be filled with the list of component instances which were serialized from disk
 			TArray<UObject*> SerializedComponents;
-			// fill the array with the component contained by this object that were actually serialized to disk through property references
-			CollectDefaultSubobjects(SerializedComponents, false);
+			// fill the array with the top level default subobjects
+			GetDefaultSubobjects(SerializedComponents);
 
 			// now, add all of the instanced components to the instance graph that will be used for instancing any components that have been added
 			// to this object's archetype since this object was last saved
diff --git a/Engine/Source/Runtime/CoreUObject/Private/UObject/UObjectBaseUtility.cpp b/Engine/Source/Runtime/CoreUObject/Private/UObject/UObjectBaseUtility.cpp
--- a/Engine/Source/Runtime/CoreUObject/Private/UObject/UObjectBaseUtility.cpp
+++ b/Engine/Source/Runtime/CoreUObject/Private/UObject/UObjectBaseUtility.cpp
@@ -588,12 +588,21 @@ void* UObjectBaseUtility::GetNativeInterfaceAddress(UClass* InterfaceClass)
 	return NULL;
 }
 
-bool UObjectBaseUtility::IsDefaultSubobject() const
+bool UObjectBaseUtility::IsTemplateForSubobjects(EObjectFlags TemplateTypes) const
 {
-	const bool bIsInstanced = !HasAnyFlags(RF_ClassDefaultObject) && GetOuter() && (GetOuter()->HasAnyFlags(RF_ClassDefaultObject) || ((UObject*)this)->GetArchetype() != GetClass()->GetDefaultObject(false));
-	return bIsInstanced;
+	// This includes archetype objects that are inside CDOs or inheritable component templates, but not the CDO itself
+	return HasAnyFlags(RF_ArchetypeObject) && !HasAnyFlags(RF_ClassDefaultObject) && IsTemplate(TemplateTypes);
 }
 
+bool UObjectBaseUtility::IsDefaultSubobject() const
+{
+	// For historical reasons this behavior does not match the RF_DefaultSubObject flag.
+	// It will return true for any object instanced using a non-CDO archetype, 
+	// but it will return false for indirectly nested subobjects of a CDO that can be used as an archetype.
+	
+	return !HasAnyFlags(RF_ClassDefaultObject) && GetOuter() &&
+		(GetOuter()->HasAnyFlags(RF_ClassDefaultObject) || ((UObject*)this)->GetArchetype() != GetClass()->GetDefaultObject(false));
+}
 
 UClass* GetParentNativeClass(UClass* Class)
 {
diff --git a/Engine/Source/Runtime/CoreUObject/Public/UObject/Object.h b/Engine/Source/Runtime/CoreUObject/Public/UObject/Object.h
--- a/Engine/Source/Runtime/CoreUObject/Public/UObject/Object.h
+++ b/Engine/Source/Runtime/CoreUObject/Public/UObject/Object.h
@@ -135,7 +135,7 @@ class UObject : public UObjectBaseUtility
 	}
 
 	/**
-	 * Create a component or subobject.
+	 * Create a component or subobject that will be instanced inside all instances of this class.
 	 * @param	TReturnType					Class of return type, all overrides must be of this type
 	 * @param	SubobjectName				Name of the new component
 	 * @param	bTransient					True if the component is being assigned to a transient property. This does not make the component itself transient, but does stop it from inheriting parent defaults
@@ -161,8 +161,8 @@ class UObject : public UObjectBaseUtility
 	}
 	
 	/**
-	 * Create an optional component or subobject. Optional subobjects will not get created
-	 * if a derived class specified DoNotCreateDefaultSubobject with the subobject's name.
+	 * Create an optional component or subobject.Optional subobjects will not get created.
+	 * if a derived class specifies DoNotCreateDefaultSubobject with the subobject name.
 	 * @param	TReturnType					Class of return type, all overrides must be of this type
 	 * @param	SubobjectName				Name of the new component
 	 * @param	bTransient					True if the component is being assigned to a transient property. This does not make the component itself transient, but does stop it from inheriting parent defaults
@@ -175,8 +175,8 @@ class UObject : public UObjectBaseUtility
 	}
 	
 	/**
-	 * Create an optional component or subobject. Optional subobjects will not get created
-	 * if a derived class specified DoNotCreateDefaultSubobject with the subobject's name.
+	 * Create an optional component or subobject. Optional subobjects will not get created.
+	 * if a derived class specifies DoNotCreateDefaultSubobject with the subobject name.
 	 * @param	TReturnType					Class of return type, all overrides must be of this type
 	 * @param	TClassToConstructByDefault	Class of object to actually construct, must be a subclass of TReturnType
 	 * @param	SubobjectName				Name of the new component
@@ -189,14 +189,14 @@ class UObject : public UObjectBaseUtility
 	}
 
 	/**
-	 * Gets all default subobjects associated with this object instance.
-	 * @param	OutDefaultSubobjects	Array containing all default subobjects of this object.
+	 * Gets all directly nested default subobjects that are associated with this object instance.
+	 * @param	OutDefaultSubobjects	Array containing the directly nested default subobjects of this object.
 	 */
 	COREUOBJECT_API void GetDefaultSubobjects(TArray<UObject*>& OutDefaultSubobjects);
 
 	/**
-	 * Finds a subobject associated with this object instance by its name
-	 * @param	Name	Object name to look for
+	 * Finds a default subobject associated with this object instance by its name.
+	 * @param	Name	Object name to look for matching the SubobjectName used to create the default subobject
 	 */
 	COREUOBJECT_API UObject* GetDefaultSubobjectByName(FName ToFind);
 
@@ -1224,13 +1224,11 @@ public:
 	 */
 	COREUOBJECT_API virtual void BuildSubobjectMapping(UObject* OtherObject, TMap<UObject*, UObject*>& ObjectMapping) const;
 
-	/**
-	 * Uses the TArchiveObjectReferenceCollector to build a list of all components referenced by this object which have this object as the outer
-	 *
-	 * @param	OutDefaultSubobjects	the array that should be populated with the default subobjects "owned" by this object
-	 * @param	bIncludeNestedSubobjects	controls whether subobjects which are contained by this object, but do not have this object
-	 *										as its direct Outer should be included
+	/** 
+	 * Gets all subobjects inside this object that return true for IsDefaultSubobject.
+	 * The nested behavior is inconsistent because IsDefaultSubobject does not work reliably for nested subobjects and it is less efficient than GetDefaultSubobjects
 	 */
+	UE_DEPRECATED(5.4, "Call GetDefaultSubobjects for top level subobjects or use ForEachObjectWithOuter with a more precise check")
 	COREUOBJECT_API void CollectDefaultSubobjects( TArray<UObject*>& OutDefaultSubobjects, bool bIncludeNestedSubobjects=false ) const;
 
 	/**
diff --git a/Engine/Source/Runtime/CoreUObject/Public/UObject/ObjectMacros.h b/Engine/Source/Runtime/CoreUObject/Public/UObject/ObjectMacros.h
--- a/Engine/Source/Runtime/CoreUObject/Public/UObject/ObjectMacros.h
+++ b/Engine/Source/Runtime/CoreUObject/Public/UObject/ObjectMacros.h
@@ -538,8 +538,8 @@ enum EObjectFlags
 	RF_Standalone				=0x00000002,	///< Keep object around for editing even if unreferenced.
 	RF_MarkAsNative				=0x00000004,	///< Object (UField) will be marked as native on construction (DO NOT USE THIS FLAG in HasAnyFlags() etc)
 	RF_Transactional			=0x00000008,	///< Object is transactional.
-	RF_ClassDefaultObject		=0x00000010,	///< This object is its class's default object
-	RF_ArchetypeObject			=0x00000020,	///< This object is a template for another object - treat like a class default object
+	RF_ClassDefaultObject		=0x00000010,	///< This object is used as the default template for all instances of a class. One object is created for each class
+	RF_ArchetypeObject			=0x00000020,	///< This object can be used as a template for instancing objects. This is set on all types of object templates
 	RF_Transient				=0x00000040,	///< Don't save object.
 
 	// This group of flags is primarily concerned with garbage collection.
@@ -558,11 +558,11 @@ enum EObjectFlags
 
 	// Misc. Flags
 	RF_BeingRegenerated			=0x00020000,	///< Flagged on UObjects that are used to create UClasses (e.g. Blueprints) while they are regenerating their UClass on load (See FLinkerLoad::CreateExport()), as well as UClass objects in the midst of being created
-	RF_DefaultSubObject			=0x00040000,	///< Flagged on subobjects that are defaults
+	RF_DefaultSubObject			=0x00040000,	///< Flagged on subobject templates that were created in a class constructor, and all instances created from those templates
 	RF_WasLoaded				=0x00080000,	///< Flagged on UObjects that were loaded
 	RF_TextExportTransient		=0x00100000,	///< Do not export object to text form (e.g. copy/paste). Generally used for sub-objects that can be regenerated from data in their parent object.
 	RF_LoadCompleted			=0x00200000,	///< Object has been completely serialized by linkerload at least once. DO NOT USE THIS FLAG, It should be replaced with RF_WasLoaded.
-	RF_InheritableComponentTemplate = 0x00400000, ///< Archetype of the object can be in its super class
+	RF_InheritableComponentTemplate = 0x00400000, ///< Flagged on subobject templates stored inside a class instead of the class default object, they are instanced after default subobjects
 	RF_DuplicateTransient		=0x00800000,	///< Object should not be included in any type of duplication (copy/paste, binary duplication, etc.)
 	RF_StrongRefOnFrame			=0x01000000,	///< References to this object from persistent function frame are handled as strong ones.
 	RF_NonPIEDuplicateTransient	=0x02000000,	///< Object should not be included for duplication unless it's being duplicated for a PIE session
diff --git a/Engine/Source/Runtime/CoreUObject/Public/UObject/UObjectBaseUtility.h b/Engine/Source/Runtime/CoreUObject/Public/UObject/UObjectBaseUtility.h
--- a/Engine/Source/Runtime/CoreUObject/Public/UObject/UObjectBaseUtility.h
+++ b/Engine/Source/Runtime/CoreUObject/Public/UObject/UObjectBaseUtility.h
@@ -495,11 +495,12 @@ public:
 	COREUOBJECT_API bool MarkPackageDirty() const;
 
 	/**
-	* Determines whether this object is a template object
-	*
-	* @return	true if this object is a template object (owned by a UClass)
-	*/
-	COREUOBJECT_API bool IsTemplate( EObjectFlags TemplateTypes = RF_ArchetypeObject|RF_ClassDefaultObject ) const;
+	 * Determines whether this object is a template object by checking flags on the object and the outer chain.
+	 *
+	 * @param	TemplateTypes	Specific flags to look for, the default checks class default and non-class templates
+	 * @return	true if this object is a template object that can be used as an archetype
+	 */
+	COREUOBJECT_API bool IsTemplate(EObjectFlags TemplateTypes = RF_ArchetypeObject|RF_ClassDefaultObject) const;
 
 	/**
 	 * Traverses the outer chain searching for the next object of a certain type.  (T must be derived from UObject)
@@ -653,9 +654,23 @@ public:
 	}
 
 	/**
-	 * Returns whether this component was instanced from a component/subobject template, or if it is a component/subobject template.
+	 * Returns whether this is a template that can be used to create instanced subobjects.
+	 * By default this includes subobjects of a class default object and other templates associated with a class.
+	 * 
+	 * @param	TemplateTypes	Specific flags to check on this object and the outer chain
+	 * @return	true if this is a template for creating instanced subobjects
+	 */
+	COREUOBJECT_API bool IsTemplateForSubobjects(EObjectFlags TemplateTypes = RF_ClassDefaultObject|RF_DefaultSubObject|RF_InheritableComponentTemplate) const;
+	
+	/**
+	 * Returns whether this is a subobject (template or instance) that was originally defined inside a default class object.
+	 * This will return true for all instanced objects created from an archetype that is not a class default object,
+	 * but will only return true for template objects nested directly inside a class default object.
+	 *
+	 * @warning The behavior of this function is inconsistent for historical reasons and is not equivalent to RF_DefaultSubobject.
+	 * 			Call IsTemplateSubobject to handle all types of subobject templates, or call GetArchetype() first if you have an instance.
 	 *
-	 * @return	true if this component was instanced from a template.  false if this component was created manually at runtime.
+	 * @return	true if this was instanced from a subobject template, or is the direct subobject of a class default object
 	 */
 	COREUOBJECT_API bool IsDefaultSubobject() const;
 	
diff --git a/Engine/Source/Runtime/CoreUObject/Public/UObject/UObjectGlobals.h b/Engine/Source/Runtime/CoreUObject/Public/UObject/UObjectGlobals.h
--- a/Engine/Source/Runtime/CoreUObject/Public/UObject/UObjectGlobals.h
+++ b/Engine/Source/Runtime/CoreUObject/Public/UObject/UObjectGlobals.h
@@ -1195,11 +1195,11 @@ public:
 	COREUOBJECT_API UClass* GetClass() const;
 
 	/**
-	 * Create a component or subobject
+	 * Create a component or subobject that will be instanced inside all instances of this class.
 	 * @param	TReturnType					class of return type, all overrides must be of this type
 	 * @param	Outer						outer to construct the subobject in
-	 * @param	SubobjectName				name of the new component
-	 * @param bTransient		true if the component is being assigned to a transient property
+	 * @param	SubobjectName				name of the new component, this will be the same for all instances of this class
+	 * @param	bTransient					true if the component is being assigned to a transient property
 	 */
 	template<class TReturnType>
 	TReturnType* CreateDefaultSubobject(UObject* Outer, FName SubobjectName, bool bTransient = false) const
@@ -1209,12 +1209,12 @@ public:
 	}
 
 	/**
-	 * Create optional component or subobject. Optional subobjects may not get created
-	 * when a derived class specified DoNotCreateDefaultSubobject with the subobject's name.
+	 * Create optional component or subobject. Optional subobjects will not get created.
+	 * if a derived class specifies DoNotCreateDefaultSubobject with the subobject name.
 	 * @param	TReturnType					class of return type, all overrides must be of this type
 	 * @param	Outer						outer to construct the subobject in
-	 * @param	SubobjectName				name of the new component
-	 * @param bTransient		true if the component is being assigned to a transient property
+	 * @param	SubobjectName				name of the new component, this will be the same for all instances of this class
+	 * @param	bTransient					true if the component is being assigned to a transient property
 	 */
 	template<class TReturnType>
 	TReturnType* CreateOptionalDefaultSubobject(UObject* Outer, FName SubobjectName, bool bTransient = false) const
@@ -1224,13 +1224,13 @@ public:
 	}
 
 	/** 
-	* Create a component or subobject 
-	* @param TReturnType class of return type, all overrides must be of this type 
-	* @param TClassToConstructByDefault class to construct by default
-	* @param Outer outer to construct the subobject in 
-	* @param SubobjectName name of the new component 
-	* @param bTransient		true if the component is being assigned to a transient property
-	*/ 
+	 * Create a component or subobject, allows creating a child class and returning the parent class.
+	 * @param	TReturnType					class of return type, all overrides must be of this type 
+	 * @param	TClassToConstructByDefault	class to construct by default
+	 * @param	Outer						outer to construct the subobject in
+	 * @param	SubobjectName				name of the new component, this will be the same for all instances of this class
+	 * @param	bTransient					true if the component is being assigned to a transient property
+	 */ 
 	template<class TReturnType, class TClassToConstructByDefault> 
 	TReturnType* CreateDefaultSubobject(UObject* Outer, FName SubobjectName, bool bTransient = false) const 
 	{ 
@@ -1241,7 +1241,7 @@ public:
 	 * Create a component or subobject only to be used with the editor.
 	 * @param	TReturnType					class of return type, all overrides must be of this type
 	 * @param	Outer						outer to construct the subobject in
-	 * @param	SubobjectName				name of the new component
+	 * @param	SubobjectName				name of the new component, this will be the same for all instances of this class
 	 * @param	bTransient					true if the component is being assigned to a transient property
 	 */
 	template<class TReturnType>
@@ -1254,14 +1254,14 @@ public:
 	/**
 	* Create a component or subobject only to be used with the editor.
 	* @param	Outer						outer to construct the subobject in
-	* @param	SubobjectName				name of the new component
+	* @param	SubobjectName				name of the new component, this will be the same for all instances of this class
 	* @param	ReturnType					type of the new component
 	* @param	bTransient					true if the component is being assigned to a transient property
 	*/
 	COREUOBJECT_API UObject* CreateEditorOnlyDefaultSubobject(UObject* Outer, FName SubobjectName, const UClass* ReturnType, bool bTransient = false) const;
 
 	/**
-	 * Create a component or subobject
+	 * Create a component or subobject that will be instanced inside all instances of this class.
 	 * @param	Outer                       outer to construct the subobject in
 	 * @param	SubobjectName               name of the new component
 	 * @param	ReturnType                  class of return type, all overrides must be of this type
