Add overload of GenerateDeterministicTextKey that works with a property name

#rnx

diff --git a/Engine/Source/Runtime/CoreUObject/Private/Internationalization/TextPackageNamespaceUtil.cpp b/Engine/Source/Runtime/CoreUObject/Private/Internationalization/TextPackageNamespaceUtil.cpp
--- a/Engine/Source/Runtime/CoreUObject/Private/Internationalization/TextPackageNamespaceUtil.cpp
+++ b/Engine/Source/Runtime/CoreUObject/Private/Internationalization/TextPackageNamespaceUtil.cpp
@@ -178,6 +178,11 @@ FString TextNamespaceUtil::GenerateRandomTextKey()
 }
 
 FString TextNamespaceUtil::GenerateDeterministicTextKey(UObject* InTextOwner, const FTextProperty* InTextProperty, const bool bApplyPackageNamespace)
+{
+	return GenerateDeterministicTextKey(InTextOwner, InTextProperty->GetFName(), bApplyPackageNamespace);
+}
+
+FString TextNamespaceUtil::GenerateDeterministicTextKey(UObject* InTextOwner, const FName InTextPropertyName, const bool bApplyPackageNamespace)
 {
 	auto GetNameKeyHash = [](const FName Name)
 	{
@@ -196,7 +201,7 @@ FString TextNamespaceUtil::GenerateDeterministicTextKey(UObject* InTextOwner, co
 	};
 
 	// Build a (hopefully) unique and deterministic key from a combination of the text owner and text property info
-	FGuid KeyGuid(GetObjectKeyHash(InTextOwner->GetOuter()), GetNameKeyHash(InTextOwner->GetFName()), GetNameKeyHash(InTextOwner->GetClass()->GetFName()), GetNameKeyHash(InTextProperty->GetFName()));
+	FGuid KeyGuid(GetObjectKeyHash(InTextOwner->GetOuter()), GetNameKeyHash(InTextOwner->GetFName()), GetNameKeyHash(InTextOwner->GetClass()->GetFName()), GetNameKeyHash(InTextPropertyName));
 #if USE_STABLE_LOCALIZATION_KEYS
 	if (const FString PackageNamespace = bApplyPackageNamespace ? EnsurePackageNamespace(InTextOwner) : FString();
 		!PackageNamespace.IsEmpty())
diff --git a/Engine/Source/Runtime/CoreUObject/Public/Internationalization/TextPackageNamespaceUtil.h b/Engine/Source/Runtime/CoreUObject/Public/Internationalization/TextPackageNamespaceUtil.h
--- a/Engine/Source/Runtime/CoreUObject/Public/Internationalization/TextPackageNamespaceUtil.h
+++ b/Engine/Source/Runtime/CoreUObject/Public/Internationalization/TextPackageNamespaceUtil.h
@@ -115,6 +115,7 @@ COREUOBJECT_API FString GenerateRandomTextKey();
  * @param bApplyPackageNamespace		If true, apply the package namespace to the generated key hash (when USE_STABLE_LOCALIZATION_KEYS is true).
  */
 COREUOBJECT_API FString GenerateDeterministicTextKey(UObject* InTextOwner, const FTextProperty* InTextProperty, const bool bApplyPackageNamespace = true);
+COREUOBJECT_API FString GenerateDeterministicTextKey(UObject* InTextOwner, const FName InTextPropertyName, const bool bApplyPackageNamespace = true);
 
 enum class ETextEditAction : uint8
 {
