Updated FPropertyTypeName to make it more convenient to use

#rnx

diff --git a/Engine/Source/Runtime/CoreUObject/Private/UObject/PropertyTypeName.cpp b/Engine/Source/Runtime/CoreUObject/Private/UObject/PropertyTypeName.cpp
--- a/Engine/Source/Runtime/CoreUObject/Private/UObject/PropertyTypeName.cpp
+++ b/Engine/Source/Runtime/CoreUObject/Private/UObject/PropertyTypeName.cpp
@@ -337,6 +337,8 @@ void FPropertyTypeNameBuilder::EndTypeParameters()
 
 void FPropertyTypeNameBuilder::AddTypeName(FName Name)
 {
+	checkf(ActiveIndex >= 0 || Nodes.IsEmpty(), TEXT("Only one type name may be added as the root node."));
+
 	FPropertyTypeNameNode& Node = Nodes.AddDefaulted_GetRef();
 	Node.Name = Name;
 	Node.InnerCount = 0;
@@ -348,6 +350,20 @@ void FPropertyTypeNameBuilder::AddTypeName(FName Name)
 	}
 }
 
+void FPropertyTypeNameBuilder::AddTypeName(FPropertyTypeName Name)
+{
+	AddTypeName(Name.GetTypeName());
+	if (const int32 Count = Name.GetTypeParameterCount())
+	{
+		BeginTypeParameters();
+		for (int32 Index = 0; Index < Count; ++Index)
+		{
+			AddTypeName(Name.GetTypeParameter(Index));
+		}
+		EndTypeParameters();
+	}
+}
+
 FPropertyTypeName FPropertyTypeNameBuilder::Build() const
 {
 	FPropertyTypeName Type;
diff --git a/Engine/Source/Runtime/CoreUObject/Public/UObject/PropertyTypeName.h b/Engine/Source/Runtime/CoreUObject/Public/UObject/PropertyTypeName.h
--- a/Engine/Source/Runtime/CoreUObject/Public/UObject/PropertyTypeName.h
+++ b/Engine/Source/Runtime/CoreUObject/Public/UObject/PropertyTypeName.h
@@ -58,11 +58,27 @@ public:
 	/**
 	 * Returns the indexed parameter under the root of this property type name.
 	 *
+	 * An out-of-bounds index will return an empty type name.
+	 *
 	 * Example: MapProperty<StructProperty<KeyStruct>,EnumProperty<ByteProperty<ByteEnum>>>
 	 * - GetTypeParameter(0) -> StructProperty<KeyStruct>
 	 * - GetTypeParameter(1) -> EnumProperty<ByteProperty<ByteEnum>>
 	 */
-	UE_API FPropertyTypeName GetTypeParameter(int32 ParamIndex) const;
+	UE_API FPropertyTypeName GetTypeParameter(int32 ParamIndex = 0) const;
+
+	/**
+	 * Returns the indexed parameter type name under the root of this property type name.
+	 *
+	 * An out-of-bounds index will return a name of None.
+	 *
+	 * Example: MapProperty<StructProperty<KeyStruct>,EnumProperty<ByteProperty<ByteEnum>>>
+	 * - GetTypeParameterName(0) -> StructProperty
+	 * - GetTypeParameterName(1) -> EnumProperty
+	 */
+	inline FName GetTypeParameterName(int32 ParamIndex = 0) const
+	{
+		return GetTypeParameter(ParamIndex).GetTypeName();
+	}
 
 	/**
 	 * Resets this to an empty type name.
@@ -93,15 +109,15 @@ private:
  * Example: MapProperty<StructProperty<KeyStruct>,EnumProperty<ByteProperty<ByteEnum>>>
  *
  * FPropertyTypeNameBuilder Builder;
- * Builder.AddTypeName(TEXT("MapProperty"));
+ * Builder.AddTypeName(NAME_MapProperty));
  * Builder.BeginTypeParameters();
- *   Builder.AddTypeName(TEXT("StructProperty"));
+ *   Builder.AddTypeName(NAME_StructProperty);
  *   Builder.BeginTypeParameters();
  *     Builder.AddTypeName(TEXT("KeyStruct"));
  *   Builder.EndTypeParameters();
- *   Builder.AddTypeName(TEXT("EnumProperty"));
+ *   Builder.AddTypeName(NAME_EnumProperty);
  *   Builder.BeginTypeParameters();
- *     Builder.AddTypeName(TEXT("ByteProperty"));
+ *     Builder.AddTypeName(NAME_ByteProperty);
  *     Builder.BeginTypeParameters();
  *       Builder.AddTypeName(TEXT("ByteEnum"));
  *     Builder.EndTypeParameters();
@@ -112,6 +128,10 @@ private:
 class FPropertyTypeNameBuilder
 {
 public:
+	/** Add a complete type name with its type parameters. */
+	UE_API void AddTypeName(FPropertyTypeName Name);
+
+	/** Add a type name without any type parameters. */
 	UE_API void AddTypeName(FName Name);
 
 	/** Mark the beginning of the type parameters for the last added type. */
diff --git a/Engine/Source/Runtime/CoreUObject/Tests/UObject/PropertyTypeNameTest.cpp b/Engine/Source/Runtime/CoreUObject/Tests/UObject/PropertyTypeNameTest.cpp
--- a/Engine/Source/Runtime/CoreUObject/Tests/UObject/PropertyTypeNameTest.cpp
+++ b/Engine/Source/Runtime/CoreUObject/Tests/UObject/PropertyTypeNameTest.cpp
@@ -200,6 +200,14 @@ TEST_CASE_NAMED(FPropertyTypeNameSmokeTest, "CoreUObject::PropertyTypeName::Smok
 		Builder.EndTypeParameters();
 		const FPropertyTypeName MapIntToIntToInt = Builder.Build();
 
+		Builder.Reset();
+		Builder.AddTypeName(NAME_MapProperty);
+		Builder.BeginTypeParameters();
+		Builder.AddTypeName(NAME_IntProperty);
+		Builder.AddTypeName(MapIntToInt);
+		Builder.EndTypeParameters();
+		const FPropertyTypeName MapIntToIntToIntAlternative = Builder.Build();
+
 		CHECK(FPropertyTypeName() == FPropertyTypeName());
 		CHECK(GetTypeHash(FPropertyTypeName()) == GetTypeHash(FPropertyTypeName()));
 
@@ -236,6 +244,8 @@ TEST_CASE_NAMED(FPropertyTypeNameSmokeTest, "CoreUObject::PropertyTypeName::Smok
 		CHECK(VectorArray < MapIntToIntToInt);
 		CHECK(GetTypeHash(MapIntToInt) == GetTypeHash(MapIntToIntToInt.GetTypeParameter(1)));
 		CHECK_FALSE(GetTypeHash(MapIntToInt) == GetTypeHash(MapIntToIntToInt));
+
+		CHECK(MapIntToIntToInt == MapIntToIntToIntAlternative);
 	}
 
 	SECTION("Archive")
