Re-enable check for invalid PackageHandles when resolving Packaged Objects. This check helps with identifying cases where garbage memory is used to construct objects

diff --git a/Engine/Source/Runtime/CoreUObject/Private/UObject/ObjectHandleIndex.cpp b/Engine/Source/Runtime/CoreUObject/Private/UObject/ObjectHandleIndex.cpp
--- a/Engine/Source/Runtime/CoreUObject/Private/UObject/ObjectHandleIndex.cpp
+++ b/Engine/Source/Runtime/CoreUObject/Private/UObject/ObjectHandleIndex.cpp
@@ -313,7 +313,7 @@ namespace UE::CoreUObject::Private
 
 		if (PackageIndex >= static_cast<uint32>(GObjectHandleIndex.PackageData.Num()))
 		{
-			//checkf(false, TEXT("FObjectHandle: PackageIndex invalid.  This ObjectHandle is from malformed data."));
+			checkf(false, TEXT("FPackageId: PackageIndex invalid. This ObjectHandle is from malformed data."));
 			return;
 		}
 		FObjectHandlePackageData& FoundPackageData = GObjectHandleIndex.PackageData[PackageIndex];
@@ -596,4 +596,4 @@ namespace UE::CoreUObject::Private
 }
 
 
-#endif
\ No newline at end of file
+#endif
