Added clear overides API to the Overridable Manager.

diff --git a/Engine/Source/Runtime/CoreUObject/Private/UObject/Class.cpp b/Engine/Source/Runtime/CoreUObject/Private/UObject/Class.cpp
--- a/Engine/Source/Runtime/CoreUObject/Private/UObject/Class.cpp
+++ b/Engine/Source/Runtime/CoreUObject/Private/UObject/Class.cpp
@@ -1424,7 +1424,7 @@ void UStruct::SerializeVersionedTaggedProperties(FStructuredArchive::FSlot Slot,
 		if (EnumHasAnyFlags(SerializationControl, EClassSerializationControlExtension::OverridableSerializationInformation))
 		{
 			checkf(!UnderlyingArchive.ArUseCustomPropertyList, TEXT("Overridable serialization does not support custom property list"))
-			EOverriddenPropertyOperation Operation = UnderlyingArchive.IsSaving() ? ControlContext.OverriddenProperties->GetOverriddenPropertyOperation(/*CurrentPropertyChain*/nullptr, /*Property*/nullptr) : EOverriddenPropertyOperation::None;
+			EOverriddenPropertyOperation Operation = UnderlyingArchive.IsSaving() ? ControlContext.OverriddenProperties->GetOverriddenPropertyOperation((FArchiveSerializedPropertyChain*)nullptr, (FProperty*)nullptr) : EOverriddenPropertyOperation::None;
 			Slot << SA_ATTRIBUTE(TEXT("OverridableOperation"), Operation);
 
 			if (UnderlyingArchive.IsLoading())
diff --git a/Engine/Source/Runtime/CoreUObject/Private/UObject/OverridableManager.cpp b/Engine/Source/Runtime/CoreUObject/Private/UObject/OverridableManager.cpp
--- a/Engine/Source/Runtime/CoreUObject/Private/UObject/OverridableManager.cpp
+++ b/Engine/Source/Runtime/CoreUObject/Private/UObject/OverridableManager.cpp
@@ -63,9 +63,9 @@ FOverriddenPropertySet& FOverridableManager::SetOverriddenProperties(UObject& Ob
 
 EOverriddenState FOverridableManager::GetOverriddenState(UObject& Object)
 {
-	if(FOverriddenPropertySet* OverriddenProperties = GetOverriddenProperties(Object))
+	if(const FOverriddenPropertySet* OverriddenProperties = GetOverriddenProperties(Object))
 	{
-		const EOverriddenPropertyOperation Operation = OverriddenProperties->GetOverriddenPropertyOperation(nullptr, nullptr);
+		const EOverriddenPropertyOperation Operation = OverriddenProperties->GetOverriddenPropertyOperation((FArchiveSerializedPropertyChain*)nullptr, (FProperty*)nullptr);
 		if (Operation != EOverriddenPropertyOperation::None)
 		{
 			return Operation == EOverriddenPropertyOperation::Replace ? EOverriddenState::AllOverridden : EOverriddenState::HasOverrides;
@@ -118,13 +118,22 @@ void FOverridableManager::PropagateOverrideToInstancedSubObjects(UObject& Object
 
 void FOverridableManager::OverrideProperty(UObject& Object, const FPropertyChangedEvent& PropertyEvent, const FEditPropertyChain& PropertyChain)
 {
-	if (FOverriddenPropertySet* ThisObjectOverriddenProperties = OverriddenObjectAnnotations.Find(Object) )
+	if (FOverriddenPropertySet* ThisObjectOverriddenProperties = OverriddenObjectAnnotations.Find(Object))
 	{
 		ThisObjectOverriddenProperties->NotifyPropertyChange(EPropertyNotificationType::PreEdit, FPropertyChangedEvent(nullptr), PropertyChain.GetHead(), &Object);
 		ThisObjectOverriddenProperties->NotifyPropertyChange(EPropertyNotificationType::PostEdit, PropertyEvent, PropertyChain.GetHead(), &Object);
 	}
 }
 
+bool FOverridableManager::ClearOverriddenProperty(UObject& Object, const FPropertyChangedEvent& PropertyEvent, const FEditPropertyChain& PropertyChain)
+{
+	if (FOverriddenPropertySet* ThisObjectOverriddenProperties = OverriddenObjectAnnotations.Find(Object))
+	{
+		return ThisObjectOverriddenProperties->ClearOverriddenProperty(PropertyEvent, PropertyChain.GetHead());
+	}
+	return false;
+}
+
 void FOverridableManager::PreOverrideProperty(UObject& Object, const FEditPropertyChain& PropertyChain)
 {
 	if (FOverriddenPropertySet* ThisObjectOverriddenProperties = OverriddenObjectAnnotations.Find(Object))
@@ -149,24 +158,15 @@ void FOverridableManager::NotifyPropertyChange(const EPropertyNotificationType N
 	}
 }
 
-EOverriddenPropertyOperation FOverridableManager::GetOverriddenPropertyOperation(UObject& Object, const FEditPropertyChain& PropertyChain)
+EOverriddenPropertyOperation FOverridableManager::GetOverriddenPropertyOperation(UObject& Object, const FPropertyChangedEvent& PropertyEvent, const FEditPropertyChain& PropertyChain, bool* bOutInheritedState)
 {
-	if (FOverriddenPropertySet* ThisObjectOverriddenProperties = OverriddenObjectAnnotations.Find(Object))
+	if (const FOverriddenPropertySet* ThisObjectOverriddenProperties = OverriddenObjectAnnotations.Find(Object))
 	{
-		return ThisObjectOverriddenProperties->GetOverriddenPropertyOperation(PropertyChain.GetHead());
+		return ThisObjectOverriddenProperties->GetOverriddenPropertyOperation(PropertyEvent, PropertyChain.GetHead(), bOutInheritedState);
 	}
 	return EOverriddenPropertyOperation::None;
 }
 
-bool FOverridableManager::SetOverriddenPropertyOperation(UObject& Object, const FEditPropertyChain& PropertyChain, EOverriddenPropertyOperation Operation)
-{
-	if (FOverriddenPropertySet* ThisObjectOverriddenProperties = OverriddenObjectAnnotations.Find(Object))
-	{
-		return ThisObjectOverriddenProperties->SetOverriddenPropertyOperation(Operation, PropertyChain.GetHead()) != nullptr;
-	}
-	return false;
-}
-
 void FOverridableManager::ClearOverrides(UObject& Object)
 {
 	if(FOverriddenPropertySet* ThisObjectOverriddenProperties = OverriddenObjectAnnotations.Find(Object))
@@ -204,7 +204,7 @@ void FOverridableManager::SerializeOverriddenProperties(UObject& Object, FStruct
 	TOptional<FStructuredArchiveSlot> OverridenPropertiesSlot = ObjectRecord.TryEnterField(TEXT("OverridenProperties"), OverriddenProperties != nullptr);
 	if (OverridenPropertiesSlot.IsSet())
 	{
-		EOverriddenPropertyOperation Operation = OverriddenProperties ? OverriddenProperties->GetOverriddenPropertyOperation(nullptr, nullptr) : EOverriddenPropertyOperation::None;
+		EOverriddenPropertyOperation Operation = OverriddenProperties ? OverriddenProperties->GetOverriddenPropertyOperation((FArchiveSerializedPropertyChain*)nullptr, (FProperty*)nullptr) : EOverriddenPropertyOperation::None;
 		*OverridenPropertiesSlot << SA_ATTRIBUTE( TEXT("OverriddenOperation"), Operation);
 
 		if (ArchiveState.IsLoading())
diff --git a/Engine/Source/Runtime/CoreUObject/Private/UObject/OverriddenPropertySet.cpp b/Engine/Source/Runtime/CoreUObject/Private/UObject/OverriddenPropertySet.cpp
--- a/Engine/Source/Runtime/CoreUObject/Private/UObject/OverriddenPropertySet.cpp
+++ b/Engine/Source/Runtime/CoreUObject/Private/UObject/OverriddenPropertySet.cpp
@@ -156,18 +156,22 @@ FOverriddenPropertyNode& FOverriddenPropertySet::FindOrAddNode(FOverriddenProper
 	return OverriddenPropertyNodes.Get(NewID);
 }
 
-EOverriddenPropertyOperation FOverriddenPropertySet::GetOverriddenPropertyOperation(const FOverriddenPropertyNode& ParentPropertyNode, const FEditPropertyChain::TDoubleLinkedListNode* PropertyNode) const
+FOverriddenPropertyNode* FOverriddenPropertySet::GetOverriddenPropertyNode(FOverriddenPropertyNode& ParentPropertyNode, const FEditPropertyChain::TDoubleLinkedListNode* PropertyNode, bool* bOutInheritedState)
 {
 	// No need to look further
 	// if it is the entire property is replaced or
 	// if it is the FOverriddenPropertySet struct which is always Overridden
 	if (ParentPropertyNode.Operation == EOverriddenPropertyOperation::Replace)
 	{
-		return EOverriddenPropertyOperation::Replace;
+		if (bOutInheritedState)
+		{
+			*bOutInheritedState = PropertyNode != nullptr;
+		}
+		return &ParentPropertyNode;
 	}
 
 	const FEditPropertyChain::TDoubleLinkedListNode* PropertyIterator = PropertyNode;
-	const FOverriddenPropertyNode* OverriddenPropertyNode = &ParentPropertyNode;
+	FOverriddenPropertyNode* OverriddenPropertyNode = &ParentPropertyNode;
 	while (PropertyIterator && OverriddenPropertyNode && OverriddenPropertyNode->Operation != EOverriddenPropertyOperation::Replace)
 	{
 		const FProperty* CurrentProperty = PropertyIterator->GetValue();
@@ -185,37 +189,16 @@ EOverriddenPropertyOperation FOverriddenPropertySet::GetOverriddenPropertyOperat
 		PropertyIterator = PropertyIterator->GetNextNode();
 	}
 
-	return OverriddenPropertyNode ? OverriddenPropertyNode->Operation : EOverriddenPropertyOperation::None;
-}
-
-FOverriddenPropertyNode* FOverriddenPropertySet::SetOverriddenPropertyOperation(EOverriddenPropertyOperation Operation, FOverriddenPropertyNode& ParentPropertyNode, const FEditPropertyChain::TDoubleLinkedListNode* PropertyNode)
-{
-	// No need to look further
-	// if it is the entire property is replaced or
-	// if it is the FOverriddenPropertySet struct which is always Overridden
-	if (ParentPropertyNode.Operation == EOverriddenPropertyOperation::Replace)
-	{
-		return nullptr;
-	}
-
-	const FEditPropertyChain::TDoubleLinkedListNode* PropertyIterator = PropertyNode;
-	FOverriddenPropertyNode* OverriddenPropertyNode = &ParentPropertyNode;
-	while (PropertyIterator && OverriddenPropertyNode->Operation != EOverriddenPropertyOperation::Replace)
-	{
-		const FProperty* CurrentProperty = PropertyIterator->GetValue();
-		const FName CurrentPropID = CurrentProperty->GetFName();
-		OverriddenPropertyNode = &FindOrAddNode(*OverriddenPropertyNode, CurrentPropID);
-		PropertyIterator = PropertyIterator->GetNextNode();
-	}
-
-	// Might have stop before as one of the parent property was completely replaced.
-	if (!PropertyIterator)
+	if (bOutInheritedState)
 	{
-		OverriddenPropertyNode->Operation = Operation;
-		return OverriddenPropertyNode;
+		*bOutInheritedState = PropertyIterator != nullptr;
 	}
+	return OverriddenPropertyNode;
+}
 
-	return nullptr;
+const FOverriddenPropertyNode* FOverriddenPropertySet::GetOverriddenPropertyNode(const FOverriddenPropertyNode& ParentPropertyNode, const FEditPropertyChain::TDoubleLinkedListNode* PropertyNode, bool* bOutInheritedState) const
+{
+	return const_cast<FOverriddenPropertySet*>(this)->GetOverriddenPropertyNode(const_cast<FOverriddenPropertyNode&>(ParentPropertyNode), PropertyNode, bOutInheritedState);
 }
 
 void FOverriddenPropertySet::NotifyPropertyChange(FOverriddenPropertyNode* ParentPropertyNode, const EPropertyNotificationType Notification, const FPropertyChangedEvent& PropertyEvent, const FEditPropertyChain::TDoubleLinkedListNode* PropertyNode, const void* Data)
@@ -757,18 +740,30 @@ void FOverriddenPropertySet::RemoveOverriddenSubProperties(FOverriddenPropertyNo
 	PropertyNode.SubPropertyNodeKeys.Empty();
 }
 
-EOverriddenPropertyOperation FOverriddenPropertySet::GetOverriddenPropertyOperation(const FEditPropertyChain::TDoubleLinkedListNode* PropertyNode) const
+EOverriddenPropertyOperation FOverriddenPropertySet::GetOverriddenPropertyOperation(const FPropertyChangedEvent& PropertyEvent, const FEditPropertyChain::TDoubleLinkedListNode* PropertyNode, bool* bOutInheritedState /*= nullptr*/) const
 {
 	if (const FOverriddenPropertyNode* RootNode = OverriddenPropertyNodes.FindByHash(GetTypeHash(RootNodeID), RootNodeID))
 	{
-		return GetOverriddenPropertyOperation(*RootNode, PropertyNode);
+		const FOverriddenPropertyNode* OverriddenPropertyNode = GetOverriddenPropertyNode(*RootNode, PropertyNode, bOutInheritedState);
+		return OverriddenPropertyNode ? OverriddenPropertyNode->Operation : EOverriddenPropertyOperation::None;
 	}
 	return EOverriddenPropertyOperation::None;
 }
 
-FOverriddenPropertyNode* FOverriddenPropertySet::SetOverriddenPropertyOperation(EOverriddenPropertyOperation Operation, const FEditPropertyChain::TDoubleLinkedListNode* PropertyNode)
+bool FOverriddenPropertySet::ClearOverriddenProperty(const FPropertyChangedEvent& PropertyEvent, const FEditPropertyChain::TDoubleLinkedListNode* PropertyNode)
 {
-	return SetOverriddenPropertyOperation(Operation, OverriddenPropertyNodes.FindOrAddByHash(GetTypeHash(RootNodeID), RootNodeID), PropertyNode);
+	if (FOverriddenPropertyNode* RootNode = OverriddenPropertyNodes.FindByHash(GetTypeHash(RootNodeID), RootNodeID))
+	{
+		bool bOutInheritedState = false;
+		FOverriddenPropertyNode* OverriddenPropertyNode = GetOverriddenPropertyNode(*RootNode, PropertyNode, &bOutInheritedState);
+		if (OverriddenPropertyNode && OverriddenPropertyNode->Operation != EOverriddenPropertyOperation::None && !bOutInheritedState)
+		{
+			RemoveOverriddenSubProperties(*OverriddenPropertyNode);
+			OverriddenPropertyNode->Operation = EOverriddenPropertyOperation::None;
+			return true;
+		}
+	}
+	return false;
 }
 
 void FOverriddenPropertySet::NotifyPropertyChange(const EPropertyNotificationType Notification, const FPropertyChangedEvent& PropertyEvent, const FEditPropertyChain::TDoubleLinkedListNode* PropertyNode, const void* Data)
diff --git a/Engine/Source/Runtime/CoreUObject/Public/UObject/OverridableManager.h b/Engine/Source/Runtime/CoreUObject/Public/UObject/OverridableManager.h
--- a/Engine/Source/Runtime/CoreUObject/Public/UObject/OverridableManager.h
+++ b/Engine/Source/Runtime/CoreUObject/Public/UObject/OverridableManager.h
@@ -166,10 +166,18 @@ public:
 	 * Override a specific property of an object (Helper methods to call Pre/PostOverride)
 	 * Note: Supports object that does not have overridable serialization enabled
 	 * @param Object owning the property
-	 * @param PropertyEvent information about the type of change
+	 * @param PropertyEvent information about the type of change including any container item index
 	 * @param PropertyChain leading to the property that is about to be overridden */
 	COREUOBJECT_API void OverrideProperty(UObject& Object, const FPropertyChangedEvent& PropertyEvent, const FEditPropertyChain& PropertyChain);
 
+	/**
+	 * Clears an overridden properties specified by the property chain
+	 * @param Object owning the property to clear
+	 * @param PropertyEvent only needed to know about the container item index in any
+	 * @param PropertyChain to the property to clear from the root of the specified object
+	 * @return true if the property was successfully cleared. */
+	COREUOBJECT_API bool ClearOverriddenProperty(UObject& Object, const FPropertyChangedEvent& PropertyEvent, const FEditPropertyChain& PropertyChain);
+
 	/**
 	 * To be called prior to override a property of the specified object
 	 * Note: Supports object that does not have overridable serialization enabled
@@ -181,7 +189,7 @@ public:
 	 * To be called after the property was overridden of the specified object
 	 * Note: Supports object that does not have overridable serialization enabled
 	 * @param Object owning the property
-	 * @param PropertyEvent information about the type of change
+	 * @param PropertyEvent information about the type of change including any container item index
 	 * @param PropertyChain leading to the property that was overridden */
 	COREUOBJECT_API void PostOverrideProperty(UObject& Object, const FPropertyChangedEvent& PropertyEvent, const FEditPropertyChain& PropertyChain);
 
@@ -190,33 +198,23 @@ public:
 	 * Note: Will ensure if the specified object does not have overridable serialization enabled
 	 * @param Notification the type of notification (pre or post
 	 * @param Object owning the property
-	 * @param PropertyEvent information about the type of change
-	 * @param PropertyNode leading to the property that is changing
-	 */
+	 * @param PropertyEvent information about the type of change including any container item index
+	 * @param PropertyNode leading to the property that is changing */
 	COREUOBJECT_API void NotifyPropertyChange(const EPropertyNotificationType Notification, UObject& Object, const FPropertyChangedEvent& PropertyEvent, const FEditPropertyChain::TDoubleLinkedListNode* PropertyNode);
 
 	/**
 	 * Retrieve the overridable operation from the specified the edit property chain
 	 * @param Object owning the property
+	 * @param PropertyEvent only needed to know about the container item index in any
 	 * @param PropertyChain leading to the property the caller is interested in
-	 * @return the current type of override operation on the property
-	 */
-	COREUOBJECT_API EOverriddenPropertyOperation GetOverriddenPropertyOperation(UObject& Object, const FEditPropertyChain& PropertyChain);
-
-	/**
-	 * Set the overridable operation from the specified the edit property chain
-	 * @param Object owning the property
-	 * @param PropertyChain leading to the property the caller is interested in
-	 * @param Operation the override operation to set on the property
-	 * @return if the operation was successful
-	 */
-	COREUOBJECT_API bool SetOverriddenPropertyOperation(UObject& Object, const FEditPropertyChain& PropertyChain, EOverriddenPropertyOperation Operation);
+	 * @param bOutInheritedState optional parameter to know if the state returned was inherited from a parent property
+	 * @return the current type of override operation on the property */
+	COREUOBJECT_API EOverriddenPropertyOperation GetOverriddenPropertyOperation(UObject& Object, const FPropertyChangedEvent& PropertyEvent, const FEditPropertyChain& PropertyChain, bool* bOutInheritedState = nullptr);
 
 	/**
 	 * Serializes the overriden properties of the specified object into the record
 	 * @param Object to serialize the overriden property
-	 * @param ObjectRecord the record to use for serialization
-	 */
+	 * @param ObjectRecord the record to use for serialization */
 	COREUOBJECT_API void SerializeOverriddenProperties(UObject& Object, FStructuredArchive::FRecord ObjectRecord);
 
 	void HandleObjectsReInstantiated(const TMap<UObject*, UObject*>& OldToNewInstanceMap);
diff --git a/Engine/Source/Runtime/CoreUObject/Public/UObject/OverriddenPropertySet.h b/Engine/Source/Runtime/CoreUObject/Public/UObject/OverriddenPropertySet.h
--- a/Engine/Source/Runtime/CoreUObject/Public/UObject/OverriddenPropertySet.h
+++ b/Engine/Source/Runtime/CoreUObject/Public/UObject/OverriddenPropertySet.h
@@ -215,17 +215,19 @@ public:
 	{}
 
 	/**
-	 * Retrieve the overridable operation from the specified the edit property chain and the specified property
+	 * Retrieve the overridable operation from the specified the edit property chain node
+	 * @param PropertyEvent only needed to know about the container item index in any
 	 * @param PropertyNode in the linked list leading to the property the caller is interested in
+	 * @param bOutInheritedState optional parameter to know if the state returned was inherited from a parent property
 	 * @return the current type of override operation on the property */
-	EOverriddenPropertyOperation GetOverriddenPropertyOperation(const FEditPropertyChain::TDoubleLinkedListNode* PropertyNode) const;
+	EOverriddenPropertyOperation GetOverriddenPropertyOperation(const FPropertyChangedEvent& PropertyEvent, const FEditPropertyChain::TDoubleLinkedListNode* PropertyNode, bool* bOutInheritedState = nullptr) const;
 
 	/**
-	 * Setup the overridable operation of the current property from the serialized property chain and the specified property
-	 * @param Operation to set for this property
+	 * Clear any properties from the serialized property chain node
+	 * @param PropertyEvent only needed to know about the container item index in any
 	 * @param PropertyNode in the linked list leading to the property the caller is interested in
-	 * @return the node containing the information of the overridden property */
-	FOverriddenPropertyNode* SetOverriddenPropertyOperation(EOverriddenPropertyOperation Operation, const FEditPropertyChain::TDoubleLinkedListNode* PropertyNode);
+	 * @return if the operation was successful */
+	bool ClearOverriddenProperty(const FPropertyChangedEvent& PropertyEvent, const FEditPropertyChain::TDoubleLinkedListNode* PropertyNode);
 
 	/**
 	 * Handling and storing modification on a property of an object
@@ -287,8 +289,8 @@ protected:
 
 	FOverriddenPropertyNode& FindOrAddNode(FOverriddenPropertyNode& ParentPropertyNode, FOverriddenPropertyNodeID NodeID);
 
-	EOverriddenPropertyOperation GetOverriddenPropertyOperation(const FOverriddenPropertyNode& ParentPropertyNode, const FEditPropertyChain::TDoubleLinkedListNode* PropertyNode) const;
-	FOverriddenPropertyNode* SetOverriddenPropertyOperation(EOverriddenPropertyOperation Operation, FOverriddenPropertyNode& ParentPropertyNode, const FEditPropertyChain::TDoubleLinkedListNode* PropertyNode);
+	FOverriddenPropertyNode* GetOverriddenPropertyNode(FOverriddenPropertyNode& ParentPropertyNode, const FEditPropertyChain::TDoubleLinkedListNode* PropertyNode, bool* bOutInheritedState);
+	const FOverriddenPropertyNode* GetOverriddenPropertyNode(const FOverriddenPropertyNode& ParentPropertyNode, const FEditPropertyChain::TDoubleLinkedListNode* PropertyNode, bool* bOutInheritedState) const;
 	void NotifyPropertyChange(FOverriddenPropertyNode* ParentPropertyNode, const EPropertyNotificationType Notification, const FPropertyChangedEvent& PropertyEvent, const FEditPropertyChain::TDoubleLinkedListNode* PropertyNode, const void* Data);
 
 	EOverriddenPropertyOperation GetOverriddenPropertyOperation(const FOverriddenPropertyNode& ParentPropertyNode, const FArchiveSerializedPropertyChain* CurrentPropertyChain, FProperty* Property) const;
