Fix some incorrect calls to IsDefaultSubobject

diff --git a/Engine/Plugins/Enterprise/DataprepEditor/Source/DataprepEditor/Private/DataprepStats.cpp b/Engine/Plugins/Enterprise/DataprepEditor/Source/DataprepEditor/Private/DataprepStats.cpp
--- a/Engine/Plugins/Enterprise/DataprepEditor/Source/DataprepEditor/Private/DataprepStats.cpp
+++ b/Engine/Plugins/Enterprise/DataprepEditor/Source/DataprepEditor/Private/DataprepStats.cpp
@@ -255,7 +255,7 @@ private:
 		}
 
 		// Owned by a default object? Not part of a level either.
-		if(InComponent->GetOuter() && InComponent->GetOuter()->IsDefaultSubobject() )
+		if(InComponent->GetOuter() && InComponent->GetOuter()->IsTemplate())
 		{
 			return; // Result;
 		}
diff --git a/Engine/Plugins/Experimental/GameplayTargetingSystem/Source/GameplayTargetingSystem/Private/TargetingSystem/TargetingSubsystem.cpp b/Engine/Plugins/Experimental/GameplayTargetingSystem/Source/GameplayTargetingSystem/Private/TargetingSystem/TargetingSubsystem.cpp
--- a/Engine/Plugins/Experimental/GameplayTargetingSystem/Source/GameplayTargetingSystem/Private/TargetingSystem/TargetingSubsystem.cpp
+++ b/Engine/Plugins/Experimental/GameplayTargetingSystem/Source/GameplayTargetingSystem/Private/TargetingSystem/TargetingSubsystem.cpp
@@ -154,7 +154,7 @@ void UTargetingSubsystem::Deinitialize()
 
 bool UTargetingSubsystem::Exec_Runtime(UWorld* Inworld, const TCHAR* Cmd, FOutputDevice& Ar)
 {
-	if (HasAnyFlags(RF_ClassDefaultObject) || IsDefaultSubobject())
+	if (IsTemplate())
 	{
 		return false;
 	}
diff --git a/Engine/Source/Editor/Kismet/Private/WatchPointViewer.cpp b/Engine/Source/Editor/Kismet/Private/WatchPointViewer.cpp
--- a/Engine/Source/Editor/Kismet/Private/WatchPointViewer.cpp
+++ b/Engine/Source/Editor/Kismet/Private/WatchPointViewer.cpp
@@ -807,7 +807,7 @@ void WatchViewer::UpdateInstancedWatchDisplay()
 			FText BlueprintName = FText::FromString(BlueprintObj->GetName());
 
 			// Don't show info for the CDO
-			if (BlueprintInstance->IsDefaultSubobject())
+			if (BlueprintInstance->IsTemplate())
 			{
 				continue;
 			}
diff --git a/Engine/Source/Editor/StatsViewer/Private/StatsPages/PrimitiveStatsPage.cpp b/Engine/Source/Editor/StatsViewer/Private/StatsPages/PrimitiveStatsPage.cpp
--- a/Engine/Source/Editor/StatsViewer/Private/StatsPages/PrimitiveStatsPage.cpp
+++ b/Engine/Source/Editor/StatsViewer/Private/StatsPages/PrimitiveStatsPage.cpp
@@ -71,7 +71,7 @@ struct PrimitiveStatsGenerator
 		}
 
 		// Owned by a default object? Not part of a level either.
-		if(InPrimitiveComponent->GetOuter() && InPrimitiveComponent->GetOuter()->IsDefaultSubobject() )
+		if(InPrimitiveComponent->GetOuter() && InPrimitiveComponent->GetOuter()->IsTemplate())
 		{
 			return NULL;
 		}
diff --git a/Engine/Source/Runtime/CoreUObject/Private/UObject/ObjectMemoryAnalyzer.cpp b/Engine/Source/Runtime/CoreUObject/Private/UObject/ObjectMemoryAnalyzer.cpp
--- a/Engine/Source/Runtime/CoreUObject/Private/UObject/ObjectMemoryAnalyzer.cpp
+++ b/Engine/Source/Runtime/CoreUObject/Private/UObject/ObjectMemoryAnalyzer.cpp
@@ -42,9 +42,8 @@ void FObjectMemoryAnalyzer::AnalyzeObjects(const TArray<class UObject*>& InObjec
 
 void FObjectMemoryAnalyzer::AnalyzeObject(UObject* Object)
 {
-	if (ObjectList.Contains(Object) ||
-		(!(AnalyzeFlags&EAnalyzeFlags::IncludeDefaultObjects) && Object->HasAnyFlags(RF_ClassDefaultObject|RF_ArchetypeObject)) ||
-		Object->IsDefaultSubobject())
+	if (ObjectList.Contains(Object) || 
+		(!(AnalyzeFlags&EAnalyzeFlags::IncludeDefaultObjects) && Object->IsTemplate()))
 	{
 		return;
 	}
@@ -84,7 +83,7 @@ void FObjectMemoryAnalyzer::AnalyzeObjects( UClass* InBaseClass )
 	{
 		UObject* Object	= *It;
 
-		if (!(AnalyzeFlags&EAnalyzeFlags::IncludeDefaultObjects) && Object->IsDefaultSubobject()) { continue; };
+		if (!(AnalyzeFlags&EAnalyzeFlags::IncludeDefaultObjects) && Object->IsTemplate()) { continue; };
 
 		FObjectMemoryUsage& Annotation = MemUsageAnnotations.GetAnnotationRef(Object);
 
@@ -100,7 +99,7 @@ void FObjectMemoryAnalyzer::AnalyzeObjects( UClass* InBaseClass )
 	{
 		UObject* Object	= *It;
 
-		if (!(AnalyzeFlags&EAnalyzeFlags::IncludeDefaultObjects) && Object->IsDefaultSubobject()) { continue; };
+		if (!(AnalyzeFlags&EAnalyzeFlags::IncludeDefaultObjects) && Object->IsTemplate()) { continue; };
 	
 		FObjectMemoryUsage& Annotation = MemUsageAnnotations.GetAnnotationRef(Object);
 			
@@ -114,7 +113,7 @@ void FObjectMemoryAnalyzer::AnalyzeObjects( UClass* InBaseClass )
 	{
 		UObject* Object	= *It;
 		
-		if (!(AnalyzeFlags&EAnalyzeFlags::IncludeDefaultObjects) && Object->IsDefaultSubobject()) { continue; };
+		if (!(AnalyzeFlags&EAnalyzeFlags::IncludeDefaultObjects) && Object->IsTemplate()) { continue; };
 
 		FObjectMemoryUsage& Annotation = MemUsageAnnotations.GetAnnotationRef(Object);
 		
@@ -329,7 +328,7 @@ int32 FObjectMemoryAnalyzer::GetResults(TArray<FObjectMemoryUsage>& Results)
 		for( FThreadSafeObjectIterator It(BaseClass, false, (EObjectFlags)ExclusionFlags); It; ++It )
 		{
 			UObject* Object	= *It;
-			if (!(AnalyzeFlags&EAnalyzeFlags::IncludeDefaultObjects) && Object->IsDefaultSubobject()) { continue; };
+			if (!(AnalyzeFlags&EAnalyzeFlags::IncludeDefaultObjects) && Object->IsTemplate()) { continue; };
 			FObjectMemoryUsage& Annotation = MemUsageAnnotations.GetAnnotationRef(Object);
 
 			if (Annotation.IsRoot())
