- Separates FFrame::ReadObject and FFrame::ReadProperty to rely on a new method, FFrame::ReadPointer(), which returns a bare ScriptPointerType.
- Changed codepaths reading FProperty objects using ReadObject() to now use ReadPropertyUnchecked(), since FProperty is not a child of UObject (ReadObject was also an unchecked implementation so validation has been lost).

These changes allow ReadObject to return a UObject fetched from a constructed TObjectPtr<UObject> which enables UObject access tracking. This corrects the previous implementation of ReadProperty which would return an FProperty from a temporarily constructed TObjectPtr<UObject>. This had the effect of running TObjectPtr tracking handlers (if any are registered) on the FProperty object which is not actually a UObject type. Using the intermediate TObjectPtr<UObject> in a handler could corrupt the underlying FProperty data as many codepaths in TObjectPtr lazily initialize UObject member variables upon calling UObject methods.

diff --git a/Engine/Source/Runtime/CoreUObject/Private/UObject/ScriptCore.cpp b/Engine/Source/Runtime/CoreUObject/Private/UObject/ScriptCore.cpp
--- a/Engine/Source/Runtime/CoreUObject/Private/UObject/ScriptCore.cpp
+++ b/Engine/Source/Runtime/CoreUObject/Private/UObject/ScriptCore.cpp
@@ -2208,8 +2208,7 @@ IMPLEMENT_VM_FUNCTION( EX_LocalVariable, execLocalVariable );
 
 DEFINE_FUNCTION(UObject::execInstanceVariable)
 {
-	FProperty* VarProperty = (FProperty*)Stack.ReadObject();
-	Stack.MostRecentProperty = VarProperty;
+	FProperty* VarProperty = (FProperty*)Stack.ReadPropertyUnchecked();
 
 	if (VarProperty == nullptr || !P_THIS->IsA((UClass*)VarProperty->InternalGetOwnerAsUObjectUnsafe()))
 	{
@@ -2235,8 +2234,7 @@ IMPLEMENT_VM_FUNCTION( EX_InstanceVariable, execInstanceVariable );
 
 DEFINE_FUNCTION(UObject::execClassSparseDataVariable)
 {
-	FProperty* VarProperty = (FProperty*)Stack.ReadObject();
-	Stack.MostRecentProperty = VarProperty;
+	FProperty* VarProperty = (FProperty*)Stack.ReadPropertyUnchecked();
 
 	if (VarProperty == nullptr || P_THIS->GetSparseClassDataStruct() == nullptr)
 	{
@@ -2262,8 +2260,7 @@ IMPLEMENT_VM_FUNCTION(EX_ClassSparseDataVariable, execClassSparseDataVariable);
 
 DEFINE_FUNCTION(UObject::execDefaultVariable)
 {
-	FProperty* VarProperty = (FProperty*)Stack.ReadObject();
-	Stack.MostRecentProperty = VarProperty;
+	FProperty* VarProperty = (FProperty*)Stack.ReadPropertyUnchecked();
 	Stack.MostRecentPropertyAddress = nullptr;
 	Stack.MostRecentPropertyContainer = nullptr;
 
@@ -3478,7 +3475,7 @@ IMPLEMENT_VM_FUNCTION( EX_TextConst, execTextConst );
 
 DEFINE_FUNCTION(UObject::execPropertyConst)
 {
-	*(FProperty**)RESULT_PARAM = (FProperty*)Stack.ReadObject();
+	*(FProperty**)RESULT_PARAM = (FProperty*)Stack.ReadPropertyUnchecked();
 }
 IMPLEMENT_VM_FUNCTION(EX_PropertyConst, execPropertyConst);
 
@@ -3724,7 +3721,7 @@ IMPLEMENT_VM_FUNCTION( EX_SetMap, execSetMap );
 
 DEFINE_FUNCTION(UObject::execArrayConst)
 {
-	FProperty* InnerProperty = CastFieldChecked<FProperty>((FField*)Stack.ReadObject());
+	FProperty* InnerProperty = CastFieldChecked<FProperty>((FField*)Stack.ReadPropertyUnchecked());
 	int32 Num = Stack.ReadInt<int32>();
 	check(RESULT_PARAM);
 	FScriptArrayHelper ArrayHelper = FScriptArrayHelper::CreateHelperFormInnerProperty(InnerProperty, RESULT_PARAM);
@@ -3744,7 +3741,7 @@ IMPLEMENT_VM_FUNCTION(EX_ArrayConst, execArrayConst);
 
 DEFINE_FUNCTION(UObject::execSetConst)
 {
-	FProperty* InnerProperty = CastFieldChecked<FProperty>((FField*)Stack.ReadObject());
+	FProperty* InnerProperty = CastFieldChecked<FProperty>((FField*)Stack.ReadPropertyUnchecked());
 	int32 Num = Stack.ReadInt<int32>();
 	check(RESULT_PARAM);
 
@@ -3764,8 +3761,8 @@ IMPLEMENT_VM_FUNCTION(EX_SetConst, execSetConst);
 
 DEFINE_FUNCTION(UObject::execMapConst)
 {
-	FProperty* KeyProperty = CastFieldChecked<FProperty>((FField*)Stack.ReadObject());
-	FProperty* ValProperty = CastFieldChecked<FProperty>((FField*)Stack.ReadObject());
+	FProperty* KeyProperty = CastFieldChecked<FProperty>((FField*)Stack.ReadPropertyUnchecked());
+	FProperty* ValProperty = CastFieldChecked<FProperty>((FField*)Stack.ReadPropertyUnchecked());
 	int32 Num = Stack.ReadInt<int32>();
 	check(RESULT_PARAM);
 
@@ -3786,7 +3783,7 @@ IMPLEMENT_VM_FUNCTION(EX_MapConst, execMapConst);
 
 DEFINE_FUNCTION(UObject::execBitFieldConst)
 {
-	FBoolProperty* BitProperty = CastFieldChecked<FBoolProperty>((FField*)Stack.ReadObject());
+	FBoolProperty* BitProperty = CastFieldChecked<FBoolProperty>((FField*)Stack.ReadPropertyUnchecked());
 	uint8 ByteValue = Stack.Read<uint8>();
 	// we could pack the bit into the lower bits of the FProperty pointer, but this instruction is rarely used
 	// and it's likely that a simple implementation will be appreciated by readers, debuggers, and even optimizers:
diff --git a/Engine/Source/Runtime/CoreUObject/Public/UObject/Stack.h b/Engine/Source/Runtime/CoreUObject/Public/UObject/Stack.h
--- a/Engine/Source/Runtime/CoreUObject/Public/UObject/Stack.h
+++ b/Engine/Source/Runtime/CoreUObject/Public/UObject/Stack.h
@@ -245,6 +245,7 @@ public:
 	TNumericType ReadInt();
 	float ReadFloat();
 	double ReadDouble();
+	ScriptPointerType ReadPointer();
 	FName ReadName();
 	UObject* ReadObject();
 	int32 ReadWord();
@@ -397,12 +398,8 @@ inline TNumericType FFrame::ReadInt()
 
 inline UObject* FFrame::ReadObject()
 {
-	// we always pull 64-bits of data out, which is really a UObject* in some representation (depending on platform)
-	ScriptPointerType TempCode = FPlatformMemory::ReadUnaligned<ScriptPointerType>(Code);
+	UObject* Result = (UObject*) ReadPointer();
 
-	// turn that uint32 into a UObject pointer
-	UObject* Result = (UObject*)(TempCode);
-	Code += sizeof(ScriptPointerType);
 #if UE_WITH_OBJECT_HANDLE_LATE_RESOLVE
 	TObjectPtr<UObject> ObjPtr(Result);
 	return ObjPtr.Get();
@@ -413,8 +410,7 @@ inline UObject* FFrame::ReadObject()
 
 inline FProperty* FFrame::ReadProperty()
 {
-	FProperty* Result = (FProperty*)ReadObject();
-	MostRecentProperty = Result;
+	FProperty* Result = ReadPropertyUnchecked();
 
 	// Callers don't check for NULL; this method is expected to succeed.
 	check(Result);
@@ -424,7 +420,7 @@ inline FProperty* FFrame::ReadProperty()
 
 inline FProperty* FFrame::ReadPropertyUnchecked()
 {
-	FProperty* Result = (FProperty*)ReadObject();
+	FProperty* Result = (FProperty*)ReadPointer();
 	MostRecentProperty = Result;
 	return Result;
 }
@@ -439,6 +435,15 @@ inline double FFrame::ReadDouble()
 	return Read<double>();
 }
 
+inline ScriptPointerType FFrame::ReadPointer()
+{
+	// Serialized pointers are always the size of ScriptPointerType 
+	ScriptPointerType Pointer = FPlatformMemory::ReadUnaligned<ScriptPointerType>(Code);
+
+	Code += sizeof(ScriptPointerType);
+	return Pointer;
+}
+
 inline int32 FFrame::ReadWord()
 {
 	return Read<uint16>();
@@ -459,7 +464,7 @@ inline VariableSizeType FFrame::ReadVariableSize( FProperty** ExpressionField )
 {
 	VariableSizeType Result=0;
 
-	FField* Field = (FField*)ReadObject(); // Is it safe to assume it's an FField?
+	FField* Field = (FField*)ReadPropertyUnchecked(); // Is it safe to assume it's an FField?
 	FProperty* Property = CastField<FProperty>(Field);
 	if (Property)
 	{
