Fixed a crash in FPropertyBag caused by an invalid type in a property tag

#rnx

diff --git a/Engine/Source/Runtime/CoreUObject/Private/UObject/Field.cpp b/Engine/Source/Runtime/CoreUObject/Private/UObject/Field.cpp
--- a/Engine/Source/Runtime/CoreUObject/Private/UObject/Field.cpp
+++ b/Engine/Source/Runtime/CoreUObject/Private/UObject/Field.cpp
@@ -972,6 +972,15 @@ FField* FField::Construct(const FName& FieldTypeName, const FFieldVariant& InOwn
 	return Instance;
 }
 
+FField* FField::TryConstruct(const FName& FieldTypeName, const FFieldVariant& InOwner, const FName& InName, EObjectFlags InFlags)
+{
+	if (FFieldClass* FieldClassPtr = FFieldClass::GetNameToFieldClassMap().FindRef(FieldTypeName))
+	{
+		return FieldClassPtr->Construct(InOwner, InName, InFlags);
+	}
+	return nullptr;
+}
+
 FName FField::GenerateFFieldName(FFieldVariant InOwner /** Unused yet */, FFieldClass* InClass)
 {
 	check(InClass);
diff --git a/Engine/Source/Runtime/CoreUObject/Private/UObject/PropertyArray.cpp b/Engine/Source/Runtime/CoreUObject/Private/UObject/PropertyArray.cpp
--- a/Engine/Source/Runtime/CoreUObject/Private/UObject/PropertyArray.cpp
+++ b/Engine/Source/Runtime/CoreUObject/Private/UObject/PropertyArray.cpp
@@ -1265,7 +1265,7 @@ bool FArrayProperty::LoadFromTag(const FPropertyTag& Tag)
 		return false;
 	}
 
-	FField* Field = FField::Construct(Tag.InnerType, {}, Tag.Name, RF_NoFlags);
+	FField* Field = FField::TryConstruct(Tag.InnerType, {}, Tag.Name, RF_NoFlags);
 	if (FProperty* Property = CastField<FProperty>(Field))
 	{
 		FPropertyTag InnerTag = Tag;
diff --git a/Engine/Source/Runtime/CoreUObject/Private/UObject/PropertyBag.cpp b/Engine/Source/Runtime/CoreUObject/Private/UObject/PropertyBag.cpp
--- a/Engine/Source/Runtime/CoreUObject/Private/UObject/PropertyBag.cpp
+++ b/Engine/Source/Runtime/CoreUObject/Private/UObject/PropertyBag.cpp
@@ -109,7 +109,7 @@ void FPropertyBag::LoadPropertyByTag(const FPropertyPathName& Path, const FPrope
 	}
 
 	// Construct a property from the tag and try to use it to serialize the value.
-	FField* Field = FField::Construct(Tag.Type, {}, Tag.Name, RF_NoFlags);
+	FField* Field = FField::TryConstruct(Tag.Type, {}, Tag.Name, RF_NoFlags);
 	if (FProperty* Property = CastField<FProperty>(Field); Property && Property->LoadFromTag(Value.Tag))
 	{
 		Property->Link(UnderlyingArchive);
diff --git a/Engine/Source/Runtime/CoreUObject/Public/UObject/Field.h b/Engine/Source/Runtime/CoreUObject/Public/UObject/Field.h
--- a/Engine/Source/Runtime/CoreUObject/Public/UObject/Field.h
+++ b/Engine/Source/Runtime/CoreUObject/Public/UObject/Field.h
@@ -511,6 +511,8 @@ public:
 	static COREUOBJECT_API FField* Construct(const FFieldVariant& InOwner, const FName& InName, EObjectFlags InFlags);
 	/** Constructs a new field given the name of its class */
 	static COREUOBJECT_API FField* Construct(const FName& FieldTypeName, const FFieldVariant& InOwner, const FName& InName, EObjectFlags InFlags);
+	/** Tries to construct a new field given the name of its class. Returns null if the type does not exist. */
+	static COREUOBJECT_API FField* TryConstruct(const FName& FieldTypeName, const FFieldVariant& InOwner, const FName& InName, EObjectFlags InFlags);
 
 	/** Fixups after duplicating a Field */
 	COREUOBJECT_API virtual void PostDuplicate(const FField& InField);
