Fix undo/redo of overrides

diff --git a/Engine/Source/Runtime/CoreUObject/Private/UObject/Class.cpp b/Engine/Source/Runtime/CoreUObject/Private/UObject/Class.cpp
--- a/Engine/Source/Runtime/CoreUObject/Private/UObject/Class.cpp
+++ b/Engine/Source/Runtime/CoreUObject/Private/UObject/Class.cpp
@@ -1662,7 +1662,11 @@ void UStruct::SerializeVersionedTaggedProperties(FStructuredArchive::FSlot Slot,
 						// The operation was set part of the tag, now that we know the associated property, restore the overridden operation on the object
 						if (FOverriddenPropertySet* OverriddenProperties = FOverridableSerializationLogic::GetOverriddenProperties())
 						{
-							OverriddenProperties->SetOverriddenPropertyOperation(Tag.OverrideOperation, UnderlyingArchive.GetSerializedPropertyChain(), Property);
+							// No need to restore none operations
+							if (Tag.OverrideOperation != EOverriddenPropertyOperation::None)
+							{
+								OverriddenProperties->SetOverriddenPropertyOperation(Tag.OverrideOperation, UnderlyingArchive.GetSerializedPropertyChain(), Property);
+							}
 						}
 
 						switch (Property->ConvertFromType(Tag, ValueSlot, Data, DefaultsStruct, Defaults))
diff --git a/Engine/Source/Runtime/CoreUObject/Private/UObject/Obj.cpp b/Engine/Source/Runtime/CoreUObject/Private/UObject/Obj.cpp
--- a/Engine/Source/Runtime/CoreUObject/Private/UObject/Obj.cpp
+++ b/Engine/Source/Runtime/CoreUObject/Private/UObject/Obj.cpp
@@ -1647,9 +1647,9 @@ void UObject::Serialize(FStructuredArchive::FRecord Record)
 		// Handle derived UClass objects (exact UClass objects are native only and shouldn't be touched)
 		if (ObjClass != UClass::StaticClass())
 		{
-			// Serializing the overriden properties for undo/redo
-			// of if using UPS (ie. object duplication)
-			if (UnderlyingArchive.IsTransacting() || UnderlyingArchive.UseUnversionedPropertySerialization())
+			// Serializing the overriden properties if using UPS (like object duplication)
+			// @todo: We probably do not want this in cooked packages
+			if (UnderlyingArchive.UseUnversionedPropertySerialization())
 			{
 				FOverridableManager::Get().SerializeOverriddenProperties(*this, Record);
 			}
