Return the added override state on an object.

diff --git a/Engine/Source/Runtime/CoreUObject/Private/UObject/OverridableManager.cpp b/Engine/Source/Runtime/CoreUObject/Private/UObject/OverridableManager.cpp
--- a/Engine/Source/Runtime/CoreUObject/Private/UObject/OverridableManager.cpp
+++ b/Engine/Source/Runtime/CoreUObject/Private/UObject/OverridableManager.cpp
@@ -65,6 +65,15 @@ EOverriddenState FOverridableManager::GetOverriddenState(UObject& Object)
 {
 	if(const FOverriddenPropertySet* OverriddenProperties = GetOverriddenProperties(Object))
 	{
+		// Consider any object that its template is a CDO as added.
+		if (UObject* Archetype = Object.GetArchetype())
+		{
+			if (Archetype->HasAnyFlags(RF_ClassDefaultObject))
+			{
+				return EOverriddenState::Added;
+			}
+		}
+
 		const EOverriddenPropertyOperation Operation = OverriddenProperties->GetOverriddenPropertyOperation((FArchiveSerializedPropertyChain*)nullptr, (FProperty*)nullptr);
 		if (Operation != EOverriddenPropertyOperation::None)
 		{
diff --git a/Engine/Source/Runtime/CoreUObject/Public/UObject/OverridableManager.h b/Engine/Source/Runtime/CoreUObject/Public/UObject/OverridableManager.h
--- a/Engine/Source/Runtime/CoreUObject/Public/UObject/OverridableManager.h
+++ b/Engine/Source/Runtime/CoreUObject/Public/UObject/OverridableManager.h
@@ -68,6 +68,7 @@ enum class EOverriddenState : uint8
 	HasOverrides, // has overrides in the object properties
 	AllOverridden, // all properties are overridden for this object and its subobjects
 	SubObjectsHasOverrides, // at least one of its subobjects has overrides
+	Added, // This object was added
 };
 
 
