[Feature] PropertyTypeLayout deatails customization of types will be used in the details for their InstanceDataObject counterparts

diff --git a/Engine/Plugins/Editor/InstanceDataObjectFixupTool/Source/Private/InstanceDataObjectFixupPanel.cpp b/Engine/Plugins/Editor/InstanceDataObjectFixupTool/Source/Private/InstanceDataObjectFixupPanel.cpp
--- a/Engine/Plugins/Editor/InstanceDataObjectFixupTool/Source/Private/InstanceDataObjectFixupPanel.cpp
+++ b/Engine/Plugins/Editor/InstanceDataObjectFixupTool/Source/Private/InstanceDataObjectFixupPanel.cpp
@@ -213,6 +213,81 @@ int32 FInstanceDataObjectFixupPanel::Find(UObject* Value) const
 	return Instances.Find(Value);
 }
 
+static bool RemoveCustomizationsWithLooseProperties(const FFieldVariant& FieldVariant, const TSharedPtr<IDetailsView>& DetailsView)
+{
+#if WITH_EDITORONLY_DATA
+	static const FName NAME_IsLooseMetadata(TEXT("IsLoose"));
+	if (FStructProperty* AsStructProperty = FieldVariant.Get<FStructProperty>())
+	{
+		if (RemoveCustomizationsWithLooseProperties(AsStructProperty->Struct, DetailsView))
+		{
+			return true;
+		}
+	}
+	else if (FObjectProperty* AsObjectProperty = FieldVariant.Get<FObjectProperty>())
+	{
+		if (AsObjectProperty->HasAnyPropertyFlags(CPF_InstancedReference))
+		{
+			if (RemoveCustomizationsWithLooseProperties(AsObjectProperty->PropertyClass, DetailsView))
+			{
+				return true;
+			}
+		}
+	}
+	else if (const FArrayProperty* AsArrayProperty = FieldVariant.Get<FArrayProperty>())
+	{
+		if (RemoveCustomizationsWithLooseProperties(AsArrayProperty->Inner, DetailsView))
+		{
+			return true;
+		}
+	}
+	else if (const FSetProperty* AsSetProperty = FieldVariant.Get<FSetProperty>())
+	{
+		if (RemoveCustomizationsWithLooseProperties(AsSetProperty->ElementProp, DetailsView))
+		{
+			return true;
+		}
+	}
+	else if (const FMapProperty* AsMapProperty = FieldVariant.Get<FMapProperty>())
+	{
+		if (RemoveCustomizationsWithLooseProperties(AsMapProperty->KeyProp, DetailsView))
+		{
+			return true;
+		}
+		if (RemoveCustomizationsWithLooseProperties(AsMapProperty->ValueProp, DetailsView))
+		{
+			return true;
+		}
+	}
+	else if (UStruct* AsStruct = FieldVariant.Get<UStruct>())
+	{
+		bool result = false;
+		for (const FProperty* Property : TFieldRange<FProperty>(AsStruct))
+		{
+			if (RemoveCustomizationsWithLooseProperties(Property, DetailsView))
+			{
+				result = true;
+			}
+		}
+		if (result)
+		{
+			// register an empty delegate to override the global rule of displaying this type with customizations
+			DetailsView->RegisterInstancedCustomPropertyTypeLayout(AsStruct->GetFName(), {});
+		}
+		return result;
+	}
+	
+	if (const FProperty* Property = FieldVariant.Get<FProperty>())
+	{
+		if (Property->HasMetaData(NAME_IsLooseMetadata))
+		{
+			return true;
+		}
+	}
+#endif
+	return false;
+}
+
 TSharedPtr<IDetailsView>& FInstanceDataObjectFixupPanel::GenerateDetailsView(bool bScrollbarOnLeft)
 {
 	FDetailsViewArgs DetailsViewArgs;
@@ -223,6 +298,10 @@ TSharedPtr<IDetailsView>& FInstanceDataObjectFixupPanel::GenerateDetailsView(boo
 	
 	FPropertyEditorModule& PropertyEditorModule = FModuleManager::GetModuleChecked<FPropertyEditorModule>("PropertyEditor");
 	DetailsView = PropertyEditorModule.CreateDetailView(DetailsViewArgs);
+	for (const UObject* Instance : Instances)
+	{
+		RemoveCustomizationsWithLooseProperties(Instance->GetClass(), DetailsView);
+	}
 
 	for (const UObject* Object : Instances)
 	{
diff --git a/Engine/Source/Editor/PropertyEditor/Private/PropertyEditorModule.cpp b/Engine/Source/Editor/PropertyEditor/Private/PropertyEditorModule.cpp
--- a/Engine/Source/Editor/PropertyEditor/Private/PropertyEditorModule.cpp
+++ b/Engine/Source/Editor/PropertyEditor/Private/PropertyEditorModule.cpp
@@ -928,6 +928,20 @@ bool FPropertyEditorModule::IsCustomizedStruct(const UStruct* Struct, const FCus
 		{
 			bFound = GlobalPropertyTypeToLayoutMap.Contains( Struct->GetFName() );
 		}
+		
+		if( !bFound )
+		{
+			static const FName NAME_PresentAsTypeMetadata(TEXT("PresentAsType"));
+			if (const FString* DisplayType = Struct->FindMetaData(NAME_PresentAsTypeMetadata))
+			{
+				// try finding DisplayType instead
+				bFound = InstancePropertyTypeLayoutMap.Contains( FName(*DisplayType) );
+				if( !bFound )
+				{
+					bFound = GlobalPropertyTypeToLayoutMap.Contains( FName(*DisplayType) );
+				}
+			}
+		}
 	}
 	
 	return bFound;
@@ -1005,6 +1019,24 @@ FPropertyTypeLayoutCallback FPropertyEditorModule::FindPropertyTypeLayoutCallbac
 		{
 			LayoutCallbacks = GlobalPropertyTypeToLayoutMap.Find(PropertyTypeName);
 		}
+		
+		if( !LayoutCallbacks )
+		{
+			if (const FStructProperty* AsStructProperty = CastField<FStructProperty>(PropertyHandle.GetProperty()))
+			{
+				static const FName NAME_PresentAsTypeMetadata(TEXT("PresentAsType"));
+				if (const FString* DisplayType = AsStructProperty->Struct->FindMetaData(NAME_PresentAsTypeMetadata))
+				{
+					// try finding DisplayType instead
+					LayoutCallbacks = InstancedPropertyTypeLayoutMap.Find(FName(*DisplayType));
+	
+					if( !LayoutCallbacks )
+					{
+						LayoutCallbacks = GlobalPropertyTypeToLayoutMap.Find(FName(*DisplayType));
+					}
+				}
+			}
+		}
 
 		if ( LayoutCallbacks )
 		{
diff --git a/Engine/Source/Runtime/CoreUObject/Private/UObject/InstanceDataObjectUtils.cpp b/Engine/Source/Runtime/CoreUObject/Private/UObject/InstanceDataObjectUtils.cpp
--- a/Engine/Source/Runtime/CoreUObject/Private/UObject/InstanceDataObjectUtils.cpp
+++ b/Engine/Source/Runtime/CoreUObject/Private/UObject/InstanceDataObjectUtils.cpp
@@ -12,7 +12,10 @@ namespace UE
 	using FWildcardPropertyPathName = FPropertyPathName;
 	
 	static FName Name_ValuesSetBySerialization(TEXT("_ValuesSetBySerialization"));
-	static const FName NAME_StructOrignalTypeMetadata(TEXT("OriginalType"));
+	static const FName NAME_StructOriginalTypeMetadata(TEXT("OriginalType"));
+	static const FName NAME_PresentAsTypeMetadata(TEXT("PresentAsType"));
+	static const FName NAME_IsLooseMetadata(TEXT("IsLoose"));
+	static const FName NAME_CategoryMetadata(TEXT("Category"));
 	
 	struct ResolvePropertyPathNameHelperParams
 	{
@@ -261,14 +264,16 @@ namespace UE
 			const FString* OriginalType = nullptr;
 #if WITH_EDITORONLY_DATA
 			//@note: Transfer existing metadata over as we build the InstanceDataObject from the struct or it owner, if any, this is useful for testing purposes
-			OriginalType = AsStructProperty->FindMetaData(NAME_StructOrignalTypeMetadata);
+			OriginalType = AsStructProperty->FindMetaData(NAME_StructOriginalTypeMetadata);
 			FField* OwnerField = OriginalType == nullptr ? AsStructProperty->Owner.ToField() : nullptr;
-			OriginalType = OwnerField ? OwnerField->FindMetaData(NAME_StructOrignalTypeMetadata) : OriginalType;
+			OriginalType = OwnerField ? OwnerField->FindMetaData(NAME_StructOriginalTypeMetadata) : OriginalType;
 #endif
 			const FName StructOriginalName = OriginalType ? FName(**OriginalType) : AsStructProperty->Struct->GetFName();
 			AsStructProperty->Struct = CreateInstanceDataObjectStructRec<UScriptStruct>(AsStructProperty->Struct, Outer, LooseProperties, Path);
 #if WITH_EDITORONLY_DATA
-			AsStructProperty->SetMetaData("OriginalType", StructOriginalName.ToString());
+			AsStructProperty->SetMetaData(NAME_StructOriginalTypeMetadata, StructOriginalName.ToString());
+			AsStructProperty->SetMetaData(NAME_PresentAsTypeMetadata, StructOriginalName.ToString());
+			AsStructProperty->Struct->SetMetaData(NAME_PresentAsTypeMetadata, *StructOriginalName.ToString());
 #endif
 		}
 		else if (const FArrayProperty* AsArrayProperty = CastField<FArrayProperty>(Property))
@@ -292,9 +297,9 @@ namespace UE
 		}
 		
 #if WITH_EDITORONLY_DATA
-		if (!Property->HasMetaData(TEXT("category")))
+		if (!Property->HasMetaData(NAME_CategoryMetadata))
 		{
-			Property->SetMetaData(TEXT("category"), TEXT("Verse"));
+			Property->SetMetaData(NAME_CategoryMetadata, TEXT("Verse"));
 		}
 #endif
 	}
@@ -376,8 +381,8 @@ namespace UE
 	static void MarkPropertyAsLoose(FProperty* Property)
 	{
 #if WITH_EDITORONLY_DATA
-		Property->SetMetaData(TEXT("isLoose"), TEXT("True"));
-		Property->SetMetaData(TEXT("category"), TEXT("Loose Properties"));
+		Property->SetMetaData(NAME_IsLooseMetadata, TEXT("True"));
+		Property->SetMetaData(NAME_CategoryMetadata, TEXT("Loose Properties"));
 #endif
 		Property->SetPropertyFlags(CPF_Edit | CPF_EditConst);
 		if (const FArrayProperty* AsArrayProperty = CastField<FArrayProperty>(Property))
@@ -418,7 +423,7 @@ namespace UE
 		if (TemplateStruct)
 		{
 			const FName SuperName(TemplateStruct->GetName() + TEXT("_Super"));
-			Super = NewObject<UStruct>(Outer, StructClass, MakeUniqueObjectName(Outer, StructClass, SuperName));
+			Super = NewObject<UStruct>(Outer, StructClass, MakeUniqueObjectName(nullptr, StructClass, SuperName));
 			
 			// Gather properties for Super Struct
 			TArray<FProperty*> SuperProperties;
@@ -434,9 +439,9 @@ namespace UE
 				FProperty* SuperProperty = CreateInstanceDataObjectProperty(TemplateProperty, Super, LooseProperties, Path);
 				
 #if WITH_EDITORONLY_DATA
-				if (!SuperProperty->HasMetaData(TEXT("category")))
+				if (!SuperProperty->HasMetaData(NAME_CategoryMetadata))
 				{
-					SuperProperty->SetMetaData(TEXT("category"), TEXT("Verse"));
+					SuperProperty->SetMetaData(NAME_CategoryMetadata, TEXT("Verse"));
 				}
 #endif
 				Path.Pop();
@@ -464,7 +469,7 @@ namespace UE
 		}
 
 		const FName InstanceDataObjectName = (TemplateStruct) ? FName(TemplateStruct->GetName() + TEXT("_InstanceDataObject")) : FName(TEXT("InstanceDataObject"));
-		UStruct* Result = NewObject<UStruct>(Outer, StructClass, MakeUniqueObjectName(Outer, StructClass, InstanceDataObjectName));
+		UStruct* Result = NewObject<UStruct>(Outer, StructClass, MakeUniqueObjectName(nullptr, StructClass, InstanceDataObjectName));
 
 		// Gather "loose" properties for child Struct
 		TArray<FProperty*> LooseInstanceDataObjectProperties;
