Report missing script (native) packages on import when using zenloader

diff --git a/Engine/Source/Runtime/CoreUObject/Private/UObject/LinkerLoad.cpp b/Engine/Source/Runtime/CoreUObject/Private/UObject/LinkerLoad.cpp
--- a/Engine/Source/Runtime/CoreUObject/Private/UObject/LinkerLoad.cpp
+++ b/Engine/Source/Runtime/CoreUObject/Private/UObject/LinkerLoad.cpp
@@ -3452,6 +3452,23 @@ bool FLinkerLoad::VerifyImportInner(const int32 ImportIndex, FString& WarningSuf
 		}
 #endif
 
+		// helper to report missing native packages when an import is requested:
+		const auto ReportMissingPackage = [&PackageToLoad, ImportIndex, this]() -> bool
+			{
+				TCHAR PackageToLoadBuffer[FName::StringBufferSize];
+				PackageToLoad.ToString(PackageToLoadBuffer);
+				if (FPackageName::IsScriptPackage(PackageToLoadBuffer))
+				{
+					if (!FLinkerLoad::IsKnownMissingPackage(PackageToLoad))
+					{
+						FLinkerLoad::AddKnownMissingPackage(PackageToLoad);
+						UE_ASSET_LOG(LogLinker, Warning, PackagePath, TEXT("VerifyImport: Failed to find script package for import object '%s'"), *GetImportFullName(ImportIndex));
+					}
+					return true;
+				}
+				return false;
+			};
+
 		// Check if the package exist first, if it already exists, it is either already loaded or being loaded
 		// In the fully loaded case we can entirely skip the loading
 		// In the other case we do not want to trigger another load of the objects in that import, in case they contain dependencies to the package we are currently loading
@@ -3461,6 +3478,7 @@ bool FLinkerLoad::VerifyImportInner(const int32 ImportIndex, FString& WarningSuf
 		{
 			if (!Package)
 			{
+				ReportMissingPackage();
 				return nullptr;
 			}
 			Import.SourceLinker = FindExistingLinkerForPackage(Package);
@@ -3476,18 +3494,9 @@ bool FLinkerLoad::VerifyImportInner(const int32 ImportIndex, FString& WarningSuf
 		}
 		if (Package == nullptr || !Package->IsFullyLoaded())
 		{
+			if (ReportMissingPackage())
 			{
-				TCHAR PackageToLoadBuffer[FName::StringBufferSize];
-				PackageToLoad.ToString(PackageToLoadBuffer);
-				if (FPackageName::IsScriptPackage(PackageToLoadBuffer))
-				{
-					if (!FLinkerLoad::IsKnownMissingPackage(PackageToLoad))
-					{
-						FLinkerLoad::AddKnownMissingPackage(PackageToLoad);
-						UE_ASSET_LOG(LogLinker, Warning, PackagePath, TEXT("VerifyImport: Failed to find script package for import object '%s'"), *GetImportFullName(ImportIndex));
-					}
-					return nullptr;
-				}
+				return nullptr;
 			}
 
 #if USE_CIRCULAR_DEPENDENCY_LOAD_DEFERRING
