Added structured serialization of FPropertyTypeName

#rnx

diff --git a/Engine/Source/Runtime/CoreUObject/Private/UObject/PropertyTypeName.cpp b/Engine/Source/Runtime/CoreUObject/Private/UObject/PropertyTypeName.cpp
--- a/Engine/Source/Runtime/CoreUObject/Private/UObject/PropertyTypeName.cpp
+++ b/Engine/Source/Runtime/CoreUObject/Private/UObject/PropertyTypeName.cpp
@@ -8,6 +8,7 @@
 #include "HAL/CriticalSection.h"
 #include "HAL/LowLevelMemTracker.h"
 #include "Misc/ScopeRWLock.h"
+#include "String/Find.h"
 
 namespace UE
 {
@@ -37,6 +38,13 @@ inline FArchive& operator<<(FArchive& Ar, FPropertyTypeNameNode& Node)
 	return Ar << Node.Name << Node.InnerCount;
 }
 
+inline void operator<<(FStructuredArchiveSlot Slot, FPropertyTypeNameNode& Node)
+{
+	FStructuredArchiveRecord Record = Slot.EnterRecord();
+	Record.EnterField(TEXT("Name")) << Node.Name;
+	Record.EnterField(TEXT("InnerCount")) << Node.InnerCount;
+}
+
 inline const FPropertyTypeNameNode* AppendNode(FStringBuilderBase& Builder, const FPropertyTypeNameNode* Node)
 {
 	Builder << Node->Name;
@@ -315,6 +323,66 @@ FArchive& operator<<(FArchive& Ar, FPropertyTypeName& TypeName)
 	return Ar;
 }
 
+void operator<<(FStructuredArchiveSlot Slot, FPropertyTypeName& TypeName)
+{
+	const FArchiveState& State = Slot.GetArchiveState();
+	if (!State.IsPersistent())
+	{
+		Slot << TypeName.Index;
+	}
+	else if (State.IsLoading())
+	{
+		if (State.IsTextFormat())
+		{
+			FString Text;
+			Slot << Text;
+			FPropertyTypeNameBuilder Builder;
+			if (Builder.TryParse(Text))
+			{
+				TypeName = Builder.Build();
+			}
+			else
+			{
+				TypeName.Reset();
+				Slot.GetUnderlyingArchive().SetError();
+			}
+		}
+		else
+		{
+			FStructuredArchiveStream Stream = Slot.EnterStream();
+			TArray<FPropertyTypeNameNode, TInlineAllocator<16>> Nodes;
+			int32 Remaining = 1;
+			do 
+			{
+				FPropertyTypeNameNode& Node = Nodes.AddDefaulted_GetRef();
+				Stream.EnterElement() << Node;
+				Remaining += Node.InnerCount - 1;
+			}
+			while (Remaining > 0);
+			TypeName.Index = GPropertyTypeNameTable.FindOrAddByName(Nodes.GetData());
+		}
+	}
+	else if (State.IsSaving())
+	{
+		if (State.IsTextFormat())
+		{
+			FString Text(WriteToString<256>(TypeName));
+			Slot << Text;
+		}
+		else
+		{
+			FStructuredArchiveStream Stream = Slot.EnterStream();
+			const FPropertyTypeNameNode* Node = GPropertyTypeNameTable.ResolveByIndex(TypeName.Index);
+			for (int32 Remaining = 1; Remaining > 0; --Remaining, ++Node)
+			{
+				FPropertyTypeNameNode NodeCopy = *Node;
+				Stream.EnterElement() << NodeCopy;
+				Remaining += NodeCopy.InnerCount;
+			}
+		}
+	}
+}
+
 FStringBuilderBase& operator<<(FStringBuilderBase& Builder, const FPropertyTypeName& TypeName)
 {
 	const FPropertyTypeNameNode* Node = GPropertyTypeNameTable.ResolveByIndex(TypeName.Index);
@@ -364,10 +432,101 @@ void FPropertyTypeNameBuilder::AddTypeName(FPropertyTypeName Name)
 	}
 }
 
+bool FPropertyTypeNameBuilder::TryParse(FStringView Name)
+{
+	const auto ResetToInitial = [this, InitialCount = Nodes.Num(), InitialActiveIndex = ActiveIndex]
+	{
+		Nodes.SetNum(InitialCount, EAllowShrinking::No);
+		OuterNodeIndex.SetNum(InitialCount, EAllowShrinking::No);
+		ActiveIndex = InitialActiveIndex;
+	};
+
+	int32 Index = 0;
+	int32 Depth = 0;
+	bool bAllowName = true;
+	bool bAllowBegin = false;
+	for (FStringView Remaining = Name;; Remaining.RightChopInline(Index + 1))
+	{
+		Index = String::FindFirstOfAnyChar(Remaining, {TEXT('<'), TEXT(','), TEXT('>')});
+
+		if (Index != 0)
+		{
+			FStringView Type = (Index > 0 ? Remaining.Left(Index) : Remaining).TrimStartAndEnd();
+			if (!Type.IsEmpty())
+			{
+				if (!bAllowName)
+				{
+					// Names must follow '<' or ',' or be at the root.
+					break;
+				}
+				AddTypeName(FName(Type));
+				bAllowBegin = true;
+				bAllowName = false;
+			}
+		}
+
+		if (Index == INDEX_NONE)
+		{
+			if (Depth == 0 && !bAllowName)
+			{
+				return true;
+			}
+			// Missing a '>' and/or missing a name.
+			break;
+		}
+
+		if (bAllowName)
+		{
+			// Missing a name.
+			break;
+		}
+
+		const TCHAR C = Remaining[Index];
+		if (C == TEXT('<'))
+		{
+			if (!bAllowBegin)
+			{
+				// '<' must follow a name.
+				break;
+			}
+			++Depth;
+			BeginTypeParameters();
+			bAllowName = true;
+		}
+		else if (C == TEXT('>'))
+		{
+			if (Depth <= 0)
+			{
+				// '>' must have a matching '<'.
+				break;
+			}
+			EndTypeParameters();
+			--Depth;
+		}
+		else // if (C == TEXT(','))
+		{
+			if (Depth <= 0)
+			{
+				// ',' must not be used at the root.
+				break;
+			}
+			bAllowName = true;
+		}
+
+		bAllowBegin = false;
+	}
+
+	ResetToInitial();
+	return false;
+}
+
 FPropertyTypeName FPropertyTypeNameBuilder::Build() const
 {
 	FPropertyTypeName Type;
-	Type.Index = GPropertyTypeNameTable.FindOrAddByName(Nodes.GetData());
+	if (!Nodes.IsEmpty())
+	{
+		Type.Index = GPropertyTypeNameTable.FindOrAddByName(Nodes.GetData());
+	}
 	return Type;
 }
 
diff --git a/Engine/Source/Runtime/CoreUObject/Public/UObject/PropertyTypeName.h b/Engine/Source/Runtime/CoreUObject/Public/UObject/PropertyTypeName.h
--- a/Engine/Source/Runtime/CoreUObject/Public/UObject/PropertyTypeName.h
+++ b/Engine/Source/Runtime/CoreUObject/Public/UObject/PropertyTypeName.h
@@ -10,6 +10,7 @@
 #define UE_API COREUOBJECT_API
 
 class FArchive;
+class FStructuredArchiveSlot;
 
 namespace UE { class FPropertyTypeNameBuilder; }
 
@@ -29,7 +30,7 @@ struct FPropertyTypeNameNode
  * - int32 -> IntProperty
  * - TArray<int32> -> ArrayProperty<IntProperty>
  * - TArray<FStructType> -> ArrayProperty<StructProperty<StructType>>
- * - TMap<FKeyStruct, EByteEnum> -> MapProperty<StructProperty<KeyStruct>,EnumProperty<ByteProperty<ByteEnum>>>
+ * - TMap<FKeyStruct, EByteEnum> -> MapProperty<StructProperty<KeyStruct>,EnumProperty<ByteEnum,ByteProperty>>
  */
 class FPropertyTypeName
 {
@@ -42,7 +43,7 @@ public:
 	/**
 	 * Returns the type at the root of this property type name.
 	 *
-	 * Example: MapProperty<StructProperty<KeyStruct>,EnumProperty<ByteProperty<ByteEnum>>>
+	 * Example: MapProperty<StructProperty<KeyStruct>,EnumProperty<ByteEnum,ByteProperty>>
 	 * - GetTypeName() -> MapProperty
 	 */
 	UE_API FName GetTypeName() const;
@@ -50,7 +51,7 @@ public:
 	/**
 	 * Returns the number of type parameters under the root of this property type name.
 	 *
-	 * Example: MapProperty<StructProperty<KeyStruct>,EnumProperty<ByteProperty<ByteEnum>>>
+	 * Example: MapProperty<StructProperty<KeyStruct>,EnumProperty<ByteEnum,ByteProperty>>
 	 * - GetTypeParameterCount() -> 2
 	 */
 	UE_API int32 GetTypeParameterCount() const;
@@ -60,9 +61,9 @@ public:
 	 *
 	 * An out-of-bounds index will return an empty type name.
 	 *
-	 * Example: MapProperty<StructProperty<KeyStruct>,EnumProperty<ByteProperty<ByteEnum>>>
+	 * Example: MapProperty<StructProperty<KeyStruct>,EnumProperty<ByteEnum,ByteProperty>>
 	 * - GetTypeParameter(0) -> StructProperty<KeyStruct>
-	 * - GetTypeParameter(1) -> EnumProperty<ByteProperty<ByteEnum>>
+	 * - GetTypeParameter(1) -> EnumProperty<ByteEnum,ByteProperty>
 	 */
 	UE_API FPropertyTypeName GetTypeParameter(int32 ParamIndex = 0) const;
 
@@ -71,7 +72,7 @@ public:
 	 *
 	 * An out-of-bounds index will return a name of None.
 	 *
-	 * Example: MapProperty<StructProperty<KeyStruct>,EnumProperty<ByteProperty<ByteEnum>>>
+	 * Example: MapProperty<StructProperty<KeyStruct>,EnumProperty<ByteEnum,ByteProperty>>
 	 * - GetTypeParameterName(0) -> StructProperty
 	 * - GetTypeParameterName(1) -> EnumProperty
 	 */
@@ -95,6 +96,7 @@ private:
 	UE_API friend bool operator<(const FPropertyTypeName& Lhs, const FPropertyTypeName& Rhs);
 
 	UE_API friend FArchive& operator<<(FArchive& Ar, FPropertyTypeName& TypeName);
+	UE_API friend void operator<<(FStructuredArchiveSlot Slot, FPropertyTypeName& TypeName);
 
 	UE_API friend FStringBuilderBase& operator<<(FStringBuilderBase& Builder, const FPropertyTypeName& TypeName);
 
@@ -106,7 +108,7 @@ private:
 /**
  * Builder for FPropertyTypeName.
  *
- * Example: MapProperty<StructProperty<KeyStruct>,EnumProperty<ByteProperty<ByteEnum>>>
+ * Example: MapProperty<StructProperty<KeyStruct>,EnumProperty<ByteEnum,ByteProperty>>
  *
  * FPropertyTypeNameBuilder Builder;
  * Builder.AddTypeName(NAME_MapProperty));
@@ -117,10 +119,8 @@ private:
  *   Builder.EndTypeParameters();
  *   Builder.AddTypeName(NAME_EnumProperty);
  *   Builder.BeginTypeParameters();
+ *     Builder.AddTypeName(TEXT("ByteEnum"));
  *     Builder.AddTypeName(NAME_ByteProperty);
- *     Builder.BeginTypeParameters();
- *       Builder.AddTypeName(TEXT("ByteEnum"));
- *     Builder.EndTypeParameters();
  *   Builder.EndTypeParameters();
  * Builder.EndTypeParameters();
  * FPropertyTypeName Name = Builder.Build();
@@ -140,11 +140,22 @@ public:
 	/** Mark the end of the type parameters for the matching begin. */
 	UE_API void EndTypeParameters();
 
+	/** Build from the types that have been added. */
 	UE_API FPropertyTypeName Build() const;
 
 	/** Resets the builder to allow it to be reused. */
 	UE_API void Reset();
 
+	/**
+	 * Try to parse and add a complete type name with its type parameters.
+	 *
+	 * Builder is restored to its previous state when parsing fails.
+	 *
+	 * @param Name A name in the format returned by operator<<(FStringBuilderBase&, const FPropertyTypeName&).
+	 * @return true if a name was parsed into the builder, false on failure.
+	 */
+	UE_API bool TryParse(FStringView Name);
+
 private:
 	TArray<FPropertyTypeNameNode, TInlineAllocator<8>> Nodes;
 	TArray<int32, TInlineAllocator<8>> OuterNodeIndex;
diff --git a/Engine/Source/Runtime/CoreUObject/Tests/UObject/PropertyTypeNameTest.cpp b/Engine/Source/Runtime/CoreUObject/Tests/UObject/PropertyTypeNameTest.cpp
--- a/Engine/Source/Runtime/CoreUObject/Tests/UObject/PropertyTypeNameTest.cpp
+++ b/Engine/Source/Runtime/CoreUObject/Tests/UObject/PropertyTypeNameTest.cpp
@@ -6,8 +6,12 @@
 
 #include "Async/ParallelFor.h"
 #include "Misc/StringBuilder.h"
+#include "Serialization/Formatters/BinaryArchiveFormatter.h"
+#include "Serialization/Formatters/JsonArchiveInputFormatter.h"
+#include "Serialization/Formatters/JsonArchiveOutputFormatter.h"
 #include "Serialization/MemoryReader.h"
 #include "Serialization/MemoryWriter.h"
+#include "Serialization/StructuredArchive.h"
 #include "Tests/TestHarnessAdapter.h"
 
 namespace UE
@@ -49,23 +53,15 @@ static FPropertyTypeName CreateEnumMap()
 		Builder.AddTypeName(NAME_EnumProperty);
 		{
 			Builder.BeginTypeParameters();
+			Builder.AddTypeName(TEXT("Key"));
 			Builder.AddTypeName(NAME_ByteProperty);
-			{
-				Builder.BeginTypeParameters();
-				Builder.AddTypeName(TEXT("Key"));
-				Builder.EndTypeParameters();
-			}
 			Builder.EndTypeParameters();
 		}
 		Builder.AddTypeName(NAME_EnumProperty);
 		{
 			Builder.BeginTypeParameters();
+			Builder.AddTypeName(TEXT("Value"));
 			Builder.AddTypeName(NAME_ByteProperty);
-			{
-				Builder.BeginTypeParameters();
-				Builder.AddTypeName(TEXT("Value"));
-				Builder.EndTypeParameters();
-			}
 			Builder.EndTypeParameters();
 		}
 		Builder.EndTypeParameters();
@@ -159,17 +155,17 @@ TEST_CASE_NAMED(FPropertyTypeNameSmokeTest, "CoreUObject::PropertyTypeName::Smok
 		CHECK(TypeName.GetTypeName() == NAME_MapProperty);
 		CHECK(TypeName.GetTypeParameterCount() == 2);
 		CHECK(TypeName.GetTypeParameter(0).GetTypeName() == NAME_EnumProperty);
-		CHECK(TypeName.GetTypeParameter(0).GetTypeParameter(0).GetTypeName() == NAME_ByteProperty);
-		CHECK(TypeName.GetTypeParameter(0).GetTypeParameter(0).GetTypeParameter(0).GetTypeName() == TEXT("Key"));
-		CHECK(TypeName.GetTypeParameter(0).GetTypeParameter(0).GetTypeParameter(0).GetTypeParameter(0).IsEmpty());
-		CHECK(TypeName.GetTypeParameter(0).GetTypeParameter(0).GetTypeParameter(1).IsEmpty());
-		CHECK(TypeName.GetTypeParameter(0).GetTypeParameter(1).IsEmpty());
+		CHECK(TypeName.GetTypeParameter(0).GetTypeParameterName(0) == TEXT("Key"));
+		CHECK(TypeName.GetTypeParameter(0).GetTypeParameterName(1) == NAME_ByteProperty);
+		CHECK(TypeName.GetTypeParameter(0).GetTypeParameter(0).GetTypeParameterCount() == 0);
+		CHECK(TypeName.GetTypeParameter(0).GetTypeParameter(1).GetTypeParameterCount() == 0);
+		CHECK(TypeName.GetTypeParameter(0).GetTypeParameter(2).IsEmpty());
 		CHECK(TypeName.GetTypeParameter(1).GetTypeName() == NAME_EnumProperty);
-		CHECK(TypeName.GetTypeParameter(1).GetTypeParameter(0).GetTypeName() == NAME_ByteProperty);
-		CHECK(TypeName.GetTypeParameter(1).GetTypeParameter(0).GetTypeParameter(0).GetTypeName() == TEXT("Value"));
-		CHECK(TypeName.GetTypeParameter(1).GetTypeParameter(0).GetTypeParameter(0).GetTypeParameter(0).IsEmpty());
-		CHECK(TypeName.GetTypeParameter(1).GetTypeParameter(0).GetTypeParameter(1).IsEmpty());
-		CHECK(TypeName.GetTypeParameter(1).GetTypeParameter(1).IsEmpty());
+		CHECK(TypeName.GetTypeParameter(1).GetTypeParameterName(0) == TEXT("Value"));
+		CHECK(TypeName.GetTypeParameter(1).GetTypeParameterName(1) == NAME_ByteProperty);
+		CHECK(TypeName.GetTypeParameter(1).GetTypeParameter(0).GetTypeParameterCount() == 0);
+		CHECK(TypeName.GetTypeParameter(1).GetTypeParameter(1).GetTypeParameterCount() == 0);
+		CHECK(TypeName.GetTypeParameter(1).GetTypeParameter(2).IsEmpty());
 		CHECK(TypeName.GetTypeParameter(2).IsEmpty());
 	}
 
@@ -256,20 +252,66 @@ TEST_CASE_NAMED(FPropertyTypeNameSmokeTest, "CoreUObject::PropertyTypeName::Smok
 		FPropertyTypeName VectorArray = PropertyTypeNameTest::CreateVectorArray();
 
 		TArray<uint8> PersistentData;
-		FMemoryWriter PersistentSaveAr(PersistentData, /*bIsPersistent*/ true);
-		PersistentSaveAr << Empty << Int << EnumMap << VectorArray;
+		{
+			FMemoryWriter Ar(PersistentData, /*bIsPersistent*/ true);
+			Ar << Empty << Int << EnumMap << VectorArray;
+		}
 
 		TArray<uint8> MemoryData;
-		FMemoryWriter MemorySaveAr(MemoryData, /*bIsPersistent*/ false);
-		MemorySaveAr << Empty << Int << EnumMap << VectorArray;
+		{
+			FMemoryWriter Ar(MemoryData, /*bIsPersistent*/ false);
+			Ar << Empty << Int << EnumMap << VectorArray;
+		}
+
+		TArray<uint8> BinaryData;
+		{
+			FMemoryWriter Ar(BinaryData, /*bIsPersistent*/ true);
+			FBinaryArchiveFormatter Formatter(Ar);
+			FStructuredArchive StructuredAr(Formatter);
+			FStructuredArchiveRecord Record = StructuredAr.Open().EnterRecord();
+			Record.EnterField(TEXT("Empty")) << Empty;
+			Record.EnterField(TEXT("Int")) << Int;
+			Record.EnterField(TEXT("EnumMap")) << EnumMap;
+			Record.EnterField(TEXT("VectorArray")) << VectorArray;
+		}
+
+	#if WITH_TEXT_ARCHIVE_SUPPORT
+		TArray<uint8> JsonData;
+		{
+			FMemoryWriter Ar(JsonData, /*bIsPersistent*/ true);
+			FJsonArchiveOutputFormatter Formatter(Ar);
+			FStructuredArchive StructuredAr(Formatter);
+			FStructuredArchiveRecord Record = StructuredAr.Open().EnterRecord();
+			Record.EnterField(TEXT("Empty")) << Empty;
+			Record.EnterField(TEXT("Int")) << Int;
+			Record.EnterField(TEXT("EnumMap")) << EnumMap;
+			Record.EnterField(TEXT("VectorArray")) << VectorArray;
+		}
+	#endif
 
 		FPropertyTypeName EmptyCopy = Int; // needs to be non-empty
 		FPropertyTypeName IntCopy;
 		FPropertyTypeName EnumMapCopy;
 		FPropertyTypeName VectorArrayCopy;
 
-		FMemoryReader PersistentReader(PersistentData, /*bIsPersistent*/ true);
-		PersistentReader << EmptyCopy << IntCopy << EnumMapCopy << VectorArrayCopy;
+		{
+			FMemoryReader Ar(PersistentData, /*bIsPersistent*/ true);
+			Ar << EmptyCopy << IntCopy << EnumMapCopy << VectorArrayCopy;
+		}
+		CHECK(Empty == EmptyCopy);
+		CHECK(Int == IntCopy);
+		CHECK(EnumMap == EnumMapCopy);
+		CHECK(VectorArray == VectorArrayCopy);
+
+		EmptyCopy = Int; // needs to be non-empty
+		IntCopy.Reset();
+		EnumMapCopy.Reset();
+		VectorArrayCopy.Reset();
+
+		{
+			FMemoryReader Ar(MemoryData, /*bIsPersistent*/ false);
+			Ar << EmptyCopy << IntCopy << EnumMapCopy << VectorArrayCopy;
+		}
 		CHECK(Empty == EmptyCopy);
 		CHECK(Int == IntCopy);
 		CHECK(EnumMap == EnumMapCopy);
@@ -280,20 +322,104 @@ TEST_CASE_NAMED(FPropertyTypeNameSmokeTest, "CoreUObject::PropertyTypeName::Smok
 		EnumMapCopy.Reset();
 		VectorArrayCopy.Reset();
 
-		FMemoryReader MemoryReader(MemoryData, /*bIsPersistent*/ false);
-		MemoryReader << EmptyCopy << IntCopy << EnumMapCopy << VectorArrayCopy;
+		{
+			FMemoryReader Ar(BinaryData, /*bIsPersistent*/ true);
+			FBinaryArchiveFormatter Formatter(Ar);
+			FStructuredArchive StructuredAr(Formatter);
+			FStructuredArchiveRecord Record = StructuredAr.Open().EnterRecord();
+			Record.EnterField(TEXT("Empty")) << EmptyCopy;
+			Record.EnterField(TEXT("Int")) << IntCopy;
+			Record.EnterField(TEXT("EnumMap")) << EnumMapCopy;
+			Record.EnterField(TEXT("VectorArray")) << VectorArrayCopy;
+		}
 		CHECK(Empty == EmptyCopy);
 		CHECK(Int == IntCopy);
 		CHECK(EnumMap == EnumMapCopy);
 		CHECK(VectorArray == VectorArrayCopy);
+
+		EmptyCopy = Int; // needs to be non-empty
+		IntCopy.Reset();
+		EnumMapCopy.Reset();
+		VectorArrayCopy.Reset();
+
+	#if WITH_TEXT_ARCHIVE_SUPPORT
+		{
+			FMemoryReader Ar(JsonData, /*bIsPersistent*/ true);
+			FJsonArchiveInputFormatter Formatter(Ar);
+			FStructuredArchive StructuredAr(Formatter);
+			FStructuredArchiveRecord Record = StructuredAr.Open().EnterRecord();
+			Record.EnterField(TEXT("Empty")) << EmptyCopy;
+			Record.EnterField(TEXT("Int")) << IntCopy;
+			Record.EnterField(TEXT("EnumMap")) << EnumMapCopy;
+			Record.EnterField(TEXT("VectorArray")) << VectorArrayCopy;
+		}
+		CHECK(Empty == EmptyCopy);
+		CHECK(Int == IntCopy);
+		CHECK(EnumMap == EnumMapCopy);
+		CHECK(VectorArray == VectorArrayCopy);
+	#endif
 	}
 
+	const auto Parse = [](FStringView String) -> FPropertyTypeName
+	{
+		FPropertyTypeNameBuilder Builder;
+		CHECK(Builder.TryParse(String));
+		return Builder.Build();
+	};
+	const auto TryParse = [](FStringView String) -> bool
+	{
+		FPropertyTypeNameBuilder Builder;
+		return Builder.TryParse(String);
+	};
+
 	SECTION("String")
 	{
-		CHECK(TEXTVIEW("None").Equals(WriteToString<128>(FPropertyTypeName())));
-		CHECK(TEXTVIEW("IntProperty").Equals(WriteToString<128>(PropertyTypeNameTest::CreateInt())));
-		CHECK(TEXTVIEW("MapProperty<EnumProperty<ByteProperty<Key>>,EnumProperty<ByteProperty<Value>>>").Equals(WriteToString<128>(PropertyTypeNameTest::CreateEnumMap())));
-		CHECK(TEXTVIEW("ArrayProperty<StructProperty<Vector>>").Equals(WriteToString<128>(PropertyTypeNameTest::CreateVectorArray())));
+		FPropertyTypeName Empty;
+		FPropertyTypeName Int = PropertyTypeNameTest::CreateInt();
+		FPropertyTypeName EnumMap = PropertyTypeNameTest::CreateEnumMap();
+		FPropertyTypeName VectorArray = PropertyTypeNameTest::CreateVectorArray();
+
+		TStringBuilder<128> EmptyString(InPlace, Empty);
+		TStringBuilder<128> IntString(InPlace, Int);
+		TStringBuilder<128> EnumMapString(InPlace, EnumMap);
+		TStringBuilder<128> VectorArrayString(InPlace, VectorArray);
+
+		CHECK(TEXTVIEW("None").Equals(EmptyString));
+		CHECK(TEXTVIEW("IntProperty").Equals(IntString));
+		CHECK(TEXTVIEW("MapProperty<EnumProperty<Key,ByteProperty>,EnumProperty<Value,ByteProperty>>").Equals(EnumMapString));
+		CHECK(TEXTVIEW("ArrayProperty<StructProperty<Vector>>").Equals(VectorArrayString));
+
+		CHECK(Empty == Parse(EmptyString));
+		CHECK(Int == Parse(IntString));
+		CHECK(EnumMap == Parse(EnumMapString));
+		CHECK(VectorArray == Parse(VectorArrayString));
+	}
+
+	SECTION("Parse")
+	{
+		// There are positive parsing tests in the String section above.
+		CHECK(Parse(TEXTVIEW(" \t IntProperty \t ")) == PropertyTypeNameTest::CreateInt());
+		CHECK(Parse(TEXTVIEW(" \t ArrayProperty < StructProperty\t<\tVector\t>\t > \t ")) == PropertyTypeNameTest::CreateVectorArray());
+
+		CHECK_FALSE(TryParse(TEXTVIEW("")));
+		CHECK_FALSE(TryParse(TEXTVIEW(",")));
+		CHECK_FALSE(TryParse(TEXTVIEW("<>")));
+		CHECK_FALSE(TryParse(TEXTVIEW(",IntProperty")));
+		CHECK_FALSE(TryParse(TEXTVIEW("IntProperty,")));
+		CHECK_FALSE(TryParse(TEXTVIEW("IntProperty,IntProperty")));
+		CHECK_FALSE(TryParse(TEXTVIEW("IntProperty>")));
+		CHECK_FALSE(TryParse(TEXTVIEW("<IntProperty>")));
+		CHECK_FALSE(TryParse(TEXTVIEW("<IntProperty")));
+		CHECK_FALSE(TryParse(TEXTVIEW("ArrayProperty<IntProperty")));
+		CHECK_FALSE(TryParse(TEXTVIEW("ArrayProperty<IntProperty>>")));
+		CHECK_FALSE(TryParse(TEXTVIEW("ArrayProperty<<IntProperty>")));
+		CHECK_FALSE(TryParse(TEXTVIEW("ArrayProperty<IntProperty>IntProperty")));
+		CHECK_FALSE(TryParse(TEXTVIEW("MapProperty<>")));
+		CHECK_FALSE(TryParse(TEXTVIEW("MapProperty<,>")));
+		CHECK_FALSE(TryParse(TEXTVIEW("MapProperty<IntProperty,>")));
+		CHECK_FALSE(TryParse(TEXTVIEW("MapProperty<,IntProperty>")));
+		CHECK_FALSE(TryParse(TEXTVIEW("MapProperty<IntProperty,,IntProperty>")));
+		CHECK_FALSE(TryParse(TEXTVIEW("MapProperty<IntProperty><IntProperty>")));
 	}
 }
 
