Fix an editor crash that could result during propagation of a PostEditChange() event after modifying a nested struct member field on an archetype.

The crash is a regression introduced at CL 24972477. The root cause was introduced at CL 2545115 but did not manifest until the calling context became necessary to use for value resolution.

diff --git a/Engine/Source/Runtime/CoreUObject/Private/UObject/Obj.cpp b/Engine/Source/Runtime/CoreUObject/Private/UObject/Obj.cpp
--- a/Engine/Source/Runtime/CoreUObject/Private/UObject/Obj.cpp
+++ b/Engine/Source/Runtime/CoreUObject/Private/UObject/Obj.cpp
@@ -509,12 +509,12 @@ void UObject::PostEditChangeChainProperty(FPropertyChangedChainEvent& PropertyCh
 	if (!FApp::IsGame())
 	{
 		if (HasAnyFlags(RF_ClassDefaultObject | RF_ArchetypeObject) && PropertyChangedEvent.PropertyChain.GetActiveMemberNode() == PropertyChangedEvent.PropertyChain.GetHead())
-	{
+		{
 			// Get a list of all archetype instances
 			TArray<UObject*> ArchetypeInstances;
 			GetArchetypeInstances(ArchetypeInstances);
 
-		// Propagate the editchange call to archetype instances
+			// Propagate the editchange call to archetype instances
 			PropagatePostEditChange(ArchetypeInstances, PropertyChangedEvent);
 		}
 		else if (GetOuter()->HasAnyFlags(RF_ClassDefaultObject | RF_ArchetypeObject))
@@ -541,7 +541,7 @@ void UObject::PostEditChangeChainProperty(FPropertyChangedChainEvent& PropertyCh
 					}
 				}
 
-				GetOuter()->PropagatePostEditChange(ArchetypeComponentInstances, PropertyChangedEvent);
+				PropagatePostEditChange(ArchetypeComponentInstances, PropertyChangedEvent);
 
 				break;
 			}
