Convenience refactor of CoreUObject.Build.cs

#rnx

diff --git a/Engine/Source/Runtime/CoreUObject/CoreUObject.Build.cs b/Engine/Source/Runtime/CoreUObject/CoreUObject.Build.cs
--- a/Engine/Source/Runtime/CoreUObject/CoreUObject.Build.cs
+++ b/Engine/Source/Runtime/CoreUObject/CoreUObject.Build.cs
@@ -5,12 +5,17 @@ using UnrealBuildTool;
 
 public class CoreUObject : ModuleRules
 {
+	public static bool IsVerseVMEnabled(ReadOnlyTargetRules Target)
+	{
+		return !Target.bUseVerseBPVM || Target.GlobalDefinitions.Contains("WITH_VERSE_VM=1")
+			|| Target.GlobalDefinitions.Contains("WITH_VERSE_VM=WITH_COREUOBJECT");
+	}
+
 	public static void AddVerseVMDependencies(ModuleRules Rules, ReadOnlyTargetRules Target)
 	{
 		// If using the new VM either by default or directly, then add in a dependency to 
 		// the core VerseVM which contains special compile flags not compatible with CoreUObject
-		if (!Target.bUseVerseBPVM || Target.GlobalDefinitions.Contains("WITH_VERSE_VM=1")
-			|| Target.GlobalDefinitions.Contains("WITH_VERSE_VM=WITH_COREUOBJECT"))
+		if (IsVerseVMEnabled(Target))
 		{
 			Rules.PublicDependencyModuleNames.AddRange(
 				new string[]
