Fixes check for DoesPackageExist to confirm the IoDispatcher is not being used. The IoDispatcher is unilaterally initialized but may never have any mounts mounted so relying on it's initialization value is not correct. Instead of exposing an IoDispatcher::HasMountedBackends(), we rely on GetTotalLoaded==0 to know that we are not actually using the IoDispatcher and can rely on the AssetRegistry to reliably tell if packages exist on the Filesystem.

diff --git a/Engine/Source/Runtime/CoreUObject/Private/Misc/PackageName.cpp b/Engine/Source/Runtime/CoreUObject/Private/Misc/PackageName.cpp
--- a/Engine/Source/Runtime/CoreUObject/Private/Misc/PackageName.cpp
+++ b/Engine/Source/Runtime/CoreUObject/Private/Misc/PackageName.cpp
@@ -1861,7 +1861,8 @@ FPackageName::EPackageLocationFilter FPackageName::DoesPackageExistEx(const FPac
 
 	// Todo: The AssetRegistry currently cannot determine if a package comes from the Filesystem 
 	// or cooked content so we avoid registry lookups since we can't provide a reliable Location
-	if (AssetRegistry && ((uint8)Filter & (uint8)EPackageLocationFilter::FileSystem) && !FIoDispatcher::IsInitialized())
+	if (AssetRegistry && ((uint8)Filter & (uint8)EPackageLocationFilter::FileSystem) && 
+		(!FIoDispatcher::IsInitialized() || FIoDispatcher::Get().GetTotalLoaded() == 0))
 	{
 		FName PackageName = PackagePath.GetPackageFName();
 		FName CorrectCasePackageName;
