AssetRegistry: Fix deadlock to to unexpected recursion in new code that calls AssetRegistry functions to accelerate FPackageName::DoesPackageExist.
#rnx

diff --git a/Engine/Source/Runtime/CoreUObject/Private/Misc/PackageName.cpp b/Engine/Source/Runtime/CoreUObject/Private/Misc/PackageName.cpp
--- a/Engine/Source/Runtime/CoreUObject/Private/Misc/PackageName.cpp
+++ b/Engine/Source/Runtime/CoreUObject/Private/Misc/PackageName.cpp
@@ -1856,7 +1856,7 @@ FPackageName::EPackageLocationFilter FPackageName::DoesPackageExistEx(const FPac
 		return EPackageLocationFilter::None;
 	}
 
-#if WITH_EDITOR
+#if 0 // WITH_EDITOR // Temporarily disabled due to deadlock, UE-205707
 	IAssetRegistryInterface* AssetRegistry = IAssetRegistryInterface::GetPtr();
 
 	// Todo: The AssetRegistry currently cannot determine if a package comes from the Filesystem 
