CompileTimes: Change UPackage's COREUOBJECT_API markup to include just the functions it needs to export rather than the entire class.
Formatting: add some line breaks for some long lines.
#rnx

diff --git a/Engine/Source/Runtime/CoreUObject/Public/UObject/Package.h b/Engine/Source/Runtime/CoreUObject/Public/UObject/Package.h
--- a/Engine/Source/Runtime/CoreUObject/Public/UObject/Package.h
+++ b/Engine/Source/Runtime/CoreUObject/Public/UObject/Package.h
@@ -139,9 +139,9 @@ struct FSavePackageResultStruct
 * A package.
 */
 PRAGMA_DISABLE_DEPRECATION_WARNINGS // Required for auto-generated functions referencing PackageFlags
-class COREUOBJECT_API UPackage : public UObject
+class UPackage : public UObject
 {
-	DECLARE_CASTED_CLASS_INTRINSIC_NO_CTOR_NO_VTABLE_CTOR(UPackage, UObject, 0, TEXT("/Script/CoreUObject"), CASTCLASS_UPackage, NO_API)
+	DECLARE_CASTED_CLASS_INTRINSIC_NO_CTOR_NO_VTABLE_CTOR(UPackage, UObject, 0, TEXT("/Script/CoreUObject"), CASTCLASS_UPackage, COREUOBJECT_API)
 	/** DO NOT USE. This constructor is for internal usage only for hot-reload purposes. */
 	UPackage(FVTableHelper& Helper)
 		: Super(Helper)
@@ -169,28 +169,28 @@ public:
 	DECLARE_MULTICAST_DELEGATE_TwoParams(FPreSavePackageWithContext, class UPackage*, FObjectPreSaveContext);
 
 	UE_DEPRECATED(5.0, "Use PreSavePackageWithContextEvent instead.")
-	static FPreSavePackage PreSavePackageEvent;
+	COREUOBJECT_API static FPreSavePackage PreSavePackageEvent;
 	/** Delegate to notify subscribers when a package is about to be saved. */
-	static FPreSavePackageWithContext PreSavePackageWithContextEvent;
+	COREUOBJECT_API static FPreSavePackageWithContext PreSavePackageWithContextEvent;
 	UE_DEPRECATED(5.0, "Use PackageSavedWithContextEvent instead.")
-	static FOnPackageSaved PackageSavedEvent;
+	COREUOBJECT_API static FOnPackageSaved PackageSavedEvent;
 	/** Delegate to notify subscribers when a package has been saved. This is triggered when the package saving
 	*  has completed and was successful. */
-	static FOnPackageSavedWithContext PackageSavedWithContextEvent;
+	COREUOBJECT_API static FOnPackageSavedWithContext PackageSavedWithContextEvent;
 	/** Delegate to notify subscribers when the dirty state of a package is changed.
 	*  Allows the editor to register the modified package as one that should be prompted for source control checkout. 
 	*  Use Package->IsDirty() to get the updated dirty state of the package */
-	static FOnPackageDirtyStateChanged PackageDirtyStateChangedEvent;
+	COREUOBJECT_API static FOnPackageDirtyStateChanged PackageDirtyStateChangedEvent;
 	/** 
 	* Delegate to notify subscribers when a package is marked as dirty via UObjectBaseUtilty::MarkPackageDirty 
 	* Note: Unlike FOnPackageDirtyStateChanged, this is always called, even when the package is already dirty
 	* Use bWasDirty to check the previous dirty state of the package
 	* Use Package->IsDirty() to get the updated dirty state of the package
 	*/
-	static FOnPackageMarkedDirty PackageMarkedDirtyEvent;
+	COREUOBJECT_API static FOnPackageMarkedDirty PackageMarkedDirtyEvent;
 
 private:
-	/** Used by the editor to determine if a package has been changed.																							*/
+	/** Used by the editor to determine if a package has been changed. */
 	uint8	bDirty:1;
 
 #if WITH_EDITORONLY_DATA
@@ -205,7 +205,7 @@ public:
 	uint8 bIsCookedForEditor:1;
 #endif
 
-	/** Whether this package has been fully loaded (aka had all it's exports created) at some point.															*/
+	/** Whether this package has been fully loaded (aka had all it's exports created) at some point. */
 	mutable uint8 bHasBeenFullyLoaded:1;
 
 	/**
@@ -239,9 +239,13 @@ public:
 #endif
 
 private:
-	// @note this should probably be entirely deprecated and removed but certain stat dump function are still using it, just compile it out in shipping for now
+	// @note this should probably be entirely deprecated and removed but certain stat dump function are still using it,
+	// just compile it out in shipping for now
 #if !UE_BUILD_SHIPPING
-	/** Time in seconds it took to fully load this package. 0 if package is either in process of being loaded or has never been fully loaded. */
+	/**
+	 * Time in seconds it took to fully load this package.
+	 * 0 if package is either in process of being loaded or has never been fully loaded.
+	 */
 	float LoadTime;
 #endif
 
@@ -264,7 +268,7 @@ private:
 #endif
 
 	/** Package Flags */
-	uint32	PackageFlagsPrivate;
+	uint32 PackageFlagsPrivate;
 	
 	/** Globally unique id */
 	FPackageId PackageId;
@@ -325,34 +329,41 @@ public:
 
 public:
 
-	virtual bool IsNameStableForNetworking() const override { return true; }		// For now, assume all packages have stable net names
-	virtual bool NeedsLoadForClient() const override { return true; }				// To avoid calling the expensive generic version, which only makes sure that the UPackage static class isn't excluded
+	// For now, assume all packages have stable net names
+	virtual bool IsNameStableForNetworking() const override { return true; }
+	// We override NeedsLoadForClient to avoid calling the expensive generic version, which only makes sure that the
+	// UPackage static class isn't excluded
+	virtual bool NeedsLoadForClient() const override { return true; }
 	virtual bool NeedsLoadForServer() const override { return true; }
 	virtual bool IsPostLoadThreadSafe() const override;
 	virtual bool IsDestructionThreadSafe() const override { return true; }
 
 #if WITH_EDITORONLY_DATA
 	/** Sets the bLoadedByEditorPropertiesOnly flag */
-	void SetLoadedByEditorPropertiesOnly(bool bIsEditorOnly, bool bRecursive = false);
+	COREUOBJECT_API void SetLoadedByEditorPropertiesOnly(bool bIsEditorOnly, bool bRecursive = false);
 	/** returns true when the package is only referenced by editor-only flag */
 	bool IsLoadedByEditorPropertiesOnly() const { return bLoadedByEditorPropertiesOnly; }
 
 	/** Sets the bIsDynamicPIEPackagePending flag */
-	void SetDynamicPIEPackagePending(bool bInIsDynamicPIEPackagePending) { bIsDynamicPIEPackagePending = bInIsDynamicPIEPackagePending; }
+	void SetDynamicPIEPackagePending(bool bInIsDynamicPIEPackagePending)
+	{
+		bIsDynamicPIEPackagePending = bInIsDynamicPIEPackagePending;
+	}
 	/** returns the bIsDynamicPIEPackagePending flag */
 	bool IsDynamicPIEPackagePending() const { return bIsDynamicPIEPackagePending; }
 #endif
 
 	/**
-	* Called after the C++ constructor and after the properties have been initialized, but before the config has been loaded, etc.
-	* mainly this is to emulate some behavior of when the constructor was called after the properties were initialized.
+	* Called after the C++ constructor and after the properties have been initialized, but before the config has been
+	* loaded, etc. Mainly this is to emulate some behavior of when the constructor was called after the properties were
+	* initialized.
 	*/
-	virtual void PostInitProperties() override;
+	COREUOBJECT_API virtual void PostInitProperties() override;
 
-	virtual void FinishDestroy() override;
+	COREUOBJECT_API virtual void FinishDestroy() override;
 
 	/** Serializer */
-	virtual void Serialize( FArchive& Ar ) override;
+	COREUOBJECT_API virtual void Serialize( FArchive& Ar ) override;
 
 	/** Packages are never assets */
 	virtual bool IsAsset() const override { return false; }
@@ -363,10 +374,10 @@ public:
 	 * each package are Objects in each package are stored in SoftGCPackageToObjectList and are
 	 * marked as referenced if the package is referenced.
 	 */
-	static bool bSupportCookerSoftGC;
-	static TMap<UPackage*, TArrayView<TObjectPtr<UObject>>> SoftGCPackageToObjectList;
+	COREUOBJECT_API static bool bSupportCookerSoftGC;
+	COREUOBJECT_API static TMap<UPackage*, TArrayView<TObjectPtr<UObject>>> SoftGCPackageToObjectList;
 	/** Static class override of UObject::AddReferencedObjects */
-	static void AddReferencedObjects(UObject* InThis, FReferenceCollector& Collector);
+	COREUOBJECT_API static void AddReferencedObjects(UObject* InThis, FReferenceCollector& Collector);
 #endif
 
 	// UPackage interface.
@@ -454,7 +465,6 @@ public:
 		return LinkerCustomVersion;
 	}
 
-
 	/**
 	 * Sets the time it took to load this package.
 	 */
@@ -511,7 +521,7 @@ public:
 	/**
 	* Marks/Unmarks the package's bDirty flag, save the package to the transaction buffer if a transaction is ongoing
 	*/
-	void SetDirtyFlag( bool bIsDirty );
+	COREUOBJECT_API void SetDirtyFlag( bool bIsDirty );
 
 	/**
 	* Returns whether the package needs to be saved.
@@ -557,22 +567,23 @@ public:
 	*
 	* @return true if fully loaded or no file associated on disk, false otherwise
 	*/
-	bool IsFullyLoaded() const;
+	COREUOBJECT_API bool IsFullyLoaded() const;
 
 	/**
 	* Fully loads this package. Safe to call multiple times and won't clobber already loaded assets.
 	*/
-	void FullyLoad();
+	COREUOBJECT_API void FullyLoad();
 
 	/**
-	 * Get the path this package was loaded from; may be different than packagename, and may not be set if the package was not loaded from disk
+	 * Get the path this package was loaded from; may be different than packagename, and may not be set if the package
+	 * was not loaded from disk
 	 */
-	const FPackagePath& GetLoadedPath() const;
+	COREUOBJECT_API const FPackagePath& GetLoadedPath() const;
 
 	/**
 	 * Set the path this package was loaded from; typically called only by the linker during load
 	 */
-	void SetLoadedPath(const FPackagePath& PackagePath);
+	COREUOBJECT_API void SetLoadedPath(const FPackagePath& PackagePath);
 
 	/**
 	* Marks/Unmarks the package's bCanBeImported flag.
@@ -594,7 +605,7 @@ public:
 	/**
 	* Tags the Package's metadata
 	*/
-	virtual void TagSubobjects(EObjectFlags NewFlags) override;
+	COREUOBJECT_API virtual void TagSubobjects(EObjectFlags NewFlags) override;
 
 	/**
 	* Called to indicate that this package contains a ULevel or UWorld object.
@@ -869,13 +880,13 @@ public:
 	 * Utility function to find Asset in this package, if any
 	 * @return the asset in the package, if any
 	 */
-	UObject* FindAssetInPackage(EObjectFlags RequiredTopLevelFlags = RF_NoFlags) const;
+	COREUOBJECT_API UObject* FindAssetInPackage(EObjectFlags RequiredTopLevelFlags = RF_NoFlags) const;
 
 	/**
 	 * Return the list of packages found assigned to object outer-ed to the top level objects of this package
 	 * @return the array of external packages
 	 */
-	TArray<UPackage*> GetExternalPackages() const;
+	COREUOBJECT_API TArray<UPackage*> GetExternalPackages() const;
 
 	////////////////////////////////////////////////////////
 	// MetaData 
@@ -900,7 +911,7 @@ public:
 	*
 	* @return A valid UMetaData pointer for all objects in this package
 	*/
-	class UMetaData* GetMetaData();
+	COREUOBJECT_API class UMetaData* GetMetaData();
 
 	/**
 	* Save one specific object (along with any objects it references contained within the same Outer) into an Unreal package.
@@ -913,13 +924,15 @@ public:
 	*
 	* @return	FSavePackageResultStruct enum value with the result of saving a package as well as extra data
 	*/
-	static FSavePackageResultStruct Save(UPackage* InOuter, UObject* InAsset, const TCHAR* Filename, const FSavePackageArgs& SaveArgs);
+	COREUOBJECT_API static FSavePackageResultStruct Save(UPackage* InOuter, UObject* InAsset, const TCHAR* Filename,
+		const FSavePackageArgs& SaveArgs);
 
 	/**
 	 * Save a list of packages concurrently using Save2 mechanism
 	 * SaveConcurrent is currently experimental and shouldn't be used until it can safely replace Save.
 	 */
-	static ESavePackageResult SaveConcurrent(TArrayView<FPackageSaveInfo> InPackages, const FSavePackageArgs& SaveArgs, TArray<FSavePackageResultStruct>& OutResults);
+	COREUOBJECT_API static ESavePackageResult SaveConcurrent(TArrayView<FPackageSaveInfo> InPackages,
+		const FSavePackageArgs& SaveArgs, TArray<FSavePackageResultStruct>& OutResults);
 
 	/**
 	* Save one specific object (along with any objects it references contained within the same Outer) into an Unreal package.
@@ -932,19 +945,21 @@ public:
 	*
 	* @return	true if the package was saved successfully.
 	*/
-	static bool SavePackage(UPackage* InOuter, UObject* InAsset, const TCHAR* Filename, const FSavePackageArgs& SaveArgs);
+	COREUOBJECT_API static bool SavePackage(UPackage* InOuter, UObject* InAsset, const TCHAR* Filename,
+		const FSavePackageArgs& SaveArgs);
 
 	UE_DEPRECATED(5.0, "Pack the arguments into FSavePackageArgs and call the function overload that takes FSavePackageArgs. Note that Conform is no longer implemented.")
-	static bool SavePackage(UPackage* InOuter, UObject* Base, EObjectFlags TopLevelFlags, const TCHAR* Filename,
-		FOutputDevice* Error = GError, FLinkerNull* Conform = nullptr, bool bForceByteSwapping = false,
-		bool bWarnOfLongFilename = true, uint32 SaveFlags = SAVE_None, const ITargetPlatform* TargetPlatform = nullptr,
-		const FDateTime& FinalTimeStamp = FDateTime::MinValue(), bool bSlowTask = true);
+	COREUOBJECT_API static bool SavePackage(UPackage* InOuter, UObject* Base, EObjectFlags TopLevelFlags,
+		const TCHAR* Filename, FOutputDevice* Error = GError, FLinkerNull* Conform = nullptr,
+		bool bForceByteSwapping = false, bool bWarnOfLongFilename = true, uint32 SaveFlags = SAVE_None,
+		const ITargetPlatform* TargetPlatform = nullptr, const FDateTime& FinalTimeStamp = FDateTime::MinValue(),
+		bool bSlowTask = true);
 
 	/** Wait for any SAVE_Async file writes to complete **/
-	static void WaitForAsyncFileWrites();
+	COREUOBJECT_API static void WaitForAsyncFileWrites();
 
 	/** Return true if SAVE_Async file writes are pending **/
-	static bool HasAsyncFileWrites();
+	COREUOBJECT_API static bool HasAsyncFileWrites();
 
 	/**
 	* Determines if a package contains no more assets.
@@ -953,7 +968,7 @@ public:
 	* @param LastReferencer	the optional last UObject referencer to this package. This object will be excluded when determining if the package is empty
 	* @return true if Package contains no more assets.
 	*/
-	static bool IsEmptyPackage(UPackage* Package, const UObject* LastReferencer = NULL);
+	COREUOBJECT_API static bool IsEmptyPackage(UPackage* Package, const UObject* LastReferencer = NULL);
 
 private:
 	static FSavePackageResultStruct Save2(UPackage* InPackage, UObject* InAsset, const TCHAR* InFilename, const FSavePackageArgs& SaveArgs);
