Add a cvar that enables property bag and IDO support for loaded Verse Class objects

diff --git a/Engine/Source/Runtime/CoreUObject/Private/UObject/InstanceDataObjectUtils.cpp b/Engine/Source/Runtime/CoreUObject/Private/UObject/InstanceDataObjectUtils.cpp
--- a/Engine/Source/Runtime/CoreUObject/Private/UObject/InstanceDataObjectUtils.cpp
+++ b/Engine/Source/Runtime/CoreUObject/Private/UObject/InstanceDataObjectUtils.cpp
@@ -5,6 +5,7 @@
 #include "UObject/UnrealType.h"
 #include "UObject/EnumProperty.h"
 #include "UObject/Field.h"
+#include "HAL/IConsoleManager.h"
 
 namespace UE
 {
@@ -16,6 +17,7 @@ namespace UE
 	static const FName NAME_PresentAsTypeMetadata(TEXT("PresentAsType"));
 	static const FName NAME_IsLooseMetadata(TEXT("IsLoose"));
 	static const FName NAME_CategoryMetadata(TEXT("Category"));
+	static const FName NAME_VerseClass("VerseClass");
 	
 	struct ResolvePropertyPathNameHelperParams
 	{
@@ -27,6 +29,20 @@ namespace UE
 		bool bAddIfNeeded = false;
 	};
 
+	bool bEnableIDOSupport = false;
+	FAutoConsoleVariableRef EnableIDOSupportCVar(
+		TEXT("IDO.Enable"),
+		bEnableIDOSupport,
+		TEXT("Allows property bags and IDOs to be created for supported classes.")
+	);
+
+	bool IsInstanceDataObjectSupportEnabled(UObject* InObject)
+	{
+		return bEnableIDOSupport
+			//@todo FH: change to check trait when available or use config object
+			&& (!InObject || InObject->GetClass()->GetClass()->GetFName() == NAME_VerseClass);
+	}
+
 	static void BuildSegmentTypeFromProperty(const FProperty* Property, FPropertyTypeNameBuilder& OutType)
 	{
 #if false // TODO: @jordan.hoffmann when complete type info is finished use this branch instead
@@ -652,4 +668,4 @@ namespace UE
 		}
 		return false;
 	}
-} // UE
\ No newline at end of file
+} // UE
diff --git a/Engine/Source/Runtime/CoreUObject/Private/UObject/InstanceDataObjectUtils.h b/Engine/Source/Runtime/CoreUObject/Private/UObject/InstanceDataObjectUtils.h
--- a/Engine/Source/Runtime/CoreUObject/Private/UObject/InstanceDataObjectUtils.h
+++ b/Engine/Source/Runtime/CoreUObject/Private/UObject/InstanceDataObjectUtils.h
@@ -13,6 +13,9 @@ namespace UE
 	class FPropertyBag;
 	class FPropertyPathName;
 
+	// Query if InstanceDataObject are enabled for a specific object, passing in nullptr return if the system itself is enabled
+	bool IsInstanceDataObjectSupportEnabled(UObject* InObject = nullptr);
+
 	// generate a UClass that unions the properties of PropertyBag and TemplateStruct
 	UClass* CreateInstanceDataObjectClass(const FPropertyBag* PropertyBag, UClass* OwnerClass, UObject* Outer);
 
diff --git a/Engine/Source/Runtime/CoreUObject/Private/UObject/LinkerLoad.cpp b/Engine/Source/Runtime/CoreUObject/Private/UObject/LinkerLoad.cpp
--- a/Engine/Source/Runtime/CoreUObject/Private/UObject/LinkerLoad.cpp
+++ b/Engine/Source/Runtime/CoreUObject/Private/UObject/LinkerLoad.cpp
@@ -12,6 +12,7 @@
 #include "Misc/ScopedSlowTask.h"
 #include "Misc/ObjectThumbnail.h"
 #include "Misc/App.h"
+#include "UObject/InstanceDataObjectUtils.h"
 #include "UObject/MetaData.h"
 #include "UObject/LinkerLoadImportBehavior.h"
 #include "UObject/ObjectRedirector.h"
@@ -19,6 +20,7 @@
 #include "UObject/PackageResourceManager.h"
 #include "UObject/PackageResourceIoDispatcherBackend.h"
 #include "UObject/PackageTrailer.h"
+#include "UObject/PropertyBagRepository.h"
 #include "UObject/SoftObjectPath.h"
 #include "UObject/UObjectHash.h"
 #include "Misc/PackageName.h"
@@ -4580,6 +4582,20 @@ void FLinkerLoad::Preload( UObject* Object )
 						ESoftObjectPathSerializeType::AlwaysSerialize);
 #endif
 
+					// Enable IDO, if needed
+					FUObjectSerializeContext* LoadContext = FUObjectThreadContext::Get().GetSerializeContext();
+					TOptional<TGuardValue<bool>> ScopedTrackSerializedPropertyPath;
+					TOptional<TGuardValue<bool>> ScopedSerializeUnknownProperty;
+
+					// Do not enable IDO when impersonation is enabled as we are probably deserializing an IDO already at that point
+					bool bIDOEnabled = UE::IsInstanceDataObjectSupportEnabled(Object) && !LoadContext->bImpersonateProperties;
+					if (bIDOEnabled )
+					{
+						// this will had property path tracking and create a property bag to hold data not matching the current class schema
+						ScopedTrackSerializedPropertyPath.Emplace(LoadContext->bTrackSerializedPropertyPath, true);
+						ScopedSerializeUnknownProperty.Emplace(LoadContext->bSerializeUnknownProperty, true);
+					}
+
 					if (Object->HasAnyFlags(RF_ClassDefaultObject))
 					{
 #if USE_CIRCULAR_DEPENDENCY_LOAD_DEFERRING
@@ -4681,6 +4697,12 @@ void FLinkerLoad::Preload( UObject* Object )
 						Object->SetFlags(RF_LoadCompleted);
 						CurrentLoadContext->SerializedObject = PrevSerializedObject;
 					}
+
+					// Object has been deserialized, if IDO is enabled, generate it
+					if (bIDOEnabled)
+					{
+						UE::FPropertyBagRepository::Get().CreateInstanceDataObject(Object);
+					}
 				}
 
 #if USE_CIRCULAR_DEPENDENCY_LOAD_DEFERRING
diff --git a/Engine/Source/Runtime/CoreUObject/Private/UObject/Obj.cpp b/Engine/Source/Runtime/CoreUObject/Private/UObject/Obj.cpp
--- a/Engine/Source/Runtime/CoreUObject/Private/UObject/Obj.cpp
+++ b/Engine/Source/Runtime/CoreUObject/Private/UObject/Obj.cpp
@@ -1857,9 +1857,12 @@ void UObject::SerializeScriptProperties( FStructuredArchive::FSlot Slot ) const
 			ObjClass = ThisObject->GetClass();
 
 			//@todo FH: Support for Default Object needs to be added when using impersonators
-			const UObject* DiffImpersonator = UE::Private::GetDataImpersonator(DiffObject);
+			const UObject* DiffImpersonator = nullptr; //UE::Private::GetDataImpersonator(DiffObject);
 			DiffObject = DiffImpersonator ? const_cast<UObject*>(DiffImpersonator) : ObjClass->GetDefaultObject(false);
-			DiffClass = DiffObject->GetClass();
+			if (DiffObject)
+			{
+				DiffClass = DiffObject->GetClass();
+			}
 			ensureAlwaysMsgf(DiffClass == ObjClass, TEXT("Impersonation of '%s' using a different default class not appropriately supported at the moment. Class: '%s', DefaultClass: '%s'")
 				, *ThisObject->GetPathName(), *ObjClass->GetPathName(), *DiffClass->GetPathName());
 		}
diff --git a/Engine/Source/Runtime/CoreUObject/Private/UObject/SavePackage2.cpp b/Engine/Source/Runtime/CoreUObject/Private/UObject/SavePackage2.cpp
--- a/Engine/Source/Runtime/CoreUObject/Private/UObject/SavePackage2.cpp
+++ b/Engine/Source/Runtime/CoreUObject/Private/UObject/SavePackage2.cpp
@@ -34,6 +34,7 @@
 #include "Serialization/UnversionedPropertySerialization.h"
 #include "UObject/DebugSerializationFlags.h"
 #include "UObject/EditorObjectVersion.h"
+#include "UObject/InstanceDataObjectUtils.h"
 #include "UObject/LinkerSave.h"
 #include "UObject/Object.h"
 #include "UObject/ObjectMacros.h"
@@ -2909,6 +2910,12 @@ ESavePackageResult InnerSave(FSaveContext& SaveContext)
 	SaveContext.SetSerializeContext(SerializeContext);
 	SaveContext.SetEDLCookChecker(&FEDLCookCheckerThreadState::Get());
 
+	TOptional<TGuardValue<bool>> IDOImpersonationScope;
+	if (UE::IsInstanceDataObjectSupportEnabled())
+	{
+		IDOImpersonationScope.Emplace(SerializeContext->bImpersonateProperties, true);
+	}
+
 	// Create slow task dialog if needed
 	const int32 TotalSaveSteps = 3;
 	FScopedSlowTask SlowTask(TotalSaveSteps, FText(), SaveContext.IsUsingSlowTask());
