Renamed the "LoadObject" named event (defined by SCOPED_NAMED_EVENT) as "StaticLoadObjectInternal" (measures time in StaticLoadObjectInternal() function) because there is another "LoadObject" cpu scope (defined using the UE_TRACE_LOG_SCOPED_T macro; the one with metadata) which measures same function.

diff --git a/Engine/Source/Runtime/CoreUObject/Private/UObject/UObjectGlobals.cpp b/Engine/Source/Runtime/CoreUObject/Private/UObject/UObjectGlobals.cpp
--- a/Engine/Source/Runtime/CoreUObject/Private/UObject/UObjectGlobals.cpp
+++ b/Engine/Source/Runtime/CoreUObject/Private/UObject/UObjectGlobals.cpp
@@ -1343,7 +1343,7 @@ UObject* StaticLoadObjectInternal(UClass* ObjectClass, UObject* InOuter, const T
 	UE_TRACE_LOG_SCOPED_T(Cpu, LoadObject, CpuChannel)
 		<< LoadObject.AssetPath(InName);
 #endif // CPUPROFILERTRACE_ENABLED
-	SCOPED_NAMED_EVENT(LoadObject, FColor::Red);
+	SCOPED_NAMED_EVENT(StaticLoadObjectInternal, FColor::Red);
 	check(InName);
 
 	FScopedLoadingState ScopedLoadingState(InName);
