[ZenLoader]
  - Fix infinite loop during flush when multiple interdependent packages are waiting on the package currently on stack

diff --git a/Engine/Source/Runtime/CoreUObject/Private/Serialization/AsyncLoading2.cpp b/Engine/Source/Runtime/CoreUObject/Private/Serialization/AsyncLoading2.cpp
--- a/Engine/Source/Runtime/CoreUObject/Private/Serialization/AsyncLoading2.cpp
+++ b/Engine/Source/Runtime/CoreUObject/Private/Serialization/AsyncLoading2.cpp
@@ -7074,6 +7074,11 @@ FAsyncPackage2* FAsyncPackage2::UpdateDependenciesStateRecursive(FAsyncLoadingTh
 			FAllDependenciesState::RemoveFromWaitList(Context.StateMemberPtr, WaitingForPackage, this);
 			WaitingForPackage = nullptr;
 		}
+		else if (ThreadState.PackagesOnStack.Contains(WaitingForPackage))
+		{
+			FAllDependenciesState::RemoveFromWaitList(Context.StateMemberPtr, WaitingForPackage, this);
+			WaitingForPackage = nullptr;
+		}
 		else
 		{
 			return WaitingForPackage;
