BulkData - removed return value from issue batch

diff --git a/Engine/Source/Runtime/CoreUObject/Private/Serialization/BulkDataStreaming.cpp b/Engine/Source/Runtime/CoreUObject/Private/Serialization/BulkDataStreaming.cpp
--- a/Engine/Source/Runtime/CoreUObject/Private/Serialization/BulkDataStreaming.cpp
+++ b/Engine/Source/Runtime/CoreUObject/Private/Serialization/BulkDataStreaming.cpp
@@ -863,14 +863,14 @@ public:
 		}
 	}
 
-	bool Issue(FBulkDataRequest::FCompletionCallback&& Callback) 
+	void Issue(FBulkDataRequest::FCompletionCallback&& Callback) 
 	{
 		CompletionCallback = MoveTemp(Callback);
 
 		if (Requests.IsEmpty())
 		{
 			CompleteBatch(FBulkDataRequest::EStatus::Ok);
-			return false;
+			return;
 		}
 
 		TRACE_COUNTER_INCREMENT(BulkDataBatchRequest_PendingCount);
@@ -894,8 +894,6 @@ public:
 
 			TRACE_COUNTER_DECREMENT(BulkDataBatchRequest_PendingCount);
 		});
-
-		return true;
 	}
 
 private:
@@ -1034,27 +1032,18 @@ FBulkDataBatchRequest::FBatchHandle& FBulkDataBatchRequest::FBuilder::GetBatch()
 	return *Batch;
 }
 
-FBulkDataRequest::EStatus FBulkDataBatchRequest::FBuilder::IssueBatch(FBulkDataBatchRequest* OutRequest, FCompletionCallback&& Callback)
+void FBulkDataBatchRequest::FBuilder::IssueBatch(FBulkDataBatchRequest* OutRequest, FCompletionCallback&& Callback)
 {
 	check(Batch.IsValid());
 	checkf(OutRequest != nullptr || Batch->GetRefCount() > 1, TEXT("At least one request handle needs to be used when creating a batch request"));
 
 	TRefCountPtr<FBulkDataBatchRequest::FBatchHandle> NewBatch = MoveTemp(Batch);
-	const bool bOk = NewBatch->Issue(MoveTemp(Callback));
+	NewBatch->Issue(MoveTemp(Callback));
 
 	if (OutRequest)
 	{
-		if (bOk)
-		{
-			*OutRequest = FBulkDataBatchRequest(NewBatch.GetReference());
-		}
-		else
-		{
-			*OutRequest = FBulkDataBatchRequest(new FHandleBase(EStatus::Error));
-		}
+		*OutRequest = FBulkDataBatchRequest(NewBatch.GetReference());
 	}
-
-	return bOk ? EStatus::Ok : EStatus::Error;
 }
 
 FBulkDataBatchRequest::FBatchBuilder::FBatchBuilder(int32 MaxCount)
@@ -1124,27 +1113,32 @@ FBulkDataBatchRequest::FBatchBuilder& FBulkDataBatchRequest::FBatchBuilder::Read
 	return *this;
 }
 
-FBulkDataRequest::EStatus FBulkDataBatchRequest::FBatchBuilder::Issue(FBulkDataBatchRequest& OutRequest)
+void FBulkDataBatchRequest::FBatchBuilder::Issue(FCompletionCallback&& Callback, FBulkDataBatchRequest& OutRequest)
 {
 	if (NumLoaded > 0 && BatchCount == 0)
 	{
 		OutRequest = FBulkDataBatchRequest(new FHandleBase(EStatus::Ok));
-		return EStatus::Ok;
+		return;
 	}
 
-	return IssueBatch(&OutRequest, FCompletionCallback());
+	IssueBatch(&OutRequest, MoveTemp(Callback));
+}
+
+void FBulkDataBatchRequest::FBatchBuilder::Issue(FBulkDataBatchRequest& OutRequest)
+{
+	Issue(FCompletionCallback(), OutRequest);
 }
 
-FBulkDataRequest::EStatus FBulkDataBatchRequest::FBatchBuilder::Issue()
+void FBulkDataBatchRequest::FBatchBuilder::Issue()
 {
 	check(NumLoaded > 0 || BatchCount > 0);
 
 	if (NumLoaded > 0 && BatchCount == 0)
 	{
-		return EStatus::Ok;	
+		return;
 	}
 
-	return IssueBatch(nullptr, FCompletionCallback());
+	IssueBatch(nullptr, FCompletionCallback());
 }
 
 FBulkDataBatchRequest::FScatterGatherBuilder::FScatterGatherBuilder(int32 MaxCount)
@@ -1184,7 +1178,7 @@ FBulkDataBatchRequest::FScatterGatherBuilder& FBulkDataBatchRequest::FScatterGat
 	return *this;
 }
 
-FBulkDataRequest::EStatus FBulkDataBatchRequest::FScatterGatherBuilder::Issue(FIoBuffer& Dst, EAsyncIOPriorityAndFlags Priority, FCompletionCallback&& Callback, FBulkDataBatchRequest& OutRequest)
+void FBulkDataBatchRequest::FScatterGatherBuilder::Issue(FIoBuffer& Dst, EAsyncIOPriorityAndFlags Priority, FCompletionCallback&& Callback, FBulkDataBatchRequest& OutRequest)
 {
 	check(Requests.IsEmpty() == false);
 
@@ -1214,7 +1208,7 @@ FBulkDataRequest::EStatus FBulkDataBatchRequest::FScatterGatherBuilder::Issue(FI
 		DstView.RightChopInline(Request.Size);
 	}
 
-	return IssueBatch(&OutRequest, MoveTemp(Callback));
+	IssueBatch(&OutRequest, MoveTemp(Callback));
 }
 
 //////////////////////////////////////////////////////////////////////////////
diff --git a/Engine/Source/Runtime/CoreUObject/Public/Serialization/BulkData.h b/Engine/Source/Runtime/CoreUObject/Public/Serialization/BulkData.h
--- a/Engine/Source/Runtime/CoreUObject/Public/Serialization/BulkData.h
+++ b/Engine/Source/Runtime/CoreUObject/Public/Serialization/BulkData.h
@@ -1416,7 +1416,7 @@ private:
 	protected:
 		COREUOBJECT_API explicit FBuilder(int32 MaxCount);
 		COREUOBJECT_API FBulkDataBatchRequest::FBatchHandle& GetBatch();
-		COREUOBJECT_API EStatus IssueBatch(FBulkDataBatchRequest* OutRequest, FCompletionCallback&& Callback);
+		COREUOBJECT_API void IssueBatch(FBulkDataBatchRequest* OutRequest, FCompletionCallback&& Callback);
 
 		int32 BatchCount = 0;
 		int32 NumLoaded = 0;
@@ -1473,17 +1473,21 @@ public:
 
 		/**
 		 * Issue the batch.
+		 * @param Callback		Completion callback.	
 		 * @param OutRequest	A handle to the batch request.
-		 * @return				Status of the issue operation.
 		 */
-		COREUOBJECT_API EStatus Issue(FBulkDataBatchRequest& OutRequest);
+		COREUOBJECT_API void Issue(FCompletionCallback&& Callback, FBulkDataBatchRequest& OutRequest);
+		/**
+		 * Issue the batch.
+		 * @param OutRequest	A handle to the batch request.
+		 */
+		COREUOBJECT_API void Issue(FBulkDataBatchRequest& OutRequest);
 		/**
 		 * Issue the batch. 
-		 * @return				Status of the issue operation.
 		 *
 		 * @note Assumes one or more handle(s) has been passed into any of the read operations.
 		 */
-		[[nodiscard]] COREUOBJECT_API EStatus Issue();
+		COREUOBJECT_API void Issue();
 
 	private:
 		COREUOBJECT_API FBatchBuilder& Read(const FBulkData& BulkData, uint64 Offset, uint64 Size, EAsyncIOPriorityAndFlags Priority, FIoBuffer& Dst, FBulkDataBatchReadRequest* OutRequest);
@@ -1502,11 +1506,10 @@ public:
 		 * @param Priority		The I/O priority.
 		 * @param Callback		A callback triggered when the operation is completed.
 		 * @param OutRequest	A handle to the batch request.
-		 * @return				Status of the issue operation.
 		 *
 		 * @note Assumes one or more handle(s) has been passed into any of the read operations.
 		 */
-		COREUOBJECT_API EStatus Issue(FIoBuffer& Dst, EAsyncIOPriorityAndFlags Priority, FCompletionCallback&& Callback, FBulkDataBatchRequest& OutRequest);
+		COREUOBJECT_API void Issue(FIoBuffer& Dst, EAsyncIOPriorityAndFlags Priority, FCompletionCallback&& Callback, FBulkDataBatchRequest& OutRequest);
 
 	private:
 		struct FRequest
