StaticAnalysis: Suppress V568: ...argument of sizeof() operator is the 'Result.GetBytes()' expression.
sizeof is correctly returning 20 for this case.
#rnx

diff --git a/Engine/Source/Runtime/CoreUObject/Private/AssetRegistry/AssetData.cpp b/Engine/Source/Runtime/CoreUObject/Private/AssetRegistry/AssetData.cpp
--- a/Engine/Source/Runtime/CoreUObject/Private/AssetRegistry/AssetData.cpp
+++ b/Engine/Source/Runtime/CoreUObject/Private/AssetRegistry/AssetData.cpp
@@ -961,8 +961,7 @@ FIoHash FAssetPackageData::GetPackageSavedHash() const
 {
 	FIoHash Result;
 	PRAGMA_DISABLE_DEPRECATION_WARNINGS;
-	FMemory::Memcpy(&Result.GetBytes(), &this->PackageGuid,
-		FMath::Min(sizeof(Result.GetBytes()), sizeof(PackageGuid)));
+	FMemory::Memcpy(&Result.GetBytes(), &this->PackageGuid, FMath::Min(sizeof(Result.GetBytes()), sizeof(PackageGuid))); // -V568
 	PRAGMA_ENABLE_DEPRECATION_WARNINGS;
 	return Result;
 }
@@ -971,7 +970,7 @@ void FAssetPackageData::SetPackageSavedHash(const FIoHash& InHash)
 {
 	PRAGMA_DISABLE_DEPRECATION_WARNINGS;
 	PackageGuid = FGuid();
-	FMemory::Memcpy(&PackageGuid, &InHash.GetBytes(), FMath::Min(sizeof(PackageGuid), sizeof(InHash.GetBytes())));
+	FMemory::Memcpy(&PackageGuid, &InHash.GetBytes(), FMath::Min(sizeof(PackageGuid), sizeof(InHash.GetBytes()))); // -V568
 	PRAGMA_ENABLE_DEPRECATION_WARNINGS;
 }
 
