Enabled FAsyncPackageDesc/FAsyncPackageDesc2 InstancingContext in cooked builds to allow using the provided InstancingContext when loading a package.
This is now required to properly support loading ExternalStreamingObjects for instanced partitioned worlds (used by the new External Data Layer feature).
The loaded instanced ExternalStreamingObject needs to have its softobjectpaths (namely the softobjectpath of its associated outer world) remapped using the provided instancing context.
#tests Tested PIE and Cooked game

diff --git a/Engine/Source/Runtime/CoreUObject/Private/Serialization/AsyncLoading2.cpp b/Engine/Source/Runtime/CoreUObject/Private/Serialization/AsyncLoading2.cpp
--- a/Engine/Source/Runtime/CoreUObject/Private/Serialization/AsyncLoading2.cpp
+++ b/Engine/Source/Runtime/CoreUObject/Private/Serialization/AsyncLoading2.cpp
@@ -513,9 +513,7 @@ struct FPackageRequest
 	uint32 LoadFlags = LOAD_None;
 	int32 PIEInstanceID = INDEX_NONE;
 #endif
-#if WITH_EDITORONLY_DATA
 	FLinkerInstancingContext InstancingContext;
-#endif
 	FName CustomName;
 	FPackagePath PackagePath;
 	TUniquePtr<FLoadPackageAsyncDelegate> PackageLoadedDelegate;
@@ -525,11 +523,7 @@ struct FPackageRequest
 
 	FLinkerInstancingContext* GetInstancingContext()
 	{
-#if WITH_EDITORONLY_DATA
 		return &InstancingContext;
-#else
-		return nullptr;
-#endif
 	}
 
 	static FPackageRequest Create(int32 RequestId, EPackageFlags PackageFlags, uint32 LoadFlags, int32 PIEInstanceID, int32 Priority, const FLinkerInstancingContext* InstancingContext, const FPackagePath& PackagePath, FName CustomName, TUniquePtr<FLoadPackageAsyncDelegate> PackageLoadedDelegate, TUniquePtr<FLoadPackageAsyncProgressDelegate> PackageProgressDelegate, FPackageReferencer PackageReferencer)
@@ -543,9 +537,7 @@ struct FPackageRequest
 			LoadFlags,
 			PIEInstanceID,
 #endif
-#if WITH_EDITORONLY_DATA
 			InstancingContext ? *InstancingContext : FLinkerInstancingContext(),
-#endif
 			CustomName,
 			PackagePath,
 			MoveTemp(PackageLoadedDelegate),
@@ -569,6 +561,8 @@ struct FAsyncPackageDesc2
 	/** PIE instance ID this package belongs to, INDEX_NONE otherwise */
 	int32 PIEInstanceID;
 #endif
+	/** Instancing context, maps original package to their instanced counterpart, used to remap imports. */
+	FLinkerInstancingContext InstancingContext;
 	// The package id of the UPackage being loaded
 	// It will be used as key when tracking active async packages
 	FPackageId UPackageId;
@@ -599,6 +593,7 @@ struct FAsyncPackageDesc2
 			Request.LoadFlags,
 			Request.PIEInstanceID,
 #endif
+			MoveTemp(Request.InstancingContext),
 			FPackageId::FromName(UPackageName),
 			PackageIdToLoad,
 			UPackageName,
@@ -628,6 +623,7 @@ struct FAsyncPackageDesc2
 			LOAD_None,
 			INDEX_NONE,
 #endif
+			ImportingPackageDesc.InstancingContext,
 			ImportedPackageId,
 			PackageIdToLoad,
 			UPackageName,
@@ -2989,9 +2985,6 @@ private:
 	/** Callbacks called for the different loading phase of this package */
 	TArray<TUniquePtr<FLoadPackageAsyncProgressDelegate>, TInlineAllocator<2>> ProgressCallbacks;
 
-	/** Set when the package is being loaded as an instance; null otherwise. */
-	TUniquePtr<FLinkerInstancingContext> InstanceContext;
-
 public:
 
 	FAsyncLoadingThread2& GetAsyncLoadingThread()
@@ -4486,7 +4479,7 @@ bool FAsyncLoadingThread2::CreateAsyncPackagesFromQueue(FAsyncLoadingThreadState
 #if ALT2_ENABLE_LINKERLOAD_SUPPORT
 						if (!bIsZenPackage)
 						{
-							Package->InitializeLinkerLoadState(Request.GetInstancingContext());
+							Package->InitializeLinkerLoadState(&PackageDesc.InstancingContext);
 						}
 						else
 #endif
@@ -6113,8 +6106,7 @@ EEventLoadNodeExecutionResult FAsyncPackage2::Event_ProcessPackageSummary(FAsync
 		FName PackageNameToLoad = Package->Desc.PackagePathToLoad.GetPackageFName();
 		if (Package->Desc.UPackageName != PackageNameToLoad)
 		{
-			Package->InstanceContext = MakeUnique<FLinkerInstancingContext>();
-			Package->InstanceContext->BuildPackageMapping(PackageNameToLoad, Package->Desc.UPackageName);
+			Package->Desc.InstancingContext.BuildPackageMapping(PackageNameToLoad, Package->Desc.UPackageName);
 		}
 
 		TRACE_LOADTIME_PACKAGE_SUMMARY(Package, Package->HeaderData.PackageName, Package->HeaderData.PackageSummary->HeaderSize, Package->HeaderData.ImportMap.Num(), Package->HeaderData.ExportMap.Num(), Package->Desc.Priority);
@@ -6193,7 +6185,7 @@ void FAsyncPackage2::InitializeExportArchive(FExportArchive& Ar, bool bIsOptiona
 #endif
 	Ar.ImportStore = &ImportStore;
 	Ar.ExternalReadDependencies = &ExternalReadDependencies;
-	Ar.InstanceContext = InstanceContext.Get();
+	Ar.InstanceContext = &Desc.InstancingContext;
 	Ar.bIsOptionalSegment = bIsOptionalSegment;
 	Ar.bExportsCookedToSeparateArchive = Ar.UEVer() >= EUnrealEngineObjectUE5Version::DATA_RESOURCES;
 }
diff --git a/Engine/Source/Runtime/CoreUObject/Private/Serialization/AsyncPackage.h b/Engine/Source/Runtime/CoreUObject/Private/Serialization/AsyncPackage.h
--- a/Engine/Source/Runtime/CoreUObject/Private/Serialization/AsyncPackage.h
+++ b/Engine/Source/Runtime/CoreUObject/Private/Serialization/AsyncPackage.h
@@ -30,11 +30,8 @@ struct FAsyncPackageDesc
 	TAsyncLoadPriority Priority;
 	/** PIE instance ID this package belongs to, INDEX_NONE otherwise */
 	int32 PIEInstanceID;
-
-#if WITH_EDITORONLY_DATA
 	/** Instancing context, maps original package to their instanced counterpart, used to remap imports. */
 	FLinkerInstancingContext InstancingContext;
-#endif
 #if UE_WITH_PACKAGE_ACCESS_TRACKING
 	FName ReferencerPackageName;
 	FName ReferencerPackageOp;
@@ -43,13 +40,8 @@ struct FAsyncPackageDesc
 	ECookLoadType CookLoadType;
 #endif
 
-#if WITH_EDITORONLY_DATA
 	const FLinkerInstancingContext* GetInstancingContext() const { return &InstancingContext; }
 	void SetInstancingContext(FLinkerInstancingContext InInstancingContext) { InstancingContext = MoveTemp(InInstancingContext); }
-#else
-	const FLinkerInstancingContext* GetInstancingContext() const { return nullptr; }
-	void SetInstancingContext(FLinkerInstancingContext) {}
-#endif 
 
 	FAsyncPackageDesc(int32 InRequestID, const FName& InName, const FPackagePath& InPackagePath, TUniquePtr<FLoadPackageAsyncDelegate>&& InCompletionDelegate = TUniquePtr<FLoadPackageAsyncDelegate>(), EPackageFlags InPackageFlags = PKG_None, int32 InPIEInstanceID = INDEX_NONE, TAsyncLoadPriority InPriority = 0)
 		: RequestID(InRequestID)
@@ -74,9 +66,7 @@ struct FAsyncPackageDesc
 		, PackageFlags(OldPackage.PackageFlags)
 		, Priority(OldPackage.Priority)
 		, PIEInstanceID(OldPackage.PIEInstanceID)
-#if WITH_EDITORONLY_DATA
 		, InstancingContext(OldPackage.InstancingContext)
-#endif
 #if UE_WITH_PACKAGE_ACCESS_TRACKING
 		, ReferencerPackageName(OldPackage.ReferencerPackageName)
 		, ReferencerPackageOp(OldPackage.ReferencerPackageOp)
