[Feature] IDO lifetime management

diff --git a/Engine/Source/Runtime/CoreUObject/Internal/UObject/PropertyBagRepository.h b/Engine/Source/Runtime/CoreUObject/Internal/UObject/PropertyBagRepository.h
--- a/Engine/Source/Runtime/CoreUObject/Internal/UObject/PropertyBagRepository.h
+++ b/Engine/Source/Runtime/CoreUObject/Internal/UObject/PropertyBagRepository.h
@@ -5,6 +5,7 @@
 #include "CoreTypes.h"
 #include "Containers/Map.h"
 #include "HAL/CriticalSection.h"
+#include "UObject/GCObject.h"
 
 class UObjectBase;
 class UObject;
@@ -16,14 +17,14 @@ class FPropertyBag;
 class FPropertyPathName;
 
 // Singleton class tracking property bag association with objects
-class FPropertyBagRepository
+class FPropertyBagRepository : public FGCObject
 {
 	struct FPropertyBagAssociationData
 	{
 		void Destroy();
 		
 		FPropertyBag* Bag = nullptr;		// The existence of an association implies the existence of the bag. TODO: Ref bags via handle? Store as value?
-		UObject* InstanceDataObject = nullptr;
+		TObjectPtr<UObject> InstanceDataObject = nullptr;
 	};
 	// TODO: Make private throughout and extend access permissions here or in wrapper classes? Don't want engine code modifying bags outside of serializers and details panels.
 	//friend UObjectBase;
@@ -46,6 +47,9 @@ private:
 	
 	//TMap<const UObjectBase*, UObject*> ObjectToInstanceDataObjectMap;
 
+	// used to make sure IDOs don't have name overlap
+	TMap<const UObjectBase*, TObjectPtr<UObject>> Namespaces;
+
 	FPropertyBagRepository() = default;
 
 public:
@@ -101,6 +105,9 @@ public:
 	// query whether a property in Struct was set when the struct was deserialized
 	COREUOBJECT_API static bool WasPropertySetBySerialization(const UStruct* Struct, const void* StructData, const FProperty* Property, int32 ArrayIndex = 0);
 
+	virtual void AddReferencedObjects(FReferenceCollector& Collector) override;
+	virtual FString GetReferencerName() const override;
+
 private:
 	void Lock() const { CriticalSection.Lock(); }
 	void Unlock() const { CriticalSection.Unlock(); }
diff --git a/Engine/Source/Runtime/CoreUObject/Private/UObject/PropertyBagRepository.cpp b/Engine/Source/Runtime/CoreUObject/Private/UObject/PropertyBagRepository.cpp
--- a/Engine/Source/Runtime/CoreUObject/Private/UObject/PropertyBagRepository.cpp
+++ b/Engine/Source/Runtime/CoreUObject/Private/UObject/PropertyBagRepository.cpp
@@ -1,6 +1,9 @@
 // Copyright Epic Games, Inc. All Rights Reserved.
 
 #include "UObject/PropertyBagRepository.h"
+
+#include "Serialization/ObjectReader.h"
+#include "Serialization/ObjectWriter.h"
 #include "UObject/GarbageCollection.h"
 #include "UObject/Object.h"
 #include "UObject/UObjectGlobals.h"
@@ -57,7 +60,6 @@ void FPropertyBagRepository::FPropertyBagAssociationData::Destroy()
 	
 	if(InstanceDataObject && InstanceDataObject->IsValidLowLevel())
 	{
-		InstanceDataObject->RemoveFromRoot();
 		InstanceDataObject = nullptr;
 	}
 }
@@ -133,6 +135,10 @@ bool FPropertyBagRepository::RemoveAssociationUnsafe(const UObjectBase* Owner)
 		OldData.Destroy();
 		return true;
 	}
+
+	// note: RemoveAssociationUnsafe is called on every object regardless of whether it has a property bag.
+	// in that scenario, there's a chance we have a namespace associated with it. Remove that namespace.
+	Namespaces.Remove(Owner);
 	return false;
 }
 
@@ -185,6 +191,23 @@ bool FPropertyBagRepository::WasPropertySetBySerialization(const UStruct* Struct
 	return UE::WasPropertySetBySerialization(Struct, StructData, Property, ArrayIndex);
 }
 
+void FPropertyBagRepository::AddReferencedObjects(FReferenceCollector& Collector)
+{
+	for (TPair<const UObjectBase*, FPropertyBagAssociationData>& Element : AssociatedData)
+	{
+		Collector.AddReferencedObject(Element.Value.InstanceDataObject);
+	}
+	for (TPair<const UObjectBase*, TObjectPtr<UObject>>& Element : Namespaces)
+	{
+		Collector.AddReferencedObject(Element.Value);
+	}
+}
+
+FString FPropertyBagRepository::GetReferencerName() const
+{
+	return TEXT("FPropertyBagRepository");
+}
+
 void FPropertyBagRepository::CreateInstanceDataObjectUnsafe(const UObjectBase* Owner, FPropertyBagAssociationData& BagData)
 {
 	check(!BagData.InstanceDataObject);	// No repeated calls
@@ -193,9 +216,25 @@ void FPropertyBagRepository::CreateInstanceDataObjectUnsafe(const UObjectBase* O
 	// TODO: should we put the InstanceDataObject or it's class in a package?
 	const UClass* InstanceDataObjectClass = CreateInstanceDataObjectClass(PropertyBag, Owner->GetClass(), GetTransientPackage());
 
+	TObjectPtr<UObject>* OuterPtr;
+	if (FPropertyBagAssociationData* OuterData = AssociatedData.Find(Owner->GetOuter()))
+	{
+		OuterPtr = &OuterData->InstanceDataObject;
+	}
+	else
+	{
+		OuterPtr = &Namespaces.FindOrAdd(Owner->GetOuter());
+		if (*OuterPtr == nullptr)
+		{
+			*OuterPtr = CreatePackage(nullptr); // TODO: replace with dummy object
+		}
+	}
+
 	// construct InstanceDataObject object
 	FStaticConstructObjectParameters Params(InstanceDataObjectClass);
 	Params.SetFlags |= EObjectFlags::RF_Transactional;
+	Params.Name = Owner->GetFName();
+	Params.Outer = *OuterPtr;
 	UObject* InstanceDataObjectObject = StaticConstructObject_Internal(Params);
 	BagData.InstanceDataObject = InstanceDataObjectObject;
 	
@@ -204,21 +243,28 @@ void FPropertyBagRepository::CreateInstanceDataObjectUnsafe(const UObjectBase* O
 	TGuardValue<bool> ScopedImpersonateProperties(LoadContext->bImpersonateProperties, true);
 	TGuardValue<bool> ScopedTrackSerializedPropertyPath(LoadContext->bTrackSerializedPropertyPath, true);
 	
-	const FDelegateHandle OnTaggedPropertySerializeHandle = LoadContext->OnTaggedPropertySerialize.AddLambda(
-		[&BagData](const FUObjectSerializeContext& Context)
-		{
-			if (!Context.SerializedPropertyPath.IsEmpty())
+	UObject* OwnerAsObject = (UObject*)Owner;
+	if (FLinkerLoad* Linker = OwnerAsObject->GetLinker())
+	{
+		const FDelegateHandle OnTaggedPropertySerializeHandle = LoadContext->OnTaggedPropertySerialize.AddLambda(
+			[&BagData](const FUObjectSerializeContext& Context)
 			{
-				MarkPropertySetBySerialization(BagData.InstanceDataObject, Context.SerializedPropertyPath);
+				if (!Context.SerializedPropertyPath.IsEmpty())
+				{
+					MarkPropertySetBySerialization(BagData.InstanceDataObject, Context.SerializedPropertyPath);
+				}
 			}
-		}
-	);
-
-	UObject* OwnerAsObject = (UObject*)Owner;
-	OwnerAsObject->SetFlags(RF_NeedLoad);
-	OwnerAsObject->GetLinker()->Preload(OwnerAsObject);
-
-	LoadContext->OnTaggedPropertySerialize.Remove(OnTaggedPropertySerializeHandle);
+		);
+		OwnerAsObject->SetFlags(RF_NeedLoad);
+		Linker->Preload(OwnerAsObject);
+		LoadContext->OnTaggedPropertySerialize.Remove(OnTaggedPropertySerializeHandle);
+	}
+	else if (ensureMsgf(BagData.Bag == nullptr, TEXT("Linker missing when generating IDO for an object with loose properties")))
+	{
+		TArray<uint8> Buffer;
+		FObjectWriter(OwnerAsObject, Buffer);
+		FObjectReader(BagData.InstanceDataObject, Buffer);
+	}
 }
 
 // Not sure this is necessary.
