StaticAnalysis: Avoid error V568: ...argument of sizeof() operator is the 'Result.GetBytes()' expression by using decltype.
#rnx

diff --git a/Engine/Source/Editor/UnrealEd/Private/Cooker/CookPackageData.cpp b/Engine/Source/Editor/UnrealEd/Private/Cooker/CookPackageData.cpp
--- a/Engine/Source/Editor/UnrealEd/Private/Cooker/CookPackageData.cpp
+++ b/Engine/Source/Editor/UnrealEd/Private/Cooker/CookPackageData.cpp
@@ -2282,7 +2282,8 @@ void FCookGenerationInfo::CreatePackageHash()
 		TOptional<FAssetPackageData> DependencyData = AssetRegistry.GetAssetPackageDataCopy(Dependency.AssetId.PackageName);
 		if (DependencyData)
 		{
-			Blake3.Update(&DependencyData->GetPackageSavedHash().GetBytes(), sizeof(DependencyData->GetPackageSavedHash().GetBytes()));
+			Blake3.Update(&DependencyData->GetPackageSavedHash().GetBytes(),
+				sizeof(decltype(DependencyData->GetPackageSavedHash().GetBytes())));
 		}
 	}
 	PackageHash = FIoHash(Blake3.Finalize());
@@ -2290,7 +2291,7 @@ void FCookGenerationInfo::CreatePackageHash()
 	// so we have to remove any data which doesn't fit into FGuid. This can be removed when we remove the deprecated
 	// Guid storage on UPackage.
 	PRAGMA_DISABLE_DEPRECATION_WARNINGS;
-	constexpr int SizeDifference = sizeof(PackageHash) - sizeof(DeclVal<UPackage>().GetGuid());
+	constexpr int SizeDifference = sizeof(PackageHash) - sizeof(decltype(DeclVal<UPackage>().GetGuid()));
 	if (SizeDifference > 0)
 	{
 		FMemory::Memset(((uint8*)&PackageHash.GetBytes()) + (sizeof(PackageHash.GetBytes()) - SizeDifference),
diff --git a/Engine/Source/Runtime/CoreUObject/Private/AssetRegistry/AssetData.cpp b/Engine/Source/Runtime/CoreUObject/Private/AssetRegistry/AssetData.cpp
--- a/Engine/Source/Runtime/CoreUObject/Private/AssetRegistry/AssetData.cpp
+++ b/Engine/Source/Runtime/CoreUObject/Private/AssetRegistry/AssetData.cpp
@@ -961,7 +961,8 @@ FIoHash FAssetPackageData::GetPackageSavedHash() const
 {
 	FIoHash Result;
 	PRAGMA_DISABLE_DEPRECATION_WARNINGS;
-	FMemory::Memcpy(&Result.GetBytes(), &this->PackageGuid, FMath::Min(sizeof(Result.GetBytes()), sizeof(PackageGuid))); // -V568
+	FMemory::Memcpy(&Result.GetBytes(), &this->PackageGuid,
+		FMath::Min(sizeof(decltype(Result.GetBytes())), sizeof(PackageGuid)));
 	PRAGMA_ENABLE_DEPRECATION_WARNINGS;
 	return Result;
 }
@@ -970,7 +971,8 @@ void FAssetPackageData::SetPackageSavedHash(const FIoHash& InHash)
 {
 	PRAGMA_DISABLE_DEPRECATION_WARNINGS;
 	PackageGuid = FGuid();
-	FMemory::Memcpy(&PackageGuid, &InHash.GetBytes(), FMath::Min(sizeof(PackageGuid), sizeof(InHash.GetBytes()))); // -V568
+	FMemory::Memcpy(&PackageGuid, &InHash.GetBytes(),
+		FMath::Min(sizeof(PackageGuid), sizeof(decltype(InHash.GetBytes()))));
 	PRAGMA_ENABLE_DEPRECATION_WARNINGS;
 }
 
diff --git a/Engine/Source/Runtime/CoreUObject/Private/UObject/Package.cpp b/Engine/Source/Runtime/CoreUObject/Private/UObject/Package.cpp
--- a/Engine/Source/Runtime/CoreUObject/Private/UObject/Package.cpp
+++ b/Engine/Source/Runtime/CoreUObject/Private/UObject/Package.cpp
@@ -401,7 +401,8 @@ FIoHash UPackage::GetSavedHash()
 #else
 	FIoHash Result;
 	PRAGMA_DISABLE_DEPRECATION_WARNINGS;
-	FMemory::Memcpy(&Result.GetBytes(), &Guid, FMath::Min(sizeof(Result.GetBytes()), sizeof(Guid)));
+	FMemory::Memcpy(&Result.GetBytes(), &Guid,
+		FMath::Min(sizeof(decltype(Result.GetBytes())), sizeof(Guid)));
 	PRAGMA_ENABLE_DEPRECATION_WARNINGS;
 	return Result;
 #endif
@@ -412,7 +413,8 @@ void UPackage::SetSavedHash(const FIoHash& InSavedHash)
 #if !UE_STRIP_DEPRECATED_PROPERTIES
 	PRAGMA_DISABLE_DEPRECATION_WARNINGS;
 	Guid = FGuid();
-	FMemory::Memcpy(&Guid, &InSavedHash.GetBytes(), FMath::Min(sizeof(Guid), sizeof(InSavedHash.GetBytes())));
+	FMemory::Memcpy(&Guid, &InSavedHash.GetBytes(),
+		FMath::Min(sizeof(Guid), sizeof(decltype(InSavedHash.GetBytes()))));
 	PRAGMA_ENABLE_DEPRECATION_WARNINGS;
 #endif
 }
diff --git a/Engine/Source/Runtime/CoreUObject/Private/UObject/PackageFileSummary.cpp b/Engine/Source/Runtime/CoreUObject/Private/UObject/PackageFileSummary.cpp
--- a/Engine/Source/Runtime/CoreUObject/Private/UObject/PackageFileSummary.cpp
+++ b/Engine/Source/Runtime/CoreUObject/Private/UObject/PackageFileSummary.cpp
@@ -275,7 +275,7 @@ void operator<<(FStructuredArchive::FSlot Slot, FPackageFileSummary& Sum)
 				FIoHash SavedHash = Sum.GetSavedHash();
 				Sum.PersistentGuid = FGuid();
 				FMemory::Memcpy(&Sum.PersistentGuid, &SavedHash.GetBytes(),
-					FMath::Min(sizeof(Sum.PersistentGuid), sizeof(SavedHash.GetBytes())));
+					FMath::Min(sizeof(Sum.PersistentGuid), sizeof(decltype(SavedHash.GetBytes()))));
 			}
 
 			// The owner persistent guid was added in VER_UE4_ADDED_PACKAGE_OWNER but removed in the next version VER_UE4_NON_OUTER_PACKAGE_IMPORT
@@ -455,7 +455,8 @@ FIoHash FPackageFileSummary::GetSavedHash() const
 {
 	FIoHash Result;
 	PRAGMA_DISABLE_DEPRECATION_WARNINGS;
-	FMemory::Memcpy(&Result.GetBytes(), &Guid, FMath::Min(sizeof(Result.GetBytes()), sizeof(Guid)));
+	FMemory::Memcpy(&Result.GetBytes(), &Guid,
+		FMath::Min(sizeof(decltype(Result.GetBytes())), sizeof(Guid)));
 	PRAGMA_ENABLE_DEPRECATION_WARNINGS;
 	return Result;
 }
@@ -464,7 +465,8 @@ void FPackageFileSummary::SetSavedHash(const FIoHash& InSavedHash)
 {
 	PRAGMA_DISABLE_DEPRECATION_WARNINGS;
 	Guid = FGuid();
-	FMemory::Memcpy(&Guid, &InSavedHash.GetBytes(), FMath::Min(sizeof(Guid), sizeof(InSavedHash.GetBytes())));
+	FMemory::Memcpy(&Guid, &InSavedHash.GetBytes(),
+		FMath::Min(sizeof(Guid), sizeof(decltype(InSavedHash.GetBytes()))));
 	PRAGMA_ENABLE_DEPRECATION_WARNINGS;
 }
 #endif
diff --git a/Engine/Source/Runtime/CoreUObject/Private/UObject/SavePackage2.cpp b/Engine/Source/Runtime/CoreUObject/Private/UObject/SavePackage2.cpp
--- a/Engine/Source/Runtime/CoreUObject/Private/UObject/SavePackage2.cpp
+++ b/Engine/Source/Runtime/CoreUObject/Private/UObject/SavePackage2.cpp
@@ -1117,7 +1117,7 @@ ESavePackageResult BuildLinker(FSaveContext& SaveContext)
 	{
 		FIoHash SavedHash;
 		FMemory::Memcpy(&SavedHash.GetBytes(), OutputPackageGuid,
-			FMath::Min(sizeof(SavedHash.GetBytes()), sizeof(*OutputPackageGuid)));
+			FMath::Min(sizeof(decltype(SavedHash.GetBytes())), sizeof(*OutputPackageGuid)));
 		SaveContext.GetLinker()->Summary.SetSavedHash(SavedHash);
 	}
 	else
@@ -1128,7 +1128,8 @@ ESavePackageResult BuildLinker(FSaveContext& SaveContext)
 			// Remove IsKeepGuid and OutputPackageGuid when that occurs.
 			FGuid Guid = FGuid::NewGuid();
 			FIoHash SavedHash;
-			FMemory::Memcpy(&SavedHash.GetBytes(), &Guid, FMath::Min(sizeof(SavedHash.GetBytes()), sizeof(Guid)));
+			FMemory::Memcpy(&SavedHash.GetBytes(), &Guid,
+				FMath::Min(sizeof(decltype(SavedHash.GetBytes())), sizeof(Guid)));
 			SaveContext.GetPackage()->SetSavedHash(SavedHash);
 		}
 		SaveContext.GetLinker()->Summary.SetSavedHash(SaveContext.GetPackage()->GetSavedHash());
