[ZenLoader]
  - Apply patching to export table before preloading and postload when the linker export table is modified during reinstantiation.
  - Dispatch constructed objects to their respective package.
  - Properly handle and postload objects constructed during postload (i.e. reinstantiation).

diff --git a/Engine/Source/Runtime/CoreUObject/Private/Serialization/AsyncLoading2.cpp b/Engine/Source/Runtime/CoreUObject/Private/Serialization/AsyncLoading2.cpp
--- a/Engine/Source/Runtime/CoreUObject/Private/Serialization/AsyncLoading2.cpp
+++ b/Engine/Source/Runtime/CoreUObject/Private/Serialization/AsyncLoading2.cpp
@@ -47,6 +47,7 @@
 #include "HAL/LowLevelMemTracker.h"
 #include "ProfilingDebugging/CsvProfiler.h"
 #include "UObject/UObjectArchetypeInternal.h"
+#include "Misc/MTAccessDetector.h"
 #include "UObject/GarbageCollectionInternal.h"
 #include "ProfilingDebugging/MiscTrace.h"
 #include "Serialization/LoadTimeTracePrivate.h"
@@ -2771,16 +2772,7 @@ struct FAsyncPackage2
 
 	void AddConstructedObject(UObject* Object, bool bSubObjectThatAlreadyExists)
 	{
-		// Skip the tracking when we reach the game-thread as we don't have to add the asyncloading flag anymore and since
-		// we're already in the postloading step, there is no value in accepting new objects and trying to postload them.
-		// Also, given that anything can create UObject outside the scope of loading once we reach postload and callbacks
-		// we could end up with objects that have nothing to do with the actual package currently being on the stack,
-		// and these package might not even have finished loading, so postloading any of these objects might be premature
-		// anyway.
-		if (AsyncPackageLoadingState >= EAsyncPackageLoadingState2::DeferredPostLoad)
-		{
-			return;
-		}
+		UE_MT_SCOPED_WRITE_ACCESS(ConstructedObjectsAccessDetector);
 
 		if (bSubObjectThatAlreadyExists)
 		{
@@ -2979,6 +2971,8 @@ private:
 	TArray<int32> RequestIDs;
 	/** List of ConstructedObjects = Exports + UPackage + ObjectsCreatedFromExports */
 	TArray<UObject*> ConstructedObjects;
+	/** Detects if the constructed objects are improperly accessed by different threads at the same time. */
+	UE_MT_DECLARE_MRSW_RECURSIVE_ACCESS_DETECTOR(ConstructedObjectsAccessDetector);
 	TArray<FExternalReadCallback> ExternalReadDependencies;
 	/** Callbacks called when we finished loading this package */
 	TArray<TUniquePtr<FLoadPackageAsyncDelegate>, TInlineAllocator<2>> CompletionCallbacks;
@@ -5615,10 +5609,23 @@ bool FAsyncPackage2::PreloadLinkerLoadExports(FAsyncLoadingThreadState2& ThreadS
 {
 	// Serialize exports
 	const int32 ExportCount = LinkerLoadState->Linker->ExportMap.Num();
+	check(LinkerLoadState->Linker->ExportMap.Num() == Data.Exports.Num());
 	while (LinkerLoadState->SerializeExportIndex < ExportCount)
 	{
 		const int32 ExportIndex = LinkerLoadState->SerializeExportIndex++;
 		FExportObject& ExportObject = Data.Exports[ExportIndex];
+		FObjectExport& LinkerExport = LinkerLoadState->Linker->ExportMap[ExportIndex];
+
+		// The linker export table can be patched during reinstantiation. We need to adjust our own export table if needed.
+		if (!LinkerExport.bExportLoadFailed)
+		{
+			if (ExportObject.Object != LinkerExport.Object)
+			{
+				UE_ASYNC_PACKAGE_LOG(Verbose, Desc, TEXT("PreloadLinkerLoadExports"), TEXT("Patching export %d: %s -> %s"), ExportIndex, *GetPathNameSafe(ExportObject.Object), *GetPathNameSafe(LinkerExport.Object));
+				ExportObject.Object = LinkerExport.Object;
+			}
+		}
+
 		if (UObject* Object = ExportObject.Object)
 		{
 			if (Object->HasAnyFlags(RF_NeedLoad))
@@ -5867,15 +5874,32 @@ EEventLoadNodeExecutionResult FAsyncPackage2::ExecuteDeferredPostLoadLinkerLoadP
 	// We can't return timeout during a flush as we're expected to be able to finish
 	const bool bIsReadyForAsyncPostLoadAllowed = ThreadState.SyncLoadContextStack.IsEmpty();
 
+	UE_MT_SCOPED_READ_ACCESS(ConstructedObjectsAccessDetector);
+
 	// Go through both ConstructedObjects and export table as its possible to reload objects in the export table
 	// without them being constructed and that would lead to missing postloads.
-	const int32 ConstructedObjectsCount = ConstructedObjects.Num();
-	const int32 ObjectCount = ConstructedObjectsCount + Data.Exports.Num();
-	while (LinkerLoadState->PostLoadExportIndex < ObjectCount)
+	// ConstructedObjects can be appended to during conditional postloads, so make sure to always take the latest value.
+	const int32 ExportsCount = Data.Exports.Num();
+	while (LinkerLoadState->PostLoadExportIndex < ExportsCount + ConstructedObjects.Num())
 	{
 		const int32 ObjectIndex = LinkerLoadState->PostLoadExportIndex++;
 
-		UObject* Object = ObjectIndex < ConstructedObjectsCount ? ConstructedObjects[ObjectIndex] : Data.Exports[ObjectIndex - ConstructedObjectsCount].Object;
+		if (ObjectIndex < ExportsCount)
+		{
+			FExportObject& ExportObject = Data.Exports[ObjectIndex];
+			FObjectExport& LinkerExport = LinkerLoadState->Linker->ExportMap[ObjectIndex];
+			// The linker export table can be patched during reinstantiation. We need to adjust our own export table if needed.
+			if (!LinkerExport.bExportLoadFailed)
+			{
+				if (ExportObject.Object != LinkerExport.Object)
+				{
+					UE_ASYNC_PACKAGE_LOG(Verbose, Desc, TEXT("ExecuteDeferredPostLoadLinkerLoadPackageExports"), TEXT("Patching export %d: %s -> %s"), ObjectIndex, *GetPathNameSafe(ExportObject.Object), *GetPathNameSafe(LinkerExport.Object));
+					ExportObject.Object = LinkerExport.Object;
+				}
+			}
+		}
+
+		UObject* Object = ObjectIndex < ExportsCount ? Data.Exports[ObjectIndex].Object : ConstructedObjects[ObjectIndex - ExportsCount];
 		if (Object && Object->HasAnyFlags(RF_NeedPostLoad))
 		{
 			// Only allow to wait when there is no flush waiting on us
@@ -8625,6 +8649,34 @@ void FAsyncLoadingThread2::NotifyConstructedDuringAsyncLoading(UObject* Object,
 	}
 
 	FAsyncPackage2* AsyncPackage2 = (FAsyncPackage2*)ThreadContext.AsyncPackage;
+
+#if WITH_EDITOR
+	// In editor, objects from other packages might be constructed with the wrong package currently on stack (i.e. reinstantiation).
+	// It's a lot cleaner and easier to properly dispatch them to their respective package than adding scopes everywhere.
+	if (UPackage* ObjectPackage = Object->GetPackage())
+	{
+		if (FLinkerLoad* LinkerLoad = ObjectPackage->GetLinker())
+		{
+			if (FAsyncPackage2* AsyncPackage = (FAsyncPackage2*)LinkerLoad->AsyncRoot)
+			{
+				AsyncPackage->AddConstructedObject(Object, bSubObjectThatAlreadyExists);
+				return;
+			}
+		}
+
+		FPackageId PackageId = ObjectPackage->GetPackageId();
+		if (PackageId.IsValid() && PackageId != AsyncPackage2->Desc.UPackageId)
+		{
+			FScopeLock LockAsyncPackages(&AsyncPackagesCritical);
+			if (FAsyncPackage2* AsyncPackage = AsyncPackageLookup.FindRef(PackageId))
+			{
+				AsyncPackage->AddConstructedObject(Object, bSubObjectThatAlreadyExists);
+				return;
+			}
+		}
+	}
+#endif
+	
 	AsyncPackage2->AddConstructedObject(Object, bSubObjectThatAlreadyExists);
 }
 
