More extended support at getting override state and clearing overrides
Added new associated unittests cases

diff --git a/Engine/Source/Runtime/CoreUObject/Private/UObject/OverridableManager.cpp b/Engine/Source/Runtime/CoreUObject/Private/UObject/OverridableManager.cpp
--- a/Engine/Source/Runtime/CoreUObject/Private/UObject/OverridableManager.cpp
+++ b/Engine/Source/Runtime/CoreUObject/Private/UObject/OverridableManager.cpp
@@ -134,11 +134,11 @@ void FOverridableManager::OverrideProperty(UObject& Object, const FPropertyChang
 	}
 }
 
-bool FOverridableManager::ClearOverriddenProperty(UObject& Object, const FPropertyChangedEvent& PropertyEvent, const FEditPropertyChain& PropertyChain)
+bool FOverridableManager::ClearOverriddenProperty(UObject& Object, const FPropertyChangedEvent& PropertyEvent, const FEditPropertyChain::TDoubleLinkedListNode* PropertyNode)
 {
 	if (FOverriddenPropertySet* ThisObjectOverriddenProperties = OverriddenObjectAnnotations.Find(Object))
 	{
-		return ThisObjectOverriddenProperties->ClearOverriddenProperty(PropertyEvent, PropertyChain.GetHead());
+		return ThisObjectOverriddenProperties->ClearOverriddenProperty(PropertyEvent, PropertyNode);
 	}
 	return false;
 }
@@ -167,11 +167,11 @@ void FOverridableManager::NotifyPropertyChange(const EPropertyNotificationType N
 	}
 }
 
-EOverriddenPropertyOperation FOverridableManager::GetOverriddenPropertyOperation(UObject& Object, const FPropertyChangedEvent& PropertyEvent, const FEditPropertyChain& PropertyChain, bool* bOutInheritedState)
+EOverriddenPropertyOperation FOverridableManager::GetOverriddenPropertyOperation(UObject& Object, const FPropertyChangedEvent& PropertyEvent, const FEditPropertyChain::TDoubleLinkedListNode* PropertyNode, bool* bOutInheritedOperation)
 {
 	if (const FOverriddenPropertySet* ThisObjectOverriddenProperties = OverriddenObjectAnnotations.Find(Object))
 	{
-		return ThisObjectOverriddenProperties->GetOverriddenPropertyOperation(PropertyEvent, PropertyChain.GetHead(), bOutInheritedState);
+		return ThisObjectOverriddenProperties->GetOverriddenPropertyOperation(PropertyEvent, PropertyNode, bOutInheritedOperation);
 	}
 	return EOverriddenPropertyOperation::None;
 }
diff --git a/Engine/Source/Runtime/CoreUObject/Private/UObject/OverriddenPropertySet.cpp b/Engine/Source/Runtime/CoreUObject/Private/UObject/OverriddenPropertySet.cpp
--- a/Engine/Source/Runtime/CoreUObject/Private/UObject/OverriddenPropertySet.cpp
+++ b/Engine/Source/Runtime/CoreUObject/Private/UObject/OverriddenPropertySet.cpp
@@ -59,20 +59,60 @@ EOverriddenPropertyOperation FOverridableSerializationLogic::GetOverriddenProper
 			{
 				if (UObject* SubObject = ObjectProperty->GetObjectPropertyValue(DataPtr))
 				{
-					if (FOverridableManager::Get().GetOverriddenState(*SubObject) != EOverriddenState::NoOverrides)
+					if(const FOverriddenPropertySet* OverriddenSubObjectProperties = FOverridableManager::Get().GetOverriddenProperties(*SubObject))
 					{
-						return EOverriddenPropertyOperation::Modified;
+						const EOverriddenPropertyOperation SubObjectOperation = OverriddenSubObjectProperties->GetOverriddenPropertyOperation((FArchiveSerializedPropertyChain*)nullptr, (FProperty*)nullptr);
+						if (SubObjectOperation != EOverriddenPropertyOperation::None)
+						{
+							return EOverriddenPropertyOperation::Modified;
+						}
 					}
 				}
 			}
 		}
 		else if (const FArrayProperty* CurrentArrayProperty = CastField<FArrayProperty>(CurrentProperty))
 		{
-			return CurrentArrayProperty->Inner->HasAnyPropertyFlags(CPF_PersistentInstance) ? EOverriddenPropertyOperation::Modified : EOverriddenPropertyOperation::None;
+			if (const FObjectProperty* ObjectProperty = CurrentArrayProperty->Inner->HasAnyPropertyFlags(CPF_PersistentInstance) ? CastField<FObjectProperty>(CurrentArrayProperty->Inner) : nullptr)
+			{
+				FScriptArrayHelper ArrayHelper(CurrentArrayProperty, DataPtr);
+
+				for (int i = 0; i < ArrayHelper.Num(); ++i)
+				{
+					if (UObject* SubObject = ObjectProperty->GetObjectPropertyValue(ArrayHelper.GetElementPtr(i)))
+					{
+						if(const FOverriddenPropertySet* OverriddenSubObjectProperties = FOverridableManager::Get().GetOverriddenProperties(*SubObject))
+						{
+							const EOverriddenPropertyOperation SubObjectOperation = OverriddenSubObjectProperties->GetOverriddenPropertyOperation((FArchiveSerializedPropertyChain*)nullptr, (FProperty*)nullptr);
+							if (SubObjectOperation != EOverriddenPropertyOperation::None)
+							{
+								return EOverriddenPropertyOperation::Modified;
+							}
+						}
+					}
+				}
+			}
 		}
 		else if(const FMapProperty* CurrentMapProperty = CastField<FMapProperty>(CurrentProperty))
 		{
-			return CurrentMapProperty->ValueProp->HasAnyPropertyFlags(CPF_PersistentInstance) ? EOverriddenPropertyOperation::Modified : EOverriddenPropertyOperation::None;
+			if (const FObjectProperty* ValueInstancedObjectProperty = CurrentMapProperty->ValueProp->HasAnyPropertyFlags(CPF_PersistentInstance) ? CastField<FObjectProperty>(CurrentMapProperty->ValueProp) : nullptr)
+			{
+				FScriptMapHelper MapHelper(CurrentMapProperty, DataPtr);
+
+				for (FScriptMapHelper::FIterator It(MapHelper); It; ++It)
+				{
+					if (UObject* SubObject = ValueInstancedObjectProperty->GetObjectPropertyValue(MapHelper.GetValuePtr(It.GetInternalIndex())))
+					{
+						if(const FOverriddenPropertySet* OverriddenSubObjectProperties = FOverridableManager::Get().GetOverriddenProperties(*SubObject))
+						{
+							const EOverriddenPropertyOperation SubObjectOperation = OverriddenSubObjectProperties->GetOverriddenPropertyOperation((FArchiveSerializedPropertyChain*)nullptr, (FProperty*)nullptr);
+							if (SubObjectOperation != EOverriddenPropertyOperation::None)
+							{
+								return EOverriddenPropertyOperation::Modified;
+							}
+						}
+					}
+				}
+			}
 		}
 	}
 	
@@ -156,49 +196,332 @@ FOverriddenPropertyNode& FOverriddenPropertySet::FindOrAddNode(FOverriddenProper
 	return OverriddenPropertyNodes.Get(NewID);
 }
 
-FOverriddenPropertyNode* FOverriddenPropertySet::GetOverriddenPropertyNode(FOverriddenPropertyNode& ParentPropertyNode, const FEditPropertyChain::TDoubleLinkedListNode* PropertyNode, bool* bOutInheritedState)
+EOverriddenPropertyOperation FOverriddenPropertySet::GetOverriddenPropertyOperation(const FOverriddenPropertyNode& ParentPropertyNode, const FPropertyChangedEvent& PropertyEvent, const FEditPropertyChain::TDoubleLinkedListNode* PropertyNode, bool* bOutInheritedOperation, const void* Data) const
 {
-	// No need to look further
-	// if it is the entire property is replaced or
-	// if it is the FOverriddenPropertySet struct which is always Overridden
-	if (ParentPropertyNode.Operation == EOverriddenPropertyOperation::Replace)
+	FOverridableManager& OverridableManager = FOverridableManager::Get();
+
+	const void* SubValuePtr = Data;
+	const FEditPropertyChain::TDoubleLinkedListNode* PropertyIterator = PropertyNode;
+	const FOverriddenPropertyNode* OverriddenPropertyNode = &ParentPropertyNode;
+	int32 ArrayIndex = INDEX_NONE;
+	while (PropertyIterator && (!OverriddenPropertyNode || OverriddenPropertyNode->Operation != EOverriddenPropertyOperation::Replace))
 	{
-		if (bOutInheritedState)
+		ArrayIndex = INDEX_NONE;
+
+		const FProperty* CurrentProperty = PropertyIterator->GetValue();
+		SubValuePtr = CurrentProperty->ContainerPtrToValuePtr<void>(SubValuePtr, 0); //@todo support static arrays
+
+		const FOverriddenPropertyNode* CurrentOverriddenPropertyNode = nullptr;
+		if (OverriddenPropertyNode)
 		{
-			*bOutInheritedState = PropertyNode != nullptr;
+			const FName CurrentPropID = CurrentProperty->GetFName();
+			if (const FOverriddenPropertyNodeID* CurrentPropKey = OverriddenPropertyNode->SubPropertyNodeKeys.Find(CurrentPropID))
+			{
+				CurrentOverriddenPropertyNode = OverriddenPropertyNodes.FindByHash(GetTypeHash(*CurrentPropKey), *CurrentPropKey);
+				checkf(CurrentOverriddenPropertyNode, TEXT("Expecting a node"));
+			}
 		}
-		return &ParentPropertyNode;
+
+
+		// Special handling for for instanced subobjects 
+		if (const FObjectProperty* ObjectProperty = CastField<FObjectProperty>(CurrentProperty))
+		{
+			if (PropertyIterator->GetNextNode())
+			{
+				// Forward any sub queries to the subobject
+				if (UObject* SubObject = ObjectProperty->GetObjectPropertyValue(SubValuePtr))
+				{
+					// This should not be needed in the property grid, as it should already been called on the subobject.
+					return OverridableManager.GetOverriddenPropertyOperation(*SubObject, PropertyEvent, PropertyIterator->GetNextNode(), bOutInheritedOperation);
+				}
+			}
+		}
+		// Special handling for array of instanced subobjects 
+		else if (const FArrayProperty* ArrayProperty = CastField<FArrayProperty>(CurrentProperty))
+		{
+			ArrayIndex = PropertyEvent.GetArrayIndex(CurrentProperty->GetName());
+
+			// Only special case is instanced subobjects, otherwise we fallback to full array override
+			checkf(ArrayProperty->Inner, TEXT("Expecting an inner type for Arrays"));
+			if (const FObjectProperty* InnerObjectProperty = ArrayProperty->Inner->HasAnyPropertyFlags(CPF_PersistentInstance) ? CastField<FObjectProperty>(ArrayProperty->Inner) : nullptr)
+			{
+				FScriptArrayHelper ArrayHelper(ArrayProperty, SubValuePtr);
+				if(ArrayHelper.IsValidIndex(ArrayIndex))
+				{
+					if (UObject* SubObject = InnerObjectProperty->GetObjectPropertyValue(ArrayHelper.GetElementPtr(ArrayIndex)))
+					{
+						if (PropertyIterator->GetNextNode())
+						{
+							// Forward any sub queries to the subobject
+							return OverridableManager.GetOverriddenPropertyOperation(*SubObject, PropertyEvent, PropertyIterator->GetNextNode(), bOutInheritedOperation);
+						}
+						else if(CurrentOverriddenPropertyNode)
+						{
+							// Caller wants to know about any override state on the reference of the subobject itself
+							const FOverriddenPropertyNodeID  SubObjectID(*SubObject);
+							if (const FOverriddenPropertyNodeID* CurrentPropKey = CurrentOverriddenPropertyNode->SubPropertyNodeKeys.Find(SubObjectID))
+							{
+								const FOverriddenPropertyNode* SubObjectOverriddenPropertyNode = OverriddenPropertyNodes.FindByHash(GetTypeHash(*CurrentPropKey), *CurrentPropKey);
+								checkf(SubObjectOverriddenPropertyNode, TEXT("Expecting a node"));
+								if (bOutInheritedOperation)
+								{
+									*bOutInheritedOperation = false;
+								}
+								return SubObjectOverriddenPropertyNode->Operation;
+							}
+						}
+					}
+				}
+			}
+		}
+		// Special handling for maps and values of instance subobjects
+		else if (const FMapProperty* MapProperty = CastField<FMapProperty>(CurrentProperty))
+		{
+			ArrayIndex = PropertyEvent.GetArrayIndex(CurrentProperty->GetName());
+
+			checkf(MapProperty->ValueProp, TEXT("Expecting a value type for Maps"));
+			FScriptMapHelper MapHelper(MapProperty, SubValuePtr);
+
+			const int32 InternalMapIndex = ArrayIndex != INDEX_NONE ? MapHelper.FindInternalIndex(ArrayIndex) : INDEX_NONE;
+			if(MapHelper.IsValidIndex(InternalMapIndex))
+			{
+				if (PropertyIterator->GetNextNode())
+				{
+					// Forward any sub queries to the subobject
+					if (const FObjectProperty* ValueInstancedObjectProperty = MapProperty->ValueProp->HasAnyPropertyFlags(CPF_PersistentInstance) ? CastField<FObjectProperty>(MapProperty->ValueProp) : nullptr)
+					{
+						if (UObject* ValueSubObject = ValueInstancedObjectProperty->GetObjectPropertyValue(MapHelper.GetValuePtr(InternalMapIndex)))
+						{
+							return OverridableManager.GetOverriddenPropertyOperation(*ValueSubObject, PropertyEvent, PropertyIterator->GetNextNode(), bOutInheritedOperation);
+						}
+					}
+				}
+				else if(CurrentOverriddenPropertyNode)
+				{
+					// Caller wants to know about any override state on the reference of the map pair itself
+					checkf(MapProperty->KeyProp, TEXT("Expecting a key type for Maps"));
+					FObjectProperty* KeyObjectProperty = CastField<FObjectProperty>(MapProperty->KeyProp);
+
+					FOverriddenPropertyNodeID OverriddenKeyID;
+					if (const UObject* OverriddenKeyObject = KeyObjectProperty ? KeyObjectProperty->GetObjectPropertyValue(MapHelper.GetKeyPtr(InternalMapIndex)) : nullptr)
+					{
+						OverriddenKeyID = FOverriddenPropertyNodeID(*OverriddenKeyObject);
+					}
+					else
+					{
+						FString OverriddenKey;
+						MapProperty->KeyProp->ExportTextItem_Direct(OverriddenKey, MapHelper.GetKeyPtr(InternalMapIndex), nullptr, nullptr, PPF_None);
+						OverriddenKeyID = FName(OverriddenKey);
+					}
+
+					if (const FOverriddenPropertyNodeID* CurrentPropKey = CurrentOverriddenPropertyNode->SubPropertyNodeKeys.Find(OverriddenKeyID))
+					{
+						const FOverriddenPropertyNode* SubObjectOverriddenPropertyNode = OverriddenPropertyNodes.FindByHash(GetTypeHash(*CurrentPropKey), *CurrentPropKey);
+						checkf(SubObjectOverriddenPropertyNode, TEXT("Expecting a node"));
+						if (bOutInheritedOperation)
+						{
+							*bOutInheritedOperation = false;
+						}
+						return SubObjectOverriddenPropertyNode->Operation;
+					}
+				}
+			}
+		}
+
+		OverriddenPropertyNode = CurrentOverriddenPropertyNode;
+		PropertyIterator = PropertyIterator->GetNextNode();
+	}
+
+	if (bOutInheritedOperation)
+	{
+		*bOutInheritedOperation = PropertyIterator != nullptr || ArrayIndex != INDEX_NONE;
 	}
+	return OverriddenPropertyNode ? OverriddenPropertyNode->Operation : EOverriddenPropertyOperation::None;
+}
+
+bool FOverriddenPropertySet::ClearOverriddenProperty(FOverriddenPropertyNode& ParentPropertyNode, const FPropertyChangedEvent& PropertyEvent, const FEditPropertyChain::TDoubleLinkedListNode* PropertyNode, const void* Data)
+{
+	FOverridableManager& OverridableManager = FOverridableManager::Get();
 
+	bool bClearedOverrides = false;
+	const void* SubValuePtr = Data;
 	const FEditPropertyChain::TDoubleLinkedListNode* PropertyIterator = PropertyNode;
 	FOverriddenPropertyNode* OverriddenPropertyNode = &ParentPropertyNode;
+	int32 ArrayIndex = INDEX_NONE;
 	while (PropertyIterator && OverriddenPropertyNode && OverriddenPropertyNode->Operation != EOverriddenPropertyOperation::Replace)
 	{
+		ArrayIndex = INDEX_NONE;
+
 		const FProperty* CurrentProperty = PropertyIterator->GetValue();
-		const FName CurrentPropID = CurrentProperty->GetFName();
-		if (const FOverriddenPropertyNodeID* CurrentPropKey = OverriddenPropertyNode->SubPropertyNodeKeys.Find(CurrentPropID))
+		SubValuePtr = CurrentProperty->ContainerPtrToValuePtr<void>(SubValuePtr, 0); //@todo support static arrays
+
+		FOverriddenPropertyNode* CurrentOverriddenPropertyNode = nullptr;
+		if (OverriddenPropertyNode)
 		{
-			OverriddenPropertyNode = OverriddenPropertyNodes.FindByHash(GetTypeHash(*CurrentPropKey), *CurrentPropKey);
-			checkf(OverriddenPropertyNode, TEXT("Expecting a node"));
+			const FName CurrentPropID = CurrentProperty->GetFName();
+			if (const FOverriddenPropertyNodeID* CurrentPropKey = OverriddenPropertyNode->SubPropertyNodeKeys.Find(CurrentPropID))
+			{
+				CurrentOverriddenPropertyNode = OverriddenPropertyNodes.FindByHash(GetTypeHash(*CurrentPropKey), *CurrentPropKey);
+				checkf(CurrentOverriddenPropertyNode, TEXT("Expecting a node"));
+			}
 		}
-		else
+
+		// Special handling for for instanced subobjects 
+		if (const FObjectProperty* ObjectProperty = CastField<FObjectProperty>(CurrentProperty))
 		{
-			OverriddenPropertyNode = nullptr;
-			break;
+			if (UObject* SubObject = ObjectProperty->GetObjectPropertyValue(SubValuePtr))
+			{
+				if (PropertyIterator->GetNextNode())
+				{
+					return OverridableManager.ClearOverriddenProperty(*SubObject, PropertyEvent, PropertyIterator->GetNextNode());
+				}
+				else
+				{
+					OverridableManager.ClearOverrides(*SubObject);
+					bClearedOverrides = true;
+				}
+			}
+		}
+		// Special handling for array of instanced subobjects 
+		else if (const FArrayProperty* ArrayProperty = CastField<FArrayProperty>(CurrentProperty))
+		{
+			ArrayIndex = PropertyEvent.GetArrayIndex(CurrentProperty->GetName());
+
+			// Only special case is instanced subobjects, otherwise we fallback to full array override
+			if (FObjectProperty* InnerObjectProperty = CastField<FObjectProperty>(ArrayProperty->Inner))
+			{
+				if (InnerObjectProperty->HasAnyPropertyFlags(CPF_PersistentInstance))
+				{
+					FScriptArrayHelper ArrayHelper(ArrayProperty, SubValuePtr);
+
+					if(ArrayIndex == INDEX_NONE)
+					{
+						// This is a case of the entire array needs to be cleared
+						// Need to loop through every sub object and clear them
+						for (int i = 0; i < ArrayHelper.Num(); ++i)
+						{
+							if (UObject* SubObject = InnerObjectProperty->GetObjectPropertyValue(ArrayHelper.GetElementPtr(i)))
+							{
+								OverridableManager.ClearInstancedSubObjectOverrides(*Owner, *SubObject);
+							}
+						}
+						bClearedOverrides = true;
+					}
+					else if(ArrayHelper.IsValidIndex(ArrayIndex))
+					{
+						if (UObject* SubObject = InnerObjectProperty->GetObjectPropertyValue(ArrayHelper.GetElementPtr(ArrayIndex)))
+						{
+							if (PropertyIterator->GetNextNode())
+							{
+								return OverridableManager.ClearOverriddenProperty(*SubObject, PropertyEvent, PropertyIterator->GetNextNode());
+							}
+							else if(CurrentOverriddenPropertyNode)
+							{
+								const FOverriddenPropertyNodeID  SubObjectID(*SubObject);
+								FOverriddenPropertyNodeID CurrentPropKey;
+								if (CurrentOverriddenPropertyNode->SubPropertyNodeKeys.RemoveAndCopyValue(SubObjectID, CurrentPropKey))
+								{
+									verifyf(OverriddenPropertyNodes.RemoveByHash(GetTypeHash(CurrentPropKey), CurrentPropKey), TEXT("Expecting a node to be removed"));
+									OverridableManager.ClearInstancedSubObjectOverrides(*Owner, *SubObject);
+									return true;
+								}
+							}
+						}
+					}
+				}
+			}
 		}
+		// Special handling for maps and values of instance subobjects 
+		else if (const FMapProperty* MapProperty = CastField<FMapProperty>(CurrentProperty))
+		{
+			ArrayIndex = PropertyEvent.GetArrayIndex(CurrentProperty->GetName());
+
+			FScriptMapHelper MapHelper(MapProperty, SubValuePtr);
+
+			const int32 InternalMapIndex = ArrayIndex != INDEX_NONE ? MapHelper.FindInternalIndex(ArrayIndex) : INDEX_NONE;
+			const FObjectProperty* ValueInstancedObjectProperty = MapProperty->ValueProp->HasAnyPropertyFlags(CPF_PersistentInstance) ? CastField<FObjectProperty>(MapProperty->ValueProp) : nullptr;
+
+			// If there is a next node, it is probably because the map value is holding a instanced subobject and the user is changing value on it.
+			// So forward the call to the instanced subobject
+			if (PropertyIterator->GetNextNode())
+			{
+				if(MapHelper.IsValidIndex(InternalMapIndex))
+				{
+					checkf(MapProperty->ValueProp, TEXT("Expecting a value type for Maps"));
+					if (UObject* ValueSubObject = ValueInstancedObjectProperty ? ValueInstancedObjectProperty->GetObjectPropertyValue(MapHelper.GetValuePtr(InternalMapIndex)) : nullptr)
+					{
+						return OverridableManager.ClearOverriddenProperty(*ValueSubObject, PropertyEvent, PropertyIterator->GetNextNode());
+					}
+				}
+			}
+			else if(InternalMapIndex == INDEX_NONE)
+			{
+				// Users want to clear all of the overrides on the array, but in the case of instanced subobject, we need to clear the overrides on them as well.
+				if (ValueInstancedObjectProperty)
+				{
+					// This is a case of the entire array needs to be cleared
+					// Need to loop through every sub object and clear them
+					for (FScriptMapHelper::FIterator It(MapHelper); It; ++It)
+					{
+						if (UObject* ValueSubObject = ValueInstancedObjectProperty->GetObjectPropertyValue(MapHelper.GetValuePtr(It.GetInternalIndex())))
+						{
+							OverridableManager.ClearInstancedSubObjectOverrides(*Owner, *ValueSubObject);
+						}
+					}
+				}
+				bClearedOverrides = true;
+			}
+			else if (MapHelper.IsValidIndex(InternalMapIndex) && CurrentOverriddenPropertyNode)
+			{
+				checkf(MapProperty->KeyProp, TEXT("Expecting a key type for Maps"));
+				FObjectProperty* KeyObjectProperty = CastField<FObjectProperty>(MapProperty->KeyProp);
+
+				// Calculate the node from the key
+				FOverriddenPropertyNodeID OverriddenKeyID;
+				if (const UObject* OverriddenKeyObject = KeyObjectProperty ? KeyObjectProperty->GetObjectPropertyValue(MapHelper.GetKeyPtr(InternalMapIndex)) : nullptr)
+				{
+					OverriddenKeyID = FOverriddenPropertyNodeID(*OverriddenKeyObject);
+				}
+				else
+				{
+					FString OverriddenKey;
+					MapProperty->KeyProp->ExportTextItem_Direct(OverriddenKey, MapHelper.GetKeyPtr(InternalMapIndex), nullptr, nullptr, PPF_None);
+					OverriddenKeyID = FName(OverriddenKey);
+				}
+
+				FOverriddenPropertyNodeID CurrentPropKey;
+				if (CurrentOverriddenPropertyNode->SubPropertyNodeKeys.RemoveAndCopyValue(OverriddenKeyID, CurrentPropKey))
+				{
+					verifyf(OverriddenPropertyNodes.RemoveByHash(GetTypeHash(CurrentPropKey), CurrentPropKey), TEXT("Expecting a node to be removed"));
+
+					if (UObject* ValueSubObject = ValueInstancedObjectProperty ? ValueInstancedObjectProperty->GetObjectPropertyValue(MapHelper.GetValuePtr(InternalMapIndex)) : nullptr)
+					{
+						// In the case of a instanced subobject, clear all the overrides on the subobject as well
+						OverridableManager.ClearInstancedSubObjectOverrides(*Owner, *ValueSubObject);
+					}
+
+					return true;
+				}
+			}
+		}
+
+		OverriddenPropertyNode = CurrentOverriddenPropertyNode;
 		PropertyIterator = PropertyIterator->GetNextNode();
 	}
 
-	if (bOutInheritedState)
+	if (PropertyIterator != nullptr || OverriddenPropertyNode == nullptr)
 	{
-		*bOutInheritedState = PropertyIterator != nullptr;
+		return bClearedOverrides;
 	}
-	return OverriddenPropertyNode;
-}
 
-const FOverriddenPropertyNode* FOverriddenPropertySet::GetOverriddenPropertyNode(const FOverriddenPropertyNode& ParentPropertyNode, const FEditPropertyChain::TDoubleLinkedListNode* PropertyNode, bool* bOutInheritedState) const
-{
-	return const_cast<FOverriddenPropertySet*>(this)->GetOverriddenPropertyNode(const_cast<FOverriddenPropertyNode&>(ParentPropertyNode), PropertyNode, bOutInheritedState);
+	if (ArrayIndex != INDEX_NONE)
+	{
+		return false;
+	}
+
+	RemoveOverriddenSubProperties(*OverriddenPropertyNode);
+	OverriddenPropertyNode->Operation = EOverriddenPropertyOperation::None;
+	return true;
 }
 
 void FOverriddenPropertySet::NotifyPropertyChange(FOverriddenPropertyNode* ParentPropertyNode, const EPropertyNotificationType Notification, const FPropertyChangedEvent& PropertyEvent, const FEditPropertyChain::TDoubleLinkedListNode* PropertyNode, const void* Data)
@@ -425,13 +748,13 @@ void FOverriddenPropertySet::NotifyPropertyChange(FOverriddenPropertyNode* Paren
 	else if (const FMapProperty* MapProperty = CastField<FMapProperty>(Property))
 	{
 		// Special handling of instanced subobjects
-		checkf(MapProperty->KeyProp, TEXT("Expecting an key type for Maps"));
+		checkf(MapProperty->KeyProp, TEXT("Expecting a key type for Maps"));
 		FObjectProperty* KeyObjectProperty = CastField<FObjectProperty>(MapProperty->KeyProp);
 
 		// SubObjects
 		checkf(!KeyObjectProperty || !MapProperty->KeyProp->HasAnyPropertyFlags(CPF_PersistentInstance) || CastField<FClassProperty>(MapProperty->KeyProp), TEXT("Keys as a instanced subobject is not supported yet"));
 
-		checkf(MapProperty->ValueProp, TEXT("Expecting an value type for Maps"));
+		checkf(MapProperty->ValueProp, TEXT("Expecting a value type for Maps"));
 		FObjectProperty* ValueInstancedObjectProperty = MapProperty->ValueProp->HasAnyPropertyFlags(CPF_PersistentInstance) ? CastField<FObjectProperty>(MapProperty->ValueProp) : nullptr;
 
 		FScriptMapHelper MapHelper(MapProperty, SubValuePtr);
@@ -502,9 +825,9 @@ void FOverriddenPropertySet::NotifyPropertyChange(FOverriddenPropertyNode* Paren
 						}
 						else
 						{
-							FString AddedKey;
-							MapProperty->KeyProp->ExportTextItem_Direct(AddedKey, MapHelper.GetKeyPtr(It.GetInternalIndex()), nullptr, nullptr, PPF_None);
-							OverriddenKeyID = FName(AddedKey);
+							FString OverriddenKey;
+							MapProperty->KeyProp->ExportTextItem_Direct(OverriddenKey, MapHelper.GetKeyPtr(It.GetInternalIndex()), nullptr, nullptr, PPF_None);
+							OverriddenKeyID = FName(OverriddenKey);
 						}
 						FOverriddenPropertyNode& OverriddenKeyNode = FindOrAddNode(*SubPropertyNode, OverriddenKeyID);
 						OverriddenKeyNode.Operation = EOverriddenPropertyOperation::Replace;
@@ -740,12 +1063,11 @@ void FOverriddenPropertySet::RemoveOverriddenSubProperties(FOverriddenPropertyNo
 	PropertyNode.SubPropertyNodeKeys.Empty();
 }
 
-EOverriddenPropertyOperation FOverriddenPropertySet::GetOverriddenPropertyOperation(const FPropertyChangedEvent& PropertyEvent, const FEditPropertyChain::TDoubleLinkedListNode* PropertyNode, bool* bOutInheritedState /*= nullptr*/) const
+EOverriddenPropertyOperation FOverriddenPropertySet::GetOverriddenPropertyOperation(const FPropertyChangedEvent& PropertyEvent, const FEditPropertyChain::TDoubleLinkedListNode* PropertyNode, bool* bOutInheritedOperation /*= nullptr*/) const
 {
 	if (const FOverriddenPropertyNode* RootNode = OverriddenPropertyNodes.FindByHash(GetTypeHash(RootNodeID), RootNodeID))
 	{
-		const FOverriddenPropertyNode* OverriddenPropertyNode = GetOverriddenPropertyNode(*RootNode, PropertyNode, bOutInheritedState);
-		return OverriddenPropertyNode ? OverriddenPropertyNode->Operation : EOverriddenPropertyOperation::None;
+		return GetOverriddenPropertyOperation(*RootNode, PropertyEvent, PropertyNode, bOutInheritedOperation, Owner);
 	}
 	return EOverriddenPropertyOperation::None;
 }
@@ -754,14 +1076,7 @@ bool FOverriddenPropertySet::ClearOverriddenProperty(const FPropertyChangedEvent
 {
 	if (FOverriddenPropertyNode* RootNode = OverriddenPropertyNodes.FindByHash(GetTypeHash(RootNodeID), RootNodeID))
 	{
-		bool bOutInheritedState = false;
-		FOverriddenPropertyNode* OverriddenPropertyNode = GetOverriddenPropertyNode(*RootNode, PropertyNode, &bOutInheritedState);
-		if (OverriddenPropertyNode && OverriddenPropertyNode->Operation != EOverriddenPropertyOperation::None && !bOutInheritedState)
-		{
-			RemoveOverriddenSubProperties(*OverriddenPropertyNode);
-			OverriddenPropertyNode->Operation = EOverriddenPropertyOperation::None;
-			return true;
-		}
+		return ClearOverriddenProperty(*RootNode, PropertyEvent, PropertyNode, Owner);
 	}
 	return false;
 }
diff --git a/Engine/Source/Runtime/CoreUObject/Public/UObject/OverridableManager.h b/Engine/Source/Runtime/CoreUObject/Public/UObject/OverridableManager.h
--- a/Engine/Source/Runtime/CoreUObject/Public/UObject/OverridableManager.h
+++ b/Engine/Source/Runtime/CoreUObject/Public/UObject/OverridableManager.h
@@ -177,7 +177,18 @@ public:
 	 * @param PropertyEvent only needed to know about the container item index in any
 	 * @param PropertyChain to the property to clear from the root of the specified object
 	 * @return true if the property was successfully cleared. */
-	COREUOBJECT_API bool ClearOverriddenProperty(UObject& Object, const FPropertyChangedEvent& PropertyEvent, const FEditPropertyChain& PropertyChain);
+	FORCEINLINE bool ClearOverriddenProperty(UObject& Object, const FPropertyChangedEvent& PropertyEvent, const FEditPropertyChain& PropertyChain)
+	{
+		return ClearOverriddenProperty(Object, PropertyEvent, PropertyChain.GetHead());
+	}
+
+	/**
+	 * Clears an overridden properties specified by the property chain
+	 * @param Object owning the property to clear
+	 * @param PropertyEvent only needed to know about the container item index in any
+	 * @param PropertyNode leading to the property to clear, null means to clear the specified object
+	 * @return true if the property was successfully cleared. */
+	COREUOBJECT_API bool ClearOverriddenProperty(UObject& Object, const FPropertyChangedEvent& PropertyEvent, const FEditPropertyChain::TDoubleLinkedListNode* PropertyNode);
 
 	/**
 	 * To be called prior to override a property of the specified object
@@ -200,7 +211,7 @@ public:
 	 * @param Notification the type of notification (pre or post
 	 * @param Object owning the property
 	 * @param PropertyEvent information about the type of change including any container item index
-	 * @param PropertyNode leading to the property that is changing */
+	 * @param PropertyNode leading to the property that is changing, null indicate that it is the specified object has changed */
 	COREUOBJECT_API void NotifyPropertyChange(const EPropertyNotificationType Notification, UObject& Object, const FPropertyChangedEvent& PropertyEvent, const FEditPropertyChain::TDoubleLinkedListNode* PropertyNode);
 
 	/**
@@ -208,9 +219,21 @@ public:
 	 * @param Object owning the property
 	 * @param PropertyEvent only needed to know about the container item index in any
 	 * @param PropertyChain leading to the property the caller is interested in
-	 * @param bOutInheritedState optional parameter to know if the state returned was inherited from a parent property
+	 * @param bOutInheritedOperation optional parameter to know if the operation returned was inherited from a parent property
+	 * @return the current type of override operation on the property */
+	FORCEINLINE EOverriddenPropertyOperation GetOverriddenPropertyOperation(UObject& Object, const FPropertyChangedEvent& PropertyEvent, const FEditPropertyChain& PropertyChain, bool* bOutInheritedOperation = nullptr)
+	{
+		return GetOverriddenPropertyOperation(Object, PropertyEvent, PropertyChain.GetHead(), bOutInheritedOperation);
+	}
+
+	/**
+	 * Retrieve the overridable operation from the specified the edit property chain
+	 * @param Object owning the property
+	 * @param PropertyEvent only needed to know about the container item index in any
+	 * @param PropertyNode leading to the property interested in, null will return the operation on the specified object 
+	 * @param bOutInheritedOperation optional parameter to know if the state returned was inherited from a parent property
 	 * @return the current type of override operation on the property */
-	COREUOBJECT_API EOverriddenPropertyOperation GetOverriddenPropertyOperation(UObject& Object, const FPropertyChangedEvent& PropertyEvent, const FEditPropertyChain& PropertyChain, bool* bOutInheritedState = nullptr);
+	COREUOBJECT_API EOverriddenPropertyOperation GetOverriddenPropertyOperation(UObject& Object, const FPropertyChangedEvent& PropertyEvent, const FEditPropertyChain::TDoubleLinkedListNode* PropertyNode, bool* bOutInheritedOperation = nullptr);
 
 	/**
 	 * Serializes the overriden properties of the specified object into the record
diff --git a/Engine/Source/Runtime/CoreUObject/Public/UObject/OverriddenPropertySet.h b/Engine/Source/Runtime/CoreUObject/Public/UObject/OverriddenPropertySet.h
--- a/Engine/Source/Runtime/CoreUObject/Public/UObject/OverriddenPropertySet.h
+++ b/Engine/Source/Runtime/CoreUObject/Public/UObject/OverriddenPropertySet.h
@@ -217,15 +217,15 @@ public:
 	/**
 	 * Retrieve the overridable operation from the specified the edit property chain node
 	 * @param PropertyEvent only needed to know about the container item index in any
-	 * @param PropertyNode in the linked list leading to the property the caller is interested in
-	 * @param bOutInheritedState optional parameter to know if the state returned was inherited from a parent property
+	 * @param PropertyNode leading to the property interested in, null will return the operation of the object itself
+	 * @param bOutInheritedOperation optional parameter to know if the state returned was inherited from a parent property
 	 * @return the current type of override operation on the property */
-	EOverriddenPropertyOperation GetOverriddenPropertyOperation(const FPropertyChangedEvent& PropertyEvent, const FEditPropertyChain::TDoubleLinkedListNode* PropertyNode, bool* bOutInheritedState = nullptr) const;
+	EOverriddenPropertyOperation GetOverriddenPropertyOperation(const FPropertyChangedEvent& PropertyEvent, const FEditPropertyChain::TDoubleLinkedListNode* PropertyNode, bool* bOutInheritedOperation = nullptr) const;
 
 	/**
 	 * Clear any properties from the serialized property chain node
 	 * @param PropertyEvent only needed to know about the container item index in any
-	 * @param PropertyNode in the linked list leading to the property the caller is interested in
+	 * @param PropertyNode leading to the property to clear, null will clear the overrides on the object itself
 	 * @return if the operation was successful */
 	bool ClearOverriddenProperty(const FPropertyChangedEvent& PropertyEvent, const FEditPropertyChain::TDoubleLinkedListNode* PropertyNode);
 
@@ -233,28 +233,28 @@ public:
 	 * Handling and storing modification on a property of an object
 	 * @param Notification type either pre/post property overridden
 	 * @param PropertyEvent information about the type of change
-	 * @param PropertyNode leading to the property that is changing
+	 * @param PropertyNode leading to the property that is changing, null means it is the object itself that is changing
 	 * @param Data memory of the current property */
 	void NotifyPropertyChange(const EPropertyNotificationType Notification, const FPropertyChangedEvent& PropertyEvent, const FEditPropertyChain::TDoubleLinkedListNode* PropertyNode, const void* Data);
 
 	/**
 	 * Retrieve the overridable operation from the specified the serialized property chain and the specified property
-	 * @param CurrentPropertyChain leading to the property being serialized if any
-	 * @param Property being serialized if any, otherwise will fallback on the last property of the chain
+	 * @param CurrentPropertyChain leading to the property being serialized if any, null or empty it will return the node of the object itself
+	 * @param Property being serialized if any, if null it will fallback on the last property of the chain
 	 * @return the current type of override operation on the property */
 	EOverriddenPropertyOperation GetOverriddenPropertyOperation(const FArchiveSerializedPropertyChain* CurrentPropertyChain, FProperty* Property) const;
 
 	/**
 	 * Setup the overridable operation of the current property from the serialized property chain and the specified property
 	 * @param Operation to set for this property
-	 * @param CurrentPropertyChain leading to the property being serialized if any
-	 * @param Property being serialized if any, otherwise will fallback on the last property of the chain
+	 * @param CurrentPropertyChain leading to the property being serialized if any, null or empty it will return the node of the object itself
+	 * @param Property being serialized if any, if null it will fallback on the last property of the chain
 	 * @return the node containing the information of the overridden property */
 	FOverriddenPropertyNode* SetOverriddenPropertyOperation(EOverriddenPropertyOperation Operation, const FArchiveSerializedPropertyChain* CurrentPropertyChain, FProperty* Property);
 
 	/**
 	 * Retrieve the overridden property node from the serialized property chain
-	 * @param CurrentPropertyChain leading to the property being serialized.
+	 * @param CurrentPropertyChain leading to the property being serialized if any, null or empty it will return the node of the object itself
 	 * @return the node containing the information of the overridden property */
 	const FOverriddenPropertyNode* GetOverriddenPropertyNode(const FArchiveSerializedPropertyChain* CurrentPropertyChain) const;
 
@@ -289,8 +289,8 @@ protected:
 
 	FOverriddenPropertyNode& FindOrAddNode(FOverriddenPropertyNode& ParentPropertyNode, FOverriddenPropertyNodeID NodeID);
 
-	FOverriddenPropertyNode* GetOverriddenPropertyNode(FOverriddenPropertyNode& ParentPropertyNode, const FEditPropertyChain::TDoubleLinkedListNode* PropertyNode, bool* bOutInheritedState);
-	const FOverriddenPropertyNode* GetOverriddenPropertyNode(const FOverriddenPropertyNode& ParentPropertyNode, const FEditPropertyChain::TDoubleLinkedListNode* PropertyNode, bool* bOutInheritedState) const;
+	EOverriddenPropertyOperation GetOverriddenPropertyOperation(const FOverriddenPropertyNode& ParentPropertyNode, const FPropertyChangedEvent& PropertyEvent, const FEditPropertyChain::TDoubleLinkedListNode* PropertyNode, bool* bOutInheritedOperation, const void* Data) const;
+	bool ClearOverriddenProperty(FOverriddenPropertyNode& ParentPropertyNode, const FPropertyChangedEvent& PropertyEvent, const FEditPropertyChain::TDoubleLinkedListNode* PropertyNode, const void* Data);
 	void NotifyPropertyChange(FOverriddenPropertyNode* ParentPropertyNode, const EPropertyNotificationType Notification, const FPropertyChangedEvent& PropertyEvent, const FEditPropertyChain::TDoubleLinkedListNode* PropertyNode, const void* Data);
 
 	EOverriddenPropertyOperation GetOverriddenPropertyOperation(const FOverriddenPropertyNode& ParentPropertyNode, const FArchiveSerializedPropertyChain* CurrentPropertyChain, FProperty* Property) const;
