Fixed IDOs crash due to creating struct IDOs for primitive struct type that used native serialization

#changelist validated
#virtualized

diff --git a/Engine/Source/Runtime/CoreUObject/Private/UObject/InstanceDataObjectUtils.cpp b/Engine/Source/Runtime/CoreUObject/Private/UObject/InstanceDataObjectUtils.cpp
--- a/Engine/Source/Runtime/CoreUObject/Private/UObject/InstanceDataObjectUtils.cpp
+++ b/Engine/Source/Runtime/CoreUObject/Private/UObject/InstanceDataObjectUtils.cpp
@@ -98,7 +98,7 @@ namespace UE
 			if (AsArrayProperty->Inner->IsA<FEnumProperty>())
 			{
 				// enum paths currently don't recurse when they're in containers
-				OutType.AddTypeName(AsArrayProperty->Inner->GetClass()->GetFName());
+				OutType.AddTypeName(AsArrayProperty->Inner->GetID());
 			}
 			else
 			{
@@ -110,15 +110,15 @@ namespace UE
 		{
 			// sets currently don't recurse their element property types
 			OutType.BeginTypeParameters();
-			OutType.AddTypeName(AsSetProperty->ElementProp->GetClass()->GetFName());
+			OutType.AddTypeName(AsSetProperty->ElementProp->GetID());
 			OutType.EndTypeParameters();
 		}
 		else if (const FMapProperty* AsMapProperty = CastField<FMapProperty>(Property))
 		{
 			// maps currently don't recurse their key/value property types
 			OutType.BeginTypeParameters();
-			OutType.AddTypeName(AsMapProperty->KeyProp->GetClass()->GetFName());
-			OutType.AddTypeName(AsMapProperty->ValueProp->GetClass()->GetFName());
+			OutType.AddTypeName(AsMapProperty->KeyProp->GetID());
+			OutType.AddTypeName(AsMapProperty->ValueProp->GetID());
 			OutType.EndTypeParameters();
 		}
 #endif
@@ -277,20 +277,23 @@ namespace UE
 	{
 		if (FStructProperty* AsStructProperty = CastField<FStructProperty>(Property))
 		{
-			const FString* OriginalType = nullptr;
+			if (!AsStructProperty->Struct->UseNativeSerialization())
+			{
+				const FString* OriginalType = nullptr;
 #if WITH_EDITORONLY_DATA
-			//@note: Transfer existing metadata over as we build the InstanceDataObject from the struct or it owner, if any, this is useful for testing purposes
-			OriginalType = AsStructProperty->FindMetaData(NAME_StructOriginalTypeMetadata);
-			FField* OwnerField = OriginalType == nullptr ? AsStructProperty->Owner.ToField() : nullptr;
-			OriginalType = OwnerField ? OwnerField->FindMetaData(NAME_StructOriginalTypeMetadata) : OriginalType;
+				//@note: Transfer existing metadata over as we build the InstanceDataObject from the struct or it owner, if any, this is useful for testing purposes
+				OriginalType = AsStructProperty->FindMetaData(NAME_StructOriginalTypeMetadata);
+				FField* OwnerField = OriginalType == nullptr ? AsStructProperty->Owner.ToField() : nullptr;
+				OriginalType = OwnerField ? OwnerField->FindMetaData(NAME_StructOriginalTypeMetadata) : OriginalType;
 #endif
-			const FName StructOriginalName = OriginalType ? FName(**OriginalType) : AsStructProperty->Struct->GetFName();
-			AsStructProperty->Struct = CreateInstanceDataObjectStructRec<UScriptStruct>(AsStructProperty->Struct, Outer, LooseProperties, Path);
+				const FName StructOriginalName = OriginalType ? FName(**OriginalType) : AsStructProperty->Struct->GetFName();
+				AsStructProperty->Struct = CreateInstanceDataObjectStructRec<UScriptStruct>(AsStructProperty->Struct, Outer, LooseProperties, Path);
 #if WITH_EDITORONLY_DATA
-			AsStructProperty->SetMetaData(NAME_StructOriginalTypeMetadata, StructOriginalName.ToString());
-			AsStructProperty->SetMetaData(NAME_PresentAsTypeMetadata, StructOriginalName.ToString());
-			AsStructProperty->Struct->SetMetaData(NAME_PresentAsTypeMetadata, *StructOriginalName.ToString());
+				AsStructProperty->SetMetaData(NAME_StructOriginalTypeMetadata, StructOriginalName.ToString());
+				AsStructProperty->SetMetaData(NAME_PresentAsTypeMetadata, StructOriginalName.ToString());
+				AsStructProperty->Struct->SetMetaData(NAME_PresentAsTypeMetadata, *StructOriginalName.ToString());
 #endif
+			}
 		}
 		else if (const FArrayProperty* AsArrayProperty = CastField<FArrayProperty>(Property))
 		{
