Adds an InternalDoesPackageExistEx implementation that the AssetRegistry may use to check the filesystem for pacakge existence without potentially re-entering held AssetRegistry locks

#changelist validated
#virtualized

diff --git a/Engine/Source/Runtime/AssetRegistry/Private/AssetRegistry.cpp b/Engine/Source/Runtime/AssetRegistry/Private/AssetRegistry.cpp
--- a/Engine/Source/Runtime/AssetRegistry/Private/AssetRegistry.cpp
+++ b/Engine/Source/Runtime/AssetRegistry/Private/AssetRegistry.cpp
@@ -25,6 +25,7 @@
 #include "Misc/Optional.h"
 #include "Misc/PackageAccessTracking.h"
 #include "Misc/PackageAccessTrackingOps.h"
+#include "Misc/PackageSegment.h"
 #include "Misc/Paths.h"
 #include "Misc/PathViews.h"
 #include "Misc/RedirectCollector.h"
@@ -2872,11 +2873,13 @@ bool UAssetRegistryImpl::DoesPackageExistOnDisk(FName PackageName, FString* OutC
 			// presence of map -> .umap
 			// But we can only assume lack of map -> .uasset if we know the type of every object in the package.
 			// If we don't, because there was a redirector, we have to check the package on disk
-			FString FileName;
-			if (FPackageName::DoesPackageExist(PackageNameStr, &FileName, false /* InAllowTextFormats */))
+
+			// Note, the 'internal' version of DoesPackageExist must be used to avoid re-entering the AssetRegistry's lock resulting in deadlock
+			FPackagePath PackagePath;
+			if (FPackageName::InternalDoesPackageExistEx(PackageNameStr, FPackageName::EPackageLocationFilter::Any, 
+				false /*bMatchCaseOnDisk*/, &PackagePath) != FPackageName::EPackageLocationFilter::None)
 			{
-				check(!FileName.IsEmpty());
-				return FPaths::GetExtension(FileName, true /* bIncludeDot */);
+				return FString(PackagePath.GetExtensionString(EPackageSegment::Header));
 			}
 		}
 		return bContainsMap ? FPackageName::GetMapPackageExtension() : FPackageName::GetAssetPackageExtension();
@@ -2903,17 +2906,19 @@ bool UAssetRegistryImpl::DoesPackageExistOnDisk(FName PackageName, FString* OutC
 		const static bool bVerifyNegativeResults = FParse::Param(FCommandLine::Get(), TEXT("AssetRegistryValidatePackageExists"));
 		if (bVerifyNegativeResults && !AssetPackageData)
 		{
-			FString FileName;
-			if (FPackageName::DoesPackageExist(PackageNameStr, &FileName, false /* InAllowTextFormats */))
+			// Note, the 'internal' version of DoesPackageExist must be used to avoid re-entering the AssetRegistry's lock resulting in deadlock
+			FPackagePath PackagePath;
+			if (FPackageName::InternalDoesPackageExistEx(PackageNameStr, FPackageName::EPackageLocationFilter::Any, 
+				false /*bMatchCaseOnDisk*/, &PackagePath) != FPackageName::EPackageLocationFilter::None)
 			{
 				UE_LOG(LogAssetRegistry, Warning, TEXT("Package %s exists on disk but does not exist in the AssetRegistry"), *PackageNameStr);
 				if (OutCorrectCasePackageName)
 				{
-					FPackageName::TryConvertLongPackageNameToFilename(FileName, *OutCorrectCasePackageName);
+					*OutCorrectCasePackageName = PackagePath.GetLocalFullPath();
 				}
 				if (OutExtension)
 				{
-					*OutExtension = FPaths::GetExtension(FileName, true /* bIncludeDot */);
+					*OutExtension = PackagePath.GetExtensionString(EPackageSegment::Header);
 				}
 				return true;
 			}
@@ -2932,11 +2937,12 @@ bool UAssetRegistryImpl::DoesPackageExistOnDisk(FName PackageName, FString* OutC
 		{
 			if (AssetPackageData->Extension == EPackageExtension::Unspecified || AssetPackageData->Extension == EPackageExtension::Custom)
 			{
-				FString FileName;
-				if (FPackageName::DoesPackageExist(PackageNameStr, &FileName, false /* InAllowTextFormats */))
+				// Note, the 'internal' version of DoesPackageExist must be used to avoid re-entering the AssetRegistry's lock resulting in deadlock
+				FPackagePath PackagePath;
+				if(FPackageName::InternalDoesPackageExistEx(PackageNameStr, FPackageName::EPackageLocationFilter::Any, 
+					false /* bMatchCaseOnDisk*/, &PackagePath) != FPackageName::EPackageLocationFilter::None)
 				{
-					check(!FileName.IsEmpty());
-					*OutExtension = FPaths::GetExtension(FileName, true /* bIncludeDot */);
+					*OutExtension = PackagePath.GetExtensionString(EPackageSegment::Header);
 				}
 				else
 				{
@@ -5134,13 +5140,15 @@ FScanPathContext::FScanPathContext(FEventContext& InEventContext, FClassInherita
 			if (!IFileManager::Get().FileExists(*LocalPath))
 			{
 				// Find the extension
+				// Note, the 'internal' version of DoesPackageExist must be used to avoid re-entering the AssetRegistry's lock resulting in deadlock
 				FPackagePath PackagePath = FPackagePath::FromLocalPath(LocalPath);
-				if (!FPackageName::DoesPackageExist(PackagePath, &PackagePath))
+				if (FPackageName::InternalDoesPackageExistEx(PackagePath, FPackageName::EPackageLocationFilter::Any,
+					false /* bMatchCaseOnDisk */, &PackagePath) == FPackageName::EPackageLocationFilter::None)
 				{
 					UE_LOG(LogAssetRegistry, Warning, TEXT("ScanPathsSynchronous: Package %s does not exist, will not scan."), *InFile);
 					continue;
 				}
-				Extension = LexToString(PackagePath.GetHeaderExtension());
+				Extension = PackagePath.GetExtensionString(EPackageSegment::Header);
 			}
 		}
 		LocalFiles.Add(LocalPath + Extension);
diff --git a/Engine/Source/Runtime/AssetRegistry/Private/AssetRegistryConsoleCommands.h b/Engine/Source/Runtime/AssetRegistry/Private/AssetRegistryConsoleCommands.h
--- a/Engine/Source/Runtime/AssetRegistry/Private/AssetRegistryConsoleCommands.h
+++ b/Engine/Source/Runtime/AssetRegistry/Private/AssetRegistryConsoleCommands.h
@@ -212,10 +212,12 @@ public:
 		{
 			const FAssetData& AssetData = AllAssets[AssetIdx];
 
-			FString PackageFilename;
-			if ( FPackageName::DoesPackageExist(AssetData.PackageName.ToString(), &PackageFilename) )
+			// Note, the 'internal' version of DoesPackageExist must be used to avoid re-entering the AssetRegistry's lock resulting in deadlock
+			FPackagePath PackagePath;
+			if (FPackageName::InternalDoesPackageExistEx(AssetData.PackageName.ToString(), FPackageName::EPackageLocationFilter::Any, 
+				false /*bMatchCaseOnDisk*/, &PackagePath) != FPackageName::EPackageLocationFilter::None)
 			{
-				if ( FPaths::GetExtension(PackageFilename, true) == FPackageName::GetAssetPackageExtension() && !AssetData.IsUAsset())
+				if (PackagePath.GetHeaderExtension() == EPackageExtension::Asset && !AssetData.IsUAsset())
 				{
 					// This asset was in a package with a uasset extension but did not share the name of the package
 					UE_LOG(LogAssetRegistry, Log, TEXT("%s"), *AssetData.GetObjectPathString());
@@ -260,8 +262,10 @@ public:
 				FString LocalPath;
 				if (FPackageName::TryConvertLongPackageNameToFilename(InPath, LocalPath))
 				{
+					// Note, the 'internal' version of DoesPackageExist must be used to avoid re-entering the AssetRegistry's lock resulting in deadlock
 					FPackagePath PackagePath = FPackagePath::FromLocalPath(LocalPath);
-					if (FPackageName::DoesPackageExist(PackagePath, &PackagePath))
+					if (FPackageName::InternalDoesPackageExistEx(PackagePath, FPackageName::EPackageLocationFilter::Any, 
+						false /* bMatchCaseOnDisk */, &PackagePath) != FPackageName::EPackageLocationFilter::None)
 					{
 						bAsFile = true;
 						bAsDir = false;
diff --git a/Engine/Source/Runtime/CoreUObject/Private/Misc/PackageName.cpp b/Engine/Source/Runtime/CoreUObject/Private/Misc/PackageName.cpp
--- a/Engine/Source/Runtime/CoreUObject/Private/Misc/PackageName.cpp
+++ b/Engine/Source/Runtime/CoreUObject/Private/Misc/PackageName.cpp
@@ -1802,24 +1802,13 @@ bool FPackageName::FixPackageNameCase(FString& LongPackageName, FStringView Exte
 
 bool FPackageName::DoesPackageExist(const FString& LongPackageName, FString* OutFilename, bool InAllowTextFormats)
 {
-	// Make sure interpreting LongPackageName as a filename is supported.
 	FPackagePath PackagePath;
+	EErrorCode FailureReason;
+	if (!TryConvertToMountedPackagePath(LongPackageName, PackagePath, FailureReason))
 	{
-		SCOPED_LOADTIMER(FPackageName_DoesPackageExist);
-		TStringBuilder<64> PackageNameRoot;
-		TStringBuilder<64> FilePathRoot;
-		TStringBuilder<256> RelPath;
-		TStringBuilder<64> UnusedObjectName; // DoesPackageExist accepts ObjectPaths and ignores the ObjectName portion and uses only the PackageName
-		TStringBuilder<16> CustomExtension;
-		EPackageExtension Extension;
-		EErrorCode FailureReason;
-		if (!FPackageName::TryConvertToMountedPathComponents(LongPackageName, PackageNameRoot, FilePathRoot, RelPath, UnusedObjectName, Extension, CustomExtension, nullptr /* OutFlexNameType */, &FailureReason))
-		{
-			FString Message = FString::Printf(TEXT("DoesPackageExist called on PackageName that will always return false. Reason: %s"), *FormatErrorAsString(LongPackageName, FailureReason));
-			UE_LOG(LogPackageName, Warning, TEXT("%s"), *Message);
-			return false;
-		}
-		PackagePath = FPackagePath::FromMountedComponents(PackageNameRoot, FilePathRoot, RelPath, Extension, CustomExtension);
+		FString Message = FString::Printf(TEXT("DoesPackageExist called on PackageName that will always return false. Reason: %s"), *FormatErrorAsString(LongPackageName, FailureReason));
+		UE_LOG(LogPackageName, Warning, TEXT("%s"), *Message);
+		return false;
 	}
 	if (!DoesPackageExist(PackagePath, false /* bMatchCaseOnDisk */, &PackagePath))
 	{
@@ -1856,12 +1845,12 @@ FPackageName::EPackageLocationFilter FPackageName::DoesPackageExistEx(const FPac
 		return EPackageLocationFilter::None;
 	}
 
-#if 0 // WITH_EDITOR // Temporarily disabled due to deadlock, UE-205707
+#if WITH_EDITOR
 	IAssetRegistryInterface* AssetRegistry = IAssetRegistryInterface::GetPtr();
 
 	// Todo: The AssetRegistry currently cannot determine if a package comes from the Filesystem 
 	// or cooked content so we avoid registry lookups since we can't provide a reliable Location
-	if (AssetRegistry && ((uint8)Filter & (uint8)EPackageLocationFilter::FileSystem) && 
+	if (AssetRegistry && ((uint8)Filter & (uint8)EPackageLocationFilter::FileSystem) &&
 		(!FIoDispatcher::IsInitialized() || FIoDispatcher::Get().GetTotalLoaded() == 0))
 	{
 		FName PackageName = PackagePath.GetPackageFName();
@@ -1889,6 +1878,42 @@ FPackageName::EPackageLocationFilter FPackageName::DoesPackageExistEx(const FPac
 	}
 #endif
 
+	return InternalDoesPackageExistEx(PackagePath, Filter, bMatchCaseOnDisk, OutPackagePath);
+}
+
+bool FPackageName::TryConvertToMountedPackagePath(const FString& InPath, FPackagePath& OutPackagePath, EErrorCode& OutFailureReason)
+{
+	SCOPED_LOADTIMER(TryConvertToPackagePath);
+	TStringBuilder<64> PackageNameRoot;
+	TStringBuilder<64> FilePathRoot;
+	TStringBuilder<256> RelPath;
+	TStringBuilder<64> UnusedObjectName; // DoesPackageExist accepts ObjectPaths and ignores the ObjectName portion and uses only the PackageName
+	TStringBuilder<16> CustomExtension;
+	EPackageExtension Extension;
+	if (!FPackageName::TryConvertToMountedPathComponents(InPath, PackageNameRoot, FilePathRoot, RelPath, UnusedObjectName, Extension, CustomExtension, nullptr /* OutFlexNameType */, &OutFailureReason))
+	{
+		return false;
+	}
+	OutPackagePath = FPackagePath::FromMountedComponents(PackageNameRoot, FilePathRoot, RelPath, Extension, CustomExtension);
+	return true;
+}
+
+FPackageName::EPackageLocationFilter FPackageName::InternalDoesPackageExistEx(const FString& LongPackageName, EPackageLocationFilter Filter, bool bMatchCaseOnDisk, FPackagePath* OutPackagePath)
+{
+	FPackagePath PackagePath;
+	EErrorCode FailureReason;
+	if (!TryConvertToMountedPackagePath(LongPackageName, PackagePath, FailureReason))
+	{
+		FString Message = FString::Printf(TEXT("DoesPackageExist called on PackageName that will always return false. Reason: %s"), *FormatErrorAsString(LongPackageName, FailureReason));
+		UE_LOG(LogPackageName, Warning, TEXT("%s"), *Message);
+		return EPackageLocationFilter::None;
+	}
+
+	return InternalDoesPackageExistEx(PackagePath, Filter, bMatchCaseOnDisk, OutPackagePath);
+}
+
+FPackageName::EPackageLocationFilter FPackageName::InternalDoesPackageExistEx(const FPackagePath & PackagePath, EPackageLocationFilter Filter, bool bMatchCaseOnDisk, FPackagePath * OutPackagePath)
+{
 	TStringBuilder<256> PackageName;
 	PackagePath.AppendPackageName(PackageName);
 
diff --git a/Engine/Source/Runtime/CoreUObject/Public/Misc/PackageName.h b/Engine/Source/Runtime/CoreUObject/Public/Misc/PackageName.h
--- a/Engine/Source/Runtime/CoreUObject/Public/Misc/PackageName.h
+++ b/Engine/Source/Runtime/CoreUObject/Public/Misc/PackageName.h
@@ -24,6 +24,7 @@ class FPackagePath;
 class UPackage;
 struct FFileStatData;
 struct FGuid;
+namespace UE::AssetRegistry::Impl { struct FScanPathContext; }
 
 DECLARE_LOG_CATEGORY_EXTERN(LogPackageName, Log, All);
 
@@ -865,6 +866,16 @@ private:
 	static COREUOBJECT_API bool TryConvertToMountedPathComponents(FStringView InPath, FStringBuilderBase& OutMountPointPackageName, FStringBuilderBase& OutMountPointFilePath, FStringBuilderBase& OutRelPath,
 		FStringBuilderBase& OutObjectName, EPackageExtension& OutExtension, FStringBuilderBase& OutCustomExtension, EFlexNameType* OutFlexNameType = nullptr, EErrorCode* OutFailureReason = nullptr);
 
+	/**
+	 * Internal helper to create a FPackagePath given LocalFilePath or PackageName, and if found, return the MountPoint, RelativePath, and PackageExtension
+	 *
+	 * @param InPath The path to test, either a LocalFilePath, PackageName, or ObjectPath
+	 * @param OutPackagePath FPackagePath to store the converted InPath
+	 * @param OutFailureReason it is set to the failurereason if the MountPoint is not found, otherwise it is set to EErrorCode::PackageNameUnknown
+	 * @return True if InPath be converted (was not malformed)
+	 */
+	static COREUOBJECT_API bool TryConvertToMountedPackagePath(const FString& InPath, FPackagePath& OutPackagePath, EErrorCode& OutFailureReason);
+
 	/** Event that is triggered when a new content path is mounted */
 	static COREUOBJECT_API FOnContentPathMountedEvent OnContentPathMountedEvent;
 
@@ -875,5 +886,12 @@ private:
 	static COREUOBJECT_API FDoesPackageExistOverride DoesPackageExistOverrideDelegate;
 
 	friend class FPackagePath;
+
+	// Internal helper not meant for public use. These versions of DoesPackageExist do not use the AssetRegistry; any code that might be holding an AssetRegistry lock should call these functions instead
+	friend class FAssetRegistryConsoleCommands;
+	friend class UAssetRegistryImpl;
+	friend struct ::UE::AssetRegistry::Impl::FScanPathContext;
+	static COREUOBJECT_API EPackageLocationFilter InternalDoesPackageExistEx(const FPackagePath& PackagePath, EPackageLocationFilter Filterconst, bool bMatchCaseOnDisk = false, FPackagePath* OutPackagePath = nullptr);
+	static COREUOBJECT_API EPackageLocationFilter InternalDoesPackageExistEx(const FString& LongPackageName, EPackageLocationFilter Filterconst, bool bMatchCaseOnDisk = false, FPackagePath* OutPackagePath = nullptr);
 };
 
