[LinkerLoad]
 - Deprecate LoadPackageLinker

#changelist validated
#virtualized

diff --git a/Engine/Plugins/Editor/ChangelistReview/Source/Private/SSourceControlReview.cpp b/Engine/Plugins/Editor/ChangelistReview/Source/Private/SSourceControlReview.cpp
--- a/Engine/Plugins/Editor/ChangelistReview/Source/Private/SSourceControlReview.cpp
+++ b/Engine/Plugins/Editor/ChangelistReview/Source/Private/SSourceControlReview.cpp
@@ -1168,17 +1168,17 @@ void SSourceControlReview::FixupRedirectors()
 	{
 		if (ChangelistFile->ReviewFileTempPath.EndsWith(TEXT(".uasset")) || ChangelistFile->ReviewFileTempPath.EndsWith(TEXT(".umap")))
 		{
-			if (FLinkerLoad* ReviewFileLoad = LoadPackageLinker(nullptr, FPackagePath::FromLocalPath(ChangelistFile->ReviewFileTempPath), LOAD_ForDiff | LOAD_DisableCompileOnLoad | LOAD_DisableEngineVersionChecks))
-            {
-            	if (IsObjectRedirector(ReviewFileLoad, ChangelistFile->AssetName))
-            	{
-            		FString Path = GetContentRelativePath(ReviewFileLoad);
-            		if (!Path.IsEmpty())
-            		{
-            			RedirectorsFound.Add(Path, ChangelistFile.ToWeakPtr());
-            		}
-            	}
-            }
+			if (FLinkerLoad* ReviewFileLoad = GetPackageLinker(nullptr, FPackagePath::FromLocalPath(ChangelistFile->ReviewFileTempPath), LOAD_ForDiff | LOAD_DisableCompileOnLoad | LOAD_DisableEngineVersionChecks, nullptr))
+			{
+				if (IsObjectRedirector(ReviewFileLoad, ChangelistFile->AssetName))
+				{
+					FString Path = GetContentRelativePath(ReviewFileLoad);
+					if (!Path.IsEmpty())
+					{
+						RedirectorsFound.Add(Path, ChangelistFile.ToWeakPtr());
+					}
+				}
+			}
 		}
 	}
 
diff --git a/Engine/Source/Editor/UnrealEd/Private/Commandlets/BaseIteratePackagesCommandlet.cpp b/Engine/Source/Editor/UnrealEd/Private/Commandlets/BaseIteratePackagesCommandlet.cpp
--- a/Engine/Source/Editor/UnrealEd/Private/Commandlets/BaseIteratePackagesCommandlet.cpp
+++ b/Engine/Source/Editor/UnrealEd/Private/Commandlets/BaseIteratePackagesCommandlet.cpp
@@ -578,29 +578,8 @@ void UBaseIteratePackagesCommandlet::LoadAndSaveOnePackage(const FString& Filena
 
 		// Get the package linker.
 		VerboseMessage(TEXT("Pre GetPackageLinker"));
-#define USE_OLD_LOADLINKERPATH 0
-#if USE_OLD_LOADLINKERPATH
-		FLinkerLoad* Linker = LoadPackageLinker(nullptr, PackagePath, LOAD_NoVerify | LOAD_SkipLoadImportedPackages);
-	
-		// Bail early if we don't have a valid linker (package was out of date, etc)
-		if( !Linker )
-		{
-			VerboseMessage(TEXT("Aborting...package could not be loaded"));
-			CollectGarbage(RF_NoFlags);
-			return;
-		}
-
-		VerboseMessage(TEXT("Post GetPackageLinker"));
 
 		bool bSavePackage = true;
-		{
-			TRACE_CPUPROFILER_EVENT_SCOPE(UBaseIteratePackagesCommandlet::LoadAndSaveOnePackage::PerformPreloadOperations);
-			PerformPreloadOperations(Linker, bSavePackage);
-		}
-		VerboseMessage(FString::Printf(TEXT("Post PerformPreloadOperations, Resave? %d"), bSavePackage));
-#else
-		bool bSavePackage = true;
-#endif
 		
 		if (bSavePackage)
 		{
@@ -632,11 +611,11 @@ void UBaseIteratePackagesCommandlet::LoadAndSaveOnePackage(const FString& Filena
 					check(Package);
 				}
 			}
-#if !USE_OLD_LOADLINKERPATH
+
 			FLinkerLoad* Linker = Package->GetLinker();
 			PerformPreloadOperations(Linker, bSavePackage); 
 			VerboseMessage(FString::Printf(TEXT("Post PerformPreloadOperations, Resave? %d"), bSavePackage));
-#endif
+
 			VerboseMessage(TEXT("Post LoadPackage"));
 
 			// if we are only saving dirty packages and the package is not dirty, then we do not want to save the package (remember the default behavior is to ALWAYS save the package)
diff --git a/Engine/Source/Editor/UnrealEd/Private/Commandlets/PackageUtilities.cpp b/Engine/Source/Editor/UnrealEd/Private/Commandlets/PackageUtilities.cpp
--- a/Engine/Source/Editor/UnrealEd/Private/Commandlets/PackageUtilities.cpp
+++ b/Engine/Source/Editor/UnrealEd/Private/Commandlets/PackageUtilities.cpp
@@ -818,7 +818,7 @@ int32 ULoadPackageCommandlet::Main( const FString& Params )
 
 		if (bCheckForLegacyPackages)
 		{
-			FLinkerLoad* Linker = LoadPackageLinker(nullptr, PackagePath, LOAD_NoVerify);
+			FLinkerLoad* Linker = GetPackageLinker(nullptr, PackagePath, LOAD_NoVerify, nullptr);
 			MinVersion = FMath::Min<int32>(MinVersion, Linker->Summary.GetFileVersionUE().ToValue());
 		}
 		else
diff --git a/Engine/Source/Runtime/CoreUObject/Private/UObject/Linker.cpp b/Engine/Source/Runtime/CoreUObject/Private/UObject/Linker.cpp
--- a/Engine/Source/Runtime/CoreUObject/Private/UObject/Linker.cpp
+++ b/Engine/Source/Runtime/CoreUObject/Private/UObject/Linker.cpp
@@ -845,39 +845,30 @@ FLinkerLoad* GetPackageLinker
 
 FLinkerLoad* LoadPackageLinker(UPackage* InOuter, const FPackagePath& PackagePath, uint32 LoadFlags, UPackageMap* Sandbox, FArchive* InReaderOverride, TFunctionRef<void(FLinkerLoad* LoadedLinker)> LinkerLoadedCallback)
 {
-	FLinkerLoad* Linker = nullptr;
-	TRefCountPtr<FUObjectSerializeContext> LoadContext(FUObjectThreadContext::Get().GetSerializeContext());
-	BeginLoad(LoadContext);
-	{
-		FUObjectSerializeContext* InOutLoadContext = LoadContext;
-		Linker = GetPackageLinker(InOuter, PackagePath, LoadFlags, Sandbox, InReaderOverride, &InOutLoadContext);
-		if (InOutLoadContext != LoadContext)
-		{
-			// The linker already existed and was associated with another context
-			LoadContext->DecrementBeginLoadCount();
-			LoadContext = InOutLoadContext;
-			LoadContext->IncrementBeginLoadCount();
-		}
-	}
-	// Allow external code to work with the linker before EndLoad()
+	FLinkerLoad* Linker = GetPackageLinker(InOuter, PackagePath, LoadFlags, Sandbox, InReaderOverride);
 	LinkerLoadedCallback(Linker);
-	EndLoad(Linker ? Linker->GetSerializeContext() : LoadContext.GetReference());
 	return Linker;
 }
 
 FLinkerLoad* LoadPackageLinker(UPackage* InOuter, const FPackagePath& PackagePath, uint32 LoadFlags, UPackageMap* Sandbox, FArchive* InReaderOverride)
 {
+	PRAGMA_DISABLE_DEPRECATION_WARNINGS
 	return LoadPackageLinker(InOuter, PackagePath, LoadFlags, Sandbox, InReaderOverride, [](FLinkerLoad* InLinker) {});
+	PRAGMA_ENABLE_DEPRECATION_WARNINGS
 }
 
 FLinkerLoad* LoadPackageLinker(UPackage* InOuter, const TCHAR* InLongPackageName, uint32 LoadFlags, UPackageMap* Sandbox, FGuid* CompatibleGuid, FArchive* InReaderOverride, TFunctionRef<void(FLinkerLoad* LoadedLinker)> LinkerLoadedCallback)
 {
+	PRAGMA_DISABLE_DEPRECATION_WARNINGS
 	return LoadPackageLinker(InOuter, GetPackagePath(InOuter, InLongPackageName), LoadFlags, Sandbox, InReaderOverride, LinkerLoadedCallback);
+	PRAGMA_ENABLE_DEPRECATION_WARNINGS
 }
 
 FLinkerLoad* LoadPackageLinker(UPackage* InOuter, const TCHAR* InLongPackageName, uint32 LoadFlags, UPackageMap* Sandbox, FGuid* CompatibleGuid, FArchive* InReaderOverride)
 {
+	PRAGMA_DISABLE_DEPRECATION_WARNINGS
 	return LoadPackageLinker(InOuter, GetPackagePath(InOuter, InLongPackageName), LoadFlags, Sandbox, InReaderOverride, [](FLinkerLoad* InLinker) {});
+	PRAGMA_ENABLE_DEPRECATION_WARNINGS
 }
 
 
diff --git a/Engine/Source/Runtime/CoreUObject/Public/UObject/Linker.h b/Engine/Source/Runtime/CoreUObject/Public/UObject/Linker.h
--- a/Engine/Source/Runtime/CoreUObject/Public/UObject/Linker.h
+++ b/Engine/Source/Runtime/CoreUObject/Public/UObject/Linker.h
@@ -667,7 +667,9 @@ typedef uint32 ELazyLoaderFlags;
  * @param LinkerLoadedCallback Callback when the linker is loaded (or not found)
  * @return Pointer to the loaded linker or null if the file didn't exist
  */
+UE_DEPRECATED(5.4, "Use GetPackageLinker instead")
 COREUOBJECT_API FLinkerLoad* LoadPackageLinker(UPackage* InOuter, const FPackagePath& PackagePath, uint32 LoadFlags, UPackageMap* Sandbox, FArchive* InReaderOverride, TFunctionRef<void(FLinkerLoad* LoadedLinker)> LinkerLoadedCallback);
+UE_DEPRECATED(5.4, "Use GetPackageLinker instead")
 COREUOBJECT_API FLinkerLoad* LoadPackageLinker(UPackage* InOuter, const FPackagePath& PackagePath, uint32 LoadFlags = LOAD_None, UPackageMap* Sandbox = nullptr, FArchive* InReaderOverride = nullptr);
 
 UE_DEPRECATED(5.0, "Use version that takes a FPackagePath without a FGuid instead")
@@ -676,7 +678,7 @@ UE_DEPRECATED(5.0, "Use version that takes a FPackagePath without a FGuid instea
 COREUOBJECT_API FLinkerLoad* LoadPackageLinker(UPackage* InOuter, const TCHAR* InLongPackageName, uint32 LoadFlags = LOAD_None, UPackageMap* Sandbox = nullptr, FGuid* CompatibleGuid = nullptr, FArchive* InReaderOverride = nullptr);
 
 /** 
- * Gets a linker for a package and returns it without loading any objects. This call must be preceeded by BeginLoad and followed by EndLoad calls
+ * Gets a linker for a package and returns it without loading any objects.
  * @param InOuter Package if known, can be null
  * @param PackagePath Package resource to load, must not be empty
  * @param LoadFlags Flags to pass to the new linker
