Fix FScriptSetHelper usage in UObject::SaveConfig

diff --git a/Engine/Source/Runtime/CoreUObject/Private/UObject/Obj.cpp b/Engine/Source/Runtime/CoreUObject/Private/UObject/Obj.cpp
--- a/Engine/Source/Runtime/CoreUObject/Private/UObject/Obj.cpp
+++ b/Engine/Source/Runtime/CoreUObject/Private/UObject/Obj.cpp
@@ -3333,10 +3333,10 @@ void UObject::SaveConfig( uint64 Flags, const TCHAR* InFilename, FConfigCacheIni
 					else if (SetProperty)
 					{
 						FScriptSetHelper_InContainer SetHelper(SetProperty, this);
-						for (int32 i = 0; i < SetHelper.Num(); i++)
+						for (FScriptSetHelper::FIterator It = SetHelper.CreateIterator(); It; ++It)
 						{
 							FString	Buffer;
-							SetProperty->ElementProp->ExportTextItem_Direct(Buffer, SetHelper.GetElementPtr(i), SetHelper.GetElementPtr(i), this, PortFlags);
+							SetProperty->ElementProp->ExportTextItem_Direct(Buffer, SetHelper.GetElementPtr(It), SetHelper.GetElementPtr(It), this, PortFlags);
 							Config->AddToSection(*Section, *CompleteKey, *Buffer, PropFileName);
 						}
 						if (SetHelper.Num() == 0 && bIsADefaultIniWrite)
