Tidied up class cast flag declarations and fixed an outdated private namespace.

diff --git a/Engine/Source/Runtime/CoreUObject/Public/Templates/Casts.h b/Engine/Source/Runtime/CoreUObject/Public/Templates/Casts.h
--- a/Engine/Source/Runtime/CoreUObject/Public/Templates/Casts.h
+++ b/Engine/Source/Runtime/CoreUObject/Public/Templates/Casts.h
@@ -33,7 +33,6 @@ class USkinnedMeshComponent;
 class UStaticMeshComponent;
 /// @cond DOXYGEN_WARNINGS
 template<class TClass> class TSubclassOf;
-template <typename Type> struct TCastFlags;
 /// @endcond
 
 [[noreturn]] COREUOBJECT_API void CastLogError(const TCHAR* FromType, const TCHAR* ToType);
@@ -75,10 +74,19 @@ FORCEINLINE FString GetTypeName()
 	}
 }
 
+namespace UE::CoreUObject::Private
+{
+	template <typename Type>
+	constexpr inline EClassCastFlags TCastFlags_V = CASTCLASS_None;
+
+	template <typename Type>
+	constexpr inline EClassCastFlags TCastFlags_V<const Type> = TCastFlags_V<Type>;
+}
+
 template <typename Type>
-struct TCastFlags
+struct UE_DEPRECATED(5.5, "TCastFlags has been deprecated - use Cast instead") TCastFlags
 {
-	static const EClassCastFlags Value = CASTCLASS_None;
+	static const EClassCastFlags Value = UE::CoreUObject::Private::TCastFlags_V<Type>;
 };
 
 // Dynamically cast an object type-safely.
@@ -113,7 +121,7 @@ FORCEINLINE To* Cast(From* Src)
 				}
 			}
 		}
-		else if constexpr (UE_USE_CAST_FLAGS && TCastFlags<To>::Value != CASTCLASS_None)
+		else if constexpr (UE_USE_CAST_FLAGS && UE::CoreUObject::Private::TCastFlags_V<To> != CASTCLASS_None)
 		{
 			if constexpr (std::is_base_of_v<To, From>)
 			{
@@ -124,7 +132,7 @@ FORCEINLINE To* Cast(From* Src)
 #if UE_ENABLE_UNRELATED_CAST_WARNINGS
 				UE_STATIC_ASSERT_WARN((std::is_base_of_v<From, To>), "Attempting to use Cast<> on types that are not related");
 #endif
-				if (((const UObject*)Src)->GetClass()->HasAnyCastFlag(TCastFlags<To>::Value))
+				if (((const UObject*)Src)->GetClass()->HasAnyCastFlag(UE::CoreUObject::Private::TCastFlags_V<To>))
 				{
 					return (To*)Src;
 				}
@@ -287,7 +295,7 @@ FORCEINLINE typename TCopyQualifiersFromTo<From, To>::Type* Cast(const TObjectPt
 
 	const FObjectPtr& Src = (const FObjectPtr&)InSrc;
 
-	if constexpr (UE_USE_CAST_FLAGS && TCastFlags<To>::Value != CASTCLASS_None)
+	if constexpr (UE_USE_CAST_FLAGS && UE::CoreUObject::Private::TCastFlags_V<To> != CASTCLASS_None)
 	{
 		if (Src)
 		{
@@ -300,7 +308,7 @@ FORCEINLINE typename TCopyQualifiersFromTo<From, To>::Type* Cast(const TObjectPt
 	#if UE_ENABLE_UNRELATED_CAST_WARNINGS
 				UE_STATIC_ASSERT_WARN((std::is_base_of_v<From, To>), "Attempting to use Cast<> on types that are not related");
 	#endif
-				if (Src.GetClass()->HasAnyCastFlag(TCastFlags<To>::Value))
+				if (Src.GetClass()->HasAnyCastFlag(UE::CoreUObject::Private::TCastFlags_V<To>))
 				{
 					return (To*)Src.Get();
 				}
@@ -375,7 +383,7 @@ FORCEINLINE typename TCopyQualifiersFromTo<From, To>::Type* CastChecked(const TO
 
 	return nullptr;
 #else
-	if constexpr (UE_USE_CAST_FLAGS && TCastFlags<To>::Value != CASTCLASS_None)
+	if constexpr (UE_USE_CAST_FLAGS && UE::CoreUObject::Private::TCastFlags_V<To> != CASTCLASS_None)
 	{
 		return (To*)((const FObjectPtr&)Src).Get();
 	}
@@ -401,94 +409,70 @@ template< class T, class U > FORCEINLINE const T* Cast       ( const U      * Sr
 template< class T          > FORCEINLINE const T* ExactCast  ( const UObject* Src                                                                   ) { return ExactCast  <T>(const_cast<UObject*>(Src)); }
 template< class T, class U > FORCEINLINE const T* CastChecked( const U      * Src, ECastCheckedType::Type CheckType = ECastCheckedType::NullChecked ) { return CastChecked<T>(const_cast<U      *>(Src), CheckType); }
 
-#define DECLARE_CAST_BY_FLAG_FWD(ClassName) class ClassName;
-#define DECLARE_CAST_BY_FLAG_CAST(ClassName) \
-	template <> \
-	struct TCastFlags<ClassName> \
-	{ \
-		static const EClassCastFlags Value = CASTCLASS_##ClassName; \
-	}; \
-	template <> \
-	struct TCastFlags<const ClassName> \
-	{ \
-		static const EClassCastFlags Value = CASTCLASS_##ClassName; \
-	};
-
 #define DECLARE_CAST_BY_FLAG(ClassName) \
-	DECLARE_CAST_BY_FLAG_FWD(ClassName) \
-	DECLARE_CAST_BY_FLAG_CAST(ClassName)
-
-#define FINISH_DECLARING_CAST_FLAGS // intentionally defined to do nothing.
-
-// Define a macro that declares all the cast flags.
-// This allows us to reuse these declarations elsewhere to define other properties for these classes.
-// Note: When adding an item to this list, you must also add a CASTCLASS_ flag in ObjectBase.h and rebuild UnrealHeaderTool.
-#define DECLARE_ALL_CAST_FLAGS \
-DECLARE_CAST_BY_FLAG(UField)							\
-DECLARE_CAST_BY_FLAG(UEnum)								\
-DECLARE_CAST_BY_FLAG(UStruct)							\
-DECLARE_CAST_BY_FLAG(UScriptStruct)						\
-DECLARE_CAST_BY_FLAG(UClass)							\
-DECLARE_CAST_BY_FLAG(FProperty)							\
-DECLARE_CAST_BY_FLAG(FObjectPropertyBase)				\
-DECLARE_CAST_BY_FLAG(FObjectProperty)					\
-DECLARE_CAST_BY_FLAG(FWeakObjectProperty)				\
-DECLARE_CAST_BY_FLAG(FLazyObjectProperty)				\
-DECLARE_CAST_BY_FLAG(FSoftObjectProperty)				\
-DECLARE_CAST_BY_FLAG(FSoftClassProperty)				\
-DECLARE_CAST_BY_FLAG(FBoolProperty)						\
-DECLARE_CAST_BY_FLAG(UFunction)							\
-DECLARE_CAST_BY_FLAG(FStructProperty)					\
-DECLARE_CAST_BY_FLAG(FByteProperty)						\
-DECLARE_CAST_BY_FLAG(FIntProperty)						\
-DECLARE_CAST_BY_FLAG(FFloatProperty)					\
-DECLARE_CAST_BY_FLAG(FDoubleProperty)					\
-DECLARE_CAST_BY_FLAG(FClassProperty)					\
-DECLARE_CAST_BY_FLAG(FInterfaceProperty)				\
-DECLARE_CAST_BY_FLAG(FNameProperty)						\
-DECLARE_CAST_BY_FLAG(FStrProperty)						\
-DECLARE_CAST_BY_FLAG(FTextProperty)						\
-DECLARE_CAST_BY_FLAG(FArrayProperty)					\
-DECLARE_CAST_BY_FLAG(FDelegateProperty)					\
-DECLARE_CAST_BY_FLAG(FMulticastDelegateProperty)		\
-DECLARE_CAST_BY_FLAG(UPackage)							\
-DECLARE_CAST_BY_FLAG(ULevel)							\
-DECLARE_CAST_BY_FLAG(AActor)							\
-DECLARE_CAST_BY_FLAG(APlayerController)					\
-DECLARE_CAST_BY_FLAG(APawn)								\
-DECLARE_CAST_BY_FLAG(USceneComponent)					\
-DECLARE_CAST_BY_FLAG(UPrimitiveComponent)				\
-DECLARE_CAST_BY_FLAG(USkinnedMeshComponent)				\
-DECLARE_CAST_BY_FLAG(USkeletalMeshComponent)			\
-DECLARE_CAST_BY_FLAG(UBlueprint)						\
-DECLARE_CAST_BY_FLAG(UDelegateFunction)					\
-DECLARE_CAST_BY_FLAG(UStaticMeshComponent)				\
-DECLARE_CAST_BY_FLAG(FEnumProperty)						\
-DECLARE_CAST_BY_FLAG(FNumericProperty)					\
-DECLARE_CAST_BY_FLAG(FInt8Property)						\
-DECLARE_CAST_BY_FLAG(FInt16Property)					\
-DECLARE_CAST_BY_FLAG(FInt64Property)					\
-DECLARE_CAST_BY_FLAG(FUInt16Property)					\
-DECLARE_CAST_BY_FLAG(FUInt32Property)					\
-DECLARE_CAST_BY_FLAG(FUInt64Property)					\
-DECLARE_CAST_BY_FLAG(FMapProperty)						\
-DECLARE_CAST_BY_FLAG(FSetProperty)						\
-DECLARE_CAST_BY_FLAG(USparseDelegateFunction)			\
-DECLARE_CAST_BY_FLAG(FMulticastInlineDelegateProperty)	\
-DECLARE_CAST_BY_FLAG(FMulticastSparseDelegateProperty)	\
-DECLARE_CAST_BY_FLAG(FOptionalProperty)					\
-DECLARE_CAST_BY_FLAG(FVValueProperty)					\
-DECLARE_CAST_BY_FLAG(FVRestValueProperty)				\
-FINISH_DECLARING_CAST_FLAGS		// This is here to hopefully remind people to include the "\" in all declarations above, especially when copy/pasting the final line.
-
-// Now actually declare the flags
-DECLARE_ALL_CAST_FLAGS
+	class ClassName; \
+	template <> \
+	constexpr inline EClassCastFlags UE::CoreUObject::Private::TCastFlags_V<ClassName> = CASTCLASS_##ClassName;
+
+DECLARE_CAST_BY_FLAG(UField)
+DECLARE_CAST_BY_FLAG(UEnum)
+DECLARE_CAST_BY_FLAG(UStruct)
+DECLARE_CAST_BY_FLAG(UScriptStruct)
+DECLARE_CAST_BY_FLAG(UClass)
+DECLARE_CAST_BY_FLAG(FProperty)
+DECLARE_CAST_BY_FLAG(FObjectPropertyBase)
+DECLARE_CAST_BY_FLAG(FObjectProperty)
+DECLARE_CAST_BY_FLAG(FWeakObjectProperty)
+DECLARE_CAST_BY_FLAG(FLazyObjectProperty)
+DECLARE_CAST_BY_FLAG(FSoftObjectProperty)
+DECLARE_CAST_BY_FLAG(FSoftClassProperty)
+DECLARE_CAST_BY_FLAG(FBoolProperty)
+DECLARE_CAST_BY_FLAG(UFunction)
+DECLARE_CAST_BY_FLAG(FStructProperty)
+DECLARE_CAST_BY_FLAG(FByteProperty)
+DECLARE_CAST_BY_FLAG(FIntProperty)
+DECLARE_CAST_BY_FLAG(FFloatProperty)
+DECLARE_CAST_BY_FLAG(FDoubleProperty)
+DECLARE_CAST_BY_FLAG(FClassProperty)
+DECLARE_CAST_BY_FLAG(FInterfaceProperty)
+DECLARE_CAST_BY_FLAG(FNameProperty)
+DECLARE_CAST_BY_FLAG(FStrProperty)
+DECLARE_CAST_BY_FLAG(FTextProperty)
+DECLARE_CAST_BY_FLAG(FArrayProperty)
+DECLARE_CAST_BY_FLAG(FDelegateProperty)
+DECLARE_CAST_BY_FLAG(FMulticastDelegateProperty)
+DECLARE_CAST_BY_FLAG(UPackage)
+DECLARE_CAST_BY_FLAG(ULevel)
+DECLARE_CAST_BY_FLAG(AActor)
+DECLARE_CAST_BY_FLAG(APlayerController)
+DECLARE_CAST_BY_FLAG(APawn)
+DECLARE_CAST_BY_FLAG(USceneComponent)
+DECLARE_CAST_BY_FLAG(UPrimitiveComponent)
+DECLARE_CAST_BY_FLAG(USkinnedMeshComponent)
+DECLARE_CAST_BY_FLAG(USkeletalMeshComponent)
+DECLARE_CAST_BY_FLAG(UBlueprint)
+DECLARE_CAST_BY_FLAG(UDelegateFunction)
+DECLARE_CAST_BY_FLAG(UStaticMeshComponent)
+DECLARE_CAST_BY_FLAG(FEnumProperty)
+DECLARE_CAST_BY_FLAG(FNumericProperty)
+DECLARE_CAST_BY_FLAG(FInt8Property)
+DECLARE_CAST_BY_FLAG(FInt16Property)
+DECLARE_CAST_BY_FLAG(FInt64Property)
+DECLARE_CAST_BY_FLAG(FUInt16Property)
+DECLARE_CAST_BY_FLAG(FUInt32Property)
+DECLARE_CAST_BY_FLAG(FUInt64Property)
+DECLARE_CAST_BY_FLAG(FMapProperty)
+DECLARE_CAST_BY_FLAG(FSetProperty)
+DECLARE_CAST_BY_FLAG(USparseDelegateFunction)
+DECLARE_CAST_BY_FLAG(FMulticastInlineDelegateProperty)
+DECLARE_CAST_BY_FLAG(FMulticastSparseDelegateProperty)
+DECLARE_CAST_BY_FLAG(FOptionalProperty)
+DECLARE_CAST_BY_FLAG(FVValueProperty)
+DECLARE_CAST_BY_FLAG(FVRestValueProperty)
 
 #undef DECLARE_CAST_BY_FLAG
-#undef DECLARE_CAST_BY_FLAG_CAST
-#undef DECLARE_CAST_BY_FLAG_FWD
 
-namespace UECasts_Private
+namespace UE::CoreUObject::Private
 {
 	template <typename T>
 	struct TIsCastable
@@ -548,7 +532,7 @@ namespace UECasts_Private
 	}
 }
 
-#define dynamic_cast UECasts_Private::DynamicCast
+#define dynamic_cast UE::CoreUObject::Private::DynamicCast
 
 #if UE_ENABLE_INCLUDE_ORDER_DEPRECATED_IN_5_2
 #include "CoreMinimal.h"
