Create placeholder types/objects for package exports when type imports are missing/unresolved on editor map loads.

Notes:
- This allows the export to deserialize its tagged property stream (TPS) data into a property bag (when enabled), while also keeping references to the object in place.
- The export can still be serialized to the same package on save with its TPS data still intact. In that case, we will serialize the property bag back into the package as a "normal" TPS data stream.
- Serialized references to placeholder objects are resolved as NULL. For FObjectPtr types, the value is resolved as NULL when dereferenced (for type safety), but is still serialized as a non-NULL reference to the original object (even if placeholder).
- Placeholders will not be created on load unless FUObjectSerializeContext::bSerializeUnknownProperty is set to TRUE in the serialization context. Otherwise, we immediately resolve the export as NULL, and its data is lost on load (previous behavior).

Change summary:
- Added FPropertyBagRepository::PropertyBagTypeRegistry (w/ thread-safe implementation). Currently using ARO to manage placeholder type object lifetime.
- Added FPropertyBagRepository::AddPropertyBagPlaceholderType()/IsPropertyBagPlaceholderType().
- Modified FLinkerLoad::CreateExport() to instance a "placeholder" package/object if the export is missing its type import on load. Otherwise falls back to returning NULL (previous behavior).
- Added UE_WITH_OBJECT_HANDLE_TYPE_SAFETY along with IsObjectHandleTypeSafe().
- Added UE::CoreUObject::Private::HasAnyFlags() to ObjectFwd.h (using int32 for the param here as the enum isn't defined at this level). Similar to internal APIs added for late resolve; avoid having to include Object.h.
- Added RF_HasPlaceholderType (reclaimed an unused bit). This allows FObjectPtr to do a faster query at resolve time for type safety when referencing a placeholder object (i.e. it will return NULL).
- Modified FObjectPtr::Get() and IsNull() variants to include an extra check for type safety. Right now this is only enabled in an editor build context.
- Modified FLinkerLoad::Preload() to narrow the serialization scope to the export's TPS data *only* (since the original base type is unknown, we need to skip any natively-serialized data that might also exist in the stream).
- Modified FLinkerLoad serialization of raw UObject* values to return NULL in the case of a placeholder object (for type safety at runtime).
- Modified FLinkerLoad serialization of FObjectPtr values to return NULL in the case of a placeholder object only if type safety features are disabled on FObjectPtr/FObjectHandle.
- Extended the LowLevel core NULL ptr tests to include additional tests for "unsafe type" objects wrapped by an FObjectPtr/TObjectPtr.

diff --git a/Engine/Source/Runtime/CoreUObject/Internal/UObject/PropertyBagRepository.h b/Engine/Source/Runtime/CoreUObject/Internal/UObject/PropertyBagRepository.h
--- a/Engine/Source/Runtime/CoreUObject/Internal/UObject/PropertyBagRepository.h
+++ b/Engine/Source/Runtime/CoreUObject/Internal/UObject/PropertyBagRepository.h
@@ -50,7 +50,10 @@ private:
 	// used to make sure IDOs don't have name overlap
 	TMap<const UObjectBase*, TObjectPtr<UObject>> Namespaces;
 
-	FPropertyBagRepository() = default;
+	/** Internal registry that tracks the current set of types for property bag container objects instanced as placeholders for package exports that have invalid or missing class imports on load. */
+	TUniquePtr<class FPropertyBagTypeRegistry> PropertyBagTypeRegistry;
+
+	FPropertyBagRepository();
 
 public:
 	FPropertyBagRepository(const FPropertyBagRepository &) = delete;
@@ -105,8 +108,15 @@ public:
 	// query whether a property in Struct was set when the struct was deserialized
 	COREUOBJECT_API static bool WasPropertySetBySerialization(const UStruct* Struct, const void* StructData, const FProperty* Property, int32 ArrayIndex = 0);
 
+	// FGCObject interface
 	virtual void AddReferencedObjects(FReferenceCollector& Collector) override;
 	virtual FString GetReferencerName() const override;
+	// End FGCObject interface
+
+	// add a new placeholder type to swap in for a missing type on load; this will be associated with a property bag when instances are serialized so we don't lose its data
+	static void AddPropertyBagPlaceholderType(UClass* ClassType);
+	// query for whether or not the given class has been registered as a placeholder type
+	static COREUOBJECT_API bool IsPropertyBagPlaceholderType(UClass* ClassType);
 
 private:
 	void Lock() const { CriticalSection.Lock(); }
diff --git a/Engine/Source/Runtime/CoreUObject/Private/UObject/LinkerLoad.cpp b/Engine/Source/Runtime/CoreUObject/Private/UObject/LinkerLoad.cpp
--- a/Engine/Source/Runtime/CoreUObject/Private/UObject/LinkerLoad.cpp
+++ b/Engine/Source/Runtime/CoreUObject/Private/UObject/LinkerLoad.cpp
@@ -40,6 +40,7 @@
 #include "UObject/LinkerPlaceholderFunction.h"
 #include "UObject/LinkerManager.h"
 #include "UObject/ObjectSerializeAccessScope.h"
+#include "UObject/PropertyBagRepository.h"
 #include "Serialization/DeferredMessageLog.h"
 #include "UObject/UObjectThreadContext.h"
 #include "Serialization/AsyncLoading.h"
@@ -4048,6 +4049,70 @@ UClass* FLinkerLoad::GetExportLoadClass(int32 Index)
 	return dynamic_cast<UClass*>(IndexToObject(Export.ClassIndex));
 }
 
+#if WITH_EDITOR
+UClass* FLinkerLoad::TryCreatePlaceholderTypeForExport(int32 ExportIndex)
+{
+	UClass* LoadClass = nullptr;
+	FObjectExport& Export = ExportMap[ExportIndex];
+
+	// If the class import is missing, create a placeholder for this export. This will allow us to instance and redirect its data into a property bag.
+	// Note: This is only possible if the missing type did not also include any custom serialization, so we currently rely on the load context for this.
+	FUObjectSerializeContext* CurrentLoadContext = GetSerializeContext();
+	if (CurrentLoadContext && CurrentLoadContext->bSerializeUnknownProperty && Export.ClassIndex.IsImport() && !GEventDrivenLoaderEnabled)
+	{
+		FObjectImport& LoadClassImport = Imp(Export.ClassIndex);
+
+		// If the outer package import is also missing, create it now so that the full path remains the same. 
+		UObject* LoadClassParent = IndexToObject(LoadClassImport.OuterIndex);
+		if (!LoadClassParent && LoadClassImport.OuterIndex.IsImport())
+		{
+			FObjectImport& LoadClassParentImport = Imp(LoadClassImport.OuterIndex);
+			if (LoadClassParentImport.OuterIndex.IsNull())
+			{
+				LoadClassParent = CreatePackage(*LoadClassParentImport.ObjectName.ToString());
+
+				// Patch it into the import table so that we resolve to this package for future reference.
+				LoadClassParentImport.XObject = LoadClassParent;
+			}
+		}
+
+		if (LoadClassParent)
+		{
+			if (UObject* LoadClassTypePackage = FindObjectFast<UPackage>(nullptr, LoadClassImport.ClassPackage, /*bExactClass =*/ true))
+			{
+				if (UClass* LoadClassType = FindObjectFast<UClass>(LoadClassTypePackage, LoadClassImport.ClassName, /*bExactClass =*/ false))
+				{
+					// Create an opaque, non-native subtype that has no reflected properties. 
+					LoadClass = NewObject<UClass>(LoadClassParent, LoadClassType, LoadClassImport.ObjectName);
+					LoadClass->SetSuperStruct(UObject::StaticClass());
+					LoadClass->Bind();
+					LoadClass->StaticLink(/*bRelinkExistingProperties =*/ true);
+
+					// Create and configure its CDO as if it were loaded - for non-native class types, this is required.
+					UObject* LoadClassDefaults = LoadClass->GetDefaultObject();
+					LoadClass->PostLoadDefaultObject(LoadClassDefaults);
+
+					// Indicate that the CDO has a placeholder type (in case something has serialized a reference to it).
+					LoadClassDefaults->SetFlags(RF_HasPlaceholderType);
+
+					// This class is for internal use and should not be exposed for selection or instancing in the editor.
+					LoadClass->ClassFlags |= CLASS_Hidden | CLASS_HideDropDown;
+
+					// Patch it into the import table so that we resolve to this class for any future exports of this type.
+					LoadClassImport.XObject = LoadClass;
+
+					// Use the property bag repository for now to manage property bag placeholder types (e.g. object lifetime).
+					// Note: The object lifetime of instances of this type will rely on existing references that are serialized.
+					UE::FPropertyBagRepository::AddPropertyBagPlaceholderType(LoadClass);
+				}
+			}
+		}
+	}
+
+	return LoadClass;
+}
+#endif
+
 #if WITH_EDITORONLY_DATA
 int32 FLinkerLoad::LoadMetaDataFromExportMap(bool bForcePreload)
 {
@@ -4535,10 +4600,21 @@ void FLinkerLoad::Preload( UObject* Object )
 				check(Export.Object==Object);
 
 				const int64 SavedPos = Loader->Tell();
+				int64 StartPos = Export.SerialOffset;
+				int64 ExpectedSerialSize = Export.SerialSize;
+
+				// for placeholder objects that have no explicit type, we only want to serialize the TPS stream
+				const bool bSerializeOnlyScriptProperties = Object->HasAnyFlags(RF_HasPlaceholderType);
+				if (bSerializeOnlyScriptProperties)
+				{
+					// note: script start/end offsets are relative to the export's offset in the file
+					StartPos += Export.ScriptSerializationStartOffset;
+					ExpectedSerialSize = Export.ScriptSerializationEndOffset;
+				}
 
 				// move to the position in the file where this object's data
 				// is stored
-				Seek(Export.SerialOffset);
+				Seek(StartPos);
 
 				FAsyncArchive* AsyncLoader = GetAsyncLoader();
 
@@ -4681,14 +4757,30 @@ void FLinkerLoad::Preload( UObject* Object )
 							{
 								FStructuredArchiveChildReader ChildReader(ExportSlot);
 								FArchiveUObjectFromStructuredArchive Adapter(ChildReader.GetRoot());
-								Object->Serialize(Adapter.GetArchive());
+
+								if (bSerializeOnlyScriptProperties)
+								{
+									Object->SerializeScriptProperties(Adapter.GetArchive());
+								}
+								else
+								{
+									Object->Serialize(Adapter.GetArchive());
+								}
 							}
 						}
 						else
 #endif
 						{
 							UE_SERIALIZE_ACCCESS_SCOPE(Object);
-							Object->Serialize(*this);
+
+							if (bSerializeOnlyScriptProperties)
+							{
+								Object->SerializeScriptProperties(*this);
+							}
+							else
+							{
+								Object->Serialize(*this);
+							}
 						}
 
 						Object->SetFlags(RF_LoadCompleted);
@@ -4767,16 +4859,16 @@ void FLinkerLoad::Preload( UObject* Object )
 
 				// Make sure we serialized the right amount of stuff.
 				int64 Pos = Tell();
-				int64 SizeSerialized = Pos - Export.SerialOffset;
-				if( SizeSerialized != Export.SerialSize )
+				int64 SizeSerialized = Pos - StartPos;
+				if( SizeSerialized != ExpectedSerialSize )
 				{
 					if (Object->GetClass()->HasAnyClassFlags(CLASS_Deprecated))
 					{
-						UE_ASSET_LOG(LogLinker, Warning, PackagePath, TEXT("%s: Serial size mismatch: Got %d, Expected %d"), *Object->GetFullName(), (int32)SizeSerialized, Export.SerialSize);
+						UE_ASSET_LOG(LogLinker, Warning, PackagePath, TEXT("%s: Serial size mismatch: Got %d, Expected %d"), *Object->GetFullName(), (int32)SizeSerialized, ExpectedSerialSize);
 					}
 					else
 					{
-						UE_ASSET_LOG(LogLinker, Fatal, PackagePath, TEXT("%s: Serial size mismatch: Got %d, Expected %d"), *Object->GetFullName(), (int32)SizeSerialized, Export.SerialSize);
+						UE_ASSET_LOG(LogLinker, Fatal, PackagePath, TEXT("%s: Serial size mismatch: Got %d, Expected %d"), *Object->GetFullName(), (int32)SizeSerialized, ExpectedSerialSize);
 					}
 				}
 
@@ -4963,12 +5055,19 @@ UObject* FLinkerLoad::CreateExport( int32 Index )
 		UClass* LoadClass = GetExportLoadClass(Index);
 		if( !LoadClass && !Export.ClassIndex.IsNull() ) // Hack to load packages with classes which do not exist.
 		{
-			Export.bExportLoadFailed = true;
+#if WITH_EDITOR
+			// Try creating a placeholder type for it. This may allow us to instance and redirect its data into a property bag (to avoid data loss).
+			LoadClass = TryCreatePlaceholderTypeForExport(Index);
+			if (!LoadClass)
+#endif
+			{
+				Export.bExportLoadFailed = true;
 
-			FString OuterName = Export.OuterIndex.IsNull() ? LinkerRoot->GetFullName() : GetFullImpExpName(Export.OuterIndex);
-			FString ClassName = GetClassName(Export.ThisIndex).ToString();
-			UE_CLOG(Export.ObjectFlags & EObjectFlags::RF_Public, LogLinker, Warning, TEXT("Unable to load %s with outer %s because its class (%s) does not exist"), *Export.ObjectName.ToString(), *OuterName, *ClassName);
-			return nullptr;
+				FString OuterName = Export.OuterIndex.IsNull() ? LinkerRoot->GetFullName() : GetFullImpExpName(Export.OuterIndex);
+				FString ClassName = GetClassName(Export.ThisIndex).ToString();
+				UE_CLOG(Export.ObjectFlags & EObjectFlags::RF_Public, LogLinker, Warning, TEXT("Unable to load %s with outer %s because its class (%s) does not exist"), *Export.ObjectName.ToString(), *OuterName, *ClassName);
+				return nullptr;
+			}
 		}
 
 #if WITH_EDITOR
@@ -5371,7 +5470,7 @@ UObject* FLinkerLoad::CreateExport( int32 Index )
 			Export.bIsInheritedInstance = false;
 		}
 
-		LoadClass->GetDefaultObject();
+		const UObject* LoadClassDefaultObject = LoadClass->GetDefaultObject();
 
 		FStaticConstructObjectParameters Params(LoadClass);
 		Params.Outer = ThisParent;
@@ -5381,6 +5480,12 @@ UObject* FLinkerLoad::CreateExport( int32 Index )
 		// if our outer is actually an import, then the package we are an export of is not in our outer chain, set our package in that case
 		Params.ExternalPackage = Export.OuterIndex.IsImport() ? LinkerRoot : nullptr;
 
+		// If we're instancing from a placeholder type, set the object flag (so that we can quickly determine that it's a placeholder object).
+		if (LoadClassDefaultObject && LoadClassDefaultObject->HasAnyFlags(RF_HasPlaceholderType))
+		{
+			Params.SetFlags |= RF_HasPlaceholderType;
+		}
+
 		// Propagate relevant properties from the outer package to the external package
 		if (Params.ExternalPackage)
 		{
@@ -6077,7 +6182,15 @@ FArchive& FLinkerLoad::operator<<( UObject*& Object )
 	Ar << Index;
 
 	Object = ResolveResource(Index);
-
+#if WITH_EDITOR
+	if (Object && Object->HasAnyFlags(RF_HasPlaceholderType))
+	{
+		// This is needed because the pointer's type is checked only at compile time, which may not match the property
+		// bag placeholder object's type at runtime, and so we can't allow it to be dereferenced as the wrong base type.
+		// Note: These currently won't be discovered for replacement at reinstancing time, so it will remain set to NULL.
+		Object = nullptr;
+	}
+#endif
 	return *this;
 }
 
@@ -6087,28 +6200,41 @@ FArchive& FLinkerLoad::operator<<(FObjectPtr& ObjectPtr)
 	FArchive& Ar = *this;
 	Ar << Index;
 
-#if !UE_WITH_OBJECT_HANDLE_LATE_RESOLVE
-	ObjectPtr = FObjectPtr(ResolveResource(Index));
-	return *this;
+	// Wrapper that only allows pointers to exports with placeholder types when type safety features are enabled.
+	auto AsTypeSafeObjectPtr_Lambda = [](UObject* ResolvedObject)
+	{
+#if WITH_EDITOR && !UE_WITH_OBJECT_HANDLE_TYPE_SAFETY
+		// If we can't mask the reference to the placeholder object instance at resolve time, set it now to NULL.
+		// Note: Similar to hard references above, this means we won't find it for replacement at reinstancing time.
+		if (ResolvedObject && ResolvedObject->HasAnyFlags(RF_HasPlaceholderType))
+		{
+			ResolvedObject = nullptr;
+		}
+#endif
+		return FObjectPtr(ResolvedObject);
+	};
 
-#else
-	auto AssetRegistry = IAssetRegistryInterface::GetPtr();
+#if UE_WITH_OBJECT_HANDLE_LATE_RESOLVE
+	IAssetRegistryInterface* AssetRegistry = IAssetRegistryInterface::GetPtr();
 
 	if (!Index.IsImport() || !AssetRegistry)
 	{
-		ObjectPtr = FObjectPtr(ResolveResource(Index));
-		return *this;
+		ObjectPtr = AsTypeSafeObjectPtr_Lambda(ResolveResource(Index));
 	}
-
-	using namespace UE::LinkerLoad;
-	FObjectImport& Import = Imp(Index);
-	if (!TryLazyImport(*AssetRegistry, Import, *this, ObjectPtr))
+	else
 	{
-		ObjectPtr = FObjectPtr(ResolveResource(Index));
+		using namespace UE::LinkerLoad;
+		FObjectImport& Import = Imp(Index);
+		if (!TryLazyImport(*AssetRegistry, Import, *this, ObjectPtr))
+		{
+			ObjectPtr = AsTypeSafeObjectPtr_Lambda(ResolveResource(Index));
+		}
 	}
-	return *this;
-	
+#else
+	ObjectPtr = AsTypeSafeObjectPtr_Lambda(ResolveResource(Index));
 #endif
+
+	return *this;
 }
 
 FArchive& FLinkerLoad::operator<<(FSoftObjectPath& Value)
diff --git a/Engine/Source/Runtime/CoreUObject/Private/UObject/ObjectFwd.cpp b/Engine/Source/Runtime/CoreUObject/Private/UObject/ObjectFwd.cpp
--- a/Engine/Source/Runtime/CoreUObject/Private/UObject/ObjectFwd.cpp
+++ b/Engine/Source/Runtime/CoreUObject/Private/UObject/ObjectFwd.cpp
@@ -34,4 +34,9 @@ FString GetFullName( const UObject* Obj, const UObject* StopOuter, EObjectFullNa
 {
 	return Obj->GetFullName(StopOuter, Flags);
 }
+
+bool HasAnyFlags(const UObject* Obj, int32 FlagsToCheck)
+{
+	return Obj->HasAnyFlags(EObjectFlags(FlagsToCheck));
+}
 }
diff --git a/Engine/Source/Runtime/CoreUObject/Private/UObject/ObjectPtr.cpp b/Engine/Source/Runtime/CoreUObject/Private/UObject/ObjectPtr.cpp
--- a/Engine/Source/Runtime/CoreUObject/Private/UObject/ObjectPtr.cpp
+++ b/Engine/Source/Runtime/CoreUObject/Private/UObject/ObjectPtr.cpp
@@ -16,7 +16,7 @@ FString FObjectPtr::GetPathName() const
 	FObjectHandle LocalHandle = Handle;
 	if (IsObjectHandleResolved(LocalHandle) && !IsObjectHandleNull(LocalHandle))
 	{
-		return Get()->GetPathName();
+		return ResolveObjectHandle(LocalHandle)->GetPathName();
 	}
 	else
 	{
@@ -30,7 +30,7 @@ FName FObjectPtr::GetFName() const
 	FObjectHandle LocalHandle = Handle;
 	if (IsObjectHandleResolved(LocalHandle) && !IsObjectHandleNull(LocalHandle))
 	{
-		return Get()->GetFName();
+		return ResolveObjectHandle(LocalHandle)->GetFName();
 	}
 	else
 	{
@@ -44,7 +44,7 @@ FString FObjectPtr::GetFullName(EObjectFullNameFlags Flags) const
 	FObjectHandle LocalHandle = Handle;
 	if (IsObjectHandleResolved(LocalHandle) && !IsObjectHandleNull(LocalHandle))
 	{
-		return Get()->GetFullName(nullptr, Flags);
+		return ResolveObjectHandle(LocalHandle)->GetFullName(nullptr, Flags);
 	}
 	else
 	{
@@ -58,7 +58,7 @@ FObjectPtr FObjectPtr::GetOuter() const
 	FObjectHandle LocalHandle = Handle;
 	if (IsObjectHandleResolved(LocalHandle) && !IsObjectHandleNull(LocalHandle))
 	{
-		return FObjectPtr(Get()->GetOuter());
+		return FObjectPtr(ResolveObjectHandle(LocalHandle)->GetOuter());
 	}
 	UE::CoreUObject::Private::FPackedObjectRef PackedRef = UE::CoreUObject::Private::GetOuter(ReadObjectHandlePackedObjectRefNoCheck(LocalHandle));
 	FObjectPtr Ptr({ PackedRef.EncodedRef });
@@ -70,7 +70,7 @@ FObjectPtr FObjectPtr::GetPackage() const
 	FObjectHandle LocalHandle = Handle;
 	if (IsObjectHandleResolved(LocalHandle) && !IsObjectHandleNull(LocalHandle))
 	{
-		return FObjectPtr(Get()->GetPackage());
+		return FObjectPtr(ResolveObjectHandle(LocalHandle)->GetPackage());
 	}
 	UE::CoreUObject::Private::FPackedObjectRef PackedRef = UE::CoreUObject::Private::GetPackage(ReadObjectHandlePackedObjectRefNoCheck(LocalHandle));
 	FObjectPtr Ptr({ PackedRef.EncodedRef });
diff --git a/Engine/Source/Runtime/CoreUObject/Private/UObject/PropertyBagRepository.cpp b/Engine/Source/Runtime/CoreUObject/Private/UObject/PropertyBagRepository.cpp
--- a/Engine/Source/Runtime/CoreUObject/Private/UObject/PropertyBagRepository.cpp
+++ b/Engine/Source/Runtime/CoreUObject/Private/UObject/PropertyBagRepository.cpp
@@ -1,7 +1,7 @@
 // Copyright Epic Games, Inc. All Rights Reserved.
 
 #include "UObject/PropertyBagRepository.h"
-
+#include "Containers/Queue.h"
 #include "Serialization/ObjectReader.h"
 #include "Serialization/ObjectWriter.h"
 #include "UObject/GarbageCollection.h"
@@ -14,12 +14,56 @@
 #include "UObject/InstanceDataObjectUtils.h"
 #include "UObject/Package.h"
 
-
 DEFINE_LOG_CATEGORY_STATIC(LogPropertyBagRepository, Log, All);
 
 namespace UE
 {
 
+class FPropertyBagTypeRegistry
+{
+public:
+	void AddReferencedObjects(FReferenceCollector& Collector)
+	{
+		ConsumePendingPlaceholderTypes();
+		Collector.AddReferencedObjects(PlaceholderTypes);
+	}
+
+	void Add(UClass* Class)
+	{
+		PendingPlaceholderTypes.Enqueue(Class);
+	}
+
+	bool Contains(UClass* Class)
+	{
+		ConsumePendingPlaceholderTypes();
+		return PlaceholderTypes.Contains(Class);
+	}
+
+protected:
+	void ConsumePendingPlaceholderTypes()
+	{
+		if (!PendingPlaceholderTypes.IsEmpty())
+		{
+			FScopeLock ScopeLock(&CriticalSection);
+
+			TObjectPtr<UClass> Class;
+			while(PendingPlaceholderTypes.Dequeue(Class))
+			{
+				PlaceholderTypes.Add(Class);
+			}
+		}
+	}
+
+private:
+	FCriticalSection CriticalSection;
+
+	// List of types that have been registered.
+	TSet<TObjectPtr<UClass>> PlaceholderTypes;
+
+	// Types that have been added but not yet registered. Utilizes a thread-safe queue so we can avoid race conditions during an async load.
+	TQueue<TObjectPtr<UClass>> PendingPlaceholderTypes;
+};
+
 class FPropertyBagRepositoryLock
 {
 #if THREADSAFE_UOBJECTS
@@ -70,6 +114,11 @@ FPropertyBagRepository& FPropertyBagRepository::Get()
 	return Repo;
 }
 
+FPropertyBagRepository::FPropertyBagRepository()
+{
+	PropertyBagTypeRegistry = MakeUnique<FPropertyBagTypeRegistry>();
+}
+
 void FPropertyBagRepository::ReassociateObjects(const TMap<UObject*, UObject*>& ReplacedObjects)
 {
 	FPropertyBagRepositoryLock LockRepo(this);
@@ -201,6 +250,8 @@ void FPropertyBagRepository::AddReferencedObjects(FReferenceCollector& Collector
 	{
 		Collector.AddReferencedObject(Element.Value);
 	}
+
+	PropertyBagTypeRegistry->AddReferencedObjects(Collector);
 }
 
 FString FPropertyBagRepository::GetReferencerName() const
@@ -274,4 +325,24 @@ void FPropertyBagRepository::ShrinkMaps()
 	AssociatedData.Compact();
 }
 
+bool FPropertyBagRepository::IsPropertyBagPlaceholderType(UClass* ClassType)
+{
+	if (!ClassType)
+	{
+		return false;
+	}
+
+	return FPropertyBagRepository::Get().PropertyBagTypeRegistry->Contains(ClassType);
+}
+
+void FPropertyBagRepository::AddPropertyBagPlaceholderType(UClass* ClassType)
+{
+	if (!ClassType)
+	{
+		return;
+	}
+
+	FPropertyBagRepository::Get().PropertyBagTypeRegistry->Add(ClassType);
+}
+
 } // UE
diff --git a/Engine/Source/Runtime/CoreUObject/Public/UObject/LinkerLoad.h b/Engine/Source/Runtime/CoreUObject/Public/UObject/LinkerLoad.h
--- a/Engine/Source/Runtime/CoreUObject/Public/UObject/LinkerLoad.h
+++ b/Engine/Source/Runtime/CoreUObject/Public/UObject/LinkerLoad.h
@@ -856,6 +856,18 @@ private:
 	 */
 	UClass* GetExportLoadClass(int32 ExportIndex);
 
+#if WITH_EDITOR
+	/**
+	 * Utility function to create a placeholder type for the specified export.
+	 * This will be called if the export could not resolve its LoadClass import.
+	 * In that case, its data might still be serializable (into a property bag).
+	 *
+	 * @param  ExportIndex    Index of the export that's missing its type.
+	 * @return The placeholder type that will be used to load the export's data.
+	 */
+	UClass* TryCreatePlaceholderTypeForExport(int32 ExportIndex);
+#endif
+
 #if WITH_EDITORONLY_DATA
 	/** 
 	 * Looks for and loads meta data object from export map.
diff --git a/Engine/Source/Runtime/CoreUObject/Public/UObject/ObjectFwd.h b/Engine/Source/Runtime/CoreUObject/Public/UObject/ObjectFwd.h
--- a/Engine/Source/Runtime/CoreUObject/Public/UObject/ObjectFwd.h
+++ b/Engine/Source/Runtime/CoreUObject/Public/UObject/ObjectFwd.h
@@ -40,4 +40,6 @@ COREUOBJECT_API FString GetPathName( const UObject* Obj, const UObject* StopOute
 COREUOBJECT_API UClass* GetClass(UObject* Obj);
 COREUOBJECT_API FName GetFName(const UObject* Obj);
 COREUOBJECT_API UObject* GetOuter(const UObject* Obj);
+
+COREUOBJECT_API bool HasAnyFlags(const UObject* Obj, int32 FlagsToCheck);
 }
diff --git a/Engine/Source/Runtime/CoreUObject/Public/UObject/ObjectHandle.h b/Engine/Source/Runtime/CoreUObject/Public/UObject/ObjectHandle.h
--- a/Engine/Source/Runtime/CoreUObject/Public/UObject/ObjectHandle.h
+++ b/Engine/Source/Runtime/CoreUObject/Public/UObject/ObjectHandle.h
@@ -41,6 +41,7 @@ using FObjectHandle = UObject*;
 
 inline bool IsObjectHandleNull(FObjectHandle Handle);
 inline bool IsObjectHandleResolved(FObjectHandle Handle);
+inline bool IsObjectHandleTypeSafe(FObjectHandle Handle);
 
 //Private functions that forced public due to inlining.
 namespace UE::CoreUObject::Private
@@ -114,6 +115,18 @@ inline bool IsObjectHandleResolved(FObjectHandle Handle)
 #endif
 }
 
+/* return true if a handle is type safe.
+ * null and unresolved handles are considered safe
+ */ 
+inline bool IsObjectHandleTypeSafe(FObjectHandle Handle)
+{
+#if UE_WITH_OBJECT_HANDLE_TYPE_SAFETY
+	return IsObjectHandleNull(Handle) || !IsObjectHandleResolved(Handle) || !UE::CoreUObject::Private::HasAnyFlags(UE::CoreUObject::Private::ReadObjectHandlePointerNoCheck(Handle), static_cast<int32>(RF_HasPlaceholderType));
+#else
+	return true;
+#endif
+}
+
 /////////////////////////////////////////////////////////////////////////////////////////////////////////////////////////
 #if UE_WITH_OBJECT_HANDLE_LATE_RESOLVE
 inline bool operator==(UE::CoreUObject::Private::FObjectHandlePrivate LHS, UE::CoreUObject::Private::FObjectHandlePrivate RHS)
diff --git a/Engine/Source/Runtime/CoreUObject/Public/UObject/ObjectHandleDefines.h b/Engine/Source/Runtime/CoreUObject/Public/UObject/ObjectHandleDefines.h
--- a/Engine/Source/Runtime/CoreUObject/Public/UObject/ObjectHandleDefines.h
+++ b/Engine/Source/Runtime/CoreUObject/Public/UObject/ObjectHandleDefines.h
@@ -14,3 +14,8 @@
 #ifndef UE_WITH_OBJECT_HANDLE_TRACKING
 	#define UE_WITH_OBJECT_HANDLE_TRACKING WITH_EDITORONLY_DATA
 #endif
+
+
+#ifndef UE_WITH_OBJECT_HANDLE_TYPE_SAFETY
+	#define UE_WITH_OBJECT_HANDLE_TYPE_SAFETY WITH_EDITORONLY_DATA
+#endif
\ No newline at end of file
diff --git a/Engine/Source/Runtime/CoreUObject/Public/UObject/ObjectMacros.h b/Engine/Source/Runtime/CoreUObject/Public/UObject/ObjectMacros.h
--- a/Engine/Source/Runtime/CoreUObject/Public/UObject/ObjectMacros.h
+++ b/Engine/Source/Runtime/CoreUObject/Public/UObject/ObjectMacros.h
@@ -570,8 +570,7 @@ enum EObjectFlags
 	// RF_Dynamic				=0x04000000,	///< Was removed along with bp nativization
 	RF_WillBeLoaded				=0x08000000,	///< This object was constructed during load and will be loaded shortly
 	RF_HasExternalPackage		=0x10000000,	///< This object has an external package assigned and should look it up when getting the outermost package
-
-	// RF_Unused				=0x20000000,
+	RF_HasPlaceholderType		=0x20000000,	///< This object was instanced from a placeholder type (e.g. on load). References to it are serialized but externally resolve to NULL from a logical point of view (for type safety).
 
 	// RF_MirroredGarbage is mirrored in EInternalObjectFlags::Garbage because checking the internal flags is much faster for the Garbage Collector
 	// while checking the object flags is much faster outside of it where the Object pointer is already available and most likely cached.
diff --git a/Engine/Source/Runtime/CoreUObject/Public/UObject/ObjectPtr.h b/Engine/Source/Runtime/CoreUObject/Public/UObject/ObjectPtr.h
--- a/Engine/Source/Runtime/CoreUObject/Public/UObject/ObjectPtr.h
+++ b/Engine/Source/Runtime/CoreUObject/Public/UObject/ObjectPtr.h
@@ -102,11 +102,23 @@ public:
 
 	FORCEINLINE UObject* Get() const
 	{
+#if UE_WITH_OBJECT_HANDLE_TYPE_SAFETY
+		// Ensure the handle is resolved first (for late resolve), even if it's not considered type safe.
+		UObject* ResolvedObject = UE::CoreUObject::Private::ResolveObjectHandle(Handle);
+		return IsObjectHandleTypeSafe(Handle) ? ResolvedObject : nullptr;
+#else
 		return UE::CoreUObject::Private::ResolveObjectHandle(Handle);
+#endif
 	}
 
 	FORCEINLINE UClass* GetClass() const
 	{
+#if UE_WITH_OBJECT_HANDLE_TYPE_SAFETY
+		if (!IsObjectHandleTypeSafe(Handle))
+		{
+			return nullptr;
+		}
+#endif
 		return UE::CoreUObject::Private::ResolveObjectHandleClass(Handle);
 	}
 
@@ -178,13 +190,13 @@ public:
 	FORCEINLINE UObject& operator*() const { return *Get(); }
 
 	UE_DEPRECATED(5.1, "IsNull is deprecated, please use operator bool instead.")
-	FORCEINLINE bool IsNull() const { return UE::CoreUObject::Private::ResolveObjectHandleNoRead(Handle) == nullptr; }
+	FORCEINLINE bool IsNull() const { return IsNullNoResolve_Internal() || (UE::CoreUObject::Private::ResolveObjectHandleNoRead(Handle) == nullptr); }
 	
 	UE_DEPRECATED(5.1, "IsNullNoResolve is deprecated, please use operator bool instead.")
-	FORCEINLINE bool IsNullNoResolve() const { return IsObjectHandleNull(Handle); }
+	FORCEINLINE bool IsNullNoResolve() const { return IsNullNoResolve_Internal(); }
 	
-	FORCEINLINE bool operator!() const { return IsObjectHandleNull(Handle); }
-	explicit FORCEINLINE operator bool() const { return !IsObjectHandleNull(Handle); }
+	FORCEINLINE bool operator!() const { return IsNullNoResolve_Internal(); }
+	explicit FORCEINLINE operator bool() const { return !IsNullNoResolve_Internal(); }
 
 	FORCEINLINE bool IsResolved() const { return IsObjectHandleResolved(Handle); }
 
@@ -299,6 +311,15 @@ private:
 		}
 	}
 #endif // UE_OBJECT_PTR_GC_BARRIER
+
+	FORCEINLINE bool IsNullNoResolve_Internal() const
+	{
+#if UE_WITH_OBJECT_HANDLE_TYPE_SAFETY
+		return IsObjectHandleNull(Handle) || !IsObjectHandleTypeSafe(Handle);
+#else
+		return IsObjectHandleNull(Handle);
+#endif
+	}
 };
 
 template <typename T>
diff --git a/Engine/Source/Runtime/CoreUObject/Tests/ObjectHandleTest.cpp b/Engine/Source/Runtime/CoreUObject/Tests/ObjectHandleTest.cpp
--- a/Engine/Source/Runtime/CoreUObject/Tests/ObjectHandleTest.cpp
+++ b/Engine/Source/Runtime/CoreUObject/Tests/ObjectHandleTest.cpp
@@ -371,4 +371,92 @@ TEST_CASE_METHOD(FObjectHandleTestBase, "CoreUObject::FObjectHandle::Hash Object
 	CHECK(GetTypeHash(DummyObjectHandle) == GetTypeHash(&DummyObjectWithInvalidIndex));
 }
 
+#if UE_WITH_OBJECT_HANDLE_TYPE_SAFETY
+TEST_CASE_METHOD(FObjectHandleTestBase, "CoreUObject::FObjectHandle::Type Safety", "[CoreUObject][ObjectHandle]")
+{
+	const FName TestPackageName(TEXT("/Engine/Test/ObjectHandle/TypeSafety/Transient"));
+	UPackage* TestPackage = NewObject<UPackage>(nullptr, TestPackageName, RF_Transient);
+	TestPackage->AddToRoot();
+	ON_SCOPE_EXIT
+	{
+		TestPackage->RemoveFromRoot();
+	};
+
+	// simulate an unsafe class type
+	UClass* TestClass = NewObject<UClass>(TestPackage, TEXT("TestClass"), RF_Transient);
+	TestClass->SetSuperStruct(UObject::StaticClass());
+	TestClass->Bind();
+	TestClass->StaticLink(/*bRelinkExistingProperties =*/ true);
+	UObject* TestClassDefaults = TestClass->GetDefaultObject();
+	TestClass->PostLoadDefaultObject(TestClassDefaults);
+	TestClassDefaults->SetFlags(RF_HasPlaceholderType);
+
+	// validate helper method(s)
+	CHECK_FALSE(UE::CoreUObject::Private::HasAnyFlags(TestClassDefaults, RF_NoFlags));
+	CHECK(UE::CoreUObject::Private::HasAnyFlags(TestClassDefaults, RF_HasPlaceholderType));
+
+	// construct objects for testing
+	UObject* TestSafeObject = NewObject<UObjectPtrTestClass>(TestPackage, TEXT("TestSafeObject"), RF_Transient);
+	UObject* TestUnsafeObject = NewObject<UObject>(TestPackage, TestClass, TEXT("TestUnsafeObject"), RF_Transient | RF_HasPlaceholderType);
+
+	// construct object handles for testing
+	FObjectHandle NullObjectHandle = UE::CoreUObject::Private::MakeObjectHandle(nullptr);
+	FObjectHandle TestSafeObjectHandle = UE::CoreUObject::Private::MakeObjectHandle(TestSafeObject);
+	FObjectHandle TestUnsafeObjectHandle = UE::CoreUObject::Private::MakeObjectHandle(TestUnsafeObject);
+#if UE_WITH_OBJECT_HANDLE_LATE_RESOLVE
+	// note that unresolved object handles are type safe by definition (since it implies the underlying type was not a placeholder)
+	FObjectRef TestSafeObjectRef(TestSafeObject);
+	UE::CoreUObject::Private::FPackedObjectRef PackedSafeObjectRef = UE::CoreUObject::Private::MakePackedObjectRef(TestSafeObjectRef);
+	FObjectHandle TestUnresolvedSafeObjectHandle = { PackedSafeObjectRef.EncodedRef };
+#endif
+
+	// NULL/type-safe objects should report as being safe
+	CHECK(IsObjectHandleTypeSafe(NullObjectHandle));
+	CHECK(IsObjectHandleTypeSafe(TestSafeObjectHandle));
+#if UE_WITH_OBJECT_HANDLE_LATE_RESOLVE
+	CHECK(IsObjectHandleTypeSafe(TestUnresolvedSafeObjectHandle));
+	CHECK(!IsObjectHandleResolved(TestUnresolvedSafeObjectHandle));	// the call above should not resolve the handle
+#endif
+
+	// unsafe type object handles should report as being unsafe
+	CHECK_FALSE(IsObjectHandleTypeSafe(TestUnsafeObjectHandle));
+
+	// object handles should resolve the class to the unsafe type
+	CHECK(UE::CoreUObject::Private::ResolveObjectHandleClass(TestUnsafeObjectHandle) == TestClass);
+
+	// object handles should resolve/evaluate to the original type object
+	CHECK(UE::CoreUObject::Private::ResolveObjectHandle(TestUnsafeObjectHandle) == TestUnsafeObject);
+
+	// construct object pointers for testing
+	TObjectPtr<UObject> NullObjectPtr(nullptr);
+	TObjectPtr<UObject> TestSafeObjectPtr(TestSafeObject);
+	TObjectPtr<UObject> TestUnsafeObjectPtr(TestUnsafeObject);
+
+	// unsafe type object pointers should evaluate to NULL/false (for type safety)
+	CHECK(!TestUnsafeObjectPtr);
+	CHECK_FALSE(!!TestUnsafeObjectPtr);
+	CHECK(nullptr == TestUnsafeObjectPtr);
+	CHECK(TestUnsafeObjectPtr == nullptr);
+
+	// an unsafe type object pointer should not equate to NULL/other pointers
+	CHECK(TestUnsafeObjectPtr != NullObjectPtr);
+	CHECK(TestUnsafeObjectPtr != TestSafeObjectPtr);
+	CHECK(TestUnsafeObjectPtr == TestUnsafeObjectPtr);
+
+	// an unsafe type object should evaluate the object's attributes correctly
+	CHECK(TestUnsafeObjectPtr.GetName() == TestUnsafeObject->GetName());
+	CHECK(TestUnsafeObjectPtr.GetFName() == TestUnsafeObject->GetFName());
+	CHECK(TestUnsafeObjectPtr.GetPathName() == TestUnsafeObject->GetPathName());
+	CHECK(TestUnsafeObjectPtr.GetFullName() == TestUnsafeObject->GetFullName());
+	CHECK(TestUnsafeObjectPtr.GetOuter() == TestUnsafeObject->GetOuter());
+	CHECK(TestUnsafeObjectPtr.GetPackage() == TestUnsafeObject->GetPackage());
+
+	// an unsafe type object should not allow direct access to the underlying type
+	CHECK(TestUnsafeObjectPtr.GetClass() == nullptr);
+
+	// an unsafe type object pointer should resolve to NULL when dereferenced (for type safety)
+	CHECK(TestUnsafeObjectPtr.Get() == nullptr);
+}
+#endif
+
 #endif
