Prevented Cast from being used on FFields.

diff --git a/Engine/Source/Runtime/CoreUObject/Public/Templates/Casts.h b/Engine/Source/Runtime/CoreUObject/Public/Templates/Casts.h
--- a/Engine/Source/Runtime/CoreUObject/Public/Templates/Casts.h
+++ b/Engine/Source/Runtime/CoreUObject/Public/Templates/Casts.h
@@ -84,7 +84,7 @@ namespace UE::CoreUObject::Private
 }
 
 template <typename Type>
-struct UE_DEPRECATED(5.5, "TCastFlags has been deprecated - use Cast instead") TCastFlags
+struct UE_DEPRECATED(5.5, "TCastFlags has been deprecated - use Cast instead.") TCastFlags
 {
 	static const EClassCastFlags Value = UE::CoreUObject::Private::TCastFlags_V<Type>;
 };
@@ -121,43 +121,46 @@ FORCEINLINE To* Cast(From* Src)
 				}
 			}
 		}
-		else if constexpr (UE_USE_CAST_FLAGS && UE::CoreUObject::Private::TCastFlags_V<To> != CASTCLASS_None)
-		{
-			if constexpr (std::is_base_of_v<To, From>)
-			{
-				return (To*)Src;
-			}
-			else
-			{
-#if UE_ENABLE_UNRELATED_CAST_WARNINGS
-				UE_STATIC_ASSERT_WARN((std::is_base_of_v<From, To>), "Attempting to use Cast<> on types that are not related");
-#endif
-				if (((const UObject*)Src)->GetClass()->HasAnyCastFlag(UE::CoreUObject::Private::TCastFlags_V<To>))
-				{
-					return (To*)Src;
-				}
-			}
-		}
 		else
 		{
 			static_assert(std::is_base_of_v<UObject, From>, "Attempting to use Cast<> on a type that is not a UObject or an Interface");
 
-			if constexpr (TIsIInterface<To>::Value)
+			if constexpr (UE_USE_CAST_FLAGS && UE::CoreUObject::Private::TCastFlags_V<To> != CASTCLASS_None)
 			{
-				return (To*)((UObject*)Src)->GetInterfaceAddress(To::UClassType::StaticClass());
-			}
-			else if constexpr (std::is_base_of_v<To, From>)
-			{
-				return Src;
+				if constexpr (std::is_base_of_v<To, From>)
+				{
+					return (To*)Src;
+				}
+				else
+				{
+#if UE_ENABLE_UNRELATED_CAST_WARNINGS
+					UE_STATIC_ASSERT_WARN((std::is_base_of_v<From, To>), "Attempting to use Cast<> on types that are not related");
+#endif
+					if (((const UObject*)Src)->GetClass()->HasAnyCastFlag(UE::CoreUObject::Private::TCastFlags_V<To>))
+					{
+						return (To*)Src;
+					}
+				}
 			}
 			else
 			{
+				if constexpr (TIsIInterface<To>::Value)
+				{
+					return (To*)((UObject*)Src)->GetInterfaceAddress(To::UClassType::StaticClass());
+				}
+				else if constexpr (std::is_base_of_v<To, From>)
+				{
+					return Src;
+				}
+				else
+				{
 #if UE_ENABLE_UNRELATED_CAST_WARNINGS
-				UE_STATIC_ASSERT_WARN((std::is_base_of_v<From, To>), "Attempting to use Cast<> on types that are not related");
+					UE_STATIC_ASSERT_WARN((std::is_base_of_v<From, To>), "Attempting to use Cast<> on types that are not related");
 #endif
-				if (((const UObject*)Src)->IsA<To>())
-				{
-					return (To*)Src;
+					if (((const UObject*)Src)->IsA<To>())
+					{
+						return (To*)Src;
+					}
 				}
 			}
 		}
@@ -264,13 +267,18 @@ FORCEINLINE T* ExactCast( UObject* Src )
 				return (To*)Obj;
 			}
 		}
-		else if constexpr (TIsIInterface<To>::Value)
-		{
-			return (To*)((UObject*)Src)->GetInterfaceAddress(To::UClassType::StaticClass());
-		}
 		else
 		{
-			return (To*)Src;
+			static_assert(std::is_base_of_v<UObject, From>, "Attempting to use Cast<> on a type that is not a UObject or an Interface");
+
+			if constexpr (TIsIInterface<To>::Value)
+			{
+				return (To*)((UObject*)Src)->GetInterfaceAddress(To::UClassType::StaticClass());
+			}
+			else
+			{
+				return (To*)Src;
+			}
 		}
 	}
 
