Removed unused TPSOverrideStructName

#rnx

diff --git a/Engine/Source/Runtime/CoreUObject/Private/UObject/PropertyStruct.cpp b/Engine/Source/Runtime/CoreUObject/Private/UObject/PropertyStruct.cpp
--- a/Engine/Source/Runtime/CoreUObject/Private/UObject/PropertyStruct.cpp
+++ b/Engine/Source/Runtime/CoreUObject/Private/UObject/PropertyStruct.cpp
@@ -404,18 +404,6 @@ EConvertFromTypeResult FStructProperty::ConvertFromType(const FPropertyTag& Tag,
 			// Old Implementation
 			return StructProperty && !StructProperty->UseBinaryOrNativeSerialization(InAr);
 		}
-#if WITH_EDITORONLY_DATA
-		if (StructProperty && StructProperty->Struct)
-		{
-			if (const FString* TPSOverrideStructName = StructProperty->FindMetaData("TPSOverrideStructName"))
-			{
-				if (*TPSOverrideStructName == PropertyTag.StructName.ToString())
-				{
-					return true;
-				}
-			}
-		}
-#endif
 		return PropertyTag.StructGuid.IsValid() && StructProperty && StructProperty->Struct && (PropertyTag.StructGuid == StructProperty->Struct->GetCustomGuid());
 	};
 
