SavePackage: PKG_NotExternallyReferenceable imports validation improvement. Make the warning fire only for directly-referenced imports, and give the warning a different description than the warning for imports of !RF_Public objects.
#rnx

diff --git a/Engine/Source/Runtime/CoreUObject/Private/UObject/SavePackage/PackageHarvester.cpp b/Engine/Source/Runtime/CoreUObject/Private/UObject/SavePackage/PackageHarvester.cpp
--- a/Engine/Source/Runtime/CoreUObject/Private/UObject/SavePackage/PackageHarvester.cpp
+++ b/Engine/Source/Runtime/CoreUObject/Private/UObject/SavePackage/PackageHarvester.cpp
@@ -1048,6 +1048,10 @@ void FPackageHarvester::HarvestImport(TObjectPtr<UObject> InObject)
 	ForEachExportHarvestingRealm([this, InObject](ESaveRealm HarvestingRealm)
 		{
 			SaveContext.GetHarvestedRealm(HarvestingRealm).AddImport(InObject);
+			if (CurrentExportDependencies.ProcessImportDepth == 0)
+			{
+				SaveContext.GetHarvestedRealm(HarvestingRealm).AddDirectImport(InObject);
+			}
 		});
 }
 
diff --git a/Engine/Source/Runtime/CoreUObject/Private/UObject/SavePackage/SaveContext.h b/Engine/Source/Runtime/CoreUObject/Private/UObject/SavePackage/SaveContext.h
--- a/Engine/Source/Runtime/CoreUObject/Private/UObject/SavePackage/SaveContext.h
+++ b/Engine/Source/Runtime/CoreUObject/Private/UObject/SavePackage/SaveContext.h
@@ -135,6 +135,11 @@ struct FHarvestedRealm
 		}
 	}
 
+	void AddDirectImport(TObjectPtr<UObject> InObject)
+	{
+		DirectImports.Add(InObject);
+	}
+
 	void AddImport(TObjectPtr<UObject> InObject)
 	{
 		Imports.Add(InObject);
@@ -194,6 +199,11 @@ struct FHarvestedRealm
 		return Exports;
 	}
 
+	const TSet<TObjectPtr<UObject>>& GetDirectImports() const
+	{
+		return DirectImports;
+	}
+
 	const TSet<TObjectPtr<UObject>>& GetImports() const
 	{
 		return Imports;
@@ -375,6 +385,11 @@ private:
 	TSet<FTaggedExport> Exports;
 	// Set of objects marked as import
 	TSet<TObjectPtr<UObject>> Imports;
+	// Set of objects that were referenced directly from an export. Some imports are transitively added by
+	// FPackageHarvester::ProcessImport (Outer, Class, CDO, CDO subobjects, others?) for long-standing reasons
+	// (performance, loading behavior, other?)/ But some features such as allowed-import access warnings need
+	// to consider only the direct imports.
+	TSet<TObjectPtr<UObject>> DirectImports;
 	// Set of names referenced from export serialization
 	TSet<FNameEntryId> NamesReferencedFromExportData;
 	// Set of names referenced from the package header (import and export table object names etc)
@@ -773,6 +788,11 @@ public:
 		GetHarvestedRealm().AddImport(InObject);
 	}
 
+	void AddDirectImport(UObject* InObject)
+	{
+		GetHarvestedRealm().AddDirectImport(InObject);
+	}
+
 	void AddExport(FTaggedExport InTagObj)
 	{
 		GetHarvestedRealm().AddExport(MoveTemp(InTagObj));
@@ -808,6 +828,11 @@ public:
 		return GetHarvestedRealm().GetImports();
 	}
 
+	const TSet<TObjectPtr<UObject>>& GetDirectImports() const
+	{
+		return GetHarvestedRealm().GetDirectImports();
+	}
+
 	const TSet<TObjectPtr<UObject>>& GetImportsUsedInGame() const
 	{
 		return GetHarvestedRealm(ESaveRealm::Game).GetImports();
diff --git a/Engine/Source/Runtime/CoreUObject/Private/UObject/SavePackage/SavePackageUtilities.cpp b/Engine/Source/Runtime/CoreUObject/Private/UObject/SavePackage/SavePackageUtilities.cpp
--- a/Engine/Source/Runtime/CoreUObject/Private/UObject/SavePackage/SavePackageUtilities.cpp
+++ b/Engine/Source/Runtime/CoreUObject/Private/UObject/SavePackage/SavePackageUtilities.cpp
@@ -419,6 +419,19 @@ void FindMostLikelyCulprit(const TArray<UObject*>& BadObjects, UObject*& MostLik
 			OutReferencer = TEXT("Unknown property");
 		}
 	}
+
+	if (MostLikelyCulprit == nullptr)
+	{
+		// Make sure we report something
+		for (UObject* BadObject : BadObjects)
+		{
+			if (BadObject)
+			{
+				MostLikelyCulprit = BadObject;
+				break;
+			}
+		}
+	}
 }
 
 ESavePackageResult FinalizeTempOutputFiles(const FPackagePath& PackagePath, const FSavePackageOutputFileArray& OutputFiles, const FDateTime& FinalTimeStamp)
diff --git a/Engine/Source/Runtime/CoreUObject/Private/UObject/SavePackage2.cpp b/Engine/Source/Runtime/CoreUObject/Private/UObject/SavePackage2.cpp
--- a/Engine/Source/Runtime/CoreUObject/Private/UObject/SavePackage2.cpp
+++ b/Engine/Source/Runtime/CoreUObject/Private/UObject/SavePackage2.cpp
@@ -667,7 +667,8 @@ ESavePackageResult ValidateExports(FSaveContext& SaveContext)
 	return ReturnSuccessOrCancel();
 }
 
-ESavePackageResult ValidateIllegalReferences(FSaveContext& SaveContext, TArray<UObject*>& PrivateObjects, TArray<UObject*>& ObjectsInOtherMaps)
+ESavePackageResult ValidateIllegalReferences(FSaveContext& SaveContext, TArray<UObject*>& PrivateObjects,
+	TArray<UObject*>& PrivateContentObjects, TArray<UObject*>& ObjectsInOtherMaps)
 {
 	FFormatNamedArguments Args;
 
@@ -711,6 +712,25 @@ ESavePackageResult ValidateIllegalReferences(FSaveContext& SaveContext, TArray<U
 		}
 		return ESavePackageResult::Error;
 	}
+
+	if (PrivateContentObjects.Num() > 0)
+	{
+		UObject* MostLikelyCulprit = nullptr;
+		FString CulpritString = TEXT("Unknown");
+		FString Referencer;
+		UE::SavePackageUtilities::FindMostLikelyCulprit(PrivateContentObjects, MostLikelyCulprit, Referencer, &SaveContext);
+		CulpritString = FString::Printf(TEXT("%s (%s)"),
+			(MostLikelyCulprit != nullptr) ? *MostLikelyCulprit->GetFullName() : TEXT("(unknown culprit)"),
+			*Referencer);
+
+		if (SaveContext.IsGenerateSaveError())
+		{
+			SaveContext.GetError()->Logf(ELogVerbosity::Warning,
+				TEXT("Can't save package: PKG_NotExternallyReferenceable: Package %s imports object %s, which is in a different mount point and its package is marked as PKG_NotExternallyReferenceable."),
+				SaveContext.GetFilename(), *CulpritString);
+		}
+		return ESavePackageResult::Error;
+	}
 	return ReturnSuccessOrCancel();
 }
 
@@ -808,8 +828,10 @@ ESavePackageResult ValidateImports(FSaveContext& SaveContext)
 
 	// Warn for private objects & map object references
 	TArray<UObject*> PrivateObjects;
+	TArray<UObject*> PrivateContentObjects;
 	TArray<UObject*> ObjectsInOtherMaps;
 	const TSet<TObjectPtr<UObject>>& Imports = SaveContext.GetImports();
+	const TSet<TObjectPtr<UObject>>& DirectImports = SaveContext.GetDirectImports();
 	for (const TObjectPtr<UObject>& Import : Imports)
 	{
 		TObjectPtr<UPackage> ImportPackage = Import.GetPackage();
@@ -877,8 +899,10 @@ ESavePackageResult ValidateImports(FSaveContext& SaveContext)
 			}
 		}
 
-		// Enforce that Private content can only be referenced by something within the same Mount Point
-		if (!ImportPackage->IsExternallyReferenceable())
+		// Enforce that Private content can only be directly referenced by something within the same Mount Point
+		// This only applies to direct imports. Transitive imports (A -> PackageB.B, B -> B.Class -> PackageC) are
+		// allowed even for private content.
+		if (!ImportPackage->IsExternallyReferenceable() && DirectImports.Contains(Import))
 		{
 			FName MountPointName = FPackageName::GetPackageMountPoint(PackageName);
 
@@ -888,18 +912,18 @@ ESavePackageResult ValidateImports(FSaveContext& SaveContext)
 			{
 				if (MountPointName != ImportMountPointName)
 				{
-					PrivateObjects.Add(Import);
+					PrivateContentObjects.Add(Import);
 				}
 			}
 			else
 			{
-				PrivateObjects.Add(Import);
+				PrivateContentObjects.Add(Import);
 			}
 		}
 	}
-	if (PrivateObjects.Num() > 0 || ObjectsInOtherMaps.Num() > 0)
+	if (PrivateObjects.Num() > 0 || PrivateContentObjects.Num() > 0 || ObjectsInOtherMaps.Num() > 0)
 	{
-		return ValidateIllegalReferences(SaveContext, PrivateObjects, ObjectsInOtherMaps);
+		return ValidateIllegalReferences(SaveContext, PrivateObjects, PrivateContentObjects, ObjectsInOtherMaps);
 	}
 
 	// Validate External Import Rules
