Disable DoesPackageExist AssetRegistry optimization as it's breaking Content Worker jobs. To be followed up by UE-206417

diff --git a/Engine/Source/Runtime/CoreUObject/Private/Misc/PackageName.cpp b/Engine/Source/Runtime/CoreUObject/Private/Misc/PackageName.cpp
--- a/Engine/Source/Runtime/CoreUObject/Private/Misc/PackageName.cpp
+++ b/Engine/Source/Runtime/CoreUObject/Private/Misc/PackageName.cpp
@@ -1845,7 +1845,7 @@ FPackageName::EPackageLocationFilter FPackageName::DoesPackageExistEx(const FPac
 		return EPackageLocationFilter::None;
 	}
 
-#if WITH_EDITOR
+#if 0 // WITH_EDITOR - Temporarily disabled until UE-206417 is resolved
 	IAssetRegistryInterface* AssetRegistry = IAssetRegistryInterface::GetPtr();
 
 	// Todo: The AssetRegistry currently cannot determine if a package comes from the Filesystem 
