Adding file missing from my last submitted changelist.

diff --git a/Engine/Source/Runtime/CoreUObject/Public/Cooker/MPCollector.h b/Engine/Source/Runtime/CoreUObject/Public/Cooker/MPCollector.h
--- a/Engine/Source/Runtime/CoreUObject/Public/Cooker/MPCollector.h
+++ b/Engine/Source/Runtime/CoreUObject/Public/Cooker/MPCollector.h
@@ -45,7 +45,7 @@ public:
 	bool IsLocal() const { return Id == LocalId; }
 	bool IsRemote() const { return Id != InvalidId && Id != LocalId; }
 	uint8 GetRemoteIndex() const { check(IsRemote()); return Id - 1U; }
-	uint8 GetLocalOrRemoteIndex() const { check(IsValid()); return Id; }
+	uint8 GetMultiprocessId() const { check(IsValid()); return Id; }
 	bool operator==(const FWorkerId& Other) const { return Id == Other.Id; }
 	bool operator!=(const FWorkerId& Other) const { return Id != Other.Id; }
 	bool operator<(const FWorkerId& Other) const { return Id < Other.Id; }
