Fixed operator< for FPropertyPathName

#rnx

diff --git a/Engine/Source/Runtime/CoreUObject/Private/UObject/PropertyPathName.cpp b/Engine/Source/Runtime/CoreUObject/Private/UObject/PropertyPathName.cpp
--- a/Engine/Source/Runtime/CoreUObject/Private/UObject/PropertyPathName.cpp
+++ b/Engine/Source/Runtime/CoreUObject/Private/UObject/PropertyPathName.cpp
@@ -24,7 +24,7 @@ inline int32 FPropertyPathName::FSegment::Compare(const FSegment& Segment) const
 	{
 		return CompareNameWithIndex;
 	}
-	return Type < Segment.Type;
+	return (Type == Segment.Type) ? 0 : (Type < Segment.Type ? -1 : 1);
 }
 
 bool FPropertyPathName::operator==(const FPropertyPathName& Path) const
diff --git a/Engine/Source/Runtime/CoreUObject/Tests/UObject/PropertyPathNameTest.cpp b/Engine/Source/Runtime/CoreUObject/Tests/UObject/PropertyPathNameTest.cpp
--- a/Engine/Source/Runtime/CoreUObject/Tests/UObject/PropertyPathNameTest.cpp
+++ b/Engine/Source/Runtime/CoreUObject/Tests/UObject/PropertyPathNameTest.cpp
@@ -14,9 +14,9 @@ TEST_CASE_NAMED(FPropertyPathNameTest, "CoreUObject::PropertyPathName", "[Core][
 {
 	const FName CountName(TEXTVIEW("Count"));
 
-	FPropertyTypeNameBuilder IntTypeBuilder;
-	IntTypeBuilder.AddTypeName(NAME_IntProperty);
-	const FPropertyTypeName IntType = IntTypeBuilder.Build();
+	FPropertyTypeNameBuilder TypeBuilder;
+	TypeBuilder.AddTypeName(NAME_IntProperty);
+	const FPropertyTypeName IntType = TypeBuilder.Build();
 
 	SECTION("Empty")
 	{
@@ -140,6 +140,49 @@ TEST_CASE_NAMED(FPropertyPathNameTest, "CoreUObject::PropertyPathName", "[Core][
 		}
 	}
 
+	SECTION("Equals+Less")
+	{
+		const FName DepthName(TEXTVIEW("Depth"));
+		
+		TypeBuilder.Reset();
+		TypeBuilder.AddTypeName(NAME_BoolProperty);
+		const FPropertyTypeName BoolType = TypeBuilder.Build();
+
+		FPropertyPathName CountBool;
+		CountBool.Push({CountName, BoolType});
+		FPropertyPathName CountInt;
+		CountInt.Push({CountName, IntType});
+		FPropertyPathName DepthBool;
+		DepthBool.Push({DepthName, BoolType});
+		FPropertyPathName DepthInt;
+		DepthInt.Push({DepthName, IntType});
+
+		CHECK(CountBool == CountBool);
+		CHECK_FALSE(CountBool == CountInt);
+		CHECK_FALSE(CountBool == DepthBool);
+		CHECK_FALSE(CountBool == DepthInt);
+
+		CHECK_FALSE(CountBool < CountBool);
+		CHECK(CountBool < CountInt);
+		CHECK(CountBool < DepthBool);
+		CHECK(CountBool < DepthInt);
+
+		CHECK_FALSE(CountInt < CountBool);
+		CHECK_FALSE(CountInt < CountInt);
+		CHECK(CountInt < DepthBool);
+		CHECK(CountInt < DepthInt);
+
+		CHECK_FALSE(DepthBool < CountBool);
+		CHECK_FALSE(DepthBool < CountInt);
+		CHECK_FALSE(DepthBool < DepthBool);
+		CHECK(DepthBool < DepthInt);
+
+		CHECK_FALSE(DepthInt < CountBool);
+		CHECK_FALSE(DepthInt < CountInt);
+		CHECK_FALSE(DepthInt < DepthBool);
+		CHECK_FALSE(DepthInt < DepthInt);
+	}
+
 	SECTION("GetTypeHash")
 	{
 		FPropertyPathName Empty, Name, NameType, NameTypeIndex;
