Renamed n to ElementCount in FArrayProperty and fixed whitespace and NULL

#rnx

diff --git a/Engine/Source/Runtime/CoreUObject/Private/UObject/PropertyArray.cpp b/Engine/Source/Runtime/CoreUObject/Private/UObject/PropertyArray.cpp
--- a/Engine/Source/Runtime/CoreUObject/Private/UObject/PropertyArray.cpp
+++ b/Engine/Source/Runtime/CoreUObject/Private/UObject/PropertyArray.cpp
@@ -66,7 +66,7 @@ void FArrayProperty::PostDuplicate(const FField& InField)
 void FArrayProperty::LinkInternal(FArchive& Ar)
 {
 	//FLinkerLoad* MyLinker = GetLinker();
-	//if( MyLinker )
+	//if (MyLinker)
 	//{
 	//	MyLinker->Preload(this);
 	//}
@@ -75,27 +75,27 @@ void FArrayProperty::LinkInternal(FArchive& Ar)
 
 	SetElementSize();
 }
-bool FArrayProperty::Identical( const void* A, const void* B, uint32 PortFlags ) const
+bool FArrayProperty::Identical(const void* A, const void* B, uint32 PortFlags) const
 {
 	checkSlow(Inner);
 
 	FScriptArrayHelper ArrayHelperA(this, A);
 
 	const int32 ArrayNum = ArrayHelperA.Num();
-	if ( B == NULL )
+	if (B == nullptr)
 	{
 		return ArrayNum == 0;
 	}
 
 	FScriptArrayHelper ArrayHelperB(this, B);
-	if ( ArrayNum != ArrayHelperB.Num() )
+	if (ArrayNum != ArrayHelperB.Num())
 	{
 		return false;
 	}
 
-	for ( int32 ArrayIndex = 0; ArrayIndex < ArrayNum; ArrayIndex++ )
+	for (int32 ArrayIndex = 0; ArrayIndex < ArrayNum; ArrayIndex++)
 	{
-		if ( !Inner->Identical( ArrayHelperA.GetRawPtr(ArrayIndex), ArrayHelperB.GetRawPtr(ArrayIndex), PortFlags) )
+		if (!Inner->Identical(ArrayHelperA.GetRawPtr(ArrayIndex), ArrayHelperB.GetRawPtr(ArrayIndex), PortFlags))
 		{
 			return false;
 		}
@@ -132,7 +132,7 @@ void FArrayProperty::SerializeItem(FStructuredArchive::FSlot Slot, void* Value,
 	//UnderlyingArchive.Preload(Inner);
 
 	FScriptArrayHelper ArrayHelper(this, Value);
-	int32		n		= ArrayHelper.Num();
+	int32 ElementCount = ArrayHelper.Num();
 
 	// Custom branch for UPS to try and take advantage of bulk serialization
 	if (bUPS && !bExperimentalOverridableLogic)
@@ -146,26 +146,26 @@ void FArrayProperty::SerializeItem(FStructuredArchive::FSlot Slot, void* Value,
 			// but which maps down to straight writes to the underlying archive.
 			FStructuredArchiveStream Stream = Slot.EnterStream();
 
-			Stream.EnterElement() << n;
+			Stream.EnterElement() << ElementCount;
 
 			if (UnderlyingArchive.IsLoading())
 			{
-				ArrayHelper.EmptyAndAddUninitializedValues(n);
+				ArrayHelper.EmptyAndAddUninitializedValues(ElementCount);
 			}
 
-			Stream.EnterElement().Serialize(ArrayHelper.GetRawPtr(), n * Inner->ElementSize);
+			Stream.EnterElement().Serialize(ArrayHelper.GetRawPtr(), ElementCount * Inner->ElementSize);
 		}
 		else
 		{
-			FStructuredArchiveArray Array = Slot.EnterArray(n);
+			FStructuredArchiveArray Array = Slot.EnterArray(ElementCount);
 
 			if (UnderlyingArchive.IsLoading())
 			{
-				ArrayHelper.EmptyAndAddValues(n);
+				ArrayHelper.EmptyAndAddValues(ElementCount);
 			}
 
 			FSerializedPropertyScope SerializedProperty(UnderlyingArchive, Inner, this);
-			for (int32 i = 0; i < n; ++i)
+			for (int32 i = 0; i < ElementCount; ++i)
 			{
 #if WITH_EDITOR
 				static const FName NAME_UArraySerialize = FName(TEXT("FArrayProperty::Serialize"));
@@ -182,7 +182,7 @@ void FArrayProperty::SerializeItem(FStructuredArchive::FSlot Slot, void* Value,
 
 	if (bIsTextFormat && Inner->IsA<FStructProperty>())
 	{
-		MaybeInnerTag.Emplace(UnderlyingArchive, Inner, 0, (uint8*)Value, (uint8*)Defaults);	
+		MaybeInnerTag.Emplace(UnderlyingArchive, Inner, 0, (uint8*)Value, (uint8*)Defaults);
 		Slot << SA_ATTRIBUTE(TEXT("InnerStructName"), MaybeInnerTag.GetValue().StructName);
 		Slot << SA_OPTIONAL_ATTRIBUTE(TEXT("InnerStructGuid"), MaybeInnerTag.GetValue().StructGuid, FGuid());
 	}
@@ -192,11 +192,11 @@ void FArrayProperty::SerializeItem(FStructuredArchive::FSlot Slot, void* Value,
 	// TODO: Should work for maps + sets too.
 	auto SerializeContainerItem = [this, &SerializeFromMismatchedTag](FStructuredArchiveSlot Slot, uint8* Item)
 	{
-		if(SerializeFromMismatchedTag.IsSet())
+		if (const FPropertyTag* Tag = SerializeFromMismatchedTag.GetPtrOrNull())
 		{
 			const int64 StartOfProperty = Slot.GetUnderlyingArchive().Tell();
 			FStructProperty* StructProperty = CastFieldChecked<FStructProperty>(Inner);
-			switch(StructProperty->ConvertFromType(SerializeFromMismatchedTag.GetValue(), Slot, Item, nullptr, nullptr))
+			switch (StructProperty->ConvertFromType(*Tag, Slot, Item, nullptr, nullptr))
 			{
 				case EConvertFromTypeResult::Converted:
 				case EConvertFromTypeResult::Serialized:
@@ -204,17 +204,17 @@ void FArrayProperty::SerializeItem(FStructuredArchive::FSlot Slot, void* Value,
 				case EConvertFromTypeResult::CannotConvert:
 					// FStructProperty::ConvertFromType doesn't handle setting the default, so do it here.
 					StructProperty->Struct->InitializeDefaultValue(Item);
-					Slot.GetUnderlyingArchive().Seek(StartOfProperty + SerializeFromMismatchedTag.GetValue().Size);	// Skip this item
+					Slot.GetUnderlyingArchive().Seek(StartOfProperty + Tag->Size);	// Skip this item
 					return;
 				case EConvertFromTypeResult::UseSerializeItem:
 					// Fall through to default serialize
 					break;
 			}
 		}
-			
+
 		Inner->SerializeItem(Slot, Item);
 	};
-	
+
 	// Make sure the container is reloading accordingly to the value set in the property tag if any
 	if (!bUPS && UnderlyingArchive.IsLoading() && FPropertyTagScope::GetCurrentPropertyTag())
 	{
@@ -222,7 +222,7 @@ void FArrayProperty::SerializeItem(FStructuredArchive::FSlot Slot, void* Value,
 	}
 
 	// *** Experimental *** Special serialization path for array with overridable serialization
-	if(bExperimentalOverridableLogic)
+	if (bExperimentalOverridableLogic)
 	{
 		checkf(!UnderlyingArchive.ArUseCustomPropertyList, TEXT("Using custom property list is not supported by overridable serialization"));
 
@@ -280,7 +280,7 @@ void FArrayProperty::SerializeItem(FStructuredArchive::FSlot Slot, void* Value,
 
 				int32 NumRemoved = 0;
 				FStructuredArchive::FArray RemovedArray = Record.EnterArray(TEXT("Removed"), NumRemoved);
-				if(NumRemoved != 0)
+				if (NumRemoved != 0)
 				{
 					TArray<int32> IndicesToRemove;
 					TempValueStorage = (uint8*)FMemory::Malloc(InnerObjectProperty->ElementSize);
@@ -293,7 +293,7 @@ void FArrayProperty::SerializeItem(FStructuredArchive::FSlot Slot, void* Value,
 							SerializeContainerItem(RemovedArray.EnterElement(), TempValueStorage);
 						}
 
-						if (UObject* RemovedSubObject = InnerObjectProperty->GetObjectPropertyValue(TempValueStorage) )
+						if (UObject* RemovedSubObject = InnerObjectProperty->GetObjectPropertyValue(TempValueStorage))
 						{
 							int32 Index = FindObject(RemovedSubObject, nullptr, DefaultsArrayHelper);
 							if (Index != INDEX_NONE)
@@ -312,7 +312,7 @@ void FArrayProperty::SerializeItem(FStructuredArchive::FSlot Slot, void* Value,
 					}
 
 					IndicesToRemove.Sort(TGreater<>());
-					for(int32 IndexToRemove : IndicesToRemove)
+					for (int32 IndexToRemove : IndicesToRemove)
 					{
 						ArrayHelper.RemoveValues(IndexToRemove);
 					}
@@ -465,7 +465,7 @@ void FArrayProperty::SerializeItem(FStructuredArchive::FSlot Slot, void* Value,
 								case EOverriddenPropertyOperation::Add:
 									{
 										const int32 Index = FindObject(Pair.Key, ArrayHelper);
-										if(Index != INDEX_NONE)
+										if (Index != INDEX_NONE)
 										{
 											AddedIndices.Add(Index);
 											ModifiedIndices.Remove(Index);
@@ -484,10 +484,10 @@ void FArrayProperty::SerializeItem(FStructuredArchive::FSlot Slot, void* Value,
 
 			int32 NumReplaced = bReplaceArray ? ArrayHelper.Num() : INDEX_NONE;
 			FStructuredArchive::FArray ReplacedArray = Record.EnterArray(TEXT("Replaced"), NumReplaced);
-			if(bReplaceArray)
+			if (bReplaceArray)
 			{
 				const int32 ArrayNum = ArrayHelper.Num();
-				for (int32 i =0; i < ArrayNum; ++i)
+				for (int32 i = 0; i < ArrayNum; ++i)
 				{
 					SerializeContainerItem(ReplacedArray.EnterElement(), ArrayHelper.GetRawPtr(i));
 				}
@@ -522,30 +522,30 @@ void FArrayProperty::SerializeItem(FStructuredArchive::FSlot Slot, void* Value,
 		return;
 	}
 
-	FStructuredArchiveArray Array = Slot.EnterArray(n);
+	FStructuredArchiveArray Array = Slot.EnterArray(ElementCount);
 
-	if( UnderlyingArchive.IsLoading() )
+	if (UnderlyingArchive.IsLoading())
 	{
 		// If using a custom property list, don't empty the array on load. Not all indices may have been serialized, so we need to preserve existing values at those slots.
 		if (UnderlyingArchive.ArUseCustomPropertyList)
 		{
 			const int32 OldNum = ArrayHelper.Num();
-			if (n > OldNum)
+			if (ElementCount > OldNum)
 			{
-				ArrayHelper.AddValues(n - OldNum);
+				ArrayHelper.AddValues(ElementCount - OldNum);
 			}
-			else if (n < OldNum)
+			else if (ElementCount < OldNum)
 			{
-				ArrayHelper.RemoveValues(n, OldNum - n);
+				ArrayHelper.RemoveValues(ElementCount, OldNum - ElementCount);
 			}
 		}
 		else
 		{
-			ArrayHelper.EmptyAndAddValues(n);
+			ArrayHelper.EmptyAndAddValues(ElementCount);
 		}
 	}
-	ArrayHelper.CountBytes( UnderlyingArchive );
-	
+	ArrayHelper.CountBytes(UnderlyingArchive);
+
 	// Serialize a PropertyTag for the inner property of this array, allows us to validate the inner struct to see if it has changed
 	if (UnderlyingArchive.UEVer() >= VER_UE4_INNER_ARRAY_TAG_INFO && Inner->IsA<FStructProperty>())
 	{
@@ -593,7 +593,7 @@ void FArrayProperty::SerializeItem(FStructuredArchive::FSlot Slot, void* Value,
 
 #if WITH_EDITOR
 					// Ensure the structure is initialized
-					for (int32 i = 0; i < n; i++)
+					for (int32 i = 0; i < ElementCount; i++)
 					{
 						StructProperty->Struct->InitializeDefaultValue(ArrayHelper.GetRawPtr(i));
 					}
@@ -656,7 +656,7 @@ void FArrayProperty::SerializeItem(FStructuredArchive::FSlot Slot, void* Value,
 		const FCustomPropertyListNode* CustomPropertyList = UnderlyingArchive.ArCustomPropertyList;
 		const FCustomPropertyListNode* PropertyNode = CustomPropertyList;
 		FSerializedPropertyScope SerializedProperty(UnderlyingArchive, Inner, this);
-		while (PropertyNode && i < n && !bSerializeRemainingItems)
+		while (PropertyNode && i < ElementCount && !bSerializeRemainingItems)
 		{
 			if (PropertyNode->Property != Inner)
 			{
@@ -681,7 +681,7 @@ void FArrayProperty::SerializeItem(FStructuredArchive::FSlot Slot, void* Value,
 				if (Context && Context->bTrackSerializedPropertyPath)
 				{
 					Context->SerializedPropertyPath.SetIndex(i);
-					
+
 					// broadcast that a property will be serialized
 					Context->OnTaggedPropertySerialize.Broadcast(*Context);
 				}
@@ -701,7 +701,7 @@ void FArrayProperty::SerializeItem(FStructuredArchive::FSlot Slot, void* Value,
 
 		// Serialize each item until we get to the end of the array
 		FSerializedPropertyScope SerializedProperty(UnderlyingArchive, Inner, this);
-		while (i < n)
+		while (i < ElementCount)
 		{
 #if WITH_EDITOR
 			static const FName NAME_UArraySerialize = FName(TEXT("FArrayProperty::Serialize"));
@@ -712,7 +712,7 @@ void FArrayProperty::SerializeItem(FStructuredArchive::FSlot Slot, void* Value,
 			if (Context && Context->bTrackSerializedPropertyPath)
 			{
 				Context->SerializedPropertyPath.SetIndex(i);
-					
+
 				// broadcast that a property will be serialized
 				Context->OnTaggedPropertySerialize.Broadcast(*Context);
 			}
@@ -750,16 +750,16 @@ void FArrayProperty::SerializeItem(FStructuredArchive::FSlot Slot, void* Value,
 	}
 }
 
-bool FArrayProperty::NetSerializeItem( FArchive& Ar, UPackageMap* Map, void* Data, TArray<uint8> * MetaData ) const
+bool FArrayProperty::NetSerializeItem(FArchive& Ar, UPackageMap* Map, void* Data, TArray<uint8> * MetaData) const
 {
-	UE_LOG( LogProperty, Fatal, TEXT( "Deprecated code path" ) );
+	UE_LOG(LogProperty, Fatal, TEXT("Deprecated code path"));
 	return 1;
 }
 
-void FArrayProperty::Serialize( FArchive& Ar )
+void FArrayProperty::Serialize(FArchive& Ar)
 {
 	Super::Serialize(Ar);
-	
+
 	SerializeSingleField(Ar, Inner, this);
 	checkSlow(Inner);
 }
@@ -774,7 +774,7 @@ void FArrayProperty::AddReferencedObjects(FReferenceCollector& Collector)
 
 FString FArrayProperty::GetCPPTypeCustom(FString* ExtendedTypeText, uint32 CPPExportFlags, const FString& InnerTypeText, const FString& InInnerExtendedTypeText) const
 {
-	if (ExtendedTypeText != NULL)
+	if (ExtendedTypeText != nullptr)
 	{
 		FString InnerExtendedTypeText = InInnerExtendedTypeText;
 		if (InnerExtendedTypeText.Len() && InnerExtendedTypeText.Right(1) == TEXT(">"))
@@ -792,12 +792,12 @@ FString FArrayProperty::GetCPPTypeCustom(FString* ExtendedTypeText, uint32 CPPEx
 	return TEXT("TArray");
 }
 
-FString FArrayProperty::GetCPPType( FString* ExtendedTypeText/*=NULL*/, uint32 CPPExportFlags/*=0*/ ) const
+FString FArrayProperty::GetCPPType(FString* ExtendedTypeText/*=nullptr*/, uint32 CPPExportFlags/*=0*/) const
 {
 	checkSlow(Inner);
 	FString InnerExtendedTypeText;
 	FString InnerTypeText;
-	if ( ExtendedTypeText != NULL )
+	if (ExtendedTypeText != nullptr)
 	{
 		InnerTypeText = Inner->GetCPPType(&InnerExtendedTypeText, CPPExportFlags & ~CPPF_ArgumentOrReturnValue); // we won't consider array inners to be "arguments or return values"
 	}
@@ -812,13 +812,13 @@ FString FArrayProperty::GetCPPTypeForwardDeclaration() const
 }
 PRAGMA_ENABLE_DEPRECATION_WARNINGS
 
-FString FArrayProperty::GetCPPMacroType( FString& ExtendedTypeText ) const
+FString FArrayProperty::GetCPPMacroType(FString& ExtendedTypeText) const
 {
 	checkSlow(Inner);
 	ExtendedTypeText = Inner->GetCPPType();
 	return TEXT("TARRAY");
 }
-void FArrayProperty::ExportText_Internal( FString& ValueStr, const void* ContainerOrPropertyPtr, EPropertyPointerType PropertyPointerType, const void* DefaultValue, UObject* Parent, int32 PortFlags, UObject* ExportRootScope ) const
+void FArrayProperty::ExportText_Internal(FString& ValueStr, const void* ContainerOrPropertyPtr, EPropertyPointerType PropertyPointerType, const void* DefaultValue, UObject* Parent, int32 PortFlags, UObject* ExportRootScope) const
 {
 	checkSlow(Inner);
 
@@ -858,7 +858,7 @@ void FArrayProperty::ExportTextInnerItem(FString& ValueStr, const FProperty* Inn
 {
 	checkSlow(Inner);
 
-	uint8* StructDefaults = NULL;
+	uint8* StructDefaults = nullptr;
 	const FStructProperty* StructProperty = CastField<FStructProperty>(Inner);
 
 	const bool bReadableForm = (0 != (PPF_BlueprintDebugView & PortFlags));
@@ -868,7 +868,7 @@ void FArrayProperty::ExportTextInnerItem(FString& ValueStr, const FProperty* Inn
 	// so we export the entire struct to ensure all data is copied over correctly. Behavior is currently inconsistent when copy/pasting between the two types.
 	// In the future, static arrays could export diffs if the property being imported to is reset to default before the import.
 	// When exporting to an external editor, we want to save defaults so all information is available for editing
-	if ( StructProperty != NULL && Inner->ArrayDim == 1 && !bExternalEditor )
+	if (StructProperty != nullptr && Inner->ArrayDim == 1 && !bExternalEditor)
 	{
 		checkSlow(StructProperty->Struct);
 		StructDefaults = (uint8*)FMemory::Malloc(StructProperty->Struct->GetStructureSize() * Inner->ArrayDim);
@@ -876,12 +876,12 @@ void FArrayProperty::ExportTextInnerItem(FString& ValueStr, const FProperty* Inn
 	}
 
 	int32 Count = 0;
-	for( int32 i=0; i<PropertySize; i++ )
+	for (int32 i = 0; i < PropertySize; i++)
 	{
 		++Count;
-		if(!bReadableForm)
+		if (!bReadableForm)
 		{
-			if ( Count == 1 )
+			if (Count == 1)
 			{
 				ValueStr += TCHAR('(');
 			}
@@ -892,7 +892,7 @@ void FArrayProperty::ExportTextInnerItem(FString& ValueStr, const FProperty* Inn
 		}
 		else
 		{
-			if(Count > 1)
+			if (Count > 1)
 			{
 				ValueStr += TCHAR('\n');
 			}
@@ -919,7 +919,7 @@ void FArrayProperty::ExportTextInnerItem(FString& ValueStr, const FProperty* Inn
 			}
 		}
 
-		Inner->ExportTextItem_Direct( ValueStr, PropData, PropDefault, Parent, PortFlags|PPF_Delimited, ExportRootScope );
+		Inner->ExportTextItem_Direct(ValueStr, PropData, PropDefault, Parent, PortFlags|PPF_Delimited, ExportRootScope);
 	}
 
 	if ((Count > 0) && !bReadableForm)
@@ -964,7 +964,7 @@ const TCHAR* FArrayProperty::ImportText_Internal(const TCHAR* Buffer, void* Cont
 	return ImportTextInnerItem(Buffer, Inner, ArrayPtr, PortFlags, OwnerObject, &ArrayHelper, ErrorText);
 }
 
-const TCHAR* FArrayProperty::ImportTextInnerItem( const TCHAR* Buffer, const FProperty* Inner, void* Data, int32 PortFlags, UObject* Parent, FScriptArrayHelper* ArrayHelper, FOutputDevice* ErrorText )
+const TCHAR* FArrayProperty::ImportTextInnerItem(const TCHAR* Buffer, const FProperty* Inner, void* Data, int32 PortFlags, UObject* Parent, FScriptArrayHelper* ArrayHelper, FOutputDevice* ErrorText)
 {
 	checkSlow(Inner);
 
@@ -979,9 +979,9 @@ const TCHAR* FArrayProperty::ImportTextInnerItem( const TCHAR* Buffer, const FPr
 		return Buffer;
 	}
 
-	if ( *Buffer++ != TCHAR('(') )
+	if (*Buffer++ != TCHAR('('))
 	{
-		return NULL;
+		return nullptr;
 	}
 
 	if (ArrayHelper)
@@ -1014,9 +1014,9 @@ const TCHAR* FArrayProperty::ImportTextInnerItem( const TCHAR* Buffer, const FPr
 			checkf(ArrayHelper == nullptr || Inner->GetOffset_ForInternal() == 0, TEXT("Expected the Inner property of the FArrayProperty."));
 			Buffer = Inner->ImportText_Direct(Buffer, Address, Parent, PortFlags | PPF_Delimited, ErrorText);
 
-			if(!Buffer)
+			if (!Buffer)
 			{
-				return NULL;
+				return nullptr;
 			}
 
 			SkipWhitespace(Buffer);
@@ -1046,7 +1046,7 @@ const TCHAR* FArrayProperty::ImportTextInnerItem( const TCHAR* Buffer, const FPr
 	// Make sure we ended on a )
 	if (*Buffer++ != TCHAR(')'))
 	{
-		return NULL;
+		return nullptr;
 	}
 
 	return Buffer;
@@ -1071,14 +1071,14 @@ void FArrayProperty::AddCppProperty(FProperty* Property)
 	Inner = Property;
 }
 
-void FArrayProperty::CopyValuesInternal( void* Dest, void const* Src, int32 Count  ) const
+void FArrayProperty::CopyValuesInternal(void* Dest, void const* Src, int32 Count) const
 {
-	check(Count==1); // this was never supported, apparently
+	check(Count == 1); // this was never supported, apparently
 	FScriptArrayHelper SrcArrayHelper(this, Src);
 	FScriptArrayHelper DestArrayHelper(this, Dest);
 
 	int32 Num = SrcArrayHelper.Num();
-	if ( !(Inner->PropertyFlags & CPF_IsPlainOldData) )
+	if (!(Inner->PropertyFlags & CPF_IsPlainOldData))
 	{
 		DestArrayHelper.EmptyAndAddValues(Num);
 	}
@@ -1091,25 +1091,25 @@ void FArrayProperty::CopyValuesInternal( void* Dest, void const* Src, int32 Coun
 		size_t Size = Inner->ElementSize;
 		uint8* SrcData = (uint8*)SrcArrayHelper.GetRawPtr();
 		uint8* DestData = (uint8*)DestArrayHelper.GetRawPtr();
-		if( !(Inner->PropertyFlags & CPF_IsPlainOldData) )
+		if (!(Inner->PropertyFlags & CPF_IsPlainOldData))
 		{
-			for( int32 i=0; i<Num; i++ )
+			for (int32 i = 0; i < Num; i++)
 			{
-				Inner->CopyCompleteValue( DestData + i * Size, SrcData + i * Size );
+				Inner->CopyCompleteValue(DestData + i * Size, SrcData + i * Size);
 			}
 		}
 		else
 		{
-			FMemory::Memcpy( DestData, SrcData, Num*Size );
+			FMemory::Memcpy(DestData, SrcData, Num * Size);
 		}
 	}
 }
-void FArrayProperty::ClearValueInternal( void* Data ) const
+void FArrayProperty::ClearValueInternal(void* Data) const
 {
 	FScriptArrayHelper ArrayHelper(this, Data);
 	ArrayHelper.EmptyValues();
 }
-void FArrayProperty::DestroyValueInternal( void* Dest ) const
+void FArrayProperty::DestroyValueInternal(void* Dest) const
 {
 	FScriptArrayHelper ArrayHelper(this, Dest);
 	ArrayHelper.EmptyValues();
@@ -1124,15 +1124,15 @@ bool FArrayProperty::PassCPPArgsByRef() const
 
 /**
  * Creates new copies of components
- * 
+ *
  * @param	Data				pointer to the address of the instanced object referenced by this UComponentProperty
  * @param	DefaultData			pointer to the address of the default value of the instanced object referenced by this UComponentProperty
  * @param	Owner				the object that contains this property's data
  * @param	InstanceGraph		contains the mappings of instanced objects and components to their templates
  */
-void FArrayProperty::InstanceSubobjects( void* Data, void const* DefaultData, UObject* InOwner, FObjectInstancingGraph* InstanceGraph )
+void FArrayProperty::InstanceSubobjects(void* Data, void const* DefaultData, UObject* InOwner, FObjectInstancingGraph* InstanceGraph)
 {
-	if( Data && Inner->ContainsInstancedObjectProperty())
+	if (Data && Inner->ContainsInstancedObjectProperty())
 	{
 		FScriptArrayHelper ArrayHelper(this, Data);
 		FScriptArrayHelper DefaultArrayHelper(this, DefaultData);
@@ -1140,11 +1140,11 @@ void FArrayProperty::InstanceSubobjects( void* Data, void const* DefaultData, UO
 		int32 InnerElementSize = Inner->ElementSize;
 		void* TempElement = FMemory_Alloca(InnerElementSize);
 
-		for( int32 ElementIndex = 0; ElementIndex < ArrayHelper.Num(); ElementIndex++ )
+		for (int32 ElementIndex = 0; ElementIndex < ArrayHelper.Num(); ElementIndex++)
 		{
 			uint8* DefaultValue = (DefaultData && ElementIndex < DefaultArrayHelper.Num()) ? DefaultArrayHelper.GetRawPtr(ElementIndex) : nullptr;
 			FMemory::Memmove(TempElement, ArrayHelper.GetRawPtr(ElementIndex), InnerElementSize);
-			Inner->InstanceSubobjects( TempElement, DefaultValue, InOwner, InstanceGraph );
+			Inner->InstanceSubobjects(TempElement, DefaultValue, InOwner, InstanceGraph);
 			if (ElementIndex < ArrayHelper.Num())
 			{
 				FMemory::Memmove(ArrayHelper.GetRawPtr(ElementIndex), TempElement, InnerElementSize);
@@ -1213,7 +1213,7 @@ EConvertFromTypeResult FArrayProperty::ConvertFromType(const FPropertyTag& Tag,
 			// TODO: Implement SerializeFromMismatchedTag handling for arrays of structs
 			else
 			{
-				UE_LOG(LogClass, Warning, TEXT("Array Inner Type mismatch in %s of %s - Previous (%s) Current(%s) for package:  %s"), *Tag.Name.ToString(), *GetName(), *Tag.InnerType.ToString(), *Inner->GetID().ToString(), *Slot.GetUnderlyingArchive().GetArchiveName() );
+				UE_LOG(LogClass, Warning, TEXT("Array Inner Type mismatch in %s of %s - Previous (%s) Current(%s) for package:  %s"), *Tag.Name.ToString(), *GetName(), *Tag.InnerType.ToString(), *Inner->GetID().ToString(), *Slot.GetUnderlyingArchive().GetArchiveName());
 				return EConvertFromTypeResult::CannotConvert;
 			}
 		}
