[UE]: More revisions to placeholder property bag type import/export serialization support during map asset loads in the editor.

Change summary:
- Modified ObjectPtr_Private::IsObjectPtrEqualToRawPtrOfRelatedType() to include a non-resolving NULL check for LHS. This allows unsafe type object pointers to equate to NULL object pointers (in addition to nullptr) - e.g. for compiled paths that do implicit type conversions from nullptr to TObjectPtr(nullptr). Also updated the unit test to reflect this behavior change.
- Replaced FPropertyBagRepository::IsPropertyBagPlaceholderType() with IsPropertyBagPlaceholderObject().
- No longer setting RF_HasPlaceholderType on placeholder import type CDOs. This allows UObject initializers to evaluate/dereference the CDO ptr normally as a UObject*, even if created with a placeholder type.
- Replaced direct RF_HasPlaceholderType flag queries with FPropertyBagRepository::IsPropertyBagPlaceholderObject() instead for placeholder export object queries in FLinkerLoad. This remains inclusive of the CDO.
- Now appending RF_HasPlaceholderType onto the ObjectFlags member for export entries created from placeholder type imports. The flag will be cleared if/when the correct instance is patched back into the export table (e.g. at reinstancing time).
- Modified FLinkerLoad::TryCreatePlaceholderTypeForExport() to remove the check for 'bSerializeUnknownProperty'. This does not get set until after we've created the placeholder type when we attempt to Preload() the export that's using it.
- Modified FLinkerLoad::Serialize() to virtualize serialization when loading a property bag for an object with a missing type import that was serialized with an asset version older than EUnrealEngineObjectUE5Version::SCRIPT_SERIALIZATION_OFFSET.
- Modified FLinkerLoad::Preload() to include an asset version check for when serialization of placeholder exports can be safely narrowed to SerializeScriptProperties(). For older asset versions, any non-TPS data serialization is now virtualized instead.
- A warning is now emitted by FLinkerLoad::operator<<() when returning NULL for placeholder export object refs in those cases where we are not able to enforce its type safety at runtime. This now includes reflected properties that might serialize a reference to a placeholder type's CDO, which should be an unlikely edge case that we'll now report on here.
- Re-enabled unit tests for object handle/pointer type safety (ObjectHandleTests.cpp).
- Added a "stress test" method for object pointers to assist with A/B testing and perf analysis (ObjectPtrTests.cpp).
- Modified natvis to extend the TObjectPtr format to display as 'nullptr' for pointers to placeholder export types. Intent is to minimize confusion while debugging since object pointers don't allow access to unsafe type objects directly.
- Added a CVar to control whether or not we will create placeholder exports as serialization targets when import types are missing on map load (SceneGraph.EnablePropertyBagPlaceholderObjectSupport). Also can be enabled at launch via command line (-EnablePropertyBagPlaceholderObjects) for iteration purposes. Currently the CVar/feature defaults to off (experimental/WiP).

diff --git a/Engine/Source/Runtime/CoreUObject/Internal/UObject/PropertyBagRepository.h b/Engine/Source/Runtime/CoreUObject/Internal/UObject/PropertyBagRepository.h
--- a/Engine/Source/Runtime/CoreUObject/Internal/UObject/PropertyBagRepository.h
+++ b/Engine/Source/Runtime/CoreUObject/Internal/UObject/PropertyBagRepository.h
@@ -115,8 +115,10 @@ public:
 
 	// add a new placeholder type to swap in for a missing type on load; this will be associated with a property bag when instances are serialized so we don't lose its data
 	static void AddPropertyBagPlaceholderType(UClass* ClassType);
-	// query for whether or not the given class has been registered as a placeholder type
-	static COREUOBJECT_API bool IsPropertyBagPlaceholderType(UClass* ClassType);
+	// query for whether or not the given object was created as a placeholder type
+	static COREUOBJECT_API bool IsPropertyBagPlaceholderObject(UObject* Object);
+	// query for whether or not creating property bag placeholder objects should be allowed
+	static COREUOBJECT_API bool IsPropertyBagPlaceholderObjectSupportEnabled();
 
 private:
 	void Lock() const { CriticalSection.Lock(); }
diff --git a/Engine/Source/Runtime/CoreUObject/Private/UObject/LinkerLoad.cpp b/Engine/Source/Runtime/CoreUObject/Private/UObject/LinkerLoad.cpp
--- a/Engine/Source/Runtime/CoreUObject/Private/UObject/LinkerLoad.cpp
+++ b/Engine/Source/Runtime/CoreUObject/Private/UObject/LinkerLoad.cpp
@@ -661,7 +661,7 @@ void FLinkerLoad::PRIVATE_PatchNewObjectIntoExport(UObject* OldObject, UObject*
 		FObjectExport& ObjExport = OldObjectLinker->ExportMap[CachedLinkerIndex];
 		
 		// Since we don't copy the internal flags, the mirrored flags need can't be set on the new object as well
-		const EObjectFlags OldObjectFlags = OldObject->GetFlags() & ~RF_MirroredGarbage;
+		const EObjectFlags OldObjectFlags = OldObject->GetFlags() & ~(RF_MirroredGarbage | RF_HasPlaceholderType);
 
 		// Detach the old object to make room for the new
 		OldObject->ClearFlags(RF_NeedLoad|RF_NeedPostLoad|RF_NeedPostLoadSubobjects);
@@ -1017,6 +1017,8 @@ FLinkerLoad::FLinkerLoad(UPackage* InParent, const FPackagePath& InPackagePath,
 #if WITH_EDITOR
 , bExportsDuplicatesFixed(false)
 , bIsPackageRelocated(false)
+, bIsLoadingToPropertyBagObject(false)
+, bIsSerializingScriptProperties(false)
 #endif // WITH_EDITOR
 , IsTimeLimitExceededCallCount(0)
 , TimeLimit(0.0f)
@@ -4052,13 +4054,17 @@ UClass* FLinkerLoad::GetExportLoadClass(int32 Index)
 #if WITH_EDITOR
 UClass* FLinkerLoad::TryCreatePlaceholderTypeForExport(int32 ExportIndex)
 {
+	const bool bAllowPlaceholderImportTypes = UE::FPropertyBagRepository::IsPropertyBagPlaceholderObjectSupportEnabled();
+	if (!bAllowPlaceholderImportTypes)
+	{
+		return nullptr;
+	}
+
 	UClass* LoadClass = nullptr;
 	FObjectExport& Export = ExportMap[ExportIndex];
 
 	// If the class import is missing, create a placeholder for this export. This will allow us to instance and redirect its data into a property bag.
-	// Note: This is only possible if the missing type did not also include any custom serialization, so we currently rely on the load context for this.
-	FUObjectSerializeContext* CurrentLoadContext = GetSerializeContext();
-	if (CurrentLoadContext && CurrentLoadContext->bSerializeUnknownProperty && Export.ClassIndex.IsImport() && !GEventDrivenLoaderEnabled)
+	if (Export.ClassIndex.IsImport() && !GEventDrivenLoaderEnabled)
 	{
 		FObjectImport& LoadClassImport = Imp(Export.ClassIndex);
 
@@ -4092,15 +4098,15 @@ UClass* FLinkerLoad::TryCreatePlaceholderTypeForExport(int32 ExportIndex)
 					UObject* LoadClassDefaults = LoadClass->GetDefaultObject();
 					LoadClass->PostLoadDefaultObject(LoadClassDefaults);
 
-					// Indicate that the CDO has a placeholder type (in case something has serialized a reference to it).
-					LoadClassDefaults->SetFlags(RF_HasPlaceholderType);
-
 					// This class is for internal use and should not be exposed for selection or instancing in the editor.
 					LoadClass->ClassFlags |= CLASS_Hidden | CLASS_HideDropDown;
 
 					// Patch it into the import table so that we resolve to this class for any future exports of this type.
 					LoadClassImport.XObject = LoadClass;
 
+					// Modify the export's object flags for instancing to indicate that its load class is a placeholder type.
+					Export.ObjectFlags |= RF_HasPlaceholderType;
+
 					// Use the property bag repository for now to manage property bag placeholder types (e.g. object lifetime).
 					// Note: The object lifetime of instances of this type will rely on existing references that are serialized.
 					UE::FPropertyBagRepository::AddPropertyBagPlaceholderType(LoadClass);
@@ -4602,16 +4608,18 @@ void FLinkerLoad::Preload( UObject* Object )
 				const int64 SavedPos = Loader->Tell();
 				int64 StartPos = Export.SerialOffset;
 				int64 ExpectedSerialSize = Export.SerialSize;
-
+#if WITH_EDITOR
 				// for placeholder objects that have no explicit type, we only want to serialize the TPS stream
-				const bool bSerializeOnlyScriptProperties = Object->HasAnyFlags(RF_HasPlaceholderType);
-				if (bSerializeOnlyScriptProperties)
+				bool bSerializeOnlyScriptProperties = false;
+				FGuardValue_Bitfield(bIsLoadingToPropertyBagObject, UE::FPropertyBagRepository::IsPropertyBagPlaceholderObject(Object));
+				if (bIsLoadingToPropertyBagObject && UEVer() >= EUnrealEngineObjectUE5Version::SCRIPT_SERIALIZATION_OFFSET)
 				{
 					// note: script start/end offsets are relative to the export's offset in the file
 					StartPos += Export.ScriptSerializationStartOffset;
 					ExpectedSerialSize = Export.ScriptSerializationEndOffset;
+					bSerializeOnlyScriptProperties = true;	// signals that we can safely narrow the load to SerializeScriptProperties()
 				}
-
+#endif
 				// move to the position in the file where this object's data
 				// is stored
 				Seek(StartPos);
@@ -4772,17 +4780,21 @@ void FLinkerLoad::Preload( UObject* Object )
 #endif
 						{
 							UE_SERIALIZE_ACCCESS_SCOPE(Object);
-
+#if WITH_EDITOR
 							if (bSerializeOnlyScriptProperties)
 							{
 								Object->SerializeScriptProperties(*this);
 							}
 							else
+#endif
 							{
 								Object->Serialize(*this);
 							}
 						}
-
+#if WITH_EDITOR
+						// Ensure begin/end marks were hit.
+						check(!bIsSerializingScriptProperties);
+#endif
 						Object->SetFlags(RF_LoadCompleted);
 						CurrentLoadContext->SerializedObject = PrevSerializedObject;
 					}
@@ -5470,7 +5482,7 @@ UObject* FLinkerLoad::CreateExport( int32 Index )
 			Export.bIsInheritedInstance = false;
 		}
 
-		const UObject* LoadClassDefaultObject = LoadClass->GetDefaultObject();
+		LoadClass->GetDefaultObject();
 
 		FStaticConstructObjectParameters Params(LoadClass);
 		Params.Outer = ThisParent;
@@ -5480,12 +5492,6 @@ UObject* FLinkerLoad::CreateExport( int32 Index )
 		// if our outer is actually an import, then the package we are an export of is not in our outer chain, set our package in that case
 		Params.ExternalPackage = Export.OuterIndex.IsImport() ? LinkerRoot : nullptr;
 
-		// If we're instancing from a placeholder type, set the object flag (so that we can quickly determine that it's a placeholder object).
-		if (LoadClassDefaultObject && LoadClassDefaultObject->HasAnyFlags(RF_HasPlaceholderType))
-		{
-			Params.SetFlags |= RF_HasPlaceholderType;
-		}
-
 		// Propagate relevant properties from the outer package to the external package
 		if (Params.ExternalPackage)
 		{
@@ -6183,11 +6189,12 @@ FArchive& FLinkerLoad::operator<<( UObject*& Object )
 
 	Object = ResolveResource(Index);
 #if WITH_EDITOR
-	if (Object && Object->HasAnyFlags(RF_HasPlaceholderType))
+	if (Object && UE::FPropertyBagRepository::IsPropertyBagPlaceholderObject(Object))
 	{
 		// This is needed because the pointer's type is checked only at compile time, which may not match the property
 		// bag placeholder object's type at runtime, and so we can't allow it to be dereferenced as the wrong base type.
 		// Note: These currently won't be discovered for replacement at reinstancing time, so it will remain set to NULL.
+		UE_LOG(LogLinker, Warning, TEXT("Serializing reference to \"%s\" as NULL to ensure type safety."), *Object->GetPathName());
 		Object = nullptr;
 	}
 #endif
@@ -6203,12 +6210,22 @@ FArchive& FLinkerLoad::operator<<(FObjectPtr& ObjectPtr)
 	// Wrapper that only allows pointers to exports with placeholder types when type safety features are enabled.
 	auto AsTypeSafeObjectPtr_Lambda = [](UObject* ResolvedObject)
 	{
-#if WITH_EDITOR && !UE_WITH_OBJECT_HANDLE_TYPE_SAFETY
-		// If we can't mask the reference to the placeholder object instance at resolve time, set it now to NULL.
+#if WITH_EDITOR
+		// If we can't mask the pointer to the placeholder object instance at access time, resolve it now to NULL.
 		// Note: Similar to hard references above, this means we won't find it for replacement at reinstancing time.
-		if (ResolvedObject && ResolvedObject->HasAnyFlags(RF_HasPlaceholderType))
-		{
-			ResolvedObject = nullptr;
+		if (ResolvedObject && UE::FPropertyBagRepository::IsPropertyBagPlaceholderObject(ResolvedObject))
+		{
+#if UE_WITH_OBJECT_HANDLE_TYPE_SAFETY
+			// Note: With type safety enabled, unlike other instances of a placeholder type, we won't mark the CDO
+			// as also being a placeholder instance. That's because there are certain paths that need to be able
+			// to resolve the pointer (e.g. - the object initialization path during class construction). However,
+			// it also means we can't resolve other references to a placeholder CDO, as they may not be type-safe.
+			if (ResolvedObject->HasAnyFlags(RF_ClassDefaultObject))
+#endif
+			{
+				UE_LOG(LogLinker, Warning, TEXT("Serializing reference to \"%s\" as NULL to ensure type safety."), *ResolvedObject->GetPathName());
+				ResolvedObject = nullptr;
+			}
 		}
 #endif
 		return FObjectPtr(ResolvedObject);
@@ -6290,6 +6307,9 @@ void FLinkerLoad::MarkScriptSerializationStart( const UObject* Obj )
 {
 	if (Obj && Obj->GetLinker() == this)
 	{
+#if WITH_EDITOR
+		bIsSerializingScriptProperties = true;
+#endif
 		int32 Index = Obj->GetLinkerIndex();
 		if (ExportMap.IsValidIndex(Index))
 		{
@@ -6313,6 +6333,9 @@ void FLinkerLoad::MarkScriptSerializationEnd( const UObject* Obj )
 {
 	if (Obj && Obj->GetLinker() == this)
 	{
+#if WITH_EDITOR
+		bIsSerializingScriptProperties = false;
+#endif
 		int32 Index = Obj->GetLinkerIndex();
 		if (ExportMap.IsValidIndex(Index))
 		{
diff --git a/Engine/Source/Runtime/CoreUObject/Private/UObject/PropertyBagRepository.cpp b/Engine/Source/Runtime/CoreUObject/Private/UObject/PropertyBagRepository.cpp
--- a/Engine/Source/Runtime/CoreUObject/Private/UObject/PropertyBagRepository.cpp
+++ b/Engine/Source/Runtime/CoreUObject/Private/UObject/PropertyBagRepository.cpp
@@ -14,6 +14,11 @@
 #include "UObject/InstanceDataObjectUtils.h"
 #include "UObject/Package.h"
 
+#if WITH_EDITOR
+#include "HAL/IConsoleManager.h"
+#include "Misc/CommandLine.h"
+#endif
+
 DEFINE_LOG_CATEGORY_STATIC(LogPropertyBagRepository, Log, All);
 
 namespace UE
@@ -325,24 +330,54 @@ void FPropertyBagRepository::ShrinkMaps()
 	AssociatedData.Compact();
 }
 
-bool FPropertyBagRepository::IsPropertyBagPlaceholderType(UClass* ClassType)
+void FPropertyBagRepository::AddPropertyBagPlaceholderType(UClass* ClassType)
 {
 	if (!ClassType)
 	{
-		return false;
+		return;
 	}
 
-	return FPropertyBagRepository::Get().PropertyBagTypeRegistry->Contains(ClassType);
+	FPropertyBagRepository::Get().PropertyBagTypeRegistry->Add(ClassType);
 }
 
-void FPropertyBagRepository::AddPropertyBagPlaceholderType(UClass* ClassType)
+bool FPropertyBagRepository::IsPropertyBagPlaceholderObject(UObject* Object)
 {
-	if (!ClassType)
+	if (!Object)
 	{
-		return;
+		return false;
 	}
 
-	FPropertyBagRepository::Get().PropertyBagTypeRegistry->Add(ClassType);
+	return Object->HasAnyFlags(RF_HasPlaceholderType|RF_ClassDefaultObject)
+		&& FPropertyBagRepository::Get().PropertyBagTypeRegistry->Contains(Object->GetClass());
+}
+
+namespace Private
+{
+#if WITH_EDITOR
+	static bool bEnablePropertyBagPlaceholderObjectSupport = 0;
+	static FAutoConsoleVariableRef CVarEnablePropertyBagPlaceholderObjectSupport(
+		TEXT("SceneGraph.EnablePropertyBagPlaceholderObjectSupport"),
+		bEnablePropertyBagPlaceholderObjectSupport,
+		TEXT("If true, allows placeholder types to be created in place of missing types on load in order to redirect serialization into a property bag."),
+		ECVF_Default
+	);
+#endif
+}
+
+bool FPropertyBagRepository::IsPropertyBagPlaceholderObjectSupportEnabled()
+{
+#if WITH_EDITOR && UE_WITH_OBJECT_HANDLE_TYPE_SAFETY
+	static bool bIsInitialized = false;
+	if (!bIsInitialized)
+	{
+		Private::bEnablePropertyBagPlaceholderObjectSupport = FParse::Param(FCommandLine::Get(), TEXT("WithPropertyBagPlaceholderObjects"));
+		bIsInitialized = true;
+	}
+	
+	return Private::bEnablePropertyBagPlaceholderObjectSupport;
+#else
+	return false;
+#endif
 }
 
 } // UE
diff --git a/Engine/Source/Runtime/CoreUObject/Public/UObject/LinkerLoad.h b/Engine/Source/Runtime/CoreUObject/Public/UObject/LinkerLoad.h
--- a/Engine/Source/Runtime/CoreUObject/Public/UObject/LinkerLoad.h
+++ b/Engine/Source/Runtime/CoreUObject/Public/UObject/LinkerLoad.h
@@ -465,6 +465,12 @@ private:
 
 	/** Cache if the package is relocated or not */
 	bool bIsPackageRelocated : 1;
+
+	/** Set when the loader is serializing to a property bag placeholder object.											*/
+	bool bIsLoadingToPropertyBagObject : 1;
+
+	/** TRUE when the loader is actively serializing an object's script properties data.									*/
+	bool bIsSerializingScriptProperties : 1;
 #endif // WITH_EDITOR
 
 	/** Call count of IsTimeLimitExceeded.																					*/
@@ -1020,8 +1026,18 @@ private:
 		checkSlow(FPlatformTLS::GetCurrentThreadId() == OwnerThread);
 #if WITH_EDITOR
 		Loader->SetSerializedProperty(GetSerializedProperty());
+
+		// This handles the case where we're only serializing to a property bag and we don't immediately know where the script
+		// property data starts/ends in the object's data stream. Allows for loading from packages saved with an older version.
+		if (UNLIKELY(bIsLoadingToPropertyBagObject && !bIsSerializingScriptProperties))
+		{
+			Loader->Seek(Tell() + Length);
+		}
+		else
 #endif
-		Loader->Serialize(V, Length);
+		{
+			Loader->Serialize(V, Length);
+		}
 	}
 	using FArchiveUObject::operator<<; // For visibility of the overloads we don't override
 	virtual FArchive& operator<<(UObject*& Object) override;
diff --git a/Engine/Source/Runtime/CoreUObject/Public/UObject/ObjectPtr.h b/Engine/Source/Runtime/CoreUObject/Public/UObject/ObjectPtr.h
--- a/Engine/Source/Runtime/CoreUObject/Public/UObject/ObjectPtr.h
+++ b/Engine/Source/Runtime/CoreUObject/Public/UObject/ObjectPtr.h
@@ -373,6 +373,12 @@ namespace ObjectPtr_Private
 	template <typename T, typename U>
 	bool IsObjectPtrEqualToRawPtrOfRelatedType(const TObjectPtr<T>& Ptr, const U* Other)
 	{
+		// simple test if Ptr is null (avoids resolving either side)
+		if (!Ptr)
+		{
+			return !Other;
+		}
+
 #if UE_WITH_OBJECT_HANDLE_LATE_RESOLVE
 		if (Ptr.IsResolved())
 		{
@@ -380,7 +386,7 @@ namespace ObjectPtr_Private
 		}
 		else if (!Other) //avoids resolving if Other is null
 		{
-			return !Ptr;
+			return false;	// from above, we already know that !Ptr is false
 		}
 #endif
 		return Ptr.GetNoReadNoCheck() == Other;
@@ -558,7 +564,13 @@ public:
 	>
 	FORCEINLINE bool operator==(const TObjectPtr<U>& Other) const
 	{
+#if UE_WITH_OBJECT_HANDLE_TYPE_SAFETY
+		// Do a NULL test first before comparing the underlying handles, in case either side is
+		// a non-NULL, but unsafe type pointer (which would equate to NULL when Get() is called).
+		return !ObjectPtr ? !Other : ObjectPtr == Other.ObjectPtr;
+#else
 		return ObjectPtr == Other.ObjectPtr;
+#endif
 	}
 
 	// Equality/Inequality comparisons against nullptr
diff --git a/Engine/Source/Runtime/CoreUObject/Tests/ObjectHandleTest.cpp b/Engine/Source/Runtime/CoreUObject/Tests/ObjectHandleTest.cpp
--- a/Engine/Source/Runtime/CoreUObject/Tests/ObjectHandleTest.cpp
+++ b/Engine/Source/Runtime/CoreUObject/Tests/ObjectHandleTest.cpp
@@ -372,7 +372,7 @@ TEST_CASE_METHOD(FObjectHandleTestBase, "CoreUObject::FObjectHandle::Hash Object
 }
 
 #if UE_WITH_OBJECT_HANDLE_TYPE_SAFETY
-DISABLED_TEST_CASE_METHOD(FObjectHandleTestBase, "CoreUObject::FObjectHandle::Type Safety", "[CoreUObject][ObjectHandle]")
+TEST_CASE_METHOD(FObjectHandleTestBase, "CoreUObject::FObjectHandle::Type Safety", "[CoreUObject][ObjectHandle]")
 {
 	const FName TestPackageName(TEXT("/Engine/Test/ObjectHandle/TypeSafety/Transient"));
 	UPackage* TestPackage = NewObject<UPackage>(nullptr, TestPackageName, RF_Transient);
@@ -389,11 +389,10 @@ DISABLED_TEST_CASE_METHOD(FObjectHandleTestBase, "CoreUObject::FObjectHandle::Ty
 	TestClass->StaticLink(/*bRelinkExistingProperties =*/ true);
 	UObject* TestClassDefaults = TestClass->GetDefaultObject();
 	TestClass->PostLoadDefaultObject(TestClassDefaults);
-	TestClassDefaults->SetFlags(RF_HasPlaceholderType);
 
 	// validate helper method(s)
 	CHECK_FALSE(UE::CoreUObject::Private::HasAnyFlags(TestClassDefaults, RF_NoFlags));
-	CHECK(UE::CoreUObject::Private::HasAnyFlags(TestClassDefaults, RF_HasPlaceholderType));
+	CHECK(UE::CoreUObject::Private::HasAnyFlags(TestClassDefaults, RF_ClassDefaultObject));
 
 	// construct objects for testing
 	UObject* TestSafeObject = NewObject<UObjectPtrTestClass>(TestPackage, TEXT("TestSafeObject"), RF_Transient);
@@ -427,6 +426,13 @@ DISABLED_TEST_CASE_METHOD(FObjectHandleTestBase, "CoreUObject::FObjectHandle::Ty
 	// object handles should resolve/evaluate to the original type object
 	CHECK(UE::CoreUObject::Private::ResolveObjectHandle(TestUnsafeObjectHandle) == TestUnsafeObject);
 
+	// an unsafe type object handle should not equate to other unsafe type object handles except for itself (including NULL)
+	CHECK(NullObjectHandle != TestUnsafeObjectHandle);			// note: this intentionally differs from object *pointers* (see below)
+	CHECK(TestUnsafeObjectHandle != NullObjectHandle);			// see note directly above
+	CHECK(TestSafeObjectHandle != TestUnsafeObjectHandle);
+	CHECK(TestUnsafeObjectHandle != TestSafeObjectHandle);
+	CHECK(TestUnsafeObjectHandle == TestUnsafeObjectHandle);
+
 	// construct object pointers for testing
 	TObjectPtr<UObject> NullObjectPtr(nullptr);
 	TObjectPtr<UObject> TestSafeObjectPtr(TestSafeObject);
@@ -435,11 +441,15 @@ DISABLED_TEST_CASE_METHOD(FObjectHandleTestBase, "CoreUObject::FObjectHandle::Ty
 	// unsafe type object pointers should evaluate to NULL/false (for type safety)
 	CHECK(!TestUnsafeObjectPtr);
 	CHECK_FALSE(!!TestUnsafeObjectPtr);
+	CHECK(NULL == TestUnsafeObjectPtr);
+	CHECK(TestUnsafeObjectPtr == NULL);
 	CHECK(nullptr == TestUnsafeObjectPtr);
 	CHECK(TestUnsafeObjectPtr == nullptr);
 
-	// an unsafe type object pointer should not equate to NULL/other pointers
-	CHECK(TestUnsafeObjectPtr != NullObjectPtr);
+	// an unsafe type object pointer should not equate to other pointers except for NULL and itself
+	CHECK(NullObjectPtr == TestUnsafeObjectPtr);				// note: this intentionally differs from object *handles* (see above)
+	CHECK(TestUnsafeObjectPtr == NullObjectPtr);				// see note directly above
+	CHECK(TestSafeObjectPtr != TestUnsafeObjectPtr);
 	CHECK(TestUnsafeObjectPtr != TestSafeObjectPtr);
 	CHECK(TestUnsafeObjectPtr == TestUnsafeObjectPtr);
 
diff --git a/Engine/Source/Runtime/CoreUObject/Tests/ObjectPtrTest.cpp b/Engine/Source/Runtime/CoreUObject/Tests/ObjectPtrTest.cpp
--- a/Engine/Source/Runtime/CoreUObject/Tests/ObjectPtrTest.cpp
+++ b/Engine/Source/Runtime/CoreUObject/Tests/ObjectPtrTest.cpp
@@ -15,6 +15,7 @@
 #include "UObject/Package.h"
 #include "UObject/SoftObjectPath.h"
 #include "Misc/ScopeExit.h"
+#include "Tests/Benchmark.h"
 #include <type_traits>
 
 namespace UE::CoreObject::Private::Tests
@@ -165,6 +166,8 @@ TEST_CASE("CoreUObject::TObjectPtr::FindLoadBehavior")
 TEST_CASE("CoreUObject::TObjectPtr::Null", "[CoreUObject][ObjectPtr]")
 {
 	TObjectPtr<UObject> NullObjectPtr(nullptr);
+	TEST_TRUE(TEXT("NULL should equal a null object pointer"), NULL == NullObjectPtr);
+	TEST_TRUE(TEXT("A null object pointer should equal NULL"), NullObjectPtr == NULL);
 	TEST_TRUE(TEXT("Nullptr should equal a null object pointer"), nullptr == NullObjectPtr);
 	TEST_TRUE(TEXT("A null object pointer should equal nullptr"), NullObjectPtr == nullptr);
 	TEST_FALSE(TEXT("A null object pointer should evaluate to false"), !!NullObjectPtr);
@@ -343,6 +346,117 @@ TEST_CASE_METHOD(FObjectPtrTestBase, "CoreUObject::TObjectPtr::Long Path", "[Cor
 }
 #endif
 
+void ObjectPtrStressTest(int32 NumTestNodes, bool bIsContiguousTest, bool bIsEvalOnlyTest)
+{
+	const FName TestPackageName(TEXT("/Engine/Test/ObjectPtr/EvalStressTest/Transient"));
+	UPackage* TestPackage = NewObject<UPackage>(nullptr, TestPackageName, RF_Transient);
+	TestPackage->AddToRoot();
+	ON_SCOPE_EXIT
+	{
+		TestPackage->RemoveFromRoot();
+		CollectGarbage(GARBAGE_COLLECTION_KEEPFLAGS);
+	};
+
+	int32 NumObjectsAllocated = 0;
+	const int32 MaxAllowedAllocations = GUObjectArray.GetObjectArrayEstimatedAvailable();
+	TArray<TObjectPtr<UObjectPtrStressTestClass>*> TestNodePtrs;
+	TArray<TObjectPtr<UObjectPtrStressTestClass>> TestNodeStore;
+
+	if (bIsContiguousTest)
+	{
+		TestNodeStore.Reserve(NumTestNodes);
+		for (int32 NodeIdx = 0; NodeIdx < NumTestNodes; ++NodeIdx)
+		{
+			TestNodePtrs.Add(&TestNodeStore.AddDefaulted_GetRef());
+		}
+	}
+	else
+	{
+		for (int32 NodeIdx = 0; NodeIdx < NumTestNodes; ++NodeIdx)
+		{
+			TestNodePtrs.Add(new TObjectPtr<UObjectPtrStressTestClass>());
+		}
+	}
+
+	for (TObjectPtr<UObjectPtrStressTestClass>* TestNodePtr : TestNodePtrs)
+	{
+		// randomize the allocation of a test object
+		UObjectPtrStressTestClass* TestObject = nullptr;
+		if (NumObjectsAllocated < MaxAllowedAllocations && (FMath::Rand() % 2))
+		{
+			TestObject = NewObject<UObjectPtrStressTestClass>(TestPackage, NAME_None, RF_Transient);
+			++NumObjectsAllocated;
+		}
+
+		*TestNodePtr = TestObject;
+	}
+
+	for (TObjectPtr<UObjectPtrStressTestClass>* TestNodePtr : TestNodePtrs)
+	{
+		// fill out the remaining unallocated slots if we can
+		TObjectPtr<UObjectPtrStressTestClass>& TestNode = *TestNodePtr;
+		if (!TestNode && NumObjectsAllocated < MaxAllowedAllocations)
+		{
+			TestNode = NewObject<UObjectPtrStressTestClass>(TestPackage, NAME_None, RF_Transient);
+			++NumObjectsAllocated;
+		}
+
+		if (!bIsEvalOnlyTest)
+		{
+			// access the pointer (i.e. resolve to a raw UObject ptr) and do something with it
+			UObjectPtrStressTestClass* ResolvedObject = TestNode.Get();
+			if (ResolvedObject)
+			{
+				FMemory::Memzero(ResolvedObject->Data, PLATFORM_CACHE_LINE_SIZE);
+			}
+		}
+
+		if (!bIsContiguousTest)
+		{
+			if (TestNode)
+			{
+				TestNode->MarkAsGarbage();
+			}
+
+			delete TestNodePtr;
+		}
+	}
+}
+
+// Enable this method to add stress test benchmarks or to do A/B comparisons with features turned on/off.
+DISABLED_TEST_CASE_METHOD(FObjectPtrTestBase, "CoreUObject::TObjectPtr::Stress Tests", "[CoreUObject][ObjectPtr]")
+{
+	constexpr bool CONTIGUOUS = true;
+	constexpr bool NON_CONTIGUOUS = false;
+	constexpr bool EVAL_ONLY = true;
+	constexpr bool EVAL_AND_RESOLVE = false;
+
+	FSnapshotObjectRefMetrics ObjectRefMetrics(*this);
+
+	// contiguous object pointer blocks, eval-only
+	UE_BENCHMARK(5, [] { ObjectPtrStressTest(1000, CONTIGUOUS, EVAL_ONLY); } );
+	UE_BENCHMARK(5, [] { ObjectPtrStressTest(10000, CONTIGUOUS, EVAL_ONLY); });
+	UE_BENCHMARK(5, [] { ObjectPtrStressTest(100000, CONTIGUOUS, EVAL_ONLY); });
+
+	// non-contiguous list of object pointers, eval-only
+	UE_BENCHMARK(5, [] { ObjectPtrStressTest(1000, NON_CONTIGUOUS, EVAL_ONLY); });
+	UE_BENCHMARK(5, [] { ObjectPtrStressTest(10000, NON_CONTIGUOUS, EVAL_ONLY); });
+	UE_BENCHMARK(5, [] { ObjectPtrStressTest(100000, NON_CONTIGUOUS, EVAL_ONLY); });
+
+	// make sure that nothing was resolved above
+	ObjectRefMetrics.TestNumResolves(TEXT("Eval-only stress tests should not have triggered a resolve attempt"), 0);
+
+	// contiguous object pointer blocks, with resolve attempts
+	UE_BENCHMARK(5, [] { ObjectPtrStressTest(1000, CONTIGUOUS, EVAL_AND_RESOLVE); });
+	UE_BENCHMARK(5, [] { ObjectPtrStressTest(10000, CONTIGUOUS, EVAL_AND_RESOLVE); });
+	UE_BENCHMARK(5, [] { ObjectPtrStressTest(100000, CONTIGUOUS, EVAL_AND_RESOLVE); });
+
+	// non-contiguous list of object pointers, with resolve attempts
+	UE_BENCHMARK(5, [] { ObjectPtrStressTest(1000, NON_CONTIGUOUS, EVAL_AND_RESOLVE); });
+	UE_BENCHMARK(5, [] { ObjectPtrStressTest(10000, NON_CONTIGUOUS, EVAL_AND_RESOLVE); });
+	UE_BENCHMARK(5, [] { ObjectPtrStressTest(100000, NON_CONTIGUOUS, EVAL_AND_RESOLVE); });
+}
+
 TEST_CASE("CoreUObject::TObjectPtr::GetPathName", "[CoreUObject][ObjectPtr]")
 {
 	const FName TestPackage1Name(TEXT("/Engine/Test/FObjectPtrTestLongPath/Transient"));
diff --git a/Engine/Source/Runtime/CoreUObject/Tests/ObjectPtrTestClass.cpp b/Engine/Source/Runtime/CoreUObject/Tests/ObjectPtrTestClass.cpp
--- a/Engine/Source/Runtime/CoreUObject/Tests/ObjectPtrTestClass.cpp
+++ b/Engine/Source/Runtime/CoreUObject/Tests/ObjectPtrTestClass.cpp
@@ -192,4 +192,6 @@ IMPLEMENT_CORE_INTRINSIC_CLASS(UObjectPtrDerrivedTestClass, UObjectPtrTestClass,
 
 IMPLEMENT_CORE_INTRINSIC_CLASS(UObjectPtrNotLazyTestClass, UObject, {});
 
+IMPLEMENT_CORE_INTRINSIC_CLASS(UObjectPtrStressTestClass, UObject, {});
+
 #endif
\ No newline at end of file
diff --git a/Engine/Source/Runtime/CoreUObject/Tests/ObjectPtrTestClass.h b/Engine/Source/Runtime/CoreUObject/Tests/ObjectPtrTestClass.h
--- a/Engine/Source/Runtime/CoreUObject/Tests/ObjectPtrTestClass.h
+++ b/Engine/Source/Runtime/CoreUObject/Tests/ObjectPtrTestClass.h
@@ -70,4 +70,14 @@ public:
 };
 
 
+//stress testing class
+class UObjectPtrStressTestClass : public UObject
+{
+	DECLARE_CLASS_INTRINSIC(UObjectPtrStressTestClass, UObject, CLASS_MatchedSerializers, TEXT("/Script/CoreUObject"))
+
+public:
+	uint8 Data[PLATFORM_CACHE_LINE_SIZE];
+};
+
+
 #endif
\ No newline at end of file
