Refactored ConvertFromType for enums to reduce branching and prepare for the complete type in tags

#rnx

diff --git a/Engine/Source/Runtime/CoreUObject/Private/UObject/EnumProperty.cpp b/Engine/Source/Runtime/CoreUObject/Private/UObject/EnumProperty.cpp
--- a/Engine/Source/Runtime/CoreUObject/Private/UObject/EnumProperty.cpp
+++ b/Engine/Source/Runtime/CoreUObject/Private/UObject/EnumProperty.cpp
@@ -419,12 +419,17 @@ bool FEnumProperty::SameType(const FProperty* Other) const
 
 EConvertFromTypeResult FEnumProperty::ConvertFromType(const FPropertyTag& Tag, FStructuredArchive::FSlot Slot , uint8* Data, UStruct* DefaultsStruct, const uint8* Defaults)
 {
-	if ((Enum == nullptr) || (UnderlyingProp == nullptr))
+	const EName* TagType = Tag.Type.ToEName();
+	if (LIKELY(!TagType || *TagType == NAME_EnumProperty || Tag.Type.GetNumber() || !Enum || !UnderlyingProp))
 	{
 		return EConvertFromTypeResult::UseSerializeItem;
 	}
 
-	if (Tag.Type == NAME_ByteProperty)
+	switch (*TagType)
+	{
+	default:
+		return EConvertFromTypeResult::UseSerializeItem;
+	case NAME_ByteProperty:
 	{
 		uint8 PreviousValue = 0;
 		if (Tag.EnumName == NAME_None)
@@ -454,45 +459,33 @@ EConvertFromTypeResult FEnumProperty::ConvertFromType(const FPropertyTag& Tag, F
 
 		// now copy the value into the object's address space
 		UnderlyingProp->SetIntPropertyValue(ContainerPtrToValuePtr<void>(Data, Tag.ArrayIndex), (uint64)PreviousValue);
+		return EConvertFromTypeResult::Converted;
 	}
-	else if (Tag.Type == NAME_Int8Property)
-	{
+	case NAME_Int8Property:
 		UEEnumProperty_Private::ConvertIntToEnumProperty<int8>(Slot, this, UnderlyingProp, Enum, ContainerPtrToValuePtr<void>(Data, Tag.ArrayIndex));
-	}
-	else if (Tag.Type == NAME_Int16Property)
-	{
+		return EConvertFromTypeResult::Converted;
+	case NAME_Int16Property:
 		UEEnumProperty_Private::ConvertIntToEnumProperty<int16>(Slot, this, UnderlyingProp, Enum, ContainerPtrToValuePtr<void>(Data, Tag.ArrayIndex));
-	}
-	else if (Tag.Type == NAME_IntProperty)
-	{
+		return EConvertFromTypeResult::Converted;
+	case NAME_IntProperty:
 		UEEnumProperty_Private::ConvertIntToEnumProperty<int32>(Slot, this, UnderlyingProp, Enum, ContainerPtrToValuePtr<void>(Data, Tag.ArrayIndex));
-	}
-	else if (Tag.Type == NAME_Int64Property)
-	{
+		return EConvertFromTypeResult::Converted;
+	case NAME_Int64Property:
 		UEEnumProperty_Private::ConvertIntToEnumProperty<int64>(Slot, this, UnderlyingProp, Enum, ContainerPtrToValuePtr<void>(Data, Tag.ArrayIndex));
-	}
-	else if (Tag.Type == NAME_UInt16Property)
-	{
+		return EConvertFromTypeResult::Converted;
+	case NAME_UInt16Property:
 		UEEnumProperty_Private::ConvertIntToEnumProperty<uint16>(Slot, this, UnderlyingProp, Enum, ContainerPtrToValuePtr<void>(Data, Tag.ArrayIndex));
-	}
-	else if (Tag.Type == NAME_UInt32Property)
-	{
+		return EConvertFromTypeResult::Converted;
+	case NAME_UInt32Property:
 		UEEnumProperty_Private::ConvertIntToEnumProperty<uint32>(Slot, this, UnderlyingProp, Enum, ContainerPtrToValuePtr<void>(Data, Tag.ArrayIndex));
-	}
-	else if (Tag.Type == NAME_UInt64Property)
-	{
+		return EConvertFromTypeResult::Converted;
+	case NAME_UInt64Property:
 		UEEnumProperty_Private::ConvertIntToEnumProperty<uint64>(Slot, this, UnderlyingProp, Enum, ContainerPtrToValuePtr<void>(Data, Tag.ArrayIndex));
-	}
-	else if (Tag.Type == NAME_BoolProperty)
-	{
+		return EConvertFromTypeResult::Converted;
+	case NAME_BoolProperty:
 		UEEnumProperty_Private::ConvertIntValueToEnumProperty<uint8>(Tag.BoolVal, this, UnderlyingProp, Enum, ContainerPtrToValuePtr<void>(Data, Tag.ArrayIndex));
+		return EConvertFromTypeResult::Converted;
 	}
-	else
-	{
-		return EConvertFromTypeResult::UseSerializeItem;
-	}
-
-	return EConvertFromTypeResult::Converted;
 }
 
 #if WITH_EDITORONLY_DATA
diff --git a/Engine/Source/Runtime/CoreUObject/Private/UObject/PropertyByte.cpp b/Engine/Source/Runtime/CoreUObject/Private/UObject/PropertyByte.cpp
--- a/Engine/Source/Runtime/CoreUObject/Private/UObject/PropertyByte.cpp
+++ b/Engine/Source/Runtime/CoreUObject/Private/UObject/PropertyByte.cpp
@@ -216,43 +216,61 @@ struct TConvertIntToEnumProperty
 
 EConvertFromTypeResult FByteProperty::ConvertFromType(const FPropertyTag& Tag, FStructuredArchive::FSlot Slot, uint8* Data, UStruct* DefaultsStruct, const uint8* Defaults)
 {
-	if (Tag.Type == NAME_ByteProperty  && ((Tag.EnumName == NAME_None) != (Enum == nullptr)))
+	const EName* TagType= Tag.Type.ToEName();
+	if (UNLIKELY(Tag.Type.GetNumber() || !TagType))
 	{
-		// a byte property gained or lost an enum
-		// attempt to convert it
-		uint8 PreviousValue = 0;
-		if (Tag.EnumName == NAME_None)
-		{
-			// If we're a nested property the EnumName tag got lost. Fail to read in this case
-			FProperty* const PropertyOwner = GetOwner<FProperty>();
-			if (PropertyOwner)
+		return EConvertFromTypeResult::UseSerializeItem;
+	}
+
+	switch (*TagType)
+	{
+	default:
+		return EConvertFromTypeResult::UseSerializeItem;
+	case NAME_ByteProperty:
+	{
+		if ((Tag.EnumName == NAME_None) != (Enum == nullptr))
+		{
+			// A byte property gained or lost an enum.
+			uint8 PreviousValue = 0;
+			if (Tag.EnumName == NAME_None)
+			{
+				// If we're a nested property the EnumName tag got lost. Fail to read in this case
+				FProperty* const PropertyOwner = GetOwner<FProperty>();
+				if (PropertyOwner)
+				{
+					return EConvertFromTypeResult::UseSerializeItem;
+				}
+
+				// simply pretend the property still doesn't have an enum and serialize the single byte
+				Slot << PreviousValue;
+			}
+			else
 			{
-				return EConvertFromTypeResult::UseSerializeItem;
+				// attempt to find the old enum and get the byte value from the serialized enum name
+				PreviousValue = (uint8)ReadEnumAsInt64(Slot, DefaultsStruct, Tag);
 			}
 
-			// simply pretend the property still doesn't have an enum and serialize the single byte
-			Slot << PreviousValue;
-		}
-		else
-		{
-			// attempt to find the old enum and get the byte value from the serialized enum name
-			PreviousValue = (uint8)ReadEnumAsInt64(Slot, DefaultsStruct, Tag);
+			// now copy the value into the object's address space
+			SetPropertyValue_InContainer(Data, PreviousValue, Tag.ArrayIndex);
+			return EConvertFromTypeResult::Converted;
 		}
-
-		// now copy the value into the object's address space
-		SetPropertyValue_InContainer(Data, PreviousValue, Tag.ArrayIndex);
+		return EConvertFromTypeResult::UseSerializeItem;
 	}
-	else if (Tag.Type == NAME_EnumProperty && (Enum == nullptr || Tag.EnumName == Enum->GetFName() || Tag.EnumName.ToString() == Enum->GetPathName()))
+	case NAME_EnumProperty:
 	{
-		// an enum property became a byte
-		// attempt to find the old enum and get the byte value from the serialized enum name
-		uint8 PreviousValue = (uint8)ReadEnumAsInt64(Slot, DefaultsStruct, Tag);
+		if (Enum == nullptr || Tag.EnumName == Enum->GetFName() || Tag.EnumName.ToString() == Enum->GetPathName())
+		{
+			// an enum property became a byte
+			// attempt to find the old enum and get the byte value from the serialized enum name
+			uint8 PreviousValue = (uint8)ReadEnumAsInt64(Slot, DefaultsStruct, Tag);
 
-		// now copy the value into the object's address space
-		SetPropertyValue_InContainer(Data, PreviousValue, Tag.ArrayIndex);
+			// now copy the value into the object's address space
+			SetPropertyValue_InContainer(Data, PreviousValue, Tag.ArrayIndex);
+			return EConvertFromTypeResult::Converted;
+		}
+		return EConvertFromTypeResult::UseSerializeItem;
 	}
-	else if (Tag.Type == NAME_Int8Property)
-	{
+	case NAME_Int8Property:
 		if (Enum)
 		{
 			TConvertIntToEnumProperty<int8>::Convert(Slot, this, Enum, Data, Tag);
@@ -261,9 +279,8 @@ EConvertFromTypeResult FByteProperty::ConvertFromType(const FPropertyTag& Tag, F
 		{
 			ConvertFromArithmeticValue<int8>(Slot, Data, Tag);
 		}
-	}
-	else if (Tag.Type == NAME_Int16Property)
-	{
+		return EConvertFromTypeResult::Converted;
+	case NAME_Int16Property:
 		if (Enum)
 		{
 			TConvertIntToEnumProperty<int16>::Convert(Slot, this, Enum, Data, Tag);
@@ -272,9 +289,8 @@ EConvertFromTypeResult FByteProperty::ConvertFromType(const FPropertyTag& Tag, F
 		{
 			ConvertFromArithmeticValue<int16>(Slot, Data, Tag);
 		}
-	}
-	else if (Tag.Type == NAME_IntProperty)
-	{
+		return EConvertFromTypeResult::Converted;
+	case NAME_IntProperty:
 		if (Enum)
 		{
 			TConvertIntToEnumProperty<int32>::Convert(Slot, this, Enum, Data, Tag);
@@ -283,9 +299,8 @@ EConvertFromTypeResult FByteProperty::ConvertFromType(const FPropertyTag& Tag, F
 		{
 			ConvertFromArithmeticValue<int32>(Slot, Data, Tag);
 		}
-	}
-	else if (Tag.Type == NAME_Int64Property)
-	{
+		return EConvertFromTypeResult::Converted;
+	case NAME_Int64Property:
 		if (Enum)
 		{
 			TConvertIntToEnumProperty<int64>::Convert(Slot, this, Enum, Data, Tag);
@@ -294,9 +309,8 @@ EConvertFromTypeResult FByteProperty::ConvertFromType(const FPropertyTag& Tag, F
 		{
 			ConvertFromArithmeticValue<int64>(Slot, Data, Tag);
 		}
-	}
-	else if (Tag.Type == NAME_UInt16Property)
-	{
+		return EConvertFromTypeResult::Converted;
+	case NAME_UInt16Property:
 		if (Enum)
 		{
 			TConvertIntToEnumProperty<uint16>::Convert(Slot, this, Enum, Data, Tag);
@@ -305,9 +319,8 @@ EConvertFromTypeResult FByteProperty::ConvertFromType(const FPropertyTag& Tag, F
 		{
 			ConvertFromArithmeticValue<uint16>(Slot, Data, Tag);
 		}
-	}
-	else if (Tag.Type == NAME_UInt32Property)
-	{
+		return EConvertFromTypeResult::Converted;
+	case NAME_UInt32Property:
 		if (Enum)
 		{
 			TConvertIntToEnumProperty<uint32>::Convert(Slot, this, Enum, Data, Tag);
@@ -316,9 +329,8 @@ EConvertFromTypeResult FByteProperty::ConvertFromType(const FPropertyTag& Tag, F
 		{
 			ConvertFromArithmeticValue<uint32>(Slot, Data, Tag);
 		}
-	}
-	else if (Tag.Type == NAME_UInt64Property)
-	{
+		return EConvertFromTypeResult::Converted;
+	case NAME_UInt64Property:
 		if (Enum)
 		{
 			TConvertIntToEnumProperty<uint64>::Convert(Slot, this, Enum, Data, Tag);
@@ -327,9 +339,8 @@ EConvertFromTypeResult FByteProperty::ConvertFromType(const FPropertyTag& Tag, F
 		{
 			ConvertFromArithmeticValue<uint64>(Slot, Data, Tag);
 		}
-	}
-	else if (Tag.Type == NAME_BoolProperty)
-	{
+		return EConvertFromTypeResult::Converted;
+	case NAME_BoolProperty:
 		if (Enum)
 		{
 			TConvertIntToEnumProperty<uint64>::ConvertValue(Tag.BoolVal, this, Enum, Data, Tag);
@@ -338,13 +349,8 @@ EConvertFromTypeResult FByteProperty::ConvertFromType(const FPropertyTag& Tag, F
 		{
 			SetPropertyValue_InContainer(Data, Tag.BoolVal, Tag.ArrayIndex);
 		}
+		return EConvertFromTypeResult::Converted;
 	}
-	else
-	{
-		return EConvertFromTypeResult::UseSerializeItem;
-	}
-
-	return EConvertFromTypeResult::Converted;
 }
 
 #if WITH_EDITORONLY_DATA
diff --git a/Engine/Source/Runtime/CoreUObject/Private/UObject/PropertyMap.cpp b/Engine/Source/Runtime/CoreUObject/Private/UObject/PropertyMap.cpp
--- a/Engine/Source/Runtime/CoreUObject/Private/UObject/PropertyMap.cpp
+++ b/Engine/Source/Runtime/CoreUObject/Private/UObject/PropertyMap.cpp
@@ -1421,198 +1421,209 @@ bool FMapProperty::SameType(const FProperty* Other) const
 
 EConvertFromTypeResult FMapProperty::ConvertFromType(const FPropertyTag& Tag, FStructuredArchive::FSlot Slot, uint8* Data, UStruct* DefaultsStruct, const uint8* Defaults)
 {
+	// Ar related calls in this function must be mirrored in FMapProperty::SerializeItem
 	FArchive& UnderlyingArchive = Slot.GetUnderlyingArchive();
 
-	// Ar related calls in this function must be mirrored in FMapProperty::SerializeItem
 	checkSlow(KeyProp);
 	checkSlow(ValueProp);
 
-	const auto SerializeOrConvert = [](FProperty* CurrentType, const FPropertyTag& InTag, FStructuredArchive::FSlot InnerSlot, uint8* InData, UStruct* InDefaultsStruct) -> bool
+	if (FStructProperty* KeyPropAsStruct = CastField<FStructProperty>(KeyProp))
 	{
-		// Serialize wants the property address, while convert wants the container address. InData is the container address
-		if(CurrentType->GetID() == InTag.Type)
+		if (!KeyPropAsStruct->Struct || (KeyPropAsStruct->Struct->GetCppStructOps() && !KeyPropAsStruct->Struct->GetCppStructOps()->HasGetTypeHash()))
 		{
-			uint8* DestAddress = CurrentType->ContainerPtrToValuePtr<uint8>(InData, InTag.ArrayIndex);
-			CurrentType->SerializeItem(InnerSlot, DestAddress, nullptr);
-			return true;
+			// If the type we contain is no longer hashable, we're going to drop the saved data here.
+			// This can happen if the native GetTypeHash function is removed.
+			ensureMsgf(false, TEXT("Map Property %s has an unhashable key type %s and will lose its saved data. Package: %s"),
+				*Tag.Name.ToString(), *KeyPropAsStruct->Struct->GetFName().ToString(), *UnderlyingArchive.GetArchiveName());
+
+			FScriptMapHelper ScriptMapHelper(this, ContainerPtrToValuePtr<void>(Data));
+			ScriptMapHelper.EmptyValues();
+
+			return EConvertFromTypeResult::CannotConvert;
 		}
+	}
 
-		EConvertFromTypeResult ConvertResult = CurrentType->ConvertFromType(InTag, InnerSlot, InData, InDefaultsStruct, nullptr);
-		return ConvertResult == EConvertFromTypeResult::Converted || ConvertResult == EConvertFromTypeResult::Serialized;
+	if (Tag.Type != NAME_MapProperty)
+	{
+		return EConvertFromTypeResult::UseSerializeItem;
+	}
+
+	const bool bKeyTypeMatches = (Tag.InnerType == KeyProp->GetID());
+	const bool bValueTypeMatches = (Tag.ValueType == ValueProp->GetID());
+	if ((bKeyTypeMatches || Tag.InnerType.IsNone()) && (bValueTypeMatches || Tag.ValueType.IsNone()))
+	{
+		return EConvertFromTypeResult::UseSerializeItem;
+	}
+
+	const auto SerializeOrConvert = [](bool bTypeMatches, FProperty* Inner, const FPropertyTag& InnerTag, FStructuredArchive::FSlot InnerSlot, uint8* InnerData, UStruct* InnerDefaultsStruct) -> bool
+	{
+		if (!bTypeMatches)
+		{
+			const int64 StartOfProperty = InnerSlot.GetUnderlyingArchive().Tell();
+			FStructProperty* StructProperty = CastFieldChecked<FStructProperty>(Inner);
+			switch (StructProperty->ConvertFromType(InnerTag, InnerSlot, InnerData, InnerDefaultsStruct, nullptr))
+			{
+				case EConvertFromTypeResult::Converted:
+				case EConvertFromTypeResult::Serialized:
+					return true;
+				case EConvertFromTypeResult::CannotConvert:
+					return false;
+				case EConvertFromTypeResult::UseSerializeItem:
+					// Fall through to default SerializeItem
+					break;
+			}
+		}
+
+		uint8* DestAddress = Inner->ContainerPtrToValuePtr<uint8>(InnerData, InnerTag.ArrayIndex);
+		Inner->SerializeItem(InnerSlot, DestAddress);
+		return true;
 	};
 
-	if (Tag.Type == NAME_MapProperty)
+	FScriptMapHelper MapHelper(this, ContainerPtrToValuePtr<void>(Data));
+
+	uint8* TempKeyValueStorage = nullptr;
+	ON_SCOPE_EXIT
 	{
-		if( (Tag.InnerType != NAME_None && Tag.InnerType != KeyProp->GetID()) || (Tag.ValueType != NAME_None && Tag.ValueType != ValueProp->GetID()) )
+		if (TempKeyValueStorage)
 		{
-			FScriptMapHelper MapHelper(this, ContainerPtrToValuePtr<void>(Data));
+			KeyProp->DestroyValue(TempKeyValueStorage);
+			FMemory::Free(TempKeyValueStorage);
+		}
+	};
 
-			uint8* TempKeyValueStorage = nullptr;
-			ON_SCOPE_EXIT
-			{
-				if (TempKeyValueStorage)
-				{
-					KeyProp->DestroyValue(TempKeyValueStorage);
-					FMemory::Free(TempKeyValueStorage);
-				}
-			};
+	FPropertyTag KeyPropertyTag;
+	KeyPropertyTag.Type = Tag.InnerType;
+	KeyPropertyTag.ArrayIndex = 0;
 
-			FPropertyTag KeyPropertyTag;
-			KeyPropertyTag.Type = Tag.InnerType;
-			KeyPropertyTag.ArrayIndex = 0;
+	FPropertyTag ValuePropertyTag;
+	ValuePropertyTag.Type = Tag.ValueType;
+	ValuePropertyTag.ArrayIndex = 0;
 
-			FPropertyTag ValuePropertyTag;
-			ValuePropertyTag.Type = Tag.ValueType;
-			ValuePropertyTag.ArrayIndex = 0;
+	bool bConversionSucceeded = true;
 
-			bool bConversionSucceeded = true;
+	FStructuredArchive::FRecord ValueRecord = Slot.EnterRecord();
 
-			FStructuredArchive::FRecord ValueRecord = Slot.EnterRecord();
+	// When we saved this instance we wrote out any elements that were in the 'Default' instance but not in the 
+	// instance that was being written. Presumably we were constructed from our defaults and must now remove 
+	// any of the elements that were not present when we saved this Map:
+	int32 NumKeysToRemove = 0;
+	FStructuredArchive::FArray KeysToRemoveArray = ValueRecord.EnterArray(TEXT("KeysToRemove"), NumKeysToRemove);
 
-			// When we saved this instance we wrote out any elements that were in the 'Default' instance but not in the 
-			// instance that was being written. Presumably we were constructed from our defaults and must now remove 
-			// any of the elements that were not present when we saved this Map:
-			int32 NumKeysToRemove = 0;
-			FStructuredArchive::FArray KeysToRemoveArray = ValueRecord.EnterArray(TEXT("KeysToRemove"), NumKeysToRemove);
+	if (NumKeysToRemove)
+	{
+		TempKeyValueStorage = (uint8*)FMemory::Malloc(MapLayout.SetLayout.Size);
+		KeyProp->InitializeValue(TempKeyValueStorage);
 
-			if( NumKeysToRemove != 0 )
+		if (SerializeOrConvert(bKeyTypeMatches, KeyProp, KeyPropertyTag, KeysToRemoveArray.EnterElement(), TempKeyValueStorage, DefaultsStruct))
+		{
+			// If the key is in the map, remove it
+			int32 Found = MapHelper.FindMapIndexWithKey(TempKeyValueStorage);
+			if (Found != INDEX_NONE)
 			{
-				TempKeyValueStorage = (uint8*)FMemory::Malloc(MapLayout.SetLayout.Size);
-				KeyProp->InitializeValue(TempKeyValueStorage);
-
-				if (SerializeOrConvert( KeyProp, KeyPropertyTag, KeysToRemoveArray.EnterElement(), TempKeyValueStorage, DefaultsStruct))
-				{
-					// If the key is in the map, remove it
-					int32 Found = MapHelper.FindMapIndexWithKey(TempKeyValueStorage);
-					if (Found != INDEX_NONE)
-					{
-						MapHelper.RemoveAt(Found);
-					}
+				MapHelper.RemoveAt(Found);
+			}
 
-					// things are going fine, remove the rest of the keys:
-					for(int32 I = 1; I < NumKeysToRemove; ++I)
-					{
-						verify(SerializeOrConvert( KeyProp, KeyPropertyTag, KeysToRemoveArray.EnterElement(), TempKeyValueStorage, DefaultsStruct));
-						Found = MapHelper.FindMapIndexWithKey(TempKeyValueStorage);
-						if (Found != INDEX_NONE)
-						{
-							MapHelper.RemoveAt(Found);
-						}
-					}
-				}
-				else
+			// things are going fine, remove the rest of the keys:
+			for (int32 I = 1; I < NumKeysToRemove; ++I)
+			{
+				verify(SerializeOrConvert(bKeyTypeMatches, KeyProp, KeyPropertyTag, KeysToRemoveArray.EnterElement(), TempKeyValueStorage, DefaultsStruct));
+				Found = MapHelper.FindMapIndexWithKey(TempKeyValueStorage);
+				if (Found != INDEX_NONE)
 				{
-					bConversionSucceeded = false;
+					MapHelper.RemoveAt(Found);
 				}
 			}
+		}
+		else
+		{
+			bConversionSucceeded = false;
+		}
+	}
+
+	int32 Num = 0;
+	FStructuredArchive::FArray EntriesArray = ValueRecord.EnterArray(TEXT("Entries"), Num);
+
+	if (bConversionSucceeded)
+	{
+		if (Num != 0)
+		{
+			if (TempKeyValueStorage == nullptr)
+			{
+				TempKeyValueStorage = (uint8*)FMemory::Malloc(MapLayout.SetLayout.Size);
+				KeyProp->InitializeValue(TempKeyValueStorage);
+			}
 
-			int32 NumEntries = 0;
-			FStructuredArchive::FArray EntriesArray = ValueRecord.EnterArray(TEXT("Entries"), NumEntries);
+			FStructuredArchive::FRecord FirstPropertyRecord = EntriesArray.EnterElement().EnterRecord();
 
-			if( bConversionSucceeded )
+			if (SerializeOrConvert(bKeyTypeMatches, KeyProp, KeyPropertyTag, FirstPropertyRecord.EnterField(TEXT("Key")), TempKeyValueStorage, DefaultsStruct))
 			{
-				if( NumEntries != 0 )
+				// Add a new default value if the key doesn't currently exist in the map
+				bool bKeyAlreadyPresent = true;
+				int32 NextPairIndex = MapHelper.FindMapIndexWithKey(TempKeyValueStorage);
+				if (NextPairIndex == INDEX_NONE)
 				{
-					if(TempKeyValueStorage == nullptr )
-					{
-						TempKeyValueStorage = (uint8*)FMemory::Malloc(MapLayout.SetLayout.Size);
-						KeyProp->InitializeValue(TempKeyValueStorage);
-					}
+					bKeyAlreadyPresent = false;
+					NextPairIndex = MapHelper.AddDefaultValue_Invalid_NeedsRehash();
+				}
 
-					FStructuredArchive::FRecord FirstPropertyRecord = EntriesArray.EnterElement().EnterRecord();
+				uint8* NextPairPtr = MapHelper.GetPairPtrWithoutCheck(NextPairIndex);
+				// This copy is unnecessary when the key was already in the map:
+				KeyProp->CopyCompleteValue_InContainer(NextPairPtr, TempKeyValueStorage);
 
-					if( SerializeOrConvert( KeyProp, KeyPropertyTag, FirstPropertyRecord.EnterField(TEXT("Key")), TempKeyValueStorage, DefaultsStruct ) )
+				// Deserialize value
+				if (SerializeOrConvert(bValueTypeMatches, ValueProp, ValuePropertyTag, FirstPropertyRecord.EnterField(TEXT("Value")), NextPairPtr, DefaultsStruct))
+				{
+					// first entry went fine, convert the rest:
+					for (int32 I = 1; I < Num; ++I)
 					{
-						// Add a new default value if the key doesn't currently exist in the map
-						bool bKeyAlreadyPresent = true;
-						int32 NextPairIndex = MapHelper.FindMapIndexWithKey(TempKeyValueStorage);
+						FStructuredArchive::FRecord PropertyRecord = EntriesArray.EnterElement().EnterRecord();
+
+						verify(SerializeOrConvert(bKeyTypeMatches, KeyProp, KeyPropertyTag, PropertyRecord.EnterField(TEXT("Key")), TempKeyValueStorage, DefaultsStruct));
+						NextPairIndex = MapHelper.FindMapIndexWithKey(TempKeyValueStorage);
 						if (NextPairIndex == INDEX_NONE)
 						{
-							bKeyAlreadyPresent = false;
 							NextPairIndex = MapHelper.AddDefaultValue_Invalid_NeedsRehash();
 						}
 
-						uint8* NextPairPtr = MapHelper.GetPairPtrWithoutCheck(NextPairIndex);
+						NextPairPtr = MapHelper.GetPairPtrWithoutCheck(NextPairIndex);
 						// This copy is unnecessary when the key was already in the map:
 						KeyProp->CopyCompleteValue_InContainer(NextPairPtr, TempKeyValueStorage);
-
-						// Deserialize value
-						if( SerializeOrConvert( ValueProp, ValuePropertyTag, FirstPropertyRecord.EnterField(TEXT("Value")), NextPairPtr, DefaultsStruct ) )
-						{
-							// first entry went fine, convert the rest:
-							for(int32 I = 1; I < NumEntries; ++I)
-							{
-								FStructuredArchive::FRecord PropertyRecord = EntriesArray.EnterElement().EnterRecord();
-
-								verify( SerializeOrConvert( KeyProp, KeyPropertyTag, PropertyRecord.EnterField(TEXT("Key")), TempKeyValueStorage, DefaultsStruct ) );
-								NextPairIndex = MapHelper.FindMapIndexWithKey(TempKeyValueStorage);
-								if (NextPairIndex == INDEX_NONE)
-								{
-									NextPairIndex = MapHelper.AddDefaultValue_Invalid_NeedsRehash();
-								}
-
-								NextPairPtr = MapHelper.GetPairPtrWithoutCheck(NextPairIndex);
-								// This copy is unnecessary when the key was already in the map:
-								KeyProp->CopyCompleteValue_InContainer(NextPairPtr, TempKeyValueStorage);
-								verify( SerializeOrConvert( ValueProp, ValuePropertyTag, PropertyRecord.EnterField(TEXT("Value")), NextPairPtr, DefaultsStruct ) );
-							}
-						}
-						else
-						{
-							if(!bKeyAlreadyPresent)
-							{
-								MapHelper.EmptyValues();
-							}
-
-							bConversionSucceeded = false;
-						}
+						verify(SerializeOrConvert(bValueTypeMatches, ValueProp, ValuePropertyTag, PropertyRecord.EnterField(TEXT("Value")), NextPairPtr, DefaultsStruct));
 					}
-					else
+				}
+				else
+				{
+					if (!bKeyAlreadyPresent)
 					{
-						bConversionSucceeded = false;
+						MapHelper.EmptyValues();
 					}
 
-					MapHelper.Rehash();
+					bConversionSucceeded = false;
 				}
 			}
-
-			// if we could not convert the property ourself, then indicate that calling code needs to advance the property
-			if(!bConversionSucceeded)
+			else
 			{
-				UE_LOG(
-					LogClass,
-					Warning,
-					TEXT("Map Element Type mismatch in %s of %s - Previous (%s to %s) Current (%s to %s) for package: %s"),
-					*Tag.Name.ToString(),
-					*GetName(),
-					*Tag.InnerType.ToString(),
-					*Tag.ValueType.ToString(),
-					*KeyProp->GetID().ToString(),
-					*ValueProp->GetID().ToString(),
-					*UnderlyingArchive.GetArchiveName()
-				);
+				bConversionSucceeded = false;
 			}
 
-			return bConversionSucceeded ? EConvertFromTypeResult::Converted : EConvertFromTypeResult::CannotConvert;
+			MapHelper.Rehash();
 		}
+	}
 
-		if (FStructProperty* KeyPropAsStruct = CastField<FStructProperty>(KeyProp))
-		{
-			if(!KeyPropAsStruct->Struct || (KeyPropAsStruct->Struct->GetCppStructOps() && !KeyPropAsStruct->Struct->GetCppStructOps()->HasGetTypeHash() ) )
-			{
-				// If the type we contain is no longer hashable, we're going to drop the saved data here. This can
-				// happen if the native GetTypeHash function is removed.
-				ensureMsgf(false, TEXT("FMapProperty %s with tag %s has an unhashable key type %s and will lose its saved data"), *GetName(), *Tag.Name.ToString(), *KeyProp->GetID().ToString());
-
-				FScriptMapHelper ScriptMapHelper(this, ContainerPtrToValuePtr<void>(Data));
-				ScriptMapHelper.EmptyValues();
-
-				return EConvertFromTypeResult::CannotConvert;
-			}
-		}
+	// if we could not convert the property ourself, then indicate that calling code needs to advance the property
+	if (!bConversionSucceeded)
+	{
+		UE_LOG(LogClass, Warning,
+			TEXT("Map Element Type mismatch in %s - Previous (%s to %s) Current (%s to %s) for package: %s"),
+			*WriteToString<32>(Tag.Name),
+			*WriteToString<32>(KeyPropertyTag.Type),
+			*WriteToString<32>(ValuePropertyTag.Type),
+			*WriteToString<32>(KeyProp->GetID()),
+			*WriteToString<32>(ValueProp->GetID()),
+			*UnderlyingArchive.GetArchiveName());
 	}
 
-	return EConvertFromTypeResult::UseSerializeItem;
+	return bConversionSucceeded ? EConvertFromTypeResult::Converted : EConvertFromTypeResult::CannotConvert;
 }
 
 #if WITH_EDITORONLY_DATA
diff --git a/Engine/Source/Runtime/CoreUObject/Private/UObject/PropertySet.cpp b/Engine/Source/Runtime/CoreUObject/Private/UObject/PropertySet.cpp
--- a/Engine/Source/Runtime/CoreUObject/Private/UObject/PropertySet.cpp
+++ b/Engine/Source/Runtime/CoreUObject/Private/UObject/PropertySet.cpp
@@ -888,155 +888,155 @@ bool FSetProperty::SameType(const FProperty* Other) const
 
 EConvertFromTypeResult FSetProperty::ConvertFromType(const FPropertyTag& Tag, FStructuredArchive::FSlot Slot, uint8* Data, UStruct* DefaultsStruct, const uint8* Defaults)
 {
+	// Ar related calls in this function must be mirrored in FSetProperty::SerializeItem
 	FArchive& UnderlyingArchive = Slot.GetUnderlyingArchive();
 
-	// Ar related calls in this function must be mirrored in FSetProperty::ConvertFromType
 	checkSlow(ElementProp);
 
-	// Ar related calls in this function must be mirrored in FSetProperty::SerializeItem
-	if (Tag.Type == NAME_SetProperty)
+	if (FStructProperty* ElementPropAsStruct = CastField<FStructProperty>(ElementProp))
 	{
-		if (Tag.InnerType != NAME_None && Tag.InnerType != ElementProp->GetID())
+		if (!ElementPropAsStruct->Struct || (ElementPropAsStruct->Struct->GetCppStructOps() && !ElementPropAsStruct->Struct->GetCppStructOps()->HasGetTypeHash()))
 		{
+			// If the type we contain is no longer hashable, we're going to drop the saved data here.
+			// This can happen if the native GetTypeHash function is removed.
+			ensureMsgf(false, TEXT("Set Property %s has an unhashable type %s and will lose its saved data. Package: %s"),
+				*Tag.Name.ToString(), *ElementPropAsStruct->Struct->GetFName().ToString(), *UnderlyingArchive.GetArchiveName());
+
 			FScriptSetHelper ScriptSetHelper(this, ContainerPtrToValuePtr<void>(Data));
+			ScriptSetHelper.EmptyElements();
 
-			uint8* TempElementStorage = nullptr;
-			ON_SCOPE_EXIT
-			{
-				if (TempElementStorage)
-				{
-					ElementProp->DestroyValue(TempElementStorage);
-					FMemory::Free(TempElementStorage);
-				}
-			};
+			return EConvertFromTypeResult::CannotConvert;
+		}
+	}
 
-			FPropertyTag InnerPropertyTag;
-			InnerPropertyTag.Type = Tag.InnerType;
-			InnerPropertyTag.ArrayIndex = 0;
+	if (Tag.Type != NAME_SetProperty)
+	{
+		return EConvertFromTypeResult::UseSerializeItem;
+	}
 
-			bool bConversionSucceeded = true;
+	if (Tag.InnerType.IsNone() || Tag.InnerType == ElementProp->GetID())
+	{
+		return EConvertFromTypeResult::UseSerializeItem;
+	}
 
-			FStructuredArchive::FRecord ValueRecord = Slot.EnterRecord();
+	FScriptSetHelper ScriptSetHelper(this, ContainerPtrToValuePtr<void>(Data));
 
-			// When we saved this instance we wrote out any elements that were in the 'Default' instance but not in the 
-			// instance that was being written. Presumably we were constructed from our defaults and must now remove 
-			// any of the elements that were not present when we saved this Set:
-			int32 NumElementsToRemove = 0;
-			FStructuredArchive::FArray ElementsToRemoveArray = ValueRecord.EnterArray(TEXT("ElementsToRemove"), NumElementsToRemove);
+	uint8* TempElementStorage = nullptr;
+	ON_SCOPE_EXIT
+	{
+		if (TempElementStorage)
+		{
+			ElementProp->DestroyAndFreeValue(TempElementStorage);
+		}
+	};
 
-			if(NumElementsToRemove)
+	FPropertyTag InnerPropertyTag;
+	InnerPropertyTag.Type = Tag.InnerType;
+	InnerPropertyTag.ArrayIndex = 0;
+
+	bool bConversionSucceeded = true;
+
+	FStructuredArchive::FRecord ValueRecord = Slot.EnterRecord();
+
+	// When we saved this instance we wrote out any elements that were in the 'Default' instance but not in the
+	// instance that was being written. Presumably we were constructed from our defaults and must now remove
+	// any of the elements that were not present when we saved this Set:
+	int32 NumElementsToRemove = 0;
+	FStructuredArchive::FArray ElementsToRemoveArray = ValueRecord.EnterArray(TEXT("ElementsToRemove"), NumElementsToRemove);
+
+	if (NumElementsToRemove)
+	{
+		TempElementStorage = (uint8*)ElementProp->AllocateAndInitializeValue();
+
+		EConvertFromTypeResult ConvertResult = ElementProp->ConvertFromType(InnerPropertyTag, ElementsToRemoveArray.EnterElement(), TempElementStorage, DefaultsStruct, nullptr);
+		if (ConvertResult == EConvertFromTypeResult::Converted || ConvertResult == EConvertFromTypeResult::Serialized)
+		{
+			int32 Found = ScriptSetHelper.FindElementIndex(TempElementStorage);
+			if (Found != INDEX_NONE)
 			{
-				TempElementStorage = (uint8*)FMemory::Malloc(SetLayout.Size);
-				ElementProp->InitializeValue(TempElementStorage);
+				ScriptSetHelper.RemoveAt(Found);
+			}
 
-				EConvertFromTypeResult ConvertResult = ElementProp->ConvertFromType(InnerPropertyTag, ElementsToRemoveArray.EnterElement(), TempElementStorage, DefaultsStruct, nullptr);
-				if (ConvertResult == EConvertFromTypeResult::Converted || ConvertResult == EConvertFromTypeResult::Serialized)
-				{
-					int32 Found = ScriptSetHelper.FindElementIndex(TempElementStorage);
-					if (Found != INDEX_NONE)
-					{
-						ScriptSetHelper.RemoveAt(Found);
-					}
+			for (int32 I = 1; I < NumElementsToRemove; ++I)
+			{
+				ConvertResult = ElementProp->ConvertFromType(InnerPropertyTag, ElementsToRemoveArray.EnterElement(), TempElementStorage, DefaultsStruct, nullptr);
+				check(ConvertResult == EConvertFromTypeResult::Converted || ConvertResult == EConvertFromTypeResult::Serialized);
 
-					for (int32 I = 1; I < NumElementsToRemove; ++I)
-					{
-						ConvertResult = ElementProp->ConvertFromType(InnerPropertyTag, ElementsToRemoveArray.EnterElement(), TempElementStorage, DefaultsStruct, nullptr);
-						check(ConvertResult == EConvertFromTypeResult::Converted || ConvertResult == EConvertFromTypeResult::Serialized);
-
-						Found = ScriptSetHelper.FindElementIndex(TempElementStorage);
-						if (Found != INDEX_NONE)
-						{
-							ScriptSetHelper.RemoveAt(Found);
-						}
-					}
-				}
-				else
+				Found = ScriptSetHelper.FindElementIndex(TempElementStorage);
+				if (Found != INDEX_NONE)
 				{
-					bConversionSucceeded = false;
+					ScriptSetHelper.RemoveAt(Found);
 				}
 			}
+		}
+		else
+		{
+			bConversionSucceeded = false;
+		}
+	}
 
-			int32 Num = 0;
-			FStructuredArchive::FArray ElementsArray = ValueRecord.EnterArray(TEXT("Elements"), Num);
+	int32 Num = 0;
+	FStructuredArchive::FArray ElementsArray = ValueRecord.EnterArray(TEXT("Elements"), Num);
+
+	if (bConversionSucceeded)
+	{
+		if (Num != 0)
+		{
+			// Allocate temporary key space if we haven't allocated it already above
+			if (TempElementStorage == nullptr)
+			{
+				TempElementStorage = (uint8*)ElementProp->AllocateAndInitializeValue();
+			}
 
-			if(bConversionSucceeded)
+			// and read the first entry, we have to check for conversion possibility again because 
+			// NumElementsToRemove may not have run (in fact, it likely did not):
+			EConvertFromTypeResult ConvertResult = ElementProp->ConvertFromType(InnerPropertyTag, ElementsArray.EnterElement(), TempElementStorage, DefaultsStruct, nullptr);
+			if (ConvertResult == EConvertFromTypeResult::Converted || ConvertResult == EConvertFromTypeResult::Serialized)
 			{
-				if (Num != 0)
+				if (ScriptSetHelper.FindElementIndex(TempElementStorage) == INDEX_NONE)
 				{
-					// Allocate temporary key space if we haven't allocated it already above
-					if( TempElementStorage == nullptr )
-					{
-						TempElementStorage = (uint8*)FMemory::Malloc(SetLayout.Size);
-						ElementProp->InitializeValue(TempElementStorage);
-					}
+					const int32 NewElementIndex = ScriptSetHelper.AddDefaultValue_Invalid_NeedsRehash();
+					uint8* NewElementPtr = ScriptSetHelper.GetElementPtrWithoutCheck(NewElementIndex);
 
-					// and read the first entry, we have to check for conversion possibility again because 
-					// NumElementsToRemove may not have run (in fact, it likely did not):
-					EConvertFromTypeResult ConvertResult = ElementProp->ConvertFromType(InnerPropertyTag, ElementsArray.EnterElement(), TempElementStorage, DefaultsStruct, nullptr);
-					if (ConvertResult == EConvertFromTypeResult::Converted || ConvertResult == EConvertFromTypeResult::Serialized)
-					{
-						if (ScriptSetHelper.FindElementIndex(TempElementStorage) == INDEX_NONE)
-						{
-							const int32 NewElementIndex = ScriptSetHelper.AddDefaultValue_Invalid_NeedsRehash();
-							uint8* NewElementPtr = ScriptSetHelper.GetElementPtrWithoutCheck(NewElementIndex);
-
-							// Copy over deserialized key from temporary storage
-							ElementProp->CopyCompleteValue_InContainer(NewElementPtr, TempElementStorage);
-						}
-
-						// Read remaining items into container
-						for (int32 I = 1; I < Num; ++I)
-						{
-							// Read key into temporary storage
-							ConvertResult = ElementProp->ConvertFromType(InnerPropertyTag, ElementsArray.EnterElement(), TempElementStorage, DefaultsStruct, nullptr);
-							check(ConvertResult == EConvertFromTypeResult::Converted || ConvertResult == EConvertFromTypeResult::Serialized);
-
-							// Add a new entry if the element doesn't currently exist in the set
-							if (ScriptSetHelper.FindElementIndex(TempElementStorage) == INDEX_NONE)
-							{
-								const int32 NewElementIndex = ScriptSetHelper.AddDefaultValue_Invalid_NeedsRehash();
-								uint8* NewElementPtr = ScriptSetHelper.GetElementPtrWithoutCheck(NewElementIndex);
-
-								// Copy over deserialized key from temporary storage
-								ElementProp->CopyCompleteValue_InContainer(NewElementPtr, TempElementStorage);
-							}
-						}
-					}
-					else
+					// Copy over deserialized key from temporary storage
+					ElementProp->CopyCompleteValue_InContainer(NewElementPtr, TempElementStorage);
+				}
+
+				// Read remaining items into container
+				for (int32 I = 1; I < Num; ++I)
+				{
+					// Read key into temporary storage
+					ConvertResult = ElementProp->ConvertFromType(InnerPropertyTag, ElementsArray.EnterElement(), TempElementStorage, DefaultsStruct, nullptr);
+					check(ConvertResult == EConvertFromTypeResult::Converted || ConvertResult == EConvertFromTypeResult::Serialized);
+
+					// Add a new entry if the element doesn't currently exist in the set
+					if (ScriptSetHelper.FindElementIndex(TempElementStorage) == INDEX_NONE)
 					{
-						bConversionSucceeded = false;
+						const int32 NewElementIndex = ScriptSetHelper.AddDefaultValue_Invalid_NeedsRehash();
+						uint8* NewElementPtr = ScriptSetHelper.GetElementPtrWithoutCheck(NewElementIndex);
+
+						// Copy over deserialized key from temporary storage
+						ElementProp->CopyCompleteValue_InContainer(NewElementPtr, TempElementStorage);
 					}
 				}
-
-				ScriptSetHelper.Rehash();
 			}
-
-			// if we could not convert the property ourself, then indicate that calling code needs to advance the property
-			if(!bConversionSucceeded)
+			else
 			{
-				UE_LOG(LogClass, Warning, TEXT("Set Element Type mismatch in %s of %s - Previous (%s) Current (%s) for package: %s"), *Tag.Name.ToString(), *GetName(), *Tag.InnerType.ToString(), *ElementProp->GetID().ToString(), *UnderlyingArchive.GetArchiveName() );
+				bConversionSucceeded = false;
 			}
 
-			return bConversionSucceeded ? EConvertFromTypeResult::Converted : EConvertFromTypeResult::CannotConvert;
+			ScriptSetHelper.Rehash();
 		}
+	}
 
-		if(FStructProperty* ElementPropAsStruct = CastField<FStructProperty>(ElementProp))
-		{
-			if(!ElementPropAsStruct->Struct || (ElementPropAsStruct->Struct->GetCppStructOps() && !ElementPropAsStruct->Struct->GetCppStructOps()->HasGetTypeHash()) )
-			{
-				// If the type we contain is no longer hashable, we're going to drop the saved data here. This can
-				// happen if the native GetTypeHash function is removed.
-				ensureMsgf(false, TEXT("FSetProperty %s with tag %s has an unhashable type %s and will lose its saved data"), *GetName(), *Tag.Name.ToString(), *ElementProp->GetID().ToString());
-
-				FScriptSetHelper ScriptSetHelper(this, ContainerPtrToValuePtr<void>(Data));
-				ScriptSetHelper.EmptyElements();
-
-				return EConvertFromTypeResult::CannotConvert;
-			}
-		}
+	// if we could not convert the property ourself, then indicate that calling code needs to advance the property
+	if (!bConversionSucceeded)
+	{
+		UE_LOG(LogClass, Warning, TEXT("Set Element Type mismatch in %s - Previous (%s) Current (%s) for package: %s"),
+			*WriteToString<32>(Tag.Name), *WriteToString<32>(InnerPropertyTag.Type), *WriteToString<32>(ElementProp->GetID()), *UnderlyingArchive.GetArchiveName());
 	}
 
-	return EConvertFromTypeResult::UseSerializeItem;
+	return bConversionSucceeded ? EConvertFromTypeResult::Converted : EConvertFromTypeResult::CannotConvert;
 }
 
 #if WITH_EDITORONLY_DATA
