Fixed an crash when using the set editor property bp node on a property that might not compatible with the pinned property type because the function PropertyAccessUtil::ArePropertiesCompatible is not robust enough. Exemple using it to write a dynamic array to a static array property or the inverse.

diff --git a/Engine/Source/Runtime/CoreUObject/Private/UObject/PropertyAccessUtil.cpp b/Engine/Source/Runtime/CoreUObject/Private/UObject/PropertyAccessUtil.cpp
--- a/Engine/Source/Runtime/CoreUObject/Private/UObject/PropertyAccessUtil.cpp
+++ b/Engine/Source/Runtime/CoreUObject/Private/UObject/PropertyAccessUtil.cpp
@@ -174,18 +174,19 @@ bool ArePropertiesCompatible(const FProperty* InSrcProp, const FProperty* InDest
 	// Containers also need to check their inner types
 	if (const FArrayProperty* SrcArrayProp = CastField<FArrayProperty>(InSrcProp))
 	{
-		const FArrayProperty* DestArrayProp = CastFieldChecked<FArrayProperty>(InDestProp);
-		return ArePropertiesCompatible(SrcArrayProp->Inner, DestArrayProp->Inner);
+		const FArrayProperty* DestArrayProp = CastField<FArrayProperty>(InDestProp);
+		return DestArrayProp && ArePropertiesCompatible(SrcArrayProp->Inner, DestArrayProp->Inner);
 	}
 	if (const FSetProperty* SrcSetProp = CastField<FSetProperty>(InSrcProp))
 	{
-		const FSetProperty* DestSetProp = CastFieldChecked<FSetProperty>(InDestProp);
-		return ArePropertiesCompatible(SrcSetProp->ElementProp, DestSetProp->ElementProp);
+		const FSetProperty* DestSetProp = CastField<FSetProperty>(InDestProp);
+		return DestSetProp && ArePropertiesCompatible(SrcSetProp->ElementProp, DestSetProp->ElementProp);
 	}
 	if (const FMapProperty* SrcMapProp = CastField<FMapProperty>(InSrcProp))
 	{
-		const FMapProperty* DestMapProp = CastFieldChecked<FMapProperty>(InDestProp);
-		return ArePropertiesCompatible(SrcMapProp->KeyProp, DestMapProp->KeyProp)
+		const FMapProperty* DestMapProp = CastField<FMapProperty>(InDestProp);
+		return DestMapProp
+			&& ArePropertiesCompatible(SrcMapProp->KeyProp, DestMapProp->KeyProp)
 			&& ArePropertiesCompatible(SrcMapProp->ValueProp, DestMapProp->ValueProp);
 	}
 
