Fix edit property chain in overridable manager when changing values on a subobject

diff --git a/Engine/Source/Runtime/CoreUObject/Private/UObject/OverridableManager.cpp b/Engine/Source/Runtime/CoreUObject/Private/UObject/OverridableManager.cpp
--- a/Engine/Source/Runtime/CoreUObject/Private/UObject/OverridableManager.cpp
+++ b/Engine/Source/Runtime/CoreUObject/Private/UObject/OverridableManager.cpp
@@ -129,8 +129,8 @@ void FOverridableManager::OverrideProperty(UObject& Object, const FPropertyChang
 {
 	if (FOverriddenPropertySet* ThisObjectOverriddenProperties = OverriddenObjectAnnotations.Find(Object))
 	{
-		ThisObjectOverriddenProperties->NotifyPropertyChange(EPropertyNotificationType::PreEdit, FPropertyChangedEvent(nullptr), PropertyChain.GetHead(), &Object);
-		ThisObjectOverriddenProperties->NotifyPropertyChange(EPropertyNotificationType::PostEdit, PropertyEvent, PropertyChain.GetHead(), &Object);
+		ThisObjectOverriddenProperties->NotifyPropertyChange(EPropertyNotificationType::PreEdit, FPropertyChangedEvent(nullptr), PropertyChain.GetActiveNode() ? PropertyChain.GetActiveNode() : PropertyChain.GetHead(), &Object);
+		ThisObjectOverriddenProperties->NotifyPropertyChange(EPropertyNotificationType::PostEdit, PropertyEvent, PropertyChain.GetActiveNode() ? PropertyChain.GetActiveNode() : PropertyChain.GetHead(), &Object);
 	}
 }
 
@@ -147,7 +147,7 @@ void FOverridableManager::PreOverrideProperty(UObject& Object, const FEditProper
 {
 	if (FOverriddenPropertySet* ThisObjectOverriddenProperties = OverriddenObjectAnnotations.Find(Object))
 	{
-		ThisObjectOverriddenProperties->NotifyPropertyChange(EPropertyNotificationType::PreEdit, FPropertyChangedEvent(nullptr), PropertyChain.GetHead(), &Object);
+		ThisObjectOverriddenProperties->NotifyPropertyChange(EPropertyNotificationType::PreEdit, FPropertyChangedEvent(nullptr), PropertyChain.GetActiveNode() ? PropertyChain.GetActiveNode() : PropertyChain.GetHead(), &Object);
 	}
 }
 
@@ -155,7 +155,7 @@ void FOverridableManager::PostOverrideProperty(UObject& Object, const FPropertyC
 {
 	if (FOverriddenPropertySet* ThisObjectOverriddenProperties = OverriddenObjectAnnotations.Find(Object))
 	{
-		ThisObjectOverriddenProperties->NotifyPropertyChange(EPropertyNotificationType::PostEdit, PropertyEvent, PropertyChain.GetHead(), &Object);
+		ThisObjectOverriddenProperties->NotifyPropertyChange(EPropertyNotificationType::PostEdit, PropertyEvent, PropertyChain.GetActiveNode() ? PropertyChain.GetActiveNode() : PropertyChain.GetHead(), &Object);
 	}
 }
 
diff --git a/Engine/Source/Runtime/CoreUObject/Public/UObject/OverridableManager.h b/Engine/Source/Runtime/CoreUObject/Public/UObject/OverridableManager.h
--- a/Engine/Source/Runtime/CoreUObject/Public/UObject/OverridableManager.h
+++ b/Engine/Source/Runtime/CoreUObject/Public/UObject/OverridableManager.h
@@ -179,7 +179,7 @@ public:
 	 * @return true if the property was successfully cleared. */
 	FORCEINLINE bool ClearOverriddenProperty(UObject& Object, const FPropertyChangedEvent& PropertyEvent, const FEditPropertyChain& PropertyChain)
 	{
-		return ClearOverriddenProperty(Object, PropertyEvent, PropertyChain.GetHead());
+		return ClearOverriddenProperty(Object, PropertyEvent, PropertyChain.GetActiveNode() ? PropertyChain.GetActiveNode() : PropertyChain.GetHead());
 	}
 
 	/**
@@ -223,7 +223,7 @@ public:
 	 * @return the current type of override operation on the property */
 	FORCEINLINE EOverriddenPropertyOperation GetOverriddenPropertyOperation(UObject& Object, const FPropertyChangedEvent& PropertyEvent, const FEditPropertyChain& PropertyChain, bool* bOutInheritedOperation = nullptr)
 	{
-		return GetOverriddenPropertyOperation(Object, PropertyEvent, PropertyChain.GetHead(), bOutInheritedOperation);
+		return GetOverriddenPropertyOperation(Object, PropertyEvent, PropertyChain.GetActiveNode() ? PropertyChain.GetActiveNode() : PropertyChain.GetHead(), bOutInheritedOperation);
 	}
 
 	/**
