Ignore properties that won't be serialized when determining whether we can discard a class's sparse data

diff --git a/Engine/Source/Runtime/CoreUObject/Private/Serialization/PropertyLocalizationDataGathering.cpp b/Engine/Source/Runtime/CoreUObject/Private/Serialization/PropertyLocalizationDataGathering.cpp
--- a/Engine/Source/Runtime/CoreUObject/Private/Serialization/PropertyLocalizationDataGathering.cpp
+++ b/Engine/Source/Runtime/CoreUObject/Private/Serialization/PropertyLocalizationDataGathering.cpp
@@ -10,6 +10,7 @@
 #include "UObject/TextProperty.h"
 #include "UObject/PropertyPortFlags.h"
 #include "UObject/PropertyOptional.h"
+#include "UObject/SparseClassDataUtils.h"
 #include "HAL/UnrealMemory.h"
 #include "Internationalization/TextNamespaceUtil.h"
 #include "Internationalization/TextPackageNamespaceUtil.h"
@@ -227,7 +228,11 @@ void FPropertyLocalizationDataGatherer::GatherLocalizationDataFromObject(const U
 	// Also gather from the sparse data on UClass types.
 	if (const UClass* Class = Cast<UClass>(Object))
 	{
-		if (Class->OverridesSparseClassDataArchetype())
+		const auto FilterTransientProperties = [](const FProperty* P)
+			{
+				return !P->HasAnyPropertyFlags(CPF_Transient);
+			};
+		if (UE::Reflection::CompareSparseClassDataToArchetype(Class, FilterTransientProperties))
 		{
 			const UScriptStruct* SparseDataStruct = Class->GetSparseClassDataStruct();
 			const void* SparseData = const_cast<UClass*>(Class)->GetOrCreateSparseClassData();
diff --git a/Engine/Source/Runtime/CoreUObject/Private/UObject/SparseClassDataUtils.cpp b/Engine/Source/Runtime/CoreUObject/Private/UObject/SparseClassDataUtils.cpp
new file mode 100644
--- /dev/null
+++ b/Engine/Source/Runtime/CoreUObject/Private/UObject/SparseClassDataUtils.cpp
@@ -0,0 +1,38 @@
+// Copyright Epic Games, Inc. All Rights Reserved.
+
+#include "UObject/SparseClassDataUtils.h"
+#include "UObject/Class.h"
+#include "UObject/UnrealType.h"
+#include "Templates/Function.h"
+
+bool UE::Reflection::CompareSparseClassDataToArchetype(const UClass* Class, const TFunctionRef<bool(FProperty*)>& Filter)
+{
+	const UScriptStruct* SparseClassDataStruct = Class->GetSparseClassDataStruct();
+	const void* SCD = const_cast<UClass*>(Class)->GetSparseClassData(EGetSparseClassDataMethod::ReturnIfNull);
+	if (SparseClassDataStruct == nullptr ||
+		SCD == nullptr)
+	{
+		return false;
+	}
+
+	if (SparseClassDataStruct != Class->GetSparseClassDataArchetypeStruct())
+	{
+		return true;
+	}
+
+	const void* ArchetypeSCD = Class->GetArchetypeForSparseClassData();
+	if (!ArchetypeSCD)
+	{
+		return true;
+	}
+
+	for (TFieldIterator<FProperty> It(Class); It; ++It)
+	{
+		if (Filter(*It) && !It->Identical_InContainer(ArchetypeSCD, SCD))
+		{
+			return false;
+		}
+	}
+
+	return true;
+}
diff --git a/Engine/Source/Runtime/CoreUObject/Public/UObject/Class.h b/Engine/Source/Runtime/CoreUObject/Public/UObject/Class.h
--- a/Engine/Source/Runtime/CoreUObject/Public/UObject/Class.h
+++ b/Engine/Source/Runtime/CoreUObject/Public/UObject/Class.h
@@ -3462,6 +3462,7 @@ public:
 	COREUOBJECT_API UScriptStruct* GetSparseClassDataArchetypeStruct() const;
 
 	/** Returns whether the sparse class data on this instance overrides that of its archetype (in type or value) */
+	UE_DEPRECATED(5.5, "Replace with UE::Reflection::CompareSparseClassDataToArchetype(Class, [](const FProperty*){return true;})")
 	COREUOBJECT_API bool OverridesSparseClassDataArchetype() const;
 
 	/**
diff --git a/Engine/Source/Runtime/CoreUObject/Public/UObject/SparseClassDataUtils.h b/Engine/Source/Runtime/CoreUObject/Public/UObject/SparseClassDataUtils.h
new file mode 100644
--- /dev/null
+++ b/Engine/Source/Runtime/CoreUObject/Public/UObject/SparseClassDataUtils.h
@@ -0,0 +1,21 @@
+// Copyright Epic Games, Inc. All Rights Reserved.
+
+#pragma once
+
+#include "Templates/FunctionFwd.h"
+#include "HAL/Platform.h" // for 'should be preceded by' error related to API decorator
+
+class UClass;
+class FProperty;
+
+namespace UE::Reflection
+{
+	/**
+	 * Returns true if all properties in the provided class's sparse class data match the class's archetype (super class)
+	 * Use the filter function to allow list properties that should be compared - if you want to compare all properties
+	 * in the sparse class data simply return true in the filter, otherwise return false if the property should be skipped.
+	 * Typically transient properties should be skipped, and when saving for a non editor target editor properties should
+	 * be skipped. Consider relying on FProperty::ShouldSerializeProperty if you have an archiver available.
+	 */
+	COREUOBJECT_API bool CompareSparseClassDataToArchetype(const UClass* Class, const TFunctionRef<bool(FProperty*)>& Filter);
+}
diff --git a/Engine/Source/Runtime/Engine/Private/BlueprintGeneratedClass.cpp b/Engine/Source/Runtime/Engine/Private/BlueprintGeneratedClass.cpp
--- a/Engine/Source/Runtime/Engine/Private/BlueprintGeneratedClass.cpp
+++ b/Engine/Source/Runtime/Engine/Private/BlueprintGeneratedClass.cpp
@@ -27,6 +27,7 @@
 #include "UObject/FortniteMainBranchObjectVersion.h"
 #include "UObject/Package.h" // IWYU pragma: keep
 #include "UObject/PrimaryAssetId.h"
+#include "UObject/SparseClassDataUtils.h"
 #include "UObject/UE5MainStreamObjectVersion.h"
 #include "GenericPlatform/GenericPlatformCrashContext.h"
 
@@ -820,7 +821,7 @@ void UBlueprintGeneratedClass::SerializeDefaultObject(UObject* Object, FStructur
 		if (UnderlyingArchive.IsSaving())
 		{
 			UScriptStruct* SerializedSparseClassDataStruct = SparseClassDataStruct;
-			if (!OverridesSparseClassDataArchetype())
+			if (!UE::Reflection::CompareSparseClassDataToArchetype(this, [&UnderlyingArchive](FProperty* P) { return P->ShouldSerializeValue(UnderlyingArchive);}))
 			{
 				// If this class doesn't override the sparse class data of its archetype, then we can skip saving it 
 				// since it can be lazily regenerated from the archetype data on load
