[Backout] - CL31459147
Original CL Desc
-----------------------------------------------------------------
Ignore properties that won't be serialized when determining whether we can discard a class's sparse data

diff --git a/Engine/Source/Runtime/CoreUObject/Private/Serialization/PropertyLocalizationDataGathering.cpp b/Engine/Source/Runtime/CoreUObject/Private/Serialization/PropertyLocalizationDataGathering.cpp
--- a/Engine/Source/Runtime/CoreUObject/Private/Serialization/PropertyLocalizationDataGathering.cpp
+++ b/Engine/Source/Runtime/CoreUObject/Private/Serialization/PropertyLocalizationDataGathering.cpp
@@ -10,7 +10,6 @@
 #include "UObject/TextProperty.h"
 #include "UObject/PropertyPortFlags.h"
 #include "UObject/PropertyOptional.h"
-#include "UObject/SparseClassDataUtils.h"
 #include "HAL/UnrealMemory.h"
 #include "Internationalization/TextNamespaceUtil.h"
 #include "Internationalization/TextPackageNamespaceUtil.h"
@@ -228,11 +227,7 @@ void FPropertyLocalizationDataGatherer::GatherLocalizationDataFromObject(const U
 	// Also gather from the sparse data on UClass types.
 	if (const UClass* Class = Cast<UClass>(Object))
 	{
-		const auto FilterTransientProperties = [](const FProperty* P)
-			{
-				return !P->HasAnyPropertyFlags(CPF_Transient);
-			};
-		if (UE::Reflection::CompareSparseClassDataToArchetype(Class, FilterTransientProperties))
+		if (Class->OverridesSparseClassDataArchetype())
 		{
 			const UScriptStruct* SparseDataStruct = Class->GetSparseClassDataStruct();
 			const void* SparseData = const_cast<UClass*>(Class)->GetOrCreateSparseClassData();
diff --git a/Engine/Source/Runtime/CoreUObject/Private/UObject/SparseClassDataUtils.cpp b/Engine/Source/Runtime/CoreUObject/Private/UObject/SparseClassDataUtils.cpp
deleted file mode 100644
diff --git a/Engine/Source/Runtime/CoreUObject/Public/UObject/Class.h b/Engine/Source/Runtime/CoreUObject/Public/UObject/Class.h
--- a/Engine/Source/Runtime/CoreUObject/Public/UObject/Class.h
+++ b/Engine/Source/Runtime/CoreUObject/Public/UObject/Class.h
@@ -3462,7 +3462,6 @@ public:
 	COREUOBJECT_API UScriptStruct* GetSparseClassDataArchetypeStruct() const;
 
 	/** Returns whether the sparse class data on this instance overrides that of its archetype (in type or value) */
-	UE_DEPRECATED(5.5, "Replace with UE::Reflection::CompareSparseClassDataToArchetype(Class, [](const FProperty*){return true;})")
 	COREUOBJECT_API bool OverridesSparseClassDataArchetype() const;
 
 	/**
diff --git a/Engine/Source/Runtime/CoreUObject/Public/UObject/SparseClassDataUtils.h b/Engine/Source/Runtime/CoreUObject/Public/UObject/SparseClassDataUtils.h
deleted file mode 100644
diff --git a/Engine/Source/Runtime/Engine/Private/BlueprintGeneratedClass.cpp b/Engine/Source/Runtime/Engine/Private/BlueprintGeneratedClass.cpp
--- a/Engine/Source/Runtime/Engine/Private/BlueprintGeneratedClass.cpp
+++ b/Engine/Source/Runtime/Engine/Private/BlueprintGeneratedClass.cpp
@@ -27,7 +27,6 @@
 #include "UObject/FortniteMainBranchObjectVersion.h"
 #include "UObject/Package.h" // IWYU pragma: keep
 #include "UObject/PrimaryAssetId.h"
-#include "UObject/SparseClassDataUtils.h"
 #include "UObject/UE5MainStreamObjectVersion.h"
 #include "GenericPlatform/GenericPlatformCrashContext.h"
 
@@ -821,7 +820,7 @@ void UBlueprintGeneratedClass::SerializeDefaultObject(UObject* Object, FStructur
 		if (UnderlyingArchive.IsSaving())
 		{
 			UScriptStruct* SerializedSparseClassDataStruct = SparseClassDataStruct;
-			if (!UE::Reflection::CompareSparseClassDataToArchetype(this, [&UnderlyingArchive](FProperty* P) { return P->ShouldSerializeValue(UnderlyingArchive);}))
+			if (!OverridesSparseClassDataArchetype())
 			{
 				// If this class doesn't override the sparse class data of its archetype, then we can skip saving it 
 				// since it can be lazily regenerated from the archetype data on load
