Fixed FPropertyBag::LoadPropertyByTag to construct a new property when it had stored a property in Add

#rnx

diff --git a/Engine/Source/Runtime/CoreUObject/Private/UObject/PropertyBag.cpp b/Engine/Source/Runtime/CoreUObject/Private/UObject/PropertyBag.cpp
--- a/Engine/Source/Runtime/CoreUObject/Private/UObject/PropertyBag.cpp
+++ b/Engine/Source/Runtime/CoreUObject/Private/UObject/PropertyBag.cpp
@@ -83,7 +83,7 @@ void FPropertyBag::LoadPropertyByTag(const FPropertyPathName& Path, const FPrope
 	FValue& Value = FindOrCreateValue(Path);
 
 	const bool bPropertyChanged =
-		(Value.Tag.Prop != Tag.Prop && Value.Tag.Prop && Tag.Prop) ||
+		(Value.Tag.Prop != Tag.Prop && (Tag.Prop || !Value.bOwnsProperty)) ||
 		Value.Tag.Type != Tag.Type ||
 		Value.Tag.Name != Tag.Name ||
 		Value.Tag.StructName != Tag.StructName ||
