Fixed TWeakFieldPtr asserting with "const T" template parameters.
Added simple tests to verify this use-case compiles.

diff --git a/Engine/Source/Runtime/CoreUObject/Public/UObject/FieldPath.h b/Engine/Source/Runtime/CoreUObject/Public/UObject/FieldPath.h
--- a/Engine/Source/Runtime/CoreUObject/Public/UObject/FieldPath.h
+++ b/Engine/Source/Runtime/CoreUObject/Public/UObject/FieldPath.h
@@ -441,7 +441,7 @@ public:
 	FORCEINLINE bool operator==(const OtherPropertyType* Other) const
 	{
 		static_assert(TPointerIsConvertibleFromTo<OtherPropertyType, const FField>::Value, "TFieldPath can only be compared with FField types");
-		static_assert(TPointerIsConvertibleFromTo<PropertyType, OtherPropertyType>::Value, "Unable to compare TFieldPath with raw pointer - types are incompatible");
+		static_assert(TPointerIsConvertibleFromTo<PropertyType, const OtherPropertyType>::Value, "Unable to compare TFieldPath with raw pointer - types are incompatible");
 
 		return Get() == Other;
 	}
@@ -461,7 +461,7 @@ public:
 	FORCEINLINE bool operator!=(const TFieldPath<OtherPropertyType> &Other) const
 	{
 		static_assert(TPointerIsConvertibleFromTo<OtherPropertyType, const FField>::Value, "TFieldPath can only be compared with FField types");
-		static_assert(TPointerIsConvertibleFromTo<PropertyType, OtherPropertyType>::Value, "Unable to compare TFieldPath with raw pointer - types are incompatible");
+		static_assert(TPointerIsConvertibleFromTo<PropertyType, const OtherPropertyType>::Value, "Unable to compare TFieldPath with raw pointer - types are incompatible");
 
 		return FFieldPath::operator!=(Other);
 	}
@@ -474,7 +474,7 @@ public:
 	FORCEINLINE bool operator!=(const OtherPropertyType* Other) const
 	{
 		static_assert(TPointerIsConvertibleFromTo<OtherPropertyType, const FField>::Value, "TFieldPath can only be compared with FField types");
-		static_assert(TPointerIsConvertibleFromTo<PropertyType, OtherPropertyType>::Value, "Unable to compare TFieldPath with raw pointer - types are incompatible");
+		static_assert(TPointerIsConvertibleFromTo<PropertyType, const OtherPropertyType>::Value, "Unable to compare TFieldPath with raw pointer - types are incompatible");
 
 		return Get() != Other;
 	}
@@ -485,7 +485,7 @@ public:
 		// It's also possible that these static_asserts may fail for valid conversions because
 		// one or both of the types have only been forward-declared.
 		static_assert(TPointerIsConvertibleFromTo<LhsT, const FField>::Value, "TFieldPath can only be compared with FField types");
-		static_assert(TPointerIsConvertibleFromTo<LhsT, PropertyType>::Value || TPointerIsConvertibleFromTo<PropertyType, LhsT>::Value, "Unable to compare TFieldPath with raw pointer - types are incompatible");
+		static_assert(TPointerIsConvertibleFromTo<LhsT, const PropertyType>::Value || TPointerIsConvertibleFromTo<PropertyType, LhsT>::Value, "Unable to compare TFieldPath with raw pointer - types are incompatible");
 
 		return Rhs == Lhs;
 	}
@@ -501,7 +501,7 @@ public:
 		// It's also possible that these static_asserts may fail for valid conversions because
 		// one or both of the types have only been forward-declared.
 		static_assert(TPointerIsConvertibleFromTo<LhsT, const FField>::Value, "TFieldPath can only be compared with FField types");
-		static_assert(TPointerIsConvertibleFromTo<LhsT, PropertyType>::Value || TPointerIsConvertibleFromTo<PropertyType, LhsT>::Value, "Unable to compare TFieldPath with raw pointer - types are incompatible");
+		static_assert(TPointerIsConvertibleFromTo<LhsT, const PropertyType>::Value || TPointerIsConvertibleFromTo<PropertyType, LhsT>::Value, "Unable to compare TFieldPath with raw pointer - types are incompatible");
 
 		return Rhs != Lhs;
 	}
diff --git a/Engine/Source/Runtime/CoreUObject/Public/UObject/WeakFieldPtr.h b/Engine/Source/Runtime/CoreUObject/Public/UObject/WeakFieldPtr.h
--- a/Engine/Source/Runtime/CoreUObject/Public/UObject/WeakFieldPtr.h
+++ b/Engine/Source/Runtime/CoreUObject/Public/UObject/WeakFieldPtr.h
@@ -227,8 +227,8 @@ public:
 	template <typename TOther>
 	FORCEINLINE bool operator==(const TWeakFieldPtr<TOther> &Other) const
 	{
-		static_assert(TPointerIsConvertibleFromTo<TOther, FField>::Value, "TWeakFieldPtr can only be compared with FField types");
-		static_assert(TPointerIsConvertibleFromTo<T, TOther>::Value, "Unable to compare TWeakFieldPtr with raw pointer - types are incompatible");
+		static_assert(TPointerIsConvertibleFromTo<TOther, const FField>::Value, "TWeakFieldPtr can only be compared with FField types");
+		static_assert(TPointerIsConvertibleFromTo<T, const TOther>::Value, "Unable to compare TWeakFieldPtr with raw pointer - types are incompatible");
 
 		return Field == Other.Field;
 	}
@@ -240,8 +240,8 @@ public:
 	template <typename TOther>
 	FORCEINLINE bool operator!=(const TWeakFieldPtr<TOther> &Other) const
 	{
-		static_assert(TPointerIsConvertibleFromTo<TOther, FField>::Value, "TWeakFieldPtr can only be compared with FField types");
-		static_assert(TPointerIsConvertibleFromTo<T, TOther>::Value, "Unable to compare TWeakFieldPtr with raw pointer - types are incompatible");
+		static_assert(TPointerIsConvertibleFromTo<TOther, const FField>::Value, "TWeakFieldPtr can only be compared with FField types");
+		static_assert(TPointerIsConvertibleFromTo<T, const TOther>::Value, "Unable to compare TWeakFieldPtr with raw pointer - types are incompatible");
 
 		return Field != Other.Field;
 	}
@@ -253,8 +253,8 @@ public:
 	template <typename TOther>
 	FORCEINLINE bool operator==(const TOther* Other) const
 	{
-		static_assert(TPointerIsConvertibleFromTo<TOther, FField>::Value, "TWeakFieldPtr can only be compared with FField types");
-		static_assert(TPointerIsConvertibleFromTo<T, TOther>::Value, "Unable to compare TWeakFieldPtr with raw pointer - types are incompatible");
+		static_assert(TPointerIsConvertibleFromTo<TOther, const FField>::Value, "TWeakFieldPtr can only be compared with FField types");
+		static_assert(TPointerIsConvertibleFromTo<T, const TOther>::Value, "Unable to compare TWeakFieldPtr with raw pointer - types are incompatible");
 
 		return Field == Other;
 	}
@@ -266,8 +266,8 @@ public:
 	template <typename TOther>
 	FORCEINLINE bool operator!=(const TOther* Other) const
 	{
-		static_assert(TPointerIsConvertibleFromTo<TOther, FField>::Value, "TWeakFieldPtr can only be compared with FField types");
-		static_assert(TPointerIsConvertibleFromTo<T, TOther>::Value, "Unable to compare TWeakFieldPtr with raw pointer - types are incompatible");
+		static_assert(TPointerIsConvertibleFromTo<TOther, const FField>::Value, "TWeakFieldPtr can only be compared with FField types");
+		static_assert(TPointerIsConvertibleFromTo<T, const TOther>::Value, "Unable to compare TWeakFieldPtr with raw pointer - types are incompatible");
 
 		return Field != Other;
 	}
@@ -285,7 +285,7 @@ FORCENOINLINE bool operator==(const LhsT* Lhs, const TWeakFieldPtr<RhsT>& Rhs)
 {
 	// It's also possible that these static_asserts may fail for valid conversions because
 	// one or both of the types have only been forward-declared.
-	static_assert(TPointerIsConvertibleFromTo<LhsT, FField>::Value, "TWeakFieldPtr can only be compared with FField types");
+	static_assert(TPointerIsConvertibleFromTo<LhsT, const FField>::Value, "TWeakFieldPtr can only be compared with FField types");
 	static_assert(TPointerIsConvertibleFromTo<LhsT, RhsT>::Value || TPointerIsConvertibleFromTo<RhsT, LhsT>::Value, "Unable to compare TWeakFieldPtr with raw pointer - types are incompatible");
 
 	return Rhs == Lhs;
@@ -308,7 +308,7 @@ FORCENOINLINE bool operator!=(const LhsT* Lhs, const TWeakFieldPtr<RhsT>& Rhs)
 {
 	// It's also possible that these static_asserts may fail for valid conversions because
 	// one or both of the types have only been forward-declared.
-	static_assert(TPointerIsConvertibleFromTo<LhsT, FField>::Value, "TWeakFieldPtr can only be compared with FField types");
+	static_assert(TPointerIsConvertibleFromTo<LhsT, const FField>::Value, "TWeakFieldPtr can only be compared with FField types");
 	static_assert(TPointerIsConvertibleFromTo<LhsT, RhsT>::Value || TPointerIsConvertibleFromTo<RhsT, LhsT>::Value, "Unable to compare TWeakFieldPtr with raw pointer - types are incompatible");
 
 	return Rhs != Lhs;
diff --git a/Engine/Source/Runtime/CoreUObject/Tests/UObject/WeakFieldPtrTest.cpp b/Engine/Source/Runtime/CoreUObject/Tests/UObject/WeakFieldPtrTest.cpp
new file mode 100644
--- /dev/null
+++ b/Engine/Source/Runtime/CoreUObject/Tests/UObject/WeakFieldPtrTest.cpp
@@ -0,0 +1,35 @@
+// Copyright Epic Games, Inc. All Rights Reserved.
+
+#include "UObject/PropertyTypeName.h"
+#include "UObject/WeakFieldPtr.h"
+#include "UObject/UnrealType.h"
+
+#if WITH_TESTS
+
+#include "Tests/TestHarnessAdapter.h"
+
+namespace UE
+{
+
+TEST_CASE_NAMED(TWeakFieldPtrSmokeTest, "CoreUObject::TWeakFieldPtr::Smoke", "[Core][UObject][SmokeFilter]")
+{
+	SECTION("Check that TWeakFieldPtr can hold 'const FProperty'")
+	{
+		const FProperty* PropPtr = nullptr;
+		TWeakFieldPtr<const FProperty> ConstPtr1;
+		TWeakFieldPtr<const FProperty> ConstPtr2;
+
+		ConstPtr1 = PropPtr;
+		ConstPtr1 = ConstPtr2;
+
+		// validate that the static_asserts inside these operators allow const FProperty as T
+		CHECK(ConstPtr1 == ConstPtr2);
+		CHECK(!(ConstPtr1 != ConstPtr2));
+		CHECK(ConstPtr1 == PropPtr);
+		CHECK(!(ConstPtr1 != PropPtr));
+	}
+}
+
+} // UE
+
+#endif // WITH_TESTS
