Added functions to FProperty to save/load/match type names and query binary/native serialization

diff --git a/Engine/Source/Runtime/CoreUObject/Private/UObject/EnumProperty.cpp b/Engine/Source/Runtime/CoreUObject/Private/UObject/EnumProperty.cpp
--- a/Engine/Source/Runtime/CoreUObject/Private/UObject/EnumProperty.cpp
+++ b/Engine/Source/Runtime/CoreUObject/Private/UObject/EnumProperty.cpp
@@ -1,16 +1,12 @@
 // Copyright Epic Games, Inc. All Rights Reserved.
 
 #include "UObject/EnumProperty.h"
-#include "UObject/Package.h"
-#include "UObject/PropertyPortFlags.h"
-#include "UObject/UObjectThreadContext.h"
-#include "UObject/PropertyTag.h"
-#include "UObject/UnrealTypePrivate.h"
-#include "Templates/IsSigned.h"
+
 #include "Algo/Find.h"
-#include "UObject/LinkerLoad.h"
-#include "Misc/EngineNetworkCustomVersion.h"
 #include "Hash/Blake3.h"
+#include "UObject/Package.h"
+#include "UObject/UnrealTypePrivate.h"
+#include "UObject/UObjectThreadContext.h"
 
 namespace UEEnumProperty_Private
 {
@@ -578,3 +574,73 @@ void FEnumProperty::SaveToTag(FPropertyTag& Tag)
 		}
 	}
 }
+
+bool FEnumProperty::LoadTypeName(UE::FPropertyTypeName Type, const FPropertyTag* Tag)
+{
+	if (!Super::LoadTypeName(Type, Tag))
+	{
+		return false;
+	}
+
+	const FName EnumName = Type.GetTypeParameterName(0);
+	UEnum* LocalEnum = FindFirstObject<UEnum>(*WriteToString<256>(EnumName), EFindFirstObjectOptions::NativeFirst);
+	if (!LocalEnum)
+	{
+		return false;
+	}
+
+	const UE::FPropertyTypeName UnderlyingType = Type.GetTypeParameter(1);
+	FField* Field = FField::TryConstruct(UnderlyingType.GetTypeName(), this, GetFName(), RF_NoFlags);
+	if (FNumericProperty* Property = CastField<FNumericProperty>(Field); Property && Property->LoadTypeName(UnderlyingType, Tag))
+	{
+		Enum = LocalEnum;
+		UE_CLOG(!Property->CanHoldValue(Enum->GetMaxEnumValue()), LogClass, Warning,
+			TEXT("Enum '%s' does not fit in a %s loading property '%s'."),
+			*WriteToString<64>(Enum->GetFName()), *WriteToString<32>(Property->GetID()), *WriteToString<32>(GetFName()));
+		AddCppProperty(Property);
+		return true;
+	}
+	delete Field;
+	return false;
+}
+
+void FEnumProperty::SaveTypeName(UE::FPropertyTypeNameBuilder& Type) const
+{
+	Super::SaveTypeName(Type);
+
+	if (const UEnum* LocalEnum = Enum)
+	{
+		TStringBuilder<256> EnumName;
+		LocalEnum->GetPathName(nullptr, EnumName);
+
+		check(UnderlyingProp);
+		Type.BeginTypeParameters();
+		Type.AddTypeName(FName(EnumName));
+		UnderlyingProp->SaveTypeName(Type);
+		Type.EndTypeParameters();
+	}
+}
+
+bool FEnumProperty::CanSerializeFromTypeName(UE::FPropertyTypeName Type) const
+{
+	if (!Super::CanSerializeFromTypeName(Type))
+	{
+		return false;
+	}
+
+	const UEnum* LocalEnum = Enum;
+	if (!LocalEnum)
+	{
+		return false;
+	}
+
+	const FName EnumName = Type.GetTypeParameterName(0);
+	if (EnumName == LocalEnum->GetFName())
+	{
+		return true;
+	}
+
+	TStringBuilder<256> EnumNameString;
+	LocalEnum->GetPathName(nullptr, EnumNameString);
+	return EnumName == EnumNameString.ToView();
+}
diff --git a/Engine/Source/Runtime/CoreUObject/Private/UObject/Property.cpp b/Engine/Source/Runtime/CoreUObject/Private/UObject/Property.cpp
--- a/Engine/Source/Runtime/CoreUObject/Private/UObject/Property.cpp
+++ b/Engine/Source/Runtime/CoreUObject/Private/UObject/Property.cpp
@@ -1,33 +1,27 @@
 // Copyright Epic Games, Inc. All Rights Reserved.
 
-/*=============================================================================
-	Property.cpp: FProperty implementation
-=============================================================================*/
+#include "UObject/UnrealType.h"
 
-#include "CoreMinimal.h"
 #include "Hash/Blake3.h"
+#include "Math/Box2D.h"
+#include "Math/InterpCurvePoint.h"
+#include "Math/RandomStream.h"
+#include "Math/Ray.h"
+#include "Math/Sphere.h"
 #include "Misc/AsciiSet.h"
 #include "Misc/Guid.h"
 #include "Misc/StringBuilder.h"
-#include "Math/RandomStream.h"
-#include "Logging/LogScopedCategoryAndVerbosityOverride.h"
-#include "UObject/CoreNetTypes.h"
-#include "UObject/ObjectMacros.h"
-#include "UObject/UObjectGlobals.h"
-#include "UObject/Class.h"
+#include "Serialization/TestUndeclaredScriptStructObjectReferences.h"
 #include "Templates/Casts.h"
-#include "UObject/UnrealType.h"
-#include "UObject/UnrealTypePrivate.h"
-#include "UObject/PropertyHelper.h"
+#include "UObject/Class.h"
+#include "UObject/CoreNetTypes.h"
 #include "UObject/CoreRedirects.h"
-#include "UObject/SoftObjectPath.h"
-#include "Math/Box2D.h"
-#include "Math/Ray.h"
-#include "Math/Sphere.h"
-#include "Math/InterpCurvePoint.h"
 #include "UObject/Package.h"
-#include "UObject/ReleaseObjectVersion.h"
-#include "Serialization/TestUndeclaredScriptStructObjectReferences.h"
+#include "UObject/PropertyHelper.h"
+#include "UObject/PropertyTypeName.h"
+#include "UObject/SoftObjectPath.h"
+#include "UObject/UnrealTypePrivate.h"
+#include "UObject/UObjectGlobals.h"
 
 DEFINE_LOG_CATEGORY(LogProperty);
 
@@ -2105,6 +2099,11 @@ UPropertyWrapper* FProperty::GetUPropertyWrapper()
 }
 #endif //  WITH_EDITORONLY_DATA
 
+bool FProperty::UseBinaryOrNativeSerialization(const FArchive& Ar) const
+{
+	return Ar.WantBinaryPropertySerialization();
+}
+
 bool FProperty::LoadFromTag(const FPropertyTag& Tag)
 {
 	checkf(GetID() == Tag.Type, TEXT("Failed to load property '%s' of type '%s' from tag of type '%s'"),
@@ -2124,6 +2123,23 @@ void FProperty::AssignToTag(FPropertyTag& Tag)
 	Tag.Prop = this;
 }
 
+bool FProperty::LoadTypeName(UE::FPropertyTypeName Type, const FPropertyTag* Tag)
+{
+	return ensureMsgf(GetID() == Type.GetTypeName(),
+		TEXT("Failed to load property '%s' of type '%s' from tag of type '%s'"),
+		*WriteToString<64>(GetFName()), *WriteToString<64>(GetID()), *WriteToString<64>(Type.GetTypeName()));
+}
+
+void FProperty::SaveTypeName(UE::FPropertyTypeNameBuilder& Type) const
+{
+	Type.AddTypeName(GetID());
+}
+
+bool FProperty::CanSerializeFromTypeName(UE::FPropertyTypeName Type) const
+{
+	return Type.GetTypeName() == GetID();
+}
+
 FProperty* UStruct::FindPropertyByName(FName InName) const
 {
 	for (FProperty* Property = PropertyLink; Property != nullptr; Property = Property->PropertyLinkNext)
diff --git a/Engine/Source/Runtime/CoreUObject/Private/UObject/PropertyArray.cpp b/Engine/Source/Runtime/CoreUObject/Private/UObject/PropertyArray.cpp
--- a/Engine/Source/Runtime/CoreUObject/Private/UObject/PropertyArray.cpp
+++ b/Engine/Source/Runtime/CoreUObject/Private/UObject/PropertyArray.cpp
@@ -1,17 +1,13 @@
 // Copyright Epic Games, Inc. All Rights Reserved.
 
-#include "CoreMinimal.h"
-#include "Misc/ScopeExit.h"
-#include "UObject/ObjectMacros.h"
-#include "Templates/Casts.h"
-#include "UObject/PropertyPathName.h"
-#include "UObject/PropertyTag.h"
 #include "UObject/UnrealType.h"
-#include "UObject/UnrealTypePrivate.h"
+
+#include "Misc/ScopeExit.h"
 #include "UObject/LinkerLoad.h"
+#include "UObject/OverriddenPropertySet.h"
 #include "UObject/PropertyHelper.h"
+#include "UObject/UnrealTypePrivate.h"
 #include "UObject/UObjectThreadContext.h"
-#include "UObject/OverriddenPropertySet.h"
 
 /*-----------------------------------------------------------------------------
 	FArrayProperty.
@@ -1258,6 +1254,18 @@ void* FArrayProperty::GetValueAddressAtIndex_Direct(const FProperty* InInner, vo
 	}
 }
 
+bool FArrayProperty::UseBinaryOrNativeSerialization(const FArchive& Ar) const
+{
+	if (Super::UseBinaryOrNativeSerialization(Ar))
+	{
+		return true;
+	}
+
+	const FProperty* LocalInner = Inner;
+	check(LocalInner);
+	return LocalInner->UseBinaryOrNativeSerialization(Ar);
+}
+
 bool FArrayProperty::LoadFromTag(const FPropertyTag& Tag)
 {
 	if (!Super::LoadFromTag(Tag))
@@ -1295,3 +1303,44 @@ void FArrayProperty::SaveToTag(FPropertyTag& Tag)
 	check(LocalInner);
 	Tag.InnerType = LocalInner->GetID();
 }
+
+bool FArrayProperty::LoadTypeName(UE::FPropertyTypeName Type, const FPropertyTag* Tag)
+{
+	if (!Super::LoadTypeName(Type, Tag))
+	{
+		return false;
+	}
+
+	const UE::FPropertyTypeName InnerType = Type.GetTypeParameter();
+	FField* Field = FField::TryConstruct(InnerType.GetTypeName(), this, GetFName(), RF_NoFlags);
+	if (FProperty* Property = CastField<FProperty>(Field); Property && Property->LoadTypeName(InnerType, Tag))
+	{
+		Inner = Property;
+		return true;
+	}
+	delete Field;
+	return false;
+}
+
+void FArrayProperty::SaveTypeName(UE::FPropertyTypeNameBuilder& Type) const
+{
+	Super::SaveTypeName(Type);
+
+	const FProperty* LocalInner = Inner;
+	check(LocalInner);
+	Type.BeginTypeParameters();
+	LocalInner->SaveTypeName(Type);
+	Type.EndTypeParameters();
+}
+
+bool FArrayProperty::CanSerializeFromTypeName(UE::FPropertyTypeName Type) const
+{
+	if (!Super::CanSerializeFromTypeName(Type))
+	{
+		return false;
+	}
+
+	const FProperty* LocalInner = Inner;
+	check(LocalInner);
+	return LocalInner->CanSerializeFromTypeName(Type.GetTypeParameter());
+}
diff --git a/Engine/Source/Runtime/CoreUObject/Private/UObject/PropertyByte.cpp b/Engine/Source/Runtime/CoreUObject/Private/UObject/PropertyByte.cpp
--- a/Engine/Source/Runtime/CoreUObject/Private/UObject/PropertyByte.cpp
+++ b/Engine/Source/Runtime/CoreUObject/Private/UObject/PropertyByte.cpp
@@ -1,18 +1,12 @@
 // Copyright Epic Games, Inc. All Rights Reserved.
 
-#include "CoreMinimal.h"
-#include "UObject/ObjectMacros.h"
-#include "UObject/Class.h"
-#include "UObject/Package.h"
-#include "UObject/PropertyPortFlags.h"
 #include "UObject/UnrealType.h"
-#include "UObject/UnrealTypePrivate.h"
-#include "UObject/UObjectThreadContext.h"
-#include "Serialization/ArchiveUObjectFromStructuredArchive.h"
+
 #include "Algo/Find.h"
-#include "UObject/LinkerLoad.h"
-#include "Misc/EngineNetworkCustomVersion.h"
 #include "Hash/Blake3.h"
+#include "UObject/Package.h"
+#include "UObject/UnrealTypePrivate.h"
+#include "UObject/UObjectThreadContext.h"
 
 /*-----------------------------------------------------------------------------
 	FByteProperty.
@@ -571,3 +565,62 @@ void FByteProperty::SaveToTag(FPropertyTag& Tag)
 		}
 	}
 }
+
+bool FByteProperty::LoadTypeName(UE::FPropertyTypeName Type, const FPropertyTag* Tag)
+{
+	if (!Super::LoadTypeName(Type, Tag))
+	{
+		return false;
+	}
+
+	const FName EnumName = Type.GetTypeParameterName();
+	if (EnumName.IsNone())
+	{
+		return true;
+	}
+
+	if (UEnum* LocalEnum = FindFirstObject<UEnum>(*WriteToString<256>(EnumName), EFindFirstObjectOptions::NativeFirst))
+	{
+		Enum = LocalEnum;
+		return true;
+	}
+
+	return false;
+}
+
+void FByteProperty::SaveTypeName(UE::FPropertyTypeNameBuilder& Type) const
+{
+	Super::SaveTypeName(Type);
+
+	if (const UEnum* LocalEnum = Enum)
+	{
+		TStringBuilder<256> EnumName;
+		LocalEnum->GetPathName(nullptr, EnumName);
+
+		Type.BeginTypeParameters();
+		Type.AddTypeName(FName(EnumName));
+		Type.EndTypeParameters();
+	}
+}
+
+bool FByteProperty::CanSerializeFromTypeName(UE::FPropertyTypeName Type) const
+{
+	if (!Super::CanSerializeFromTypeName(Type))
+	{
+		return false;
+	}
+
+	const FName EnumName = Type.GetTypeParameterName();
+	if (const UEnum* LocalEnum = Enum)
+	{
+		if (EnumName == LocalEnum->GetFName())
+		{
+			return true;
+		}
+
+		TStringBuilder<256> EnumNameString;
+		LocalEnum->GetPathName(nullptr, EnumNameString);
+		return EnumName == EnumNameString.ToView();
+	}
+	return EnumName.IsNone();
+}
diff --git a/Engine/Source/Runtime/CoreUObject/Private/UObject/PropertyMap.cpp b/Engine/Source/Runtime/CoreUObject/Private/UObject/PropertyMap.cpp
--- a/Engine/Source/Runtime/CoreUObject/Private/UObject/PropertyMap.cpp
+++ b/Engine/Source/Runtime/CoreUObject/Private/UObject/PropertyMap.cpp
@@ -1,19 +1,14 @@
 // Copyright Epic Games, Inc. All Rights Reserved.
 
-#include "CoreMinimal.h"
-#include "UObject/ObjectMacros.h"
-#include "Templates/Casts.h"
+#include "UObject/UnrealType.h"
+
+#include "Misc/ScopeExit.h"
+#include "UObject/OverriddenPropertySet.h"
+#include "UObject/PropertyHelper.h"
 #include "UObject/PropertyPathName.h"
 #include "UObject/PropertyTag.h"
-#include "UObject/UnrealType.h"
 #include "UObject/UnrealTypePrivate.h"
-#include "UObject/LinkerLoad.h"
-#include "UObject/PropertyHelper.h"
 #include "UObject/UObjectThreadContext.h"
-#include "Misc/ScopeExit.h"
-#include "Serialization/ArchiveUObjectFromStructuredArchive.h"
-#include "UObject/UObjectThreadContext.h"
-#include "UObject/OverriddenPropertySet.h"
 
 namespace UEMapProperty_Private
 {
@@ -1696,6 +1691,20 @@ void* FMapProperty::GetValueAddressAtIndex_Direct(const FProperty* Inner, void*
 	return nullptr;
 }
 
+bool FMapProperty::UseBinaryOrNativeSerialization(const FArchive& Ar) const
+{
+	if (Super::UseBinaryOrNativeSerialization(Ar))
+	{
+		return true;
+	}
+
+	const FProperty* LocalKeyProp = KeyProp;
+	const FProperty* LocalValueProp = ValueProp;
+	check(LocalKeyProp);
+	check(LocalValueProp);
+	return LocalKeyProp->UseBinaryOrNativeSerialization(Ar) || LocalValueProp->UseBinaryOrNativeSerialization(Ar);
+}
+
 bool FMapProperty::LoadFromTag(const FPropertyTag& Tag)
 {
 	if (!Super::LoadFromTag(Tag))
@@ -1717,3 +1726,55 @@ void FMapProperty::SaveToTag(FPropertyTag& Tag)
 	Tag.InnerType = LocalKeyProp->GetID();
 	Tag.ValueType = LocalValueProp->GetID();
 }
+
+bool FMapProperty::LoadTypeName(UE::FPropertyTypeName Type, const FPropertyTag* Tag)
+{
+	if (!Super::LoadTypeName(Type, Tag))
+	{
+		return false;
+	}
+
+	const UE::FPropertyTypeName KeyType = Type.GetTypeParameter(0);
+	const UE::FPropertyTypeName ValueType = Type.GetTypeParameter(1);
+	FField* KeyField = FField::TryConstruct(KeyType.GetTypeName(), this, GetFName(), RF_NoFlags);
+	FField* ValueField = FField::TryConstruct(ValueType.GetTypeName(), this, GetFName(), RF_NoFlags);
+	FProperty* KeyProperty = CastField<FProperty>(KeyField);
+	FProperty* ValueProperty = CastField<FProperty>(ValueField);
+	if (KeyProperty && ValueProperty && KeyProperty->LoadTypeName(KeyType, Tag) && ValueProperty->LoadTypeName(ValueType, Tag))
+	{
+		KeyProp = KeyProperty;
+		ValueProp = ValueProperty;
+		return true;
+	}
+	delete KeyField;
+	delete ValueField;
+	return false;
+}
+
+void FMapProperty::SaveTypeName(UE::FPropertyTypeNameBuilder& Type) const
+{
+	Super::SaveTypeName(Type);
+
+	const FProperty* LocalKeyProp = KeyProp;
+	const FProperty* LocalValueProp = ValueProp;
+	check(LocalKeyProp);
+	check(LocalValueProp);
+	Type.BeginTypeParameters();
+	LocalKeyProp->SaveTypeName(Type);
+	LocalValueProp->SaveTypeName(Type);
+	Type.EndTypeParameters();
+}
+
+bool FMapProperty::CanSerializeFromTypeName(UE::FPropertyTypeName Type) const
+{
+	if (!Super::CanSerializeFromTypeName(Type))
+	{
+		return false;
+	}
+
+	const FProperty* LocalKeyProp = KeyProp;
+	const FProperty* LocalValueProp = ValueProp;
+	check(LocalKeyProp);
+	check(LocalValueProp);
+	return LocalKeyProp->CanSerializeFromTypeName(Type.GetTypeParameter(0)) && LocalValueProp->CanSerializeFromTypeName(Type.GetTypeParameter(1));
+}
diff --git a/Engine/Source/Runtime/CoreUObject/Private/UObject/PropertyMulticastDelegate.cpp b/Engine/Source/Runtime/CoreUObject/Private/UObject/PropertyMulticastDelegate.cpp
--- a/Engine/Source/Runtime/CoreUObject/Private/UObject/PropertyMulticastDelegate.cpp
+++ b/Engine/Source/Runtime/CoreUObject/Private/UObject/PropertyMulticastDelegate.cpp
@@ -1,13 +1,12 @@
 // Copyright Epic Games, Inc. All Rights Reserved.
 
-#include "CoreMinimal.h"
-#include "UObject/ObjectMacros.h"
-#include "UObject/Class.h"
 #include "UObject/UnrealType.h"
-#include "UObject/UnrealTypePrivate.h"
-#include "UObject/PropertyHelper.h"
-#include "UObject/LinkerPlaceholderFunction.h"
+
 #include "Serialization/ArchiveUObjectFromStructuredArchive.h"
+#include "UObject/LinkerPlaceholderFunction.h"
+#include "UObject/PropertyHelper.h"
+#include "UObject/PropertyTypeName.h"
+#include "UObject/UnrealTypePrivate.h"
 
 FMulticastScriptDelegate FMulticastDelegateProperty::EmptyDelegate;
 
@@ -701,11 +700,15 @@ bool FMulticastSparseDelegateProperty::LoadFromTag(const FPropertyTag& Tag)
 	return false;
 }
 
-void FMulticastSparseDelegateProperty::SaveToTag(FPropertyTag& Tag)
+bool FMulticastSparseDelegateProperty::LoadTypeName(UE::FPropertyTypeName Type, const FPropertyTag* Tag)
 {
-	Super::SaveToTag(Tag);
+	if (!Super::LoadTypeName(Type, Tag))
+	{
+		return false;
+	}
 
-	// This needs a way to save SignatureFunction to the tag.
+	// This cannot be used without its SignatureFunction and the tag lacks the information needed to load it.
+	return false;
 }
 
 IMPLEMENT_FIELD(FMulticastSparseDelegateProperty)
diff --git a/Engine/Source/Runtime/CoreUObject/Private/UObject/PropertyOptional.cpp b/Engine/Source/Runtime/CoreUObject/Private/UObject/PropertyOptional.cpp
--- a/Engine/Source/Runtime/CoreUObject/Private/UObject/PropertyOptional.cpp
+++ b/Engine/Source/Runtime/CoreUObject/Private/UObject/PropertyOptional.cpp
@@ -1,11 +1,14 @@
 // Copyright Epic Games, Inc. All Rights Reserved.
 
 #include "UObject/PropertyOptional.h"
+
 #include "Misc/Guid.h"
 #include "Serialization/CustomVersion.h"
 #include "String/LexFromString.h"
 #include "UObject/GarbageCollectionSchema.h"
+#include "UObject/LinkerLoad.h"
 #include "UObject/PropertyHelper.h"
+#include "UObject/PropertyTypeName.h"
 
 DEFINE_LOG_CATEGORY_STATIC(LogOptionalProperty, Log, All);
 
@@ -552,6 +555,18 @@ uint32 FOptionalProperty::GetValueTypeHashInternal(const void* Src) const
 	}
 }
 
+bool FOptionalProperty::UseBinaryOrNativeSerialization(const FArchive& Ar) const
+{
+	if (Super::UseBinaryOrNativeSerialization(Ar))
+	{
+		return true;
+	}
+
+	const FProperty* LocalValueProperty = ValueProperty;
+	check(LocalValueProperty);
+	return LocalValueProperty->UseBinaryOrNativeSerialization(Ar);
+}
+
 bool FOptionalProperty::LoadFromTag(const FPropertyTag& Tag)
 {
 	if (!Super::LoadFromTag(Tag))
@@ -570,3 +585,44 @@ void FOptionalProperty::SaveToTag(FPropertyTag& Tag)
 	check(LocalValueProperty);
 	Tag.InnerType = LocalValueProperty->GetID();
 }
+
+bool FOptionalProperty::LoadTypeName(UE::FPropertyTypeName Type, const FPropertyTag* Tag)
+{
+	if (!Super::LoadTypeName(Type, Tag))
+	{
+		return false;
+	}
+
+	const UE::FPropertyTypeName ValueType = Type.GetTypeParameter();
+	FField* Field = FField::TryConstruct(ValueType.GetTypeName(), this, GetFName(), RF_NoFlags);
+	if (FProperty* Property = CastField<FProperty>(Field); Property && Property->LoadTypeName(ValueType, Tag))
+	{
+		ValueProperty = Property;
+		return true;
+	}
+	delete Field;
+	return false;
+}
+
+void FOptionalProperty::SaveTypeName(UE::FPropertyTypeNameBuilder& Type) const
+{
+	Super::SaveTypeName(Type);
+
+	const FProperty* LocalValueProperty = ValueProperty;
+	check(LocalValueProperty);
+	Type.BeginTypeParameters();
+	LocalValueProperty->SaveTypeName(Type);
+	Type.EndTypeParameters();
+}
+
+bool FOptionalProperty::CanSerializeFromTypeName(UE::FPropertyTypeName Type) const
+{
+	if (!Super::CanSerializeFromTypeName(Type))
+	{
+		return false;
+	}
+
+	const FProperty* LocalValueProperty = ValueProperty;
+	check(LocalValueProperty);
+	return LocalValueProperty->CanSerializeFromTypeName(Type.GetTypeParameter());
+}
diff --git a/Engine/Source/Runtime/CoreUObject/Private/UObject/PropertySet.cpp b/Engine/Source/Runtime/CoreUObject/Private/UObject/PropertySet.cpp
--- a/Engine/Source/Runtime/CoreUObject/Private/UObject/PropertySet.cpp
+++ b/Engine/Source/Runtime/CoreUObject/Private/UObject/PropertySet.cpp
@@ -1,16 +1,11 @@
 // Copyright Epic Games, Inc. All Rights Reserved.
 
-#include "CoreMinimal.h"
-#include "UObject/ObjectMacros.h"
-#include "UObject/PropertyPathName.h"
-#include "UObject/PropertyTag.h"
 #include "UObject/UnrealType.h"
-#include "UObject/UnrealTypePrivate.h"
-#include "UObject/LinkerLoad.h"
+
+#include "Misc/ScopeExit.h"
 #include "UObject/PropertyHelper.h"
+#include "UObject/UnrealTypePrivate.h"
 #include "UObject/UObjectThreadContext.h"
-#include "Misc/ScopeExit.h"
-#include "Serialization/ArchiveUObjectFromStructuredArchive.h"
 
 namespace UESetProperty_Private
 {
@@ -1090,6 +1085,18 @@ void* FSetProperty::GetValueAddressAtIndex_Direct(const FProperty* Inner, void*
 	return nullptr;
 }
 
+bool FSetProperty::UseBinaryOrNativeSerialization(const FArchive& Ar) const
+{
+	if (Super::UseBinaryOrNativeSerialization(Ar))
+	{
+		return true;
+	}
+
+	const FProperty* LocalElementProp = ElementProp;
+	check(LocalElementProp);
+	return LocalElementProp->UseBinaryOrNativeSerialization(Ar);
+}
+
 bool FSetProperty::LoadFromTag(const FPropertyTag& Tag)
 {
 	if (!Super::LoadFromTag(Tag))
@@ -1108,3 +1115,44 @@ void FSetProperty::SaveToTag(FPropertyTag& Tag)
 	check(LocalElementProp);
 	Tag.InnerType = LocalElementProp->GetID();
 }
+
+bool FSetProperty::LoadTypeName(UE::FPropertyTypeName Type, const FPropertyTag* Tag)
+{
+	if (!Super::LoadTypeName(Type, Tag))
+	{
+		return false;
+	}
+
+	const UE::FPropertyTypeName ElementType = Type.GetTypeParameter();
+	FField* Field = FField::TryConstruct(ElementType.GetTypeName(), this, GetFName(), RF_NoFlags);
+	if (FProperty* Property = CastField<FProperty>(Field); Property && Property->LoadTypeName(ElementType, Tag))
+	{
+		ElementProp = Property;
+		return true;
+	}
+	delete Field;
+	return false;
+}
+
+void FSetProperty::SaveTypeName(UE::FPropertyTypeNameBuilder& Type) const
+{
+	Super::SaveTypeName(Type);
+
+	const FProperty* LocalElementProp = ElementProp;
+	check(LocalElementProp);
+	Type.BeginTypeParameters();
+	LocalElementProp->SaveTypeName(Type);
+	Type.EndTypeParameters();
+}
+
+bool FSetProperty::CanSerializeFromTypeName(UE::FPropertyTypeName Type) const
+{
+	if (!Super::CanSerializeFromTypeName(Type))
+	{
+		return false;
+	}
+
+	const FProperty* LocalElementProp = ElementProp;
+	check(LocalElementProp);
+	return LocalElementProp->CanSerializeFromTypeName(Type.GetTypeParameter());
+}
diff --git a/Engine/Source/Runtime/CoreUObject/Private/UObject/PropertyStruct.cpp b/Engine/Source/Runtime/CoreUObject/Private/UObject/PropertyStruct.cpp
--- a/Engine/Source/Runtime/CoreUObject/Private/UObject/PropertyStruct.cpp
+++ b/Engine/Source/Runtime/CoreUObject/Private/UObject/PropertyStruct.cpp
@@ -1,18 +1,16 @@
 // Copyright Epic Games, Inc. All Rights Reserved.
 
-#include "CoreMinimal.h"
-#include "UObject/ObjectMacros.h"
-#include "UObject/UObjectGlobals.h"
-#include "UObject/Class.h"
 #include "UObject/UnrealType.h"
-#include "UObject/UnrealTypePrivate.h"
-#include "UObject/PropertyHelper.h"
-#include "UObject/LinkerPlaceholderBase.h"
-#include "UObject/UObjectThreadContext.h"
-#include "Serialization/ArchiveUObjectFromStructuredArchive.h"
+
 #include "Hash/Blake3.h"
 #include "IO/IoHash.h"
 #include "Misc/StringBuilder.h"
+#include "Serialization/ArchiveUObjectFromStructuredArchive.h"
+#include "UObject/LinkerLoad.h"
+#include "UObject/LinkerPlaceholderBase.h"
+#include "UObject/PropertyHelper.h"
+#include "UObject/UnrealTypePrivate.h"
+#include "UObject/UObjectThreadContext.h"
 
 static inline void PreloadInnerStructMembers(FStructProperty* StructProperty)
 {
@@ -152,8 +150,12 @@ bool FStructProperty::Identical( const void* A, const void* B, uint32 PortFlags
 
 bool FStructProperty::UseBinaryOrNativeSerialization(const FArchive& Ar) const
 {
-	check(Struct);
+	if (Super::UseBinaryOrNativeSerialization(Ar))
+	{
+		return true;
+	}
 
+	check(Struct);
 	const bool bUseBinarySerialization = Struct->UseBinarySerialization(Ar);
 	const bool bUseNativeSerialization = Struct->UseNativeSerialization();
 	return bUseBinarySerialization || bUseNativeSerialization;
@@ -468,6 +470,28 @@ void FStructProperty::AppendSchemaHash(FBlake3& Builder, bool bSkipEditorOnly) c
 }
 #endif
 
+static const FName NAME_StructOriginalType(ANSITEXTVIEW("OriginalType"));
+
+static const FString* FindOriginalType(const FStructProperty* Struct)
+{
+	FUObjectSerializeContext* Context = FUObjectThreadContext::Get().GetSerializeContext();
+	if (Context && Context->bImpersonateProperties)
+	{
+#if WITH_EDITORONLY_DATA
+		if (const FString* OriginalType = Struct->FindMetaData(NAME_StructOriginalType))
+		{
+			return OriginalType;
+		}
+		//@note: To support metadata defined on array of struct in UPROPERTY for testing purposes
+		if (FField* OwnerField = Struct->Owner.ToField())
+		{
+			return OwnerField->FindMetaData(NAME_StructOriginalType);
+		}
+#endif
+	}
+	return nullptr;
+}
+
 bool FStructProperty::LoadFromTag(const FPropertyTag& Tag)
 {
 	if (!Super::LoadFromTag(Tag))
@@ -488,25 +512,11 @@ bool FStructProperty::LoadFromTag(const FPropertyTag& Tag)
 	return false;
 }
 
-static const FName NAME_StructOriginalType(TEXT("OriginalType"));
-
 void FStructProperty::SaveToTag(FPropertyTag& Tag)
 {
 	Super::SaveToTag(Tag);
 
-	//@note FH: revisit this once UE-197352 lands to use the new property type name
-	FUObjectSerializeContext* Context = FUObjectThreadContext::Get().GetSerializeContext();
-	const FString* OriginalType = nullptr;
-	if (Context && Context->bImpersonateProperties)
-	{
-#if WITH_EDITORONLY_DATA
-		OriginalType = FindMetaData(NAME_StructOriginalType);
-		//@note: To support metadata defined on array of struct in UPROPERTY for testing purposes
-		FField* OwnerField = OriginalType == nullptr ? Owner.ToField() : nullptr;
-		OriginalType = OwnerField ? OwnerField->FindMetaData(NAME_StructOriginalType) : OriginalType;
-#endif
-	}
-
+	const FString* OriginalType = FindOriginalType(this);
 	const UScriptStruct* LocalStruct = Struct;
 	check(LocalStruct);
 	Tag.StructName = OriginalType ? FName(**OriginalType) : LocalStruct->GetFName();
@@ -517,23 +527,75 @@ void FStructProperty::AssignToTag(FPropertyTag& Tag)
 {
 	Super::AssignToTag(Tag);
 
-	//@note FH: revisit this once UE-197352 lands to use the new property type name
-	FUObjectSerializeContext* Context = FUObjectThreadContext::Get().GetSerializeContext();
-	const FString* OriginalType = nullptr;
-	if (Context && Context->bImpersonateProperties)
+	const FString* OriginalType = FindOriginalType(this);
+	const UScriptStruct* LocalStruct = Struct;
+	check(LocalStruct);
+	if (OriginalType && FName(**OriginalType) == Tag.StructName)
 	{
-#if WITH_EDITORONLY_DATA
-		OriginalType = FindMetaData(NAME_StructOriginalType);
-		//@note: To support metadata defined on array of struct in UPROPERTY
-		FField* OwnerField = OriginalType == nullptr ? Owner.ToField() : nullptr;
-		OriginalType = OwnerField ? OwnerField->FindMetaData(NAME_StructOriginalType) : OriginalType;
-#endif
+		Tag.StructName = LocalStruct->GetFName();
+	}
+}
+
+bool FStructProperty::LoadTypeName(UE::FPropertyTypeName Type, const FPropertyTag* Tag)
+{
+	if (!Super::LoadTypeName(Type, Tag))
+	{
+		return false;
 	}
 
+	if (const FName Name = Type.GetTypeParameterName(0); !Name.IsNone())
+	{
+		TStringBuilder<256> NameString(InPlace, Name);
+		if (Struct = FindFirstObject<UScriptStruct>(*NameString, EFindFirstObjectOptions::NativeFirst); Struct)
+		{
+			return true;
+		}
+		// TODO: Look up the struct based on the guid.
+		// TODO: Use the fallback struct if allowed.
+	}
+
+	return false;
+}
+
+void FStructProperty::SaveTypeName(UE::FPropertyTypeNameBuilder& Type) const
+{
+	Super::SaveTypeName(Type);
+
+	const FString* OriginalType = FindOriginalType(this);
 	const UScriptStruct* LocalStruct = Struct;
 	check(LocalStruct);
-	if (OriginalType && FName(**OriginalType) == Tag.StructName)
+
+	Type.BeginTypeParameters();
+	Type.AddTypeName(OriginalType ? FName(**OriginalType) : LocalStruct->GetFName());
+	if (const FGuid StructGuid = LocalStruct->GetCustomGuid(); StructGuid.IsValid())
 	{
-		Tag.StructName = LocalStruct->GetFName();
+		Type.AddTypeName(FName(WriteToString<48>(StructGuid)));
+	}
+	Type.EndTypeParameters();
+}
+
+bool FStructProperty::CanSerializeFromTypeName(UE::FPropertyTypeName Type) const
+{
+	if (!Super::CanSerializeFromTypeName(Type))
+	{
+		return false;
 	}
+
+	const UScriptStruct* LocalStruct = Struct;
+	check(LocalStruct);
+
+	const FName StructName = Type.GetTypeParameterName(0);
+	if (StructName == LocalStruct->GetFName())
+	{
+		return true;
+	}
+
+	const FName StructGuidName = Type.GetTypeParameterName(1);
+	if (FGuid StructGuid; !StructGuidName.IsNone() && FGuid::Parse(StructGuidName.ToString(), StructGuid) && StructGuid.IsValid())
+	{
+		return StructGuid == LocalStruct->GetCustomGuid();
+	}
+
+	const FString* OriginalType = FindOriginalType(this);
+	return StructName == (OriginalType ? FName(**OriginalType) : LocalStruct->GetFName());
 }
diff --git a/Engine/Source/Runtime/CoreUObject/Public/UObject/EnumProperty.h b/Engine/Source/Runtime/CoreUObject/Public/UObject/EnumProperty.h
--- a/Engine/Source/Runtime/CoreUObject/Public/UObject/EnumProperty.h
+++ b/Engine/Source/Runtime/CoreUObject/Public/UObject/EnumProperty.h
@@ -80,6 +80,9 @@ public:
 #endif
 	COREUOBJECT_API virtual bool LoadFromTag(const FPropertyTag& Tag) override;
 	COREUOBJECT_API virtual void SaveToTag(FPropertyTag& Tag) override;
+	COREUOBJECT_API virtual bool LoadTypeName(UE::FPropertyTypeName Type, const FPropertyTag* Tag = nullptr) override;
+	COREUOBJECT_API virtual void SaveTypeName(UE::FPropertyTypeNameBuilder& Type) const override;
+	COREUOBJECT_API virtual bool CanSerializeFromTypeName(UE::FPropertyTypeName Type) const override;
 	// End of FProperty interface
 
 	/**
diff --git a/Engine/Source/Runtime/CoreUObject/Public/UObject/PropertyOptional.h b/Engine/Source/Runtime/CoreUObject/Public/UObject/PropertyOptional.h
--- a/Engine/Source/Runtime/CoreUObject/Public/UObject/PropertyOptional.h
+++ b/Engine/Source/Runtime/CoreUObject/Public/UObject/PropertyOptional.h
@@ -229,7 +229,11 @@ public:
 	virtual EConvertFromTypeResult ConvertFromType(const FPropertyTag& Tag, FStructuredArchive::FSlot Slot, uint8* Data, UStruct* DefaultsStruct, const uint8* Defaults) override;
 	virtual uint32 GetValueTypeHashInternal(const void* Src) const override;
 	virtual void AddReferencedObjects(FReferenceCollector& Collector) override;
+	virtual bool UseBinaryOrNativeSerialization(const FArchive& Ar) const override;
 	virtual bool LoadFromTag(const FPropertyTag& Tag) override;
 	virtual void SaveToTag(FPropertyTag& Tag) override;
+	virtual bool LoadTypeName(UE::FPropertyTypeName Type, const FPropertyTag* Tag = nullptr) override;
+	virtual void SaveTypeName(UE::FPropertyTypeNameBuilder& Type) const override;
+	virtual bool CanSerializeFromTypeName(UE::FPropertyTypeName Type) const override;
 	// End of FProperty interface
 };
diff --git a/Engine/Source/Runtime/CoreUObject/Public/UObject/UnrealType.h b/Engine/Source/Runtime/CoreUObject/Public/UObject/UnrealType.h
--- a/Engine/Source/Runtime/CoreUObject/Public/UObject/UnrealType.h
+++ b/Engine/Source/Runtime/CoreUObject/Public/UObject/UnrealType.h
@@ -74,6 +74,8 @@ enum ELifetimeCondition : int;
 struct CGetTypeHashable;
 struct FUObjectSerializeContext;
 template <typename FuncType> class TFunctionRef;
+namespace UE { class FPropertyTypeName; }
+namespace UE { class FPropertyTypeNameBuilder; }
 namespace UE::GC
 {
 	class FPropertyStack;
@@ -1043,6 +1045,8 @@ public:
 	}
 	COREUOBJECT_API bool ShouldSerializeValue( FArchive& Ar ) const;
 
+	COREUOBJECT_API virtual bool UseBinaryOrNativeSerialization(const FArchive& Ar) const;
+
 	/**
 	 * Determines whether this property value is eligible for copying when duplicating an object
 	 * 
@@ -1070,6 +1074,23 @@ public:
 	 */
 	COREUOBJECT_API virtual void AssignToTag(FPropertyTag& Tag);
 
+	/**
+	 * Restores this property and its owned properties from the type name.
+	 *
+	 * @return true if this property loaded from the type name and is in a valid and usable state.
+	 */
+	COREUOBJECT_API virtual bool LoadTypeName(UE::FPropertyTypeName Type, const FPropertyTag* Tag = nullptr);
+
+	/**
+	 * Saves the type name of this property and its owned properties.
+	 */
+	COREUOBJECT_API virtual void SaveTypeName(UE::FPropertyTypeNameBuilder& Type) const;
+
+	/**
+	 * Returns true if this property can use SerializeItem to load from a property with this type name.
+	 */
+	COREUOBJECT_API virtual bool CanSerializeFromTypeName(UE::FPropertyTypeName Type) const;
+
 	/**
 	 * Returns the first FProperty in this property's Outer chain that does not have a FProperty for an Outer
 	 */
@@ -1757,7 +1778,7 @@ private:
 public:
 	virtual EConvertFromTypeResult ConvertFromType(const FPropertyTag& Tag, FStructuredArchive::FSlot Slot, uint8* Data, UStruct* DefaultsStruct, const uint8* Defaults) override
 	{
-		if (const EName* TagType = Tag.Type.ToEName())
+		if (const EName* TagType = Tag.Type.ToEName(); TagType && Tag.Type.GetNumber() == NAME_NO_NUMBER_INTERNAL)
 		{
 			switch (*TagType)
 			{
@@ -1784,11 +1805,8 @@ public:
 					this->SetPropertyValue_InContainer(Data, (TCppType)PreviousValue, Tag.ArrayIndex);
 					return EConvertFromTypeResult::Converted;
 				}
-				else
-				{
-					ConvertFromArithmeticValue<int8>(Slot, Data, Tag);
-					return std::is_same_v<TCppType, uint8> ? EConvertFromTypeResult::Serialized : EConvertFromTypeResult::Converted;
-				}
+				ConvertFromArithmeticValue<int8>(Slot, Data, Tag);
+				return std::is_same_v<TCppType, uint8> ? EConvertFromTypeResult::Serialized : EConvertFromTypeResult::Converted;
 
 			case NAME_EnumProperty:
 			{
@@ -1989,6 +2007,9 @@ public:
 #endif
 	virtual bool LoadFromTag(const FPropertyTag& Tag) override;
 	virtual void SaveToTag(FPropertyTag& Tag) override;
+	virtual bool LoadTypeName(UE::FPropertyTypeName Type, const FPropertyTag* Tag = nullptr) override;
+	virtual void SaveTypeName(UE::FPropertyTypeNameBuilder& Type) const override;
+	virtual bool CanSerializeFromTypeName(UE::FPropertyTypeName Type) const override;
 	// End of FProperty interface
 
 	// FNumericProperty interface.
@@ -3684,8 +3705,12 @@ public:
 	}
 
 	virtual void* GetValueAddressAtIndex_Direct(const FProperty* Inner, void* InValueAddress, int32 Index) const override;
+	virtual bool UseBinaryOrNativeSerialization(const FArchive& Ar) const override;
 	virtual bool LoadFromTag(const FPropertyTag& Tag) override;
 	virtual void SaveToTag(FPropertyTag& Tag) override;
+	virtual bool LoadTypeName(UE::FPropertyTypeName Type, const FPropertyTag* Tag = nullptr) override;
+	virtual void SaveTypeName(UE::FPropertyTypeNameBuilder& Type) const override;
+	virtual bool CanSerializeFromTypeName(UE::FPropertyTypeName Type) const override;
 	// End of FProperty interface
 
 	FString GetCPPTypeCustom(FString* ExtendedTypeText, uint32 CPPExportFlags, const FString& InnerTypeText, const FString& InInnerExtendedTypeText) const;
@@ -3817,8 +3842,12 @@ public:
 	virtual bool SameType(const FProperty* Other) const override;
 	virtual EConvertFromTypeResult ConvertFromType(const FPropertyTag& Tag, FStructuredArchive::FSlot Slot, uint8* Data, UStruct* DefaultsStruct, const uint8* Defaults) override;
 	virtual void* GetValueAddressAtIndex_Direct(const FProperty* Inner, void* InValueAddress, int32 LogicalIndex) const override;
+	virtual bool UseBinaryOrNativeSerialization(const FArchive& Ar) const override;
 	virtual bool LoadFromTag(const FPropertyTag& Tag) override;
 	virtual void SaveToTag(FPropertyTag& Tag) override;
+	virtual bool LoadTypeName(UE::FPropertyTypeName Type, const FPropertyTag* Tag = nullptr) override;
+	virtual void SaveTypeName(UE::FPropertyTypeNameBuilder& Type) const override;
+	virtual bool CanSerializeFromTypeName(UE::FPropertyTypeName Type) const override;
 	// End of FProperty interface
 
 	FString GetCPPTypeCustom(FString* ExtendedTypeText, uint32 CPPExportFlags, const FString& KeyTypeText, const FString& InKeyExtendedTypeText, const FString& ValueTypeText, const FString& InValueExtendedTypeText) const;
@@ -3945,8 +3974,12 @@ public:
 	virtual bool SameType(const FProperty* Other) const override;
 	virtual EConvertFromTypeResult ConvertFromType(const FPropertyTag& Tag, FStructuredArchive::FSlot Slot, uint8* Data, UStruct* DefaultsStruct, const uint8* Defaults) override;
 	virtual void* GetValueAddressAtIndex_Direct(const FProperty* Inner, void* InValueAddress, int32 LogicalIndex) const override;
+	virtual bool UseBinaryOrNativeSerialization(const FArchive& Ar) const override;
 	virtual bool LoadFromTag(const FPropertyTag& Tag) override;
 	virtual void SaveToTag(FPropertyTag& Tag) override;
+	virtual bool LoadTypeName(UE::FPropertyTypeName Type, const FPropertyTag* Tag = nullptr) override;
+	virtual void SaveTypeName(UE::FPropertyTypeNameBuilder& Type) const override;
+	virtual bool CanSerializeFromTypeName(UE::FPropertyTypeName Type) const override;
 	// End of FProperty interface
 
 	FString GetCPPTypeCustom(FString* ExtendedTypeText, uint32 CPPExportFlags, const FString& ElementTypeText, const FString& InElementExtendedTypeText) const;
@@ -6067,12 +6100,15 @@ public:
 #if WITH_EDITORONLY_DATA
 	virtual void AppendSchemaHash(FBlake3& Builder, bool bSkipEditorOnly) const override;
 #endif
+	virtual bool UseBinaryOrNativeSerialization(const FArchive& Ar) const override;
 	virtual bool LoadFromTag(const FPropertyTag& Tag) override;
 	virtual void SaveToTag(FPropertyTag& Tag) override;
 	virtual void AssignToTag(FPropertyTag& Tag) override;
+	virtual bool LoadTypeName(UE::FPropertyTypeName Type, const FPropertyTag* Tag = nullptr) override;
+	virtual void SaveTypeName(UE::FPropertyTypeNameBuilder& Type) const override;
+	virtual bool CanSerializeFromTypeName(UE::FPropertyTypeName Type) const override;
 	// End of FProperty interface
 
-	bool UseBinaryOrNativeSerialization(const FArchive& Ar) const;
 	bool FindInnerPropertyInstance(FName PropertyName, const void* Data, const FProperty*& OutProp, const void*& OutData) const;
 
 private:
@@ -6389,7 +6425,7 @@ PRAGMA_ENABLE_DEPRECATION_WARNINGS
 	virtual void SerializeItem(FStructuredArchive::FSlot Slot, void* Value, void const* Defaults) const override;
 	virtual const TCHAR* ImportText_Internal(const TCHAR* Buffer, void* ContainerOrPropertyPtr, EPropertyPointerType PropertyPointerType, UObject* OwnerObject, int32 PortFlags, FOutputDevice* ErrorText) const override;
 	virtual bool LoadFromTag(const FPropertyTag& Tag) override;
-	virtual void SaveToTag(FPropertyTag& Tag) override;
+	virtual bool LoadTypeName(UE::FPropertyTypeName Type, const FPropertyTag* Tag = nullptr) override;
 	// End of FProperty interface
 
 	// FMulticastDelegateProperty interface
