[RegressionFix] added launch fixup tool button back to the component action menu

diff --git a/Engine/Source/Runtime/CoreUObject/Internal/UObject/PropertyBagRepository.h b/Engine/Source/Runtime/CoreUObject/Internal/UObject/PropertyBagRepository.h
--- a/Engine/Source/Runtime/CoreUObject/Internal/UObject/PropertyBagRepository.h
+++ b/Engine/Source/Runtime/CoreUObject/Internal/UObject/PropertyBagRepository.h
@@ -95,7 +95,7 @@ public:
 	COREUOBJECT_API bool RequiresFixup(const UObjectBase* Object) const;
 	
 	// Accessors
-	bool HasBag(const UObjectBase* Owner) const;
+	COREUOBJECT_API bool HasBag(const UObjectBase* Owner) const;
 	COREUOBJECT_API FPropertyBag* FindBag(const UObjectBase* Owner);
 	COREUOBJECT_API const FPropertyBag* FindBag(const UObjectBase* Owner) const;
 
