[SmallFix] ensure that property exists when generating wildcard paths during IDO generation
[SmallFix] use CopyCompleteValue if property types match when Copying a CDO into an IDO

diff --git a/Engine/Source/Runtime/CoreUObject/Private/UObject/InstanceDataObjectUtils.cpp b/Engine/Source/Runtime/CoreUObject/Private/UObject/InstanceDataObjectUtils.cpp
--- a/Engine/Source/Runtime/CoreUObject/Private/UObject/InstanceDataObjectUtils.cpp
+++ b/Engine/Source/Runtime/CoreUObject/Private/UObject/InstanceDataObjectUtils.cpp
@@ -389,7 +389,11 @@ namespace UE
 			{
 				FWildcardPropertyPathName ParentPath = ConvertToWildcardPath(Itr.GetPath());
 				ParentPath.Pop();
-				AddWildcardedProperties(LooseProperties, ParentPath, Itr.GetProperty());
+				const FProperty* Property = Itr.GetProperty();
+				if (ensure(Property))
+				{
+					AddWildcardedProperties(LooseProperties, ParentPath, Property);
+				}
 			}
 		}
 		
@@ -526,17 +530,26 @@ namespace UE
 		return Result;
 	}
 
-	void CopyCDO(const UStruct* SourceStruct, const void* SourceData, const UStruct* DestinationStruct, void* DestinationData)
+	void CopyCDO(const UObject* Source, UObject* Destination)
 	{
-		for (const FProperty* SourceProperty : TFieldRange<FProperty>(SourceStruct))
+		for (const FProperty* SourceProperty : TFieldRange<FProperty>(Source->GetClass()))
 		{
-			if (const FProperty* DestinationProperty = DestinationStruct->FindPropertyByName(SourceProperty->GetFName()))
+			if (const FProperty* DestinationProperty = Destination->GetClass()->FindPropertyByName(SourceProperty->GetFName()))
 			{
-				FString ValueText;
-				const void* SourceValue = SourceProperty->ContainerPtrToValuePtr<void>(SourceData);
-				void* DestinationValue = DestinationProperty->ContainerPtrToValuePtr<void>(DestinationData);
-				SourceProperty->ExportText_Direct(ValueText, SourceValue, SourceValue, nullptr, PPF_None);
-				DestinationProperty->ImportText_Direct(*ValueText, DestinationValue, nullptr, PPF_None);
+				if (SourceProperty->SameType(DestinationProperty))
+				{
+					const void* SourceValue = SourceProperty->ContainerPtrToValuePtr<void>(Source);
+					void* DestinationValue = DestinationProperty->ContainerPtrToValuePtr<void>(Destination);
+					DestinationProperty->CopyCompleteValue(DestinationValue, SourceValue);
+				}
+				else
+				{
+					FString ValueText;
+					const void* SourceValue = SourceProperty->ContainerPtrToValuePtr<void>(Source);
+					void* DestinationValue = DestinationProperty->ContainerPtrToValuePtr<void>(Destination);
+					SourceProperty->ExportText_Direct(ValueText, SourceValue, SourceValue, const_cast<UObject*>(Source), PPF_None);
+					DestinationProperty->ImportText_Direct(*ValueText, DestinationValue, Destination, PPF_None);
+				}
 			}
 		}
 	}
@@ -551,7 +564,7 @@ namespace UE
 		UObject* ResultCDO = Result->GetDefaultObject(true);
 		if (ensure(OwnerCDO && ResultCDO))
 		{
-			CopyCDO(OwnerClass, OwnerCDO, Result, ResultCDO);
+			CopyCDO(OwnerCDO, ResultCDO);
 		}
 		return Result;
 	}
