Property bag tests 1

First set of tests for the new property bag.
This exercises: Add, LoadFromTag, Remove, Iteration, storage

https://epicgames.slack.com/archives/C05HDDK8VUP/p1706200984745359?thread_ts=1706198804.429209&cid=C05HDDK8VUP

diff --git a/Engine/Source/Runtime/CoreUObject/Tests/PropertyBagTest.cpp b/Engine/Source/Runtime/CoreUObject/Tests/PropertyBagTest.cpp
new file mode 100644
--- /dev/null
+++ b/Engine/Source/Runtime/CoreUObject/Tests/PropertyBagTest.cpp
@@ -0,0 +1,959 @@
+// Copyright Epic Games, Inc. All Rights Reserved.
+#if WITH_TESTS
+
+#include "PropertyBagTest.h"
+
+#include "UObject/Class.h"
+#include "UObject/PropertyBag.h"
+#include "UObject/PropertyBagRepository.h"
+#include "UObject/UObjectGlobals.h"
+#include "UObject/CoreRedirects.h"
+#include "UObject/SavePackage.h"
+#include "UObject/UObjectThreadContext.h"
+#include "UObject/UObjectHash.h"
+#include "Misc/StringBuilder.h"
+#include "Misc/PackageName.h"
+#include "HAL/FileManager.h"
+#include "Serialization/MemoryReader.h"
+#include "Serialization/MemoryWriter.h"
+#include "Tasks/Task.h"
+#include "Tests/TestHarnessAdapter.h"
+
+DEFINE_LOG_CATEGORY_STATIC(LogPropertyBagTests, Log, All);
+
+namespace UE
+{
+namespace PropertyBagTestUtils
+{
+	UPackage* CreateTestPackage(const UClass* Class)
+	{
+		const FString UniqueAssetName = Class->GetName() + TEXT("_") + FGuid::NewGuid().ToString();
+		return CreatePackage(*(TEXT("/Temp/") + UniqueAssetName));
+	}
+
+	template<typename T>
+	UPackage* CreateTestPackage()
+	{
+		return CreateTestPackage(T::StaticClass());
+	}
+
+	template<typename T>
+	T* CreateTestObject()
+	{
+		return ::NewObject<T>((UObject*)CreateTestPackage<T>());
+	}
+
+	int BuildTypeName(const std::initializer_list<FName>& InNames, int I, FPropertyTypeNameBuilder& Dst)
+	{
+		const FName* Names = InNames.begin();
+		verify(I < (int)InNames.size());
+		verify(Names[I] != NAME_Exit);
+
+		Dst.AddTypeName(Names[I]);
+
+		if (Names[I] == NAME_MapProperty || Names[I] == NAME_EnumProperty)
+		{
+			Dst.BeginTypeParameters();
+			check(0 < (I = BuildTypeName(InNames, I + 1, Dst))); // Key
+			check(0 < (I = BuildTypeName(InNames, I, Dst))); // Value
+			Dst.EndTypeParameters();
+			return I;
+		}
+
+		if (Names[I] == NAME_StructProperty)
+		{
+			Dst.BeginTypeParameters();
+			while (I < (int)InNames.size())
+			{
+				if (Names[I] == NAME_Exit)
+				{
+					break;
+				}
+				check(0 < (I = BuildTypeName(InNames, I, Dst))); // Key
+			}
+			Dst.EndTypeParameters();
+		}
+
+		return I + 1;
+	}
+
+	FPropertyTypeName BuildTypeName(const std::initializer_list<FName>& Names)
+	{
+		FPropertyTypeNameBuilder Builder;
+		check(BuildTypeName(Names, 0, Builder) > 0);
+		return Builder.Build();
+	}
+
+	void BuildTypeName(FPropertyTypeNameBuilder& Builder, FProperty* Property)
+	{
+		if (FArrayProperty* ArrayProperty = CastField<FArrayProperty>(Property)) 
+		{
+			Builder.AddTypeName(NAME_ArrayProperty);
+			Builder.BeginTypeParameters();
+			FProperty* ValueType = CastField<FProperty>(ArrayProperty->GetInnerFieldByName(Property->GetFName()));
+			BuildTypeName(Builder, ValueType);
+			Builder.EndTypeParameters();
+
+			return;
+		}
+
+		if (FMapProperty* MapProperty = CastField<FMapProperty>(Property))
+		{
+			Builder.AddTypeName(NAME_MapProperty);
+			Builder.BeginTypeParameters();
+			BuildTypeName(Builder, const_cast<FProperty*>(MapProperty->GetKeyProperty()));
+			BuildTypeName(Builder, const_cast<FProperty*>(MapProperty->GetValueProperty()));
+			Builder.EndTypeParameters();
+
+			return;
+		}
+
+		Builder.AddTypeName(Property->GetClass()->GetFName());
+	}
+
+	FPropertyTypeName BuildTypeName(FProperty* Property)
+	{
+		FPropertyTypeNameBuilder Builder;
+		BuildTypeName(Builder, Property);
+		return Builder.Build();
+	}
+
+	FPropertyTypeName BuildPairTypeName(FProperty* Key, FProperty* Value)
+	{
+		FPropertyTypeNameBuilder Builder;
+		Builder.AddTypeName(NAME_StructProperty);
+		Builder.BeginTypeParameters();
+		BuildTypeName(Builder, Key);
+		BuildTypeName(Builder, Value);
+		Builder.EndTypeParameters();
+		return Builder.Build();
+	}
+
+	void BagIt(FPropertyBag& Dst, FPropertyPathName& Path, UClass* Def, void* Subject);
+
+	void BagIt(FPropertyBag& Dst, FPropertyPathName& Path, FProperty* Property, void* Subject)
+	{
+		// if TObjectPtr<T>
+		if (FObjectProperty* ObjectProperty = CastField<FObjectProperty>(Property))
+		{
+			void** PtrPtr = Property->ContainerPtrToValuePtr<void*>(Subject, 0);
+			if (*PtrPtr)
+			{
+				BagIt(Dst, Path, (UClass*)ObjectProperty->PropertyClass.Get(), *PtrPtr);
+			}
+				
+			return;
+		}
+
+		// if TArray<T>
+		if (FArrayProperty* ArrayProperty = CastField<FArrayProperty>(Property))
+		{
+			FProperty* ValueType = CastField<FProperty>(ArrayProperty->GetInnerFieldByName(Property->GetFName()));
+			size_t Stride = ValueType->ElementSize;
+
+			FPropertyTypeName TypeNameInner = BuildTypeName(ValueType);
+
+			TArray<uint8>* ArrayPtr = Property->ContainerPtrToValuePtr<TArray<uint8>>(Subject, 0);
+			for (int32 I = 0; I < ArrayPtr->Num(); I++)
+			{
+				void* Item = ArrayPtr->GetData() + Stride * I;
+				Path.SetSegment(Path.GetSegmentCount() - 1, { Property->GetFName(), TypeNameInner, I });
+					
+				BagIt(Dst, Path, ValueType, Item);
+			}
+
+			return;
+		}
+
+		// if TMap<TKey, TValue>
+		if (FMapProperty* MapProperty = CastField<FMapProperty>(Property))
+		{
+			FProperty* KeyProperty = const_cast<FProperty*>(MapProperty->GetKeyProperty());
+			FProperty* ValueProperty = const_cast<FProperty*>(MapProperty->GetValueProperty());
+
+			void* MapPtr = Property->ContainerPtrToValuePtr<void>(Subject, 0);
+
+			for (int32 I = 0, C = MapProperty->GetNum(MapPtr); I < C; I++)
+			{
+				void* Item = MapProperty->GetPairPtr(MapPtr, I);
+				Path.SetSegment(Path.GetSegmentCount() - 1, { Property->GetFName(), BuildTypeName(Property), I});
+
+				Path.Push({ FName(TEXT("Key")), BuildTypeName(KeyProperty) });
+				BagIt(Dst, Path, KeyProperty, Item);
+				Path.Pop();
+
+				Path.Push({ FName(TEXT("Value")), BuildTypeName(ValueProperty) });
+				BagIt(Dst, Path, ValueProperty, Item);
+				Path.Pop();
+			}
+
+			return;
+		}
+
+		// plain old field
+		Dst.Add(Path, Property, Subject);
+	}
+
+	void BagIt(FPropertyBag& Dst, FPropertyPathName& Path, UClass* Def, void* Subject)
+	{
+		for (FProperty* Property = Def->PropertyLink; Property; Property = Property->PropertyLinkNext)
+		{
+			Path.Push({ Property->GetFName(), BuildTypeName(Property) });
+
+			BagIt(Dst, Path, Property, Subject);
+
+			Path.Pop();
+		}
+	}
+
+	template<typename T>
+	FPropertyBag UObject2PropertyBag(T* Subject)
+	{
+		FPropertyBag Dst;
+		FPropertyPathName Path;
+		BagIt(Dst, Path, T::StaticClass(), Subject);
+		return Dst;
+	}
+
+	int Count(const FPropertyBag& Bag)
+	{
+		int Result = 0;
+		for (FPropertyBag::FConstIterator I = Bag.CreateConstIterator(); I; ++I, ++Result) {}
+		return Result;
+	}
+
+	template<typename T>
+	TArray<uint8> EncodeToBuffer(T InValue)
+	{
+		TArray<uint8> Buffer;
+
+		{ // create test data
+			FMemoryWriter Writer(Buffer);
+			FBinaryArchiveFormatter Formatter(Writer);
+			FStructuredArchive Archive(Formatter);
+
+			FStructuredArchiveSlot Value = Archive.Open();
+			Value << InValue;
+
+			Archive.Close();
+		}
+
+		return Buffer;
+	}
+
+	template<typename T>
+	constexpr EName GetENameForType()
+	{
+		if constexpr (std::is_same_v<T, int>)
+		{
+			return NAME_IntProperty;
+		}
+
+		if constexpr (std::is_same_v<T, float>)
+		{
+			return NAME_FloatProperty;
+		}
+	}
+
+	void LoadDataByTag(FPropertyBag& Dst, const FPropertyPathName& Path, const FPropertyTag& Tag, TArray<uint8>& Buffer)
+	{
+		FMemoryReader Reader(Buffer);
+		FBinaryArchiveFormatter Formatter(Reader);
+		FStructuredArchive Archive(Formatter);
+		FStructuredArchiveSlot Value = Archive.Open();
+
+		Dst.LoadPropertyByTag(Path, Tag, Value);
+	}
+
+	template<typename T>
+	void LoadDataByTag(FPropertyBag& Dst, const FPropertyPathName& Path, T InValue)
+	{
+		TArray<uint8> Buffer = EncodeToBuffer(InValue);
+
+		FPropertyTag Tag;
+		Tag.Type = GetENameForType<T>();
+		Tag.Name = FName(TEXT("TagTmp"));
+
+		LoadDataByTag(Dst, Path, Tag, Buffer);
+	}
+
+	template<typename T>
+	void LoadDataByTag(FPropertyBag& Dst, const FPropertyPathName& Path, FProperty* Prop, T InValue)
+	{
+		TArray<uint8> Buffer = EncodeToBuffer(InValue);
+
+		FPropertyTag Tag;
+		Tag.Prop = Prop;
+		Tag.Name = FName(TEXT("TagTmp"));
+
+		LoadDataByTag(Dst, Path, Tag, Buffer);
+	}
+} // PropertyBagTestUtils
+
+TEST_CASE_NAMED(FPropertyBagTest_Add, "CoreUObject::PropertyBag::Add", "[Core][UObject][PropertyBag]")
+{
+	FProperty* Property = CastFieldChecked<FProperty>(FField::Construct(NAME_IntProperty, {}, FName(TEXT("Tmp")), RF_NoFlags));
+	ON_SCOPE_EXIT
+	{
+		delete Property;
+	};
+
+	FPropertyTypeName IntType = PropertyBagTestUtils::BuildTypeName({ NAME_IntProperty });
+
+	// populate the path
+	FPropertyPathName Path;
+	Path.Push({ FName(TEXT("A")), IntType });
+
+	SECTION("AddAnInt_CheckStorage_ChangeValue_StorageValueChanges_LocationRemains")
+	{
+		FPropertyBag Dst;
+		int Value = 0x12345678;
+		Dst.Add(Path, Property, &Value);
+
+		void* StoredPointer = Dst.CreateConstIterator().GetValue();
+		CHECK(StoredPointer != &Value); // Its not storing the pointer we gave it.
+		CHECK(*static_cast<int*>(StoredPointer) == Value); // It did store the value
+
+		int SecondValue = 0x76543218;
+		Dst.Add(Path, Property, &SecondValue);
+
+		void* SecondStoredPointer = Dst.CreateConstIterator().GetValue();
+		CHECK(SecondStoredPointer == StoredPointer); // Same Property, Same Storage 
+		CHECK(*static_cast<int*>(SecondStoredPointer) == SecondValue); // It did store the value
+	}
+
+	SECTION("AddAnInt_CheckStorage_SamePathDifferentProperty_StorageValueChanges_LocationChanges")
+	{
+		TUniquePtr<FProperty> SecondProperty(CastFieldChecked<FProperty>(FField::Construct(NAME_IntProperty, {}, FName(TEXT("Tmp")), RF_NoFlags)));
+
+		FPropertyBag Dst;
+		int Value = 0x12345678;
+		Dst.Add(Path, Property, &Value);
+
+		void* StoredPointer = Dst.CreateConstIterator().GetValue();
+		const FProperty* Property1 = Dst.CreateConstIterator().GetProperty();
+		CHECK(*static_cast<int*>(StoredPointer) == Value); // It did store the value
+		CHECK(Property == Property1); // It did store the property
+
+		int SecondValue = 0x76543218;
+		Dst.Add(Path, SecondProperty.Get(), &SecondValue);
+
+		void* SecondStoredPointer = Dst.CreateConstIterator().GetValue();
+		const FProperty* Property2 = Dst.CreateConstIterator().GetProperty();
+		CHECK(SecondStoredPointer != StoredPointer); // Different Property, Different Storage for the same path
+		CHECK(*static_cast<int*>(SecondStoredPointer) == SecondValue); // It did store the value
+		CHECK(SecondProperty.Get() == Property2); // It did store the new property
+	}
+
+	{
+		FPropertyPathName PathAB = Path;
+		PathAB.Push({ FName(TEXT("B")), IntType });
+
+		FPropertyPathName PathABC = PathAB;
+		PathABC.Push({ FName(TEXT("C")), IntType });
+
+		SECTION("AddAnIntsInPathOrder_Descending_CheckStored")
+		{
+			FPropertyBag Dst;
+			int ValueAB = 0x12345678;
+			int ValueABC = 0x76543218;
+
+			Dst.Add(PathAB, Property, &ValueAB);
+			Dst.Add(PathABC, Property, &ValueABC);
+
+			CHECK(2 == PropertyBagTestUtils::Count(Dst));
+			for (FPropertyBag::FConstIterator I = Dst.CreateConstIterator(); I; ++I)
+			{
+				bool bExpectedPath = I.GetPath() == PathAB
+					|| I.GetPath() == PathABC;
+				CHECK(bExpectedPath);
+
+				void* StoredPointer = I.GetValue();
+				int* PTestValue = (I.GetPath() == PathAB) ? &ValueAB : &ValueABC;
+
+				CHECK(*static_cast<int*>(StoredPointer) == *PTestValue);
+			}
+		}
+
+		SECTION("AddAnIntsInPathOrder_Ascending_CheckStored")
+		{
+			FPropertyBag Dst;
+			int ValueABC = 0x12345678;
+			int ValueAB = 0x76543218;
+
+			Dst.Add(PathABC, Property, &ValueABC);
+			Dst.Add(PathAB, Property, &ValueAB);
+
+			CHECK(2 == PropertyBagTestUtils::Count(Dst));
+			for (FPropertyBag::FConstIterator I = Dst.CreateConstIterator(); I; ++I)
+			{
+				bool bExpectedPath = I.GetPath() == PathAB
+					|| I.GetPath() == PathABC;
+				CHECK(bExpectedPath);
+
+				void* StoredPointer = I.GetValue();
+				int* PTestValue = (I.GetPath() == PathABC) ? &ValueABC : &ValueAB;
+
+				CHECK(*static_cast<int*>(StoredPointer) == *PTestValue);
+			}
+		}
+	}
+}
+
+TEST_CASE_NAMED(FPropertyBagTest_Remove, "CoreUObject::PropertyBag::Remove", "[Core][UObject][PropertyBag]")
+{
+	FProperty* Property = CastFieldChecked<FProperty>(FField::Construct(NAME_IntProperty, {}, FName(TEXT("Tmp")), RF_NoFlags));
+	ON_SCOPE_EXIT
+	{
+		delete Property;
+	};
+
+	FPropertyTypeName IntType = PropertyBagTestUtils::BuildTypeName({ NAME_IntProperty });
+
+	// populate the path
+	FPropertyPathName PathA;
+	PathA.Push({ FName(TEXT("A")), IntType });
+
+	FPropertyPathName PathAV = PathA;
+	PathAV.Push({ FName(TEXT("V")), IntType }); // A/V to simulate a more complex type
+
+	FPropertyPathName PathAB = PathA;
+	PathAB.Push({ FName(TEXT("B")), IntType });
+
+	FPropertyPathName PathABC = PathAB;
+	PathABC.Push({ FName(TEXT("C")), IntType });
+
+	SECTION("AddA_AddABC_RemoveABC")
+	{
+		FPropertyBag Dst;
+		int ValueA = 0x76543218;
+		int ValueABC = 0x12345678;
+
+		Dst.Add(PathAV, Property, &ValueA);
+		Dst.Add(PathABC, Property, &ValueABC);
+
+		// A/V and A/B/C
+		CHECK(2 == PropertyBagTestUtils::Count(Dst));
+
+		Dst.Remove(PathABC);
+
+		// leaving A/V
+		CHECK(1 == PropertyBagTestUtils::Count(Dst));
+		bool bExpectedPath = PathAV == Dst.CreateConstIterator().GetPath();
+		CHECK(bExpectedPath);
+	}
+
+	SECTION("AddA_AddABC_RemoveA")
+	{
+		FPropertyBag Dst;
+		int ValueA = 0x76543218;
+		int ValueABC = 0x12345678;
+
+		Dst.Add(PathAV, Property, &ValueA);
+		Dst.Add(PathABC, Property, &ValueABC);
+
+		// A/V and A/B/C
+		CHECK(2 == PropertyBagTestUtils::Count(Dst));
+
+		Dst.Remove(PathAV);
+
+		// leaving A/V
+		CHECK(1 == PropertyBagTestUtils::Count(Dst));
+		bool bExpectedPath = PathABC == Dst.CreateConstIterator().GetPath();
+		CHECK(bExpectedPath);
+	}
+
+	SECTION("AddA_RemoveA_IsEmpty")
+	{
+		FPropertyBag Dst;
+		int ValueA = 0x76543218;
+
+		Dst.Add(PathAV, Property, &ValueA);
+
+		// A/V 
+		CHECK(1 == PropertyBagTestUtils::Count(Dst));
+
+		Dst.Remove(PathAV);
+
+		// now should be empty
+		CHECK(0 == PropertyBagTestUtils::Count(Dst));
+		CHECK(Dst.IsEmpty());
+	}
+}
+
+TEST_CASE_NAMED(FPropertyBagTest_Empty, "CoreUObject::PropertyBag::Empty", "[Core][UObject][PropertyBag]")
+{
+	FProperty* Property = CastFieldChecked<FProperty>(FField::Construct(NAME_IntProperty, {}, FName(TEXT("Tmp")), RF_NoFlags));
+	ON_SCOPE_EXIT
+	{
+		delete Property;
+	};
+
+	FPropertyTypeName IntType = PropertyBagTestUtils::BuildTypeName({ NAME_IntProperty });
+
+	// populate the path
+	FPropertyPathName PathA;
+	PathA.Push({ FName(TEXT("A")), IntType });
+
+	FPropertyPathName PathAB = PathA;
+	PathAB.Push({ FName(TEXT("B")), IntType });
+
+	SECTION("AddA_AddAB_Empty_IsEmpty")
+	{
+		FPropertyBag Dst;
+		int ValueA = 0x76543218;
+		int ValueAB = 0x12345678;
+
+		Dst.Add(PathA, Property, &ValueA);
+		Dst.Add(PathAB, Property, &ValueAB);
+
+		CHECK(2 == PropertyBagTestUtils::Count(Dst));
+		CHECK(!Dst.IsEmpty());
+
+		Dst.Empty();
+
+		CHECK(0 == PropertyBagTestUtils::Count(Dst));
+		CHECK(Dst.IsEmpty());
+	}
+}
+
+TEST_CASE_NAMED(FPropertyBagTest_LoadPropertyByTag, "CoreUObject::PropertyBag::LoadPropertyByTag", "[Core][UObject][PropertyBag]")
+{
+	FPropertyTypeName IntType = PropertyBagTestUtils::BuildTypeName({ NAME_IntProperty });
+
+	// populate the path
+	FPropertyPathName PathA;
+	PathA.Push({ FName(TEXT("A")), IntType });
+
+	FPropertyPathName PathAB = PathA;
+	PathAB.Push({ FName(TEXT("B")), IntType });
+
+	FPropertyPathName PathABC = PathAB;
+	PathABC.Push({ FName(TEXT("C")), IntType });
+
+	SECTION("LoadByTag")
+	{
+		FPropertyBag Dst;
+
+		PropertyBagTestUtils::LoadDataByTag(Dst, PathA, 0x12345678);
+
+		CHECK(*static_cast<int*>(Dst.CreateConstIterator().GetValue()) == 0x12345678);
+	}
+
+	SECTION("LoadByTagAnInt_CheckStorage_ChangeValue_StorageValueChanges_LocationRemains")
+	{
+		FPropertyBag Dst;
+		int ValueA = 0x12345678;
+
+		PropertyBagTestUtils::LoadDataByTag(Dst, PathA, ValueA);
+
+		void* StoredPointer = Dst.CreateConstIterator().GetValue();
+		CHECK(StoredPointer != &ValueA); // Its not storing the pointer we gave it.
+		CHECK(*static_cast<int*>(StoredPointer) == ValueA); // It did store the value
+
+		int SecondValue = 0x76543218;
+		PropertyBagTestUtils::LoadDataByTag(Dst, PathA, SecondValue);
+
+		void* SecondStoredPointer = Dst.CreateConstIterator().GetValue();
+		CHECK(SecondStoredPointer == StoredPointer); // Same Property, Same Storage 
+		CHECK(*static_cast<int*>(SecondStoredPointer) == SecondValue); // It did store the value
+	}
+
+	SECTION("LoadByTagAnInt_CheckStorage_ChangeValueType_StorageValueChanges_LocationChanges")
+	{
+		FPropertyBag Dst;
+		int ValueA = 0x12345678;
+
+		PropertyBagTestUtils::LoadDataByTag(Dst, PathA, ValueA);
+
+		void* StoredPointer = Dst.CreateConstIterator().GetValue();
+		CHECK(StoredPointer != &ValueA); // Its not storing the pointer we gave it.
+		CHECK(*static_cast<int*>(StoredPointer) == ValueA); // It did store the value
+
+		float SecondValue = 3.1415f;
+		PropertyBagTestUtils::LoadDataByTag(Dst, PathA, SecondValue);
+
+		void* SecondStoredPointer = Dst.CreateConstIterator().GetValue();
+		CHECK(SecondStoredPointer != StoredPointer); // Different Property, Different Storage for the same path
+		CHECK(*static_cast<float*>(SecondStoredPointer) == SecondValue); // It did store the value
+	}
+
+	SECTION("LoadByTagAnInt_CheckStorage_ChangeProperty_StorageValueChanges_LocationChanges")
+	{
+		FProperty* Property = CastFieldChecked<FProperty>(FField::Construct(NAME_IntProperty, {}, FName(TEXT("Tmp")), RF_NoFlags));
+		ON_SCOPE_EXIT
+		{
+			delete Property;
+		};
+
+		FPropertyBag Dst;
+		int ValueA = 0x12345678;
+
+		PropertyBagTestUtils::LoadDataByTag(Dst, PathA, ValueA);
+
+		void* StoredPointer = Dst.CreateConstIterator().GetValue();
+		CHECK(StoredPointer != &ValueA); // Its not storing the pointer we gave it.
+		CHECK(*static_cast<int*>(StoredPointer) == ValueA); // It did store the value
+
+		int SecondValue = 0x76543218;
+		PropertyBagTestUtils::LoadDataByTag(Dst, PathA, Property, SecondValue);
+
+		void* SecondStoredPointer = Dst.CreateConstIterator().GetValue();
+		CHECK(SecondStoredPointer != StoredPointer); // Different Property, Different Storage for the same path
+		CHECK(*static_cast<int*>(SecondStoredPointer) == SecondValue); // It did store the value
+	}
+
+	SECTION("LoadByTagAnInt_CheckStorage_ChangeProperty_StorageValueChanges_LocationChanges")
+	{
+		FProperty* Property1 = CastFieldChecked<FProperty>(FField::Construct(NAME_IntProperty, {}, FName(TEXT("Tmp")), RF_NoFlags));
+		FProperty* Property2 = CastFieldChecked<FProperty>(FField::Construct(NAME_IntProperty, {}, FName(TEXT("Tmp")), RF_NoFlags));
+		ON_SCOPE_EXIT
+		{
+			delete Property1;
+			delete Property2;
+		};
+
+		FPropertyBag Dst;
+		int ValueA = 0x12345678;
+
+		PropertyBagTestUtils::LoadDataByTag(Dst, PathA, Property1, ValueA);
+
+		void* StoredPointer = Dst.CreateConstIterator().GetValue();
+		CHECK(StoredPointer != &ValueA); // Its not storing the pointer we gave it.
+		CHECK(*static_cast<int*>(StoredPointer) == ValueA); // It did store the value
+		CHECK(Property1 == Dst.CreateConstIterator().GetProperty());
+
+		int SecondValue = 0x76543218;
+		PropertyBagTestUtils::LoadDataByTag(Dst, PathA, Property2, SecondValue);
+
+		void* SecondStoredPointer = Dst.CreateConstIterator().GetValue();
+		CHECK(SecondStoredPointer != StoredPointer); // Different Property, Different Storage for the same path
+		CHECK(*static_cast<int*>(SecondStoredPointer) == SecondValue); // It did store the value
+		CHECK(Property2 == Dst.CreateConstIterator().GetProperty());
+	}
+
+	SECTION("LoadByTagIntsInPathOrder_Descending_CheckStored")
+	{
+		FPropertyBag Dst;
+		int ValueAB = 0x12345678;
+		int ValueABC = 0x76543218;
+
+		PropertyBagTestUtils::LoadDataByTag(Dst, PathAB, ValueAB);
+		PropertyBagTestUtils::LoadDataByTag(Dst, PathABC, ValueABC);
+
+		CHECK(2 == PropertyBagTestUtils::Count(Dst));
+		for (FPropertyBag::FConstIterator I = Dst.CreateConstIterator(); I; ++I)
+		{
+			bool bExpectedPath = I.GetPath() == PathAB
+				|| I.GetPath() == PathABC;
+			CHECK(bExpectedPath);
+
+			void* StoredPointer = I.GetValue();
+			int* PTestValue = (I.GetPath() == PathAB) ? &ValueAB : &ValueABC;
+
+			CHECK(*static_cast<int*>(StoredPointer) == *PTestValue);
+		}
+	}
+
+	SECTION("AddAnIntsInPathOrder_Ascending_CheckStored")
+	{
+		FPropertyBag Dst;
+		int ValueABC = 0x12345678;
+		int ValueAB = 0x76543218;
+
+		PropertyBagTestUtils::LoadDataByTag(Dst, PathABC, ValueABC);
+		PropertyBagTestUtils::LoadDataByTag(Dst, PathAB, ValueAB);
+
+		CHECK(2 == PropertyBagTestUtils::Count(Dst));
+		for (FPropertyBag::FConstIterator I = Dst.CreateConstIterator(); I; ++I)
+		{
+			bool bExpectedPath = I.GetPath() == PathAB
+				|| I.GetPath() == PathABC;
+			CHECK(bExpectedPath);
+
+			void* StoredPointer = I.GetValue();
+			int* PTestValue = (I.GetPath() == PathABC) ? &ValueABC : &ValueAB;
+
+			CHECK(*static_cast<int*>(StoredPointer) == *PTestValue);
+		}
+	}
+}
+
+TEST_CASE_NAMED(FPropertyBagTest_LoadPropertyByTag_And_Add, "CoreUObject::PropertyBag::LoadPropertyByTag_And_Add", "[Core][UObject][PropertyBag]")
+{
+	FProperty* Property = CastFieldChecked<FProperty>(FField::Construct(NAME_IntProperty, {}, FName(TEXT("Tmp")), RF_NoFlags));
+	ON_SCOPE_EXIT
+	{
+		delete Property;
+	};
+
+	FPropertyTypeName IntType = PropertyBagTestUtils::BuildTypeName({ NAME_IntProperty });
+
+	// populate the path
+	FPropertyPathName PathA;
+	PathA.Push({ FName(TEXT("A")), IntType });
+
+	FPropertyPathName PathAB = PathA;
+	PathAB.Push({ FName(TEXT("B")), IntType });
+
+	SECTION("AddAnIntsInPathOrder_Ascending_CheckStored")
+	{
+		FPropertyBag Dst;
+		int ValueA = 0x12345678;
+		int ValueAB = 0x76543218;
+
+		PropertyBagTestUtils::LoadDataByTag(Dst, PathA, ValueA);
+		Dst.Add(PathAB, Property, &ValueAB);
+
+		CHECK(2 == PropertyBagTestUtils::Count(Dst));
+		for (FPropertyBag::FConstIterator I = Dst.CreateConstIterator(); I; ++I)
+		{
+			bool bExpectedPath = I.GetPath() == PathA
+				|| I.GetPath() == PathAB;
+			CHECK(bExpectedPath);
+
+			void* StoredPointer = I.GetValue();
+			int* PTestValue = (I.GetPath() == PathA) ? &ValueA : &ValueAB;
+
+			CHECK(*static_cast<int*>(StoredPointer) == *PTestValue);
+		}
+	}
+
+	SECTION("AddAnIntsInPathOrder_Descending_CheckStored")
+	{
+		FPropertyBag Dst;
+		int ValueA = 0x12345678;
+		int ValueAB = 0x76543218;
+
+		Dst.Add(PathAB, Property, &ValueAB);
+		PropertyBagTestUtils::LoadDataByTag(Dst, PathA, ValueA);
+
+		CHECK(2 == PropertyBagTestUtils::Count(Dst));
+		for (FPropertyBag::FConstIterator I = Dst.CreateConstIterator(); I; ++I)
+		{
+			bool bExpectedPath = I.GetPath() == PathA
+				|| I.GetPath() == PathAB;
+			CHECK(bExpectedPath);
+
+			void* StoredPointer = I.GetValue();
+			int* PTestValue = (I.GetPath() == PathA) ? &ValueA : &ValueAB;
+
+			CHECK(*static_cast<int*>(StoredPointer) == *PTestValue);
+		}
+	}
+
+	SECTION("UseBothAddAndLoadTag_UseTheSamePath_CheckStorage")
+	{
+		FPropertyBag Dst;
+		int ValueA = 0x12345678;
+		int ValueAB = 0x76543218;
+
+		Dst.Add(PathA, Property, &ValueA);
+
+		{
+			void* StoredPointer = Dst.CreateConstIterator().GetValue();
+			CHECK(StoredPointer != &ValueA); // Its not storing the pointer we gave it.
+			CHECK(*static_cast<int*>(StoredPointer) == ValueA); // It did store the value
+			CHECK(Property == Dst.CreateConstIterator().GetProperty());
+		}
+
+		PropertyBagTestUtils::LoadDataByTag(Dst, PathA, ValueAB);
+
+		{
+			void* StoredPointer = Dst.CreateConstIterator().GetValue();
+			CHECK(StoredPointer != &ValueAB); // Its not storing the pointer we gave it.
+			CHECK(*static_cast<int*>(StoredPointer) == ValueAB); // Value updated
+			CHECK(Property != Dst.CreateConstIterator().GetProperty());
+		}
+	}
+
+	SECTION("UseBothLoadTagAndAdd_UseTheSamePath_CheckStorage")
+	{
+		FPropertyBag Dst;
+		int ValueA = 0x12345678;
+		int ValueAB = 0x76543218;
+
+		PropertyBagTestUtils::LoadDataByTag(Dst, PathA, ValueA);
+
+		{
+			void* StoredPointer = Dst.CreateConstIterator().GetValue();
+			CHECK(StoredPointer != &ValueA); // Its not storing the pointer we gave it.
+			CHECK(*static_cast<int*>(StoredPointer) == ValueA); // It did store the value
+			CHECK(Property != Dst.CreateConstIterator().GetProperty());
+		}
+
+		Dst.Add(PathA, Property, &ValueAB);
+
+		{
+			void* StoredPointer = Dst.CreateConstIterator().GetValue();
+			CHECK(StoredPointer != &ValueAB); // Its not storing the pointer we gave it.
+			CHECK(*static_cast<int*>(StoredPointer) == ValueAB); // Value updated
+			CHECK(Property == Dst.CreateConstIterator().GetProperty());
+		}
+	}
+
+	SECTION("UseBothAddAndLoadTag_SameFName_UseTheSamePath_CheckStorage")
+	{
+		FProperty* PropertyTmpTag = CastFieldChecked<FProperty>(FField::Construct(NAME_IntProperty, {}, FName(TEXT("TagTmp")), RF_NoFlags));
+		ON_SCOPE_EXIT
+		{
+			delete PropertyTmpTag;
+		};
+
+		FPropertyBag Dst;
+		int ValueA = 0x12345678;
+		int ValueAB = 0x76543218;
+
+		Dst.Add(PathA, PropertyTmpTag, &ValueA);
+
+		{
+			void* StoredPointer = Dst.CreateConstIterator().GetValue();
+			CHECK(StoredPointer != &ValueA); // Its not storing the pointer we gave it.
+			CHECK(*static_cast<int*>(StoredPointer) == ValueA); // It did store the value
+			CHECK(PropertyTmpTag == Dst.CreateConstIterator().GetProperty());
+		}
+
+		PropertyBagTestUtils::LoadDataByTag(Dst, PathA, ValueAB);
+
+		{
+			void* StoredPointer = Dst.CreateConstIterator().GetValue();
+			CHECK(StoredPointer != &ValueAB); // Its not storing the pointer we gave it.
+			CHECK(*static_cast<int*>(StoredPointer) == ValueAB); // Value updated
+			CHECK(PropertyTmpTag != Dst.CreateConstIterator().GetProperty());
+		}
+	}
+}
+
+TEST_CASE_NAMED(FPropertyBagTest_Iteration, "CoreUObject::PropertyBag::Iteration", "[Core][UObject][PropertyBag]")
+{
+	SECTION("Iteration_Parents_Before_Children")
+	{
+		using namespace PropertyBagTestUtils;
+
+		UTestPropertyBagABCDABEF* TestInstance = CreateTestObject<UTestPropertyBagABCDABEF>();
+		TestInstance->ABCD = NewObject<UTestPropertyBagABCD>();
+		TestInstance->ABEF = NewObject<UTestPropertyBagABEF>();
+
+		auto Dst = UObject2PropertyBag(TestInstance);
+
+		static const TCHAR* ExpectedPaths[] = {
+			TEXT("A (StrProperty)"),
+			TEXT("ABCD (ObjectProperty)/A (StrProperty)"),
+			TEXT("ABCD (ObjectProperty)/B (EnumProperty)"),
+			TEXT("ABCD (ObjectProperty)/C (IntProperty)"),
+			TEXT("ABCD (ObjectProperty)/D (FloatProperty)"),
+			TEXT("E (IntProperty)"),
+			TEXT("ABEF (ObjectProperty)/A (StrProperty)"),
+			TEXT("ABEF (ObjectProperty)/B (EnumProperty)"),
+			TEXT("ABEF (ObjectProperty)/E (IntProperty)"),
+			TEXT("ABEF (ObjectProperty)/F (FloatProperty)"),
+			TEXT("F (FloatProperty)"),
+		};
+
+		int ExpectedI = 0;
+		for (FPropertyBag::FConstIterator I = Dst.CreateConstIterator(); I; ++I, ++ExpectedI)
+		{
+			TStringBuilderWithBuffer<TCHAR, 24> ToTest;
+			I.GetPath().ToString(ToTest, TEXT("/"));
+			CHECK(FCString::Strcmp(ToTest.ToString(), ExpectedPaths[ExpectedI]) == 0);
+		}
+	}
+}
+
+TEST_CASE_NAMED(FPropertyBagTest_Paths_With_Index, "CoreUObject::PropertyBag::Paths_With_Index", "[Core][UObject][PropertyBag]")
+{
+	FProperty* Property = CastFieldChecked<FProperty>(FField::Construct(NAME_IntProperty, {}, FName(TEXT("Tmp")), RF_NoFlags));
+	ON_SCOPE_EXIT
+	{
+		delete Property;
+	};
+
+	FPropertyTypeName NilType;
+	FPropertyTypeName IntType = PropertyBagTestUtils::BuildTypeName({ NAME_IntProperty });
+
+	SECTION("Path_With_Mid_Index")
+	{
+		FPropertyBag Dst;
+
+		for (int I = 0; I < 3; ++I) {
+			// populate the path
+			FPropertyPathName PathABnC;
+			PathABnC.Push({ FName(TEXT("A")), NilType });
+			PathABnC.Push({ FName(TEXT("B")), NilType, I });
+			PathABnC.Push({ FName(TEXT("C")), IntType });
+
+			int Value = I << 24
+				| I << 16
+				| I << 8
+				| I << 0;
+
+			Dst.Add(PathABnC, Property, &Value);
+		}
+
+		static const TCHAR* ExpectedPaths[] = {
+			TEXT("A/B[0]/C (IntProperty)"),
+			TEXT("A/B[1]/C (IntProperty)"),
+			TEXT("A/B[2]/C (IntProperty)"),
+		};
+
+		int ExpectedI = 0;
+		for (FPropertyBag::FConstIterator I = Dst.CreateConstIterator(); I; ++I, ++ExpectedI)
+		{
+			TStringBuilderWithBuffer<TCHAR, 24> ToTest;
+			I.GetPath().ToString(ToTest, TEXT("/"));
+			CHECK(FString(ToTest.ToString()) == ExpectedPaths[ExpectedI]);
+
+			int Value = ExpectedI << 24
+				| ExpectedI << 16
+				| ExpectedI << 8
+				| ExpectedI << 0;
+
+			CHECK(*static_cast<int*>(I.GetValue()) == Value);			
+		}
+	}
+
+	SECTION("Path_With_End_Index")
+	{
+		FPropertyBag Dst;
+
+		for (int I = 0; I < 3; ++I) {
+			// populate the path
+			FPropertyPathName PathABnC;
+			PathABnC.Push({ FName(TEXT("A")), NilType });
+			PathABnC.Push({ FName(TEXT("B")), NilType });
+			PathABnC.Push({ FName(TEXT("C")), IntType, I });
+
+			int Value = I << 24
+				| I << 16
+				| I << 8
+				| I << 0;
+
+			Dst.Add(PathABnC, Property, &Value);
+		}
+
+		static const TCHAR* ExpectedPaths[] = {
+			TEXT("A/B/C[0] (IntProperty)"),
+			TEXT("A/B/C[1] (IntProperty)"),
+			TEXT("A/B/C[2] (IntProperty)"),
+		};
+
+		int ExpectedI = 0;
+		for (FPropertyBag::FConstIterator I = Dst.CreateConstIterator(); I; ++I, ++ExpectedI)
+		{
+			TStringBuilderWithBuffer<TCHAR, 24> ToTest;
+			I.GetPath().ToString(ToTest, TEXT("/"));
+			CHECK(FString(ToTest.ToString()) == ExpectedPaths[ExpectedI]);
+
+			int Value = ExpectedI << 24
+				| ExpectedI << 16
+				| ExpectedI << 8
+				| ExpectedI << 0;
+
+			CHECK(*static_cast<int*>(I.GetValue()) == Value);
+		}
+	}
+}
+
+} // UE
+
+#endif // WITH_TESTS
diff --git a/Engine/Source/Runtime/CoreUObject/Tests/PropertyBagTest.h b/Engine/Source/Runtime/CoreUObject/Tests/PropertyBagTest.h
new file mode 100644
--- /dev/null
+++ b/Engine/Source/Runtime/CoreUObject/Tests/PropertyBagTest.h
@@ -0,0 +1,142 @@
+// Copyright Epic Games, Inc. All Rights Reserved.
+
+#pragma once
+
+#include "CoreTypes.h"
+#include "UObject/Object.h"
+#include "PropertyBagTest.generated.h"
+
+UCLASS()
+class UTestPropertyBag1Int : public UObject
+{
+	GENERATED_BODY()
+
+public:
+	UPROPERTY(EditAnywhere, Category = "Text")
+	int32 TheInt = -1;
+};
+
+
+UENUM(BlueprintType)
+enum class ETestPropertyBagEnum : uint8
+{
+	ValueA,
+	ValueB,
+	ValueC,
+};
+
+UCLASS()
+class UTestPropertyBagABCD : public UObject
+{
+	GENERATED_BODY()
+	
+public:
+	UPROPERTY(EditAnywhere, Category = "Text")
+	FString A;
+	
+	UPROPERTY(EditAnywhere, Category = "Text")
+	ETestPropertyBagEnum B = ETestPropertyBagEnum::ValueA;
+	
+	UPROPERTY(EditAnywhere, Category = "Text")
+	int32 C = -1;
+	
+	UPROPERTY(EditAnywhere, Category = "Text")
+	float D = 100.f;
+};
+
+UCLASS()
+class UTestPropertyBagABEF : public UObject
+{
+	GENERATED_BODY()
+
+public:
+	UPROPERTY(EditAnywhere, Category = "Text")
+	FString A;
+
+	UPROPERTY(EditAnywhere, Category = "Text")
+	ETestPropertyBagEnum B = ETestPropertyBagEnum::ValueA;
+
+	UPROPERTY(EditAnywhere, Category = "Text")
+	int32 E = -1;
+
+	UPROPERTY(EditAnywhere, Category = "Text")
+	float F = 100.f;
+};
+
+UCLASS()
+class UTestPropertyBagABCDABEF : public UObject
+{
+	GENERATED_BODY()
+
+public:
+	UPROPERTY(EditAnywhere, Category = "Text")
+	FString A;
+
+	UPROPERTY(EditAnywhere, Category = "Text")
+	TObjectPtr<UTestPropertyBagABCD> ABCD;
+
+	UPROPERTY(EditAnywhere, Category = "Text")
+	int32 E = -1;
+
+	UPROPERTY(EditAnywhere, Category = "Text")
+	TObjectPtr<UTestPropertyBagABEF> ABEF;
+
+	UPROPERTY(EditAnywhere, Category = "Text")
+	float F = 100.f;
+};
+
+UCLASS()
+class UTestPropertyBagArrayOfString : public UObject
+{
+	GENERATED_BODY()
+
+public:
+	UPROPERTY(EditAnywhere, Category = "Text")
+	TArray<FString> Arr;
+};
+
+UCLASS()
+class UTestPropertyBagAllSupportedTypes : public UObject
+{
+	GENERATED_BODY()
+
+public:
+	UPROPERTY(EditAnywhere, Category = "Text")
+	TArray<TObjectPtr<UTestPropertyBagAllSupportedTypes>> Arr;
+
+	UPROPERTY(EditAnywhere, Category = "Text")
+	TMap<FString, int32> Map;
+
+	UPROPERTY(EditAnywhere, Category = "Text")
+	FString A;
+
+	UPROPERTY(EditAnywhere, Category = "Text")
+	int32 B = -1;
+
+	UPROPERTY(EditAnywhere, Category = "Text")
+	int8 B8 = -1;
+
+	UPROPERTY(EditAnywhere, Category = "Text")
+	int16 B16 = -1;
+
+	UPROPERTY(EditAnywhere, Category = "Text")
+	int64 B64 = -1;
+
+	UPROPERTY(EditAnywhere, Category = "Text")
+	float C = 100.f;
+
+	UPROPERTY(EditAnywhere, Category = "Text")
+	double D = 100.;
+
+	UPROPERTY(EditAnywhere, Category = "Text")
+	uint32 E = -1;
+
+	UPROPERTY(EditAnywhere, Category = "Text")
+	uint8 E8 = -1;
+
+	UPROPERTY(EditAnywhere, Category = "Text")
+	uint16 E16 = -1;
+
+	UPROPERTY(EditAnywhere, Category = "Text")
+	uint64 E64 = -1;
+};
