Fix crash when running game with legacy pakfiles (-skipiostore). FPakReadRequestBase and all subclasses of IAsyncReadRequest need a default implementation for the new EnsureCompletion function rather than a fatal error if called.
#rnx

diff --git a/Engine/Source/Editor/UnrealEd/Private/EditorDomain/EditorDomainArchive.cpp b/Engine/Source/Editor/UnrealEd/Private/EditorDomain/EditorDomainArchive.cpp
--- a/Engine/Source/Editor/UnrealEd/Private/EditorDomain/EditorDomainArchive.cpp
+++ b/Engine/Source/Editor/UnrealEd/Private/EditorDomain/EditorDomainArchive.cpp
@@ -715,11 +715,6 @@ protected:
 	{
 	}
 
-	virtual void EnsureCompletion() override
-	{
-		WaitCompletionImpl(0.0f);
-	}
-
 	virtual void CancelImpl() override
 	{
 	}
@@ -762,11 +757,6 @@ protected:
 	{
 	}
 
-	virtual void EnsureCompletion() override
-	{
-		WaitCompletionImpl(0.0f);
-	}
-
 	virtual void CancelImpl() override
 	{
 	}
diff --git a/Engine/Source/Runtime/Core/Private/Microsoft/MicrosoftAsyncIO.h b/Engine/Source/Runtime/Core/Private/Microsoft/MicrosoftAsyncIO.h
--- a/Engine/Source/Runtime/Core/Private/Microsoft/MicrosoftAsyncIO.h
+++ b/Engine/Source/Runtime/Core/Private/Microsoft/MicrosoftAsyncIO.h
@@ -323,11 +323,6 @@ public:
 		}
 	}
 
-	virtual void EnsureCompletion() override
-	{
-		WaitCompletionImpl(0.0f);
-	}
-
 	virtual void CancelImpl() override
 	{
 		// no cancel support
@@ -361,11 +356,6 @@ public:
 		while (!*(volatile bool*)&bCompleteAndCallbackCalled);
 	}
 
-	virtual void EnsureCompletion() override
-	{
-		WaitCompletionImpl(0.0f);
-	}
-
 	virtual void CancelImpl() override
 	{
 	}
@@ -391,11 +381,6 @@ public:
 		while (!*(volatile bool*)&bCompleteAndCallbackCalled);
 	}
 
-	virtual void EnsureCompletion() override
-	{
-		WaitCompletionImpl(0.0f);
-	}
-
 	virtual void CancelImpl() override
 	{
 	}
diff --git a/Engine/Source/Runtime/Core/Public/Async/AsyncFileHandle.h b/Engine/Source/Runtime/Core/Public/Async/AsyncFileHandle.h
--- a/Engine/Source/Runtime/Core/Public/Async/AsyncFileHandle.h
+++ b/Engine/Source/Runtime/Core/Public/Async/AsyncFileHandle.h
@@ -102,7 +102,11 @@ public:
 	*/
 	virtual void EnsureCompletion()
 	{
-		unimplemented();
+		// Default implementation is the same as WaitCompletion(0.0f) except that it skips the testing of
+		// PollCompletion. This is potentially slower because we do not early exit if PollCompletion is true, but it
+		// provides a stronger guarantee of completion because PollCompletion can sometimes return true while
+		// completion steps are still in progress.
+		WaitCompletionImpl(0.0f);
 	}
 
 	/** Cancel the request. This is a non-blocking async call and so does not ensure completion! **/
diff --git a/Engine/Source/Runtime/Core/Public/Async/AsyncFileHandleNull.h b/Engine/Source/Runtime/Core/Public/Async/AsyncFileHandleNull.h
--- a/Engine/Source/Runtime/Core/Public/Async/AsyncFileHandleNull.h
+++ b/Engine/Source/Runtime/Core/Public/Async/AsyncFileHandleNull.h
@@ -24,11 +24,6 @@ protected:
 	{
 	}
 
-	virtual void EnsureCompletion() override
-	{
-		WaitCompletionImpl(0.0f);
-	}
-
 	virtual void CancelImpl() override
 	{
 	}
diff --git a/Engine/Source/Runtime/CoreUObject/Private/Serialization/BulkDataStreaming.cpp b/Engine/Source/Runtime/CoreUObject/Private/Serialization/BulkDataStreaming.cpp
--- a/Engine/Source/Runtime/CoreUObject/Private/Serialization/BulkDataStreaming.cpp
+++ b/Engine/Source/Runtime/CoreUObject/Private/Serialization/BulkDataStreaming.cpp
@@ -214,12 +214,6 @@ public:
 	
 	virtual void WaitCompletionImpl(float TimeLimitSeconds) override;
 
-	virtual void EnsureCompletion() override
-	{
-		WaitCompletionImpl(0.0f);
-	}
-
-
 	virtual void CancelImpl() override;
 	virtual void ReleaseMemoryOwnershipImpl() override;
 	virtual void HandleChunkResult(TIoStatusOr<FIoBuffer>&& Result) override;
@@ -294,11 +288,6 @@ private:
 		while (!*(volatile bool*)&bCompleteAndCallbackCalled);
 	}
 
-	virtual void EnsureCompletion() override
-	{
-		WaitCompletionImpl(0.0f);
-	}
-
 	virtual void CancelImpl() override
 	{
 	}
diff --git a/Engine/Source/Runtime/PakFile/Private/IPlatformFilePak.cpp b/Engine/Source/Runtime/PakFile/Private/IPlatformFilePak.cpp
--- a/Engine/Source/Runtime/PakFile/Private/IPlatformFilePak.cpp
+++ b/Engine/Source/Runtime/PakFile/Private/IPlatformFilePak.cpp
@@ -4150,11 +4150,6 @@ public:
 		}
 	}
 
-	virtual void EnsureCompletion() override
-	{
-		WaitCompletionImpl(0.0f);
-	}
-
 	virtual void CancelImpl() override
 	{
 		check(!WaitEvent); // you canceled from a different thread that you waited from
