Attempted fix for issue=430784

diff --git a/Engine/Source/Runtime/CoreUObject/Tests/PropertyBagTest.cpp b/Engine/Source/Runtime/CoreUObject/Tests/PropertyBagTest.cpp
--- a/Engine/Source/Runtime/CoreUObject/Tests/PropertyBagTest.cpp
+++ b/Engine/Source/Runtime/CoreUObject/Tests/PropertyBagTest.cpp
@@ -242,17 +242,19 @@ namespace PropertyBagTestUtils
 	}
 
 	template<typename T>
-	constexpr EName GetENameForType()
+	EName GetENameForType()
 	{
-		if constexpr (std::is_same_v<T, int>)
+		if (std::is_same_v<T, int>)
 		{
 			return NAME_IntProperty;
 		}
 
-		if constexpr (std::is_same_v<T, float>)
+		if (std::is_same_v<T, float>)
 		{
 			return NAME_FloatProperty;
 		}
+
+		return NAME_None;
 	}
 
 	void LoadDataByTag(FPropertyBag& Dst, const FPropertyPathName& Path, const FPropertyTag& Tag, TArray<uint8>& Buffer)
