[TRACE] Added more information about trace scopes
- AsyncLoading gated by LoadTime channel activation
- GameFeaturePlugin statemachine name during transitions.
- Marked GC under region, because GC destroy objects over multiple frames
#tests PIE

diff --git a/Engine/Plugins/Runtime/GameFeatures/Source/GameFeatures/GameFeatures.Build.cs b/Engine/Plugins/Runtime/GameFeatures/Source/GameFeatures/GameFeatures.Build.cs
--- a/Engine/Plugins/Runtime/GameFeatures/Source/GameFeatures/GameFeatures.Build.cs
+++ b/Engine/Plugins/Runtime/GameFeatures/Source/GameFeatures/GameFeatures.Build.cs
@@ -29,6 +29,7 @@ namespace UnrealBuildTool.Rules
 					"PakFile",
 					"Projects",
 					"RenderCore", // required for FDeferredCleanupInterface
+					"TraceLog",
 				}
 			);
 
diff --git a/Engine/Plugins/Runtime/GameFeatures/Source/GameFeatures/Private/GameFeaturePluginStateMachine.cpp b/Engine/Plugins/Runtime/GameFeatures/Source/GameFeatures/Private/GameFeaturePluginStateMachine.cpp
--- a/Engine/Plugins/Runtime/GameFeatures/Source/GameFeatures/Private/GameFeaturePluginStateMachine.cpp
+++ b/Engine/Plugins/Runtime/GameFeatures/Source/GameFeatures/Private/GameFeaturePluginStateMachine.cpp
@@ -35,6 +35,7 @@
 #include "Misc/PathViews.h"
 #include "Containers/Queue.h"
 #include "ShaderCodeLibrary.h"
+#include "Trace/Trace.h"
 
 #include UE_INLINE_GENERATED_CPP_BY_NAME(GameFeaturePluginStateMachine)
 
@@ -3335,6 +3336,10 @@ bool UGameFeaturePluginStateMachine::IsValidErrorState(EGameFeaturePluginState I
 	return AllStates[InDestinationState]->GetStateType() == EGameFeaturePluginStateType::Error;
 }
 
+UE_TRACE_EVENT_BEGIN(Cpu, GFP_UpdateStateMachine, NoSync)
+	UE_TRACE_EVENT_FIELD(UE::Trace::WideString, PluginName)
+UE_TRACE_EVENT_END()
+
 void UGameFeaturePluginStateMachine::UpdateStateMachine()
 {
 	EGameFeaturePluginState CurrentState = GetCurrentState();
@@ -3344,7 +3349,8 @@ void UGameFeaturePluginStateMachine::UpdateStateMachine()
 		return;
 	}
 
-	TRACE_CPUPROFILER_EVENT_SCOPE(GFP_UpdateStateMachine);
+	UE_TRACE_LOG_SCOPED_T(Cpu, GFP_UpdateStateMachine, CpuChannel)
+		<< GFP_UpdateStateMachine.PluginName(*GetPluginName());
 
 	TOptional<TGuardValue<bool>> ScopeGuard(InPlace, bInUpdateStateMachine, true);
 	
diff --git a/Engine/Source/Runtime/CoreUObject/Private/Serialization/AsyncLoading2.cpp b/Engine/Source/Runtime/CoreUObject/Private/Serialization/AsyncLoading2.cpp
--- a/Engine/Source/Runtime/CoreUObject/Private/Serialization/AsyncLoading2.cpp
+++ b/Engine/Source/Runtime/CoreUObject/Private/Serialization/AsyncLoading2.cpp
@@ -81,6 +81,7 @@
 #include "IO/IoPriorityQueue.h"
 #include "UObject/CoreRedirects.h"
 #include "Serialization/ZenPackageHeader.h"
+#include "Trace/Trace.h"
 
 #include <atomic>
 
@@ -4312,6 +4313,11 @@ void FAsyncLoadingThread2::IncludePackageInSyncLoadContextRecursive(FAsyncLoadin
 	}
 }
 
+UE_TRACE_EVENT_BEGIN(CUSTOM_LOADTIMER_LOG, CreateAsyncPackage, NoSync)
+	UE_TRACE_EVENT_FIELD(uint64, PackageId)
+	UE_TRACE_EVENT_FIELD(UE::Trace::WideString, PackageName)
+UE_TRACE_EVENT_END()
+
 bool FAsyncLoadingThread2::CreateAsyncPackagesFromQueue(FAsyncLoadingThreadState2& ThreadState)
 {
 	TRACE_CPUPROFILER_EVENT_SCOPE(CreateAsyncPackagesFromQueue);
@@ -4331,6 +4337,9 @@ bool FAsyncLoadingThread2::CreateAsyncPackagesFromQueue(FAsyncLoadingThreadState
 			PendingPackage->Desc.UPackageName, PackageEntry);
 		if (PendingPackageStatus == EPackageStoreEntryStatus::Ok)
 		{
+			SCOPED_CUSTOM_LOADTIMER(CreateAsyncPackage)
+				ADD_CUSTOM_LOADTIMER_META(CreateAsyncPackage, PackageId, PendingPackage->Desc.UPackageId.ValueForDebugging())
+				ADD_CUSTOM_LOADTIMER_META(CreateAsyncPackage, PackageName, *WriteToString<256>(PendingPackage->Desc.UPackageName));
 			InitializeAsyncPackageFromPackageStore(ThreadState, &IoBatch, PendingPackage, PackageEntry);
 			PendingPackage->StartLoading(ThreadState, IoBatch);
 			It.RemoveCurrent();
@@ -4469,6 +4478,10 @@ bool FAsyncLoadingThread2::CreateAsyncPackagesFromQueue(FAsyncLoadingThreadState
 					}
 					else
 					{
+						SCOPED_CUSTOM_LOADTIMER(CreateAsyncPackage)
+							ADD_CUSTOM_LOADTIMER_META(CreateAsyncPackage, PackageId, PackageDesc.UPackageId.ValueForDebugging())
+							ADD_CUSTOM_LOADTIMER_META(CreateAsyncPackage, PackageName, NameBuffer);
+
 						check(PackageStatus == EPackageStoreEntryStatus::Ok);
 #if ALT2_ENABLE_LINKERLOAD_SUPPORT
 						if (!bIsZenPackage)
@@ -9119,6 +9132,10 @@ void FAsyncPackage2::AddProgressCallback(TUniquePtr<FLoadPackageAsyncProgressDel
 	ProgressCallbacks.Emplace(MoveTemp(Callback));
 }
 
+UE_TRACE_EVENT_BEGIN(CUSTOM_LOADTIMER_LOG, LoadAsyncPackageInternal, NoSync)
+	UE_TRACE_EVENT_FIELD(UE::Trace::WideString, PackageName)
+UE_TRACE_EVENT_END()
+
 int32 FAsyncLoadingThread2::LoadPackageInternal(const FPackagePath& InPackagePath, FName InCustomName, TUniquePtr<FLoadPackageAsyncDelegate>&& InCompletionDelegate, TUniquePtr<FLoadPackageAsyncProgressDelegate>&& InProgressDelegate, EPackageFlags InPackageFlags, int32 InPIEInstanceID, int32 InPackagePriority, const FLinkerInstancingContext* InInstancingContext, uint32 InLoadFlags)
 {
 	TRACE_CPUPROFILER_EVENT_SCOPE(LoadPackage);
@@ -9128,6 +9145,8 @@ int32 FAsyncLoadingThread2::LoadPackageInternal(const FPackagePath& InPackagePat
 	{
 		InCustomName = NAME_None;
 	}
+	SCOPED_CUSTOM_LOADTIMER(LoadAsyncPackageInternal)
+		ADD_CUSTOM_LOADTIMER_META(LoadAsyncPackageInternal, PackageName, *WriteToString<256>(PackageNameToLoad));
 
 PRAGMA_DISABLE_DEPRECATION_WARNINGS
 	if (FCoreDelegates::OnAsyncLoadPackage.IsBound())
diff --git a/Engine/Source/Runtime/CoreUObject/Private/UObject/GarbageCollection.cpp b/Engine/Source/Runtime/CoreUObject/Private/UObject/GarbageCollection.cpp
--- a/Engine/Source/Runtime/CoreUObject/Private/UObject/GarbageCollection.cpp
+++ b/Engine/Source/Runtime/CoreUObject/Private/UObject/GarbageCollection.cpp
@@ -5384,6 +5384,7 @@ void PreCollectGarbageImpl(EObjectFlags KeepFlags)
 
 	DECLARE_SCOPE_CYCLE_COUNTER( TEXT( "CollectGarbageInternal" ), STAT_CollectGarbageInternal, STATGROUP_GC );
 	STAT_ADD_CUSTOMMESSAGE_NAME( STAT_NamedMarker, TEXT( "GarbageCollection - Begin" ) );
+	TRACE_BEGIN_REGION(TEXT("GarbageCollection"));
 
 	// We can't collect garbage while there's a load in progress. E.g. one potential issue is Import.XObject
 	check(!IsLoading());
@@ -5602,6 +5603,7 @@ void PostCollectGarbageImpl(EObjectFlags KeepFlags)
 	GTimingInfo.LastGCTime = PostCollectEndTime;
 	GGCStats.TotalTime += PostCollectEndTime - PostCollectStartTime;
 	STAT_ADD_CUSTOMMESSAGE_NAME(STAT_NamedMarker, TEXT("GarbageCollection - End"));
+	TRACE_END_REGION(TEXT("GarbageCollection"));
 
 	if (bPerformFullPurge)
 	{
