Support ValueSet event in the Containers when changes comes from Details Panel. I though the change type was Unspecified for this type of operation. The ValueSet was not handled at all.

diff --git a/Engine/Source/Runtime/CoreUObject/Private/UObject/OverriddenPropertySet.cpp b/Engine/Source/Runtime/CoreUObject/Private/UObject/OverriddenPropertySet.cpp
--- a/Engine/Source/Runtime/CoreUObject/Private/UObject/OverriddenPropertySet.cpp
+++ b/Engine/Source/Runtime/CoreUObject/Private/UObject/OverriddenPropertySet.cpp
@@ -680,6 +680,9 @@ void FOverriddenPropertySet::NotifyPropertyChange(FOverriddenPropertyNode* Paren
 
 					switch(PropertyEvent.ChangeType)
 					{
+					case EPropertyChangeType::ValueSet:
+						checkf(ArrayIndex != INDEX_NONE, TEXT("ValueSet change type should have associated indexes"));
+						// Intentional fall thru
 					case EPropertyChangeType::Unspecified:
 						{
 							if (ArrayIndex != INDEX_NONE)
@@ -926,6 +929,9 @@ void FOverriddenPropertySet::NotifyPropertyChange(FOverriddenPropertyNode* Paren
 
 			switch (PropertyEvent.ChangeType)
 			{
+			case EPropertyChangeType::ValueSet:
+				checkf(LogicalMapIndex != INDEX_NONE, TEXT("ValueSet change type should have associated indexes"));
+				// Intentional fall thru
 			case EPropertyChangeType::Unspecified:
 				{
 					if(LogicalMapIndex != INDEX_NONE)
