Updated FPropertyTypeName to explicitly support guids and paths

#rnx

diff --git a/Engine/Source/Runtime/CoreUObject/Private/UObject/Class.cpp b/Engine/Source/Runtime/CoreUObject/Private/UObject/Class.cpp
--- a/Engine/Source/Runtime/CoreUObject/Private/UObject/Class.cpp
+++ b/Engine/Source/Runtime/CoreUObject/Private/UObject/Class.cpp
@@ -1577,29 +1577,29 @@ void UStruct::SerializeVersionedTaggedProperties(FStructuredArchive::FSlot Slot,
 				if (SerializeContext && SerializeContext->bTrackSerializedPropertyPath)
 				{
 					FPropertyTypeNameBuilder TypeBuilder;
-					TypeBuilder.AddTypeName(Tag.Type);
+					TypeBuilder.AddName(Tag.Type);
 
 					if (!Tag.StructName.IsNone())
 					{
-						TypeBuilder.BeginTypeParameters();
-						TypeBuilder.AddTypeName(Tag.StructName);
-						TypeBuilder.EndTypeParameters();
+						TypeBuilder.BeginParameters();
+						TypeBuilder.AddName(Tag.StructName);
+						TypeBuilder.EndParameters();
 					}
 					else if (!Tag.EnumName.IsNone())
 					{
-						TypeBuilder.BeginTypeParameters();
-						TypeBuilder.AddTypeName(Tag.EnumName);
-						TypeBuilder.EndTypeParameters();
+						TypeBuilder.BeginParameters();
+						TypeBuilder.AddName(Tag.EnumName);
+						TypeBuilder.EndParameters();
 					}
 					else if (!Tag.InnerType.IsNone())
 					{
-						TypeBuilder.BeginTypeParameters();
-						TypeBuilder.AddTypeName(Tag.InnerType);
+						TypeBuilder.BeginParameters();
+						TypeBuilder.AddName(Tag.InnerType);
 						if (!Tag.ValueType.IsNone())
 						{
-							TypeBuilder.AddTypeName(Tag.ValueType);
+							TypeBuilder.AddName(Tag.ValueType);
 						}
-						TypeBuilder.EndTypeParameters();
+						TypeBuilder.EndParameters();
 					}
 
 					const FName Name = Property ? Property->GetFName() : Tag.Name;
diff --git a/Engine/Source/Runtime/CoreUObject/Private/UObject/EnumProperty.cpp b/Engine/Source/Runtime/CoreUObject/Private/UObject/EnumProperty.cpp
--- a/Engine/Source/Runtime/CoreUObject/Private/UObject/EnumProperty.cpp
+++ b/Engine/Source/Runtime/CoreUObject/Private/UObject/EnumProperty.cpp
@@ -582,15 +582,15 @@ bool FEnumProperty::LoadTypeName(UE::FPropertyTypeName Type, const FPropertyTag*
 		return false;
 	}
 
-	const FName EnumName = Type.GetTypeParameterName(0);
+	const FName EnumName = Type.GetParameterName(0);
 	UEnum* LocalEnum = FindFirstObject<UEnum>(*WriteToString<256>(EnumName), EFindFirstObjectOptions::NativeFirst);
 	if (!LocalEnum)
 	{
 		return false;
 	}
 
-	const UE::FPropertyTypeName UnderlyingType = Type.GetTypeParameter(1);
-	FField* Field = FField::TryConstruct(UnderlyingType.GetTypeName(), this, GetFName(), RF_NoFlags);
+	const UE::FPropertyTypeName UnderlyingType = Type.GetParameter(1);
+	FField* Field = FField::TryConstruct(UnderlyingType.GetName(), this, GetFName(), RF_NoFlags);
 	if (FNumericProperty* Property = CastField<FNumericProperty>(Field); Property && Property->LoadTypeName(UnderlyingType, Tag))
 	{
 		Enum = LocalEnum;
@@ -610,14 +610,11 @@ void FEnumProperty::SaveTypeName(UE::FPropertyTypeNameBuilder& Type) const
 
 	if (const UEnum* LocalEnum = Enum)
 	{
-		TStringBuilder<256> EnumName;
-		LocalEnum->GetPathName(nullptr, EnumName);
-
 		check(UnderlyingProp);
-		Type.BeginTypeParameters();
-		Type.AddTypeName(FName(EnumName));
+		Type.BeginParameters();
+		Type.AddPath(LocalEnum);
 		UnderlyingProp->SaveTypeName(Type);
-		Type.EndTypeParameters();
+		Type.EndParameters();
 	}
 }
 
@@ -634,7 +631,7 @@ bool FEnumProperty::CanSerializeFromTypeName(UE::FPropertyTypeName Type) const
 		return false;
 	}
 
-	const FName EnumName = Type.GetTypeParameterName(0);
+	const FName EnumName = Type.GetParameterName(0);
 	if (EnumName == LocalEnum->GetFName())
 	{
 		return true;
diff --git a/Engine/Source/Runtime/CoreUObject/Private/UObject/InstanceDataObjectUtils.cpp b/Engine/Source/Runtime/CoreUObject/Private/UObject/InstanceDataObjectUtils.cpp
--- a/Engine/Source/Runtime/CoreUObject/Private/UObject/InstanceDataObjectUtils.cpp
+++ b/Engine/Source/Runtime/CoreUObject/Private/UObject/InstanceDataObjectUtils.cpp
@@ -46,80 +46,80 @@ namespace UE
 	static void BuildSegmentTypeFromProperty(const FProperty* Property, FPropertyTypeNameBuilder& OutType)
 	{
 #if false // TODO: @jordan.hoffmann when complete type info is finished use this branch instead
-		OutType.AddTypeName(Property->GetID());
+		OutType.AddName(Property->GetID());
 		if (const FObjectProperty* AsObjectProperty = CastField<FObjectProperty>(Property))
 		{
-			OutType.BeginTypeParameters();
-			OutType.AddTypeName(AsObjectProperty->GetID());
-			OutType.EndTypeParameters();
+			OutType.BeginParameters();
+			OutType.AddName(AsObjectProperty->GetID());
+			OutType.EndParameters();
 		}
 		else if (const FEnumProperty* AsEnumProperty = CastField<FEnumProperty>(Property))
 		{
-			OutType.BeginTypeParameters();
-			OutType.AddTypeName(AsEnumProperty->GetEnum()->GetFName());
-			OutType.EndTypeParameters();
+			OutType.BeginParameters();
+			OutType.AddPath(AsEnumProperty->GetEnum());
+			OutType.EndParameters();
 		}
 		else if (const FArrayProperty* AsArrayProperty = CastField<FArrayProperty>(Property))
 		{
-			OutType.BeginTypeParameters();
+			OutType.BeginParameters();
 			BuildSegmentTypeFromProperty(AsArrayProperty->Inner, OutType);
-			OutType.EndTypeParameters();
+			OutType.EndParameters();
 		}
 		else if (const FSetProperty* AsSetProperty = CastField<FSetProperty>(Property))
 		{
-			OutType.BeginTypeParameters();
+			OutType.BeginParameters();
 			BuildSegmentTypeFromProperty(AsSetProperty->ElementProp, OutType);
-			OutType.EndTypeParameters();
+			OutType.EndParameters();
 		}
 		else if (const FMapProperty* AsMapProperty = CastField<FMapProperty>(Property))
 		{
-			OutType.BeginTypeParameters();
+			OutType.BeginParameters();
 			BuildSegmentTypeFromProperty(AsMapProperty->KeyProp, OutType);
 			BuildSegmentTypeFromProperty(AsMapProperty->ValueProp, OutType);
-			OutType.EndTypeParameters();
+			OutType.EndParameters();
 		}
 #else
-		OutType.AddTypeName(Property->GetID());
+		OutType.AddName(Property->GetID());
 		if (const FStructProperty* AsStructProperty = CastField<FStructProperty>(Property))
 		{
-			OutType.BeginTypeParameters();
-			OutType.AddTypeName(AsStructProperty->Struct->GetFName());
-			OutType.EndTypeParameters();
+			OutType.BeginParameters();
+			OutType.AddName(AsStructProperty->Struct->GetFName());
+			OutType.EndParameters();
 		}
 		else if (const FEnumProperty* AsEnumProperty = CastField<FEnumProperty>(Property))
 		{
-			OutType.BeginTypeParameters();
-			OutType.AddTypeName(AsEnumProperty->GetEnum()->GetFName());
-			OutType.EndTypeParameters();
+			OutType.BeginParameters();
+			OutType.AddName(AsEnumProperty->GetEnum()->GetFName());
+			OutType.EndParameters();
 		}
 		else if (const FArrayProperty* AsArrayProperty = CastField<FArrayProperty>(Property))
 		{
-			OutType.BeginTypeParameters();
+			OutType.BeginParameters();
 			if (AsArrayProperty->Inner->IsA<FEnumProperty>())
 			{
 				// enum paths currently don't recurse when they're in containers
-				OutType.AddTypeName(AsArrayProperty->Inner->GetID());
+				OutType.AddName(AsArrayProperty->Inner->GetID());
 			}
 			else
 			{
 				BuildSegmentTypeFromProperty(AsArrayProperty->Inner, OutType);
 			}
-			OutType.EndTypeParameters();
+			OutType.EndParameters();
 		}
 		else if (const FSetProperty* AsSetProperty = CastField<FSetProperty>(Property))
 		{
 			// sets currently don't recurse their element property types
-			OutType.BeginTypeParameters();
-			OutType.AddTypeName(AsSetProperty->ElementProp->GetID());
-			OutType.EndTypeParameters();
+			OutType.BeginParameters();
+			OutType.AddName(AsSetProperty->ElementProp->GetID());
+			OutType.EndParameters();
 		}
 		else if (const FMapProperty* AsMapProperty = CastField<FMapProperty>(Property))
 		{
 			// maps currently don't recurse their key/value property types
-			OutType.BeginTypeParameters();
-			OutType.AddTypeName(AsMapProperty->KeyProp->GetID());
-			OutType.AddTypeName(AsMapProperty->ValueProp->GetID());
-			OutType.EndTypeParameters();
+			OutType.BeginParameters();
+			OutType.AddName(AsMapProperty->KeyProp->GetID());
+			OutType.AddName(AsMapProperty->ValueProp->GetID());
+			OutType.EndParameters();
 		}
 #endif
 	}
@@ -684,4 +684,4 @@ namespace UE
 		}
 		return false;
 	}
-} // UE
+} // UE
\ No newline at end of file
diff --git a/Engine/Source/Runtime/CoreUObject/Private/UObject/Property.cpp b/Engine/Source/Runtime/CoreUObject/Private/UObject/Property.cpp
--- a/Engine/Source/Runtime/CoreUObject/Private/UObject/Property.cpp
+++ b/Engine/Source/Runtime/CoreUObject/Private/UObject/Property.cpp
@@ -2125,19 +2125,19 @@ void FProperty::AssignToTag(FPropertyTag& Tag)
 
 bool FProperty::LoadTypeName(UE::FPropertyTypeName Type, const FPropertyTag* Tag)
 {
-	return ensureMsgf(GetID() == Type.GetTypeName(),
+	return ensureMsgf(GetID() == Type.GetName(),
 		TEXT("Failed to load property '%s' of type '%s' from tag of type '%s'"),
-		*WriteToString<64>(GetFName()), *WriteToString<64>(GetID()), *WriteToString<64>(Type.GetTypeName()));
+		*WriteToString<64>(GetFName()), *WriteToString<64>(GetID()), *WriteToString<64>(Type.GetName()));
 }
 
 void FProperty::SaveTypeName(UE::FPropertyTypeNameBuilder& Type) const
 {
-	Type.AddTypeName(GetID());
+	Type.AddName(GetID());
 }
 
 bool FProperty::CanSerializeFromTypeName(UE::FPropertyTypeName Type) const
 {
-	return Type.GetTypeName() == GetID();
+	return Type.GetName() == GetID();
 }
 
 FProperty* UStruct::FindPropertyByName(FName InName) const
diff --git a/Engine/Source/Runtime/CoreUObject/Private/UObject/PropertyArray.cpp b/Engine/Source/Runtime/CoreUObject/Private/UObject/PropertyArray.cpp
--- a/Engine/Source/Runtime/CoreUObject/Private/UObject/PropertyArray.cpp
+++ b/Engine/Source/Runtime/CoreUObject/Private/UObject/PropertyArray.cpp
@@ -613,25 +613,27 @@ void FArrayProperty::SerializeItem(FStructuredArchive::FSlot Slot, void* Value,
 	}
 
 	FUObjectSerializeContext* Context = FUObjectThreadContext::Get().GetSerializeContext();
-	if (Context && Context->bTrackSerializedPropertyPath && MaybeInnerTag)
+	if (Context && Context->bTrackSerializedPropertyPath)
 	{
-		if (const int32 SegmentCount = Context->SerializedPropertyPath.GetSegmentCount(); SegmentCount > 0)
+		// Update the path with types from the inner tag if the outer tag is incomplete.
+		if (const int32 SegmentIndex = Context->SerializedPropertyPath.GetSegmentCount() - 1; SegmentIndex >= 0)
 		{
-			const UE::FPropertyTypeName ExistingArrayType = Context->SerializedPropertyPath.GetSegment(SegmentCount - 1).Type;
-			if (ExistingArrayType.GetTypeName() == NAME_ArrayProperty && ExistingArrayType.GetTypeParameterCount() == 1)
+			UE::FPropertyPathNameSegment Segment = Context->SerializedPropertyPath.GetSegment(SegmentIndex);
+			if (Segment.Type.GetName() == NAME_ArrayProperty)
 			{
-				const UE::FPropertyTypeName ExistingStructType = ExistingArrayType.GetTypeParameter(0);
-				if (ExistingStructType.GetTypeName() == NAME_StructProperty && ExistingStructType.GetTypeParameterCount() == 0)
+				const UE::FPropertyTypeName InnerTypeName = Segment.Type.GetParameter();
+				if (InnerTypeName.GetName() == NAME_StructProperty && InnerTypeName.GetParameterCount() == 0)
 				{
 					UE::FPropertyTypeNameBuilder NewTypeBuilder;
-					NewTypeBuilder.AddTypeName(NAME_ArrayProperty);
-					NewTypeBuilder.BeginTypeParameters();
-					NewTypeBuilder.AddTypeName(NAME_StructProperty);
-					NewTypeBuilder.BeginTypeParameters();
-					NewTypeBuilder.AddTypeName(MaybeInnerTag->StructName);
-					NewTypeBuilder.EndTypeParameters();
-					NewTypeBuilder.EndTypeParameters();
-					Context->SerializedPropertyPath.SetType(NewTypeBuilder.Build());
+					NewTypeBuilder.AddName(NAME_ArrayProperty);
+					NewTypeBuilder.BeginParameters();
+					NewTypeBuilder.AddName(NAME_StructProperty);
+					NewTypeBuilder.BeginParameters();
+					NewTypeBuilder.AddName(MaybeInnerTag->StructName);
+					NewTypeBuilder.EndParameters();
+					NewTypeBuilder.EndParameters();
+					Segment.Type = NewTypeBuilder.Build();
+					Context->SerializedPropertyPath.SetSegment(SegmentIndex, Segment);
 				}
 			}
 		}
@@ -1311,8 +1313,8 @@ bool FArrayProperty::LoadTypeName(UE::FPropertyTypeName Type, const FPropertyTag
 		return false;
 	}
 
-	const UE::FPropertyTypeName InnerType = Type.GetTypeParameter();
-	FField* Field = FField::TryConstruct(InnerType.GetTypeName(), this, GetFName(), RF_NoFlags);
+	const UE::FPropertyTypeName InnerType = Type.GetParameter();
+	FField* Field = FField::TryConstruct(InnerType.GetName(), this, GetFName(), RF_NoFlags);
 	if (FProperty* Property = CastField<FProperty>(Field); Property && Property->LoadTypeName(InnerType, Tag))
 	{
 		Inner = Property;
@@ -1328,9 +1330,9 @@ void FArrayProperty::SaveTypeName(UE::FPropertyTypeNameBuilder& Type) const
 
 	const FProperty* LocalInner = Inner;
 	check(LocalInner);
-	Type.BeginTypeParameters();
+	Type.BeginParameters();
 	LocalInner->SaveTypeName(Type);
-	Type.EndTypeParameters();
+	Type.EndParameters();
 }
 
 bool FArrayProperty::CanSerializeFromTypeName(UE::FPropertyTypeName Type) const
@@ -1342,5 +1344,5 @@ bool FArrayProperty::CanSerializeFromTypeName(UE::FPropertyTypeName Type) const
 
 	const FProperty* LocalInner = Inner;
 	check(LocalInner);
-	return LocalInner->CanSerializeFromTypeName(Type.GetTypeParameter());
+	return LocalInner->CanSerializeFromTypeName(Type.GetParameter());
 }
diff --git a/Engine/Source/Runtime/CoreUObject/Private/UObject/PropertyByte.cpp b/Engine/Source/Runtime/CoreUObject/Private/UObject/PropertyByte.cpp
--- a/Engine/Source/Runtime/CoreUObject/Private/UObject/PropertyByte.cpp
+++ b/Engine/Source/Runtime/CoreUObject/Private/UObject/PropertyByte.cpp
@@ -573,7 +573,7 @@ bool FByteProperty::LoadTypeName(UE::FPropertyTypeName Type, const FPropertyTag*
 		return false;
 	}
 
-	const FName EnumName = Type.GetTypeParameterName();
+	const FName EnumName = Type.GetParameterName();
 	if (EnumName.IsNone())
 	{
 		return true;
@@ -594,12 +594,9 @@ void FByteProperty::SaveTypeName(UE::FPropertyTypeNameBuilder& Type) const
 
 	if (const UEnum* LocalEnum = Enum)
 	{
-		TStringBuilder<256> EnumName;
-		LocalEnum->GetPathName(nullptr, EnumName);
-
-		Type.BeginTypeParameters();
-		Type.AddTypeName(FName(EnumName));
-		Type.EndTypeParameters();
+		Type.BeginParameters();
+		Type.AddPath(LocalEnum);
+		Type.EndParameters();
 	}
 }
 
@@ -610,7 +607,7 @@ bool FByteProperty::CanSerializeFromTypeName(UE::FPropertyTypeName Type) const
 		return false;
 	}
 
-	const FName EnumName = Type.GetTypeParameterName();
+	const FName EnumName = Type.GetParameterName();
 	if (const UEnum* LocalEnum = Enum)
 	{
 		if (EnumName == LocalEnum->GetFName())
diff --git a/Engine/Source/Runtime/CoreUObject/Private/UObject/PropertyMap.cpp b/Engine/Source/Runtime/CoreUObject/Private/UObject/PropertyMap.cpp
--- a/Engine/Source/Runtime/CoreUObject/Private/UObject/PropertyMap.cpp
+++ b/Engine/Source/Runtime/CoreUObject/Private/UObject/PropertyMap.cpp
@@ -1734,10 +1734,10 @@ bool FMapProperty::LoadTypeName(UE::FPropertyTypeName Type, const FPropertyTag*
 		return false;
 	}
 
-	const UE::FPropertyTypeName KeyType = Type.GetTypeParameter(0);
-	const UE::FPropertyTypeName ValueType = Type.GetTypeParameter(1);
-	FField* KeyField = FField::TryConstruct(KeyType.GetTypeName(), this, GetFName(), RF_NoFlags);
-	FField* ValueField = FField::TryConstruct(ValueType.GetTypeName(), this, GetFName(), RF_NoFlags);
+	const UE::FPropertyTypeName KeyType = Type.GetParameter(0);
+	const UE::FPropertyTypeName ValueType = Type.GetParameter(1);
+	FField* KeyField = FField::TryConstruct(KeyType.GetName(), this, GetFName(), RF_NoFlags);
+	FField* ValueField = FField::TryConstruct(ValueType.GetName(), this, GetFName(), RF_NoFlags);
 	FProperty* KeyProperty = CastField<FProperty>(KeyField);
 	FProperty* ValueProperty = CastField<FProperty>(ValueField);
 	if (KeyProperty && ValueProperty && KeyProperty->LoadTypeName(KeyType, Tag) && ValueProperty->LoadTypeName(ValueType, Tag))
@@ -1759,10 +1759,10 @@ void FMapProperty::SaveTypeName(UE::FPropertyTypeNameBuilder& Type) const
 	const FProperty* LocalValueProp = ValueProp;
 	check(LocalKeyProp);
 	check(LocalValueProp);
-	Type.BeginTypeParameters();
+	Type.BeginParameters();
 	LocalKeyProp->SaveTypeName(Type);
 	LocalValueProp->SaveTypeName(Type);
-	Type.EndTypeParameters();
+	Type.EndParameters();
 }
 
 bool FMapProperty::CanSerializeFromTypeName(UE::FPropertyTypeName Type) const
@@ -1776,5 +1776,5 @@ bool FMapProperty::CanSerializeFromTypeName(UE::FPropertyTypeName Type) const
 	const FProperty* LocalValueProp = ValueProp;
 	check(LocalKeyProp);
 	check(LocalValueProp);
-	return LocalKeyProp->CanSerializeFromTypeName(Type.GetTypeParameter(0)) && LocalValueProp->CanSerializeFromTypeName(Type.GetTypeParameter(1));
+	return LocalKeyProp->CanSerializeFromTypeName(Type.GetParameter(0)) && LocalValueProp->CanSerializeFromTypeName(Type.GetParameter(1));
 }
diff --git a/Engine/Source/Runtime/CoreUObject/Private/UObject/PropertyOptional.cpp b/Engine/Source/Runtime/CoreUObject/Private/UObject/PropertyOptional.cpp
--- a/Engine/Source/Runtime/CoreUObject/Private/UObject/PropertyOptional.cpp
+++ b/Engine/Source/Runtime/CoreUObject/Private/UObject/PropertyOptional.cpp
@@ -593,8 +593,8 @@ bool FOptionalProperty::LoadTypeName(UE::FPropertyTypeName Type, const FProperty
 		return false;
 	}
 
-	const UE::FPropertyTypeName ValueType = Type.GetTypeParameter();
-	FField* Field = FField::TryConstruct(ValueType.GetTypeName(), this, GetFName(), RF_NoFlags);
+	const UE::FPropertyTypeName ValueType = Type.GetParameter();
+	FField* Field = FField::TryConstruct(ValueType.GetName(), this, GetFName(), RF_NoFlags);
 	if (FProperty* Property = CastField<FProperty>(Field); Property && Property->LoadTypeName(ValueType, Tag))
 	{
 		ValueProperty = Property;
@@ -610,9 +610,9 @@ void FOptionalProperty::SaveTypeName(UE::FPropertyTypeNameBuilder& Type) const
 
 	const FProperty* LocalValueProperty = ValueProperty;
 	check(LocalValueProperty);
-	Type.BeginTypeParameters();
+	Type.BeginParameters();
 	LocalValueProperty->SaveTypeName(Type);
-	Type.EndTypeParameters();
+	Type.EndParameters();
 }
 
 bool FOptionalProperty::CanSerializeFromTypeName(UE::FPropertyTypeName Type) const
@@ -624,5 +624,5 @@ bool FOptionalProperty::CanSerializeFromTypeName(UE::FPropertyTypeName Type) con
 
 	const FProperty* LocalValueProperty = ValueProperty;
 	check(LocalValueProperty);
-	return LocalValueProperty->CanSerializeFromTypeName(Type.GetTypeParameter());
+	return LocalValueProperty->CanSerializeFromTypeName(Type.GetParameter());
 }
diff --git a/Engine/Source/Runtime/CoreUObject/Private/UObject/PropertySet.cpp b/Engine/Source/Runtime/CoreUObject/Private/UObject/PropertySet.cpp
--- a/Engine/Source/Runtime/CoreUObject/Private/UObject/PropertySet.cpp
+++ b/Engine/Source/Runtime/CoreUObject/Private/UObject/PropertySet.cpp
@@ -1123,8 +1123,8 @@ bool FSetProperty::LoadTypeName(UE::FPropertyTypeName Type, const FPropertyTag*
 		return false;
 	}
 
-	const UE::FPropertyTypeName ElementType = Type.GetTypeParameter();
-	FField* Field = FField::TryConstruct(ElementType.GetTypeName(), this, GetFName(), RF_NoFlags);
+	const UE::FPropertyTypeName ElementType = Type.GetParameter();
+	FField* Field = FField::TryConstruct(ElementType.GetName(), this, GetFName(), RF_NoFlags);
 	if (FProperty* Property = CastField<FProperty>(Field); Property && Property->LoadTypeName(ElementType, Tag))
 	{
 		ElementProp = Property;
@@ -1140,9 +1140,9 @@ void FSetProperty::SaveTypeName(UE::FPropertyTypeNameBuilder& Type) const
 
 	const FProperty* LocalElementProp = ElementProp;
 	check(LocalElementProp);
-	Type.BeginTypeParameters();
+	Type.BeginParameters();
 	LocalElementProp->SaveTypeName(Type);
-	Type.EndTypeParameters();
+	Type.EndParameters();
 }
 
 bool FSetProperty::CanSerializeFromTypeName(UE::FPropertyTypeName Type) const
@@ -1154,5 +1154,5 @@ bool FSetProperty::CanSerializeFromTypeName(UE::FPropertyTypeName Type) const
 
 	const FProperty* LocalElementProp = ElementProp;
 	check(LocalElementProp);
-	return LocalElementProp->CanSerializeFromTypeName(Type.GetTypeParameter());
+	return LocalElementProp->CanSerializeFromTypeName(Type.GetParameter());
 }
diff --git a/Engine/Source/Runtime/CoreUObject/Private/UObject/PropertyStruct.cpp b/Engine/Source/Runtime/CoreUObject/Private/UObject/PropertyStruct.cpp
--- a/Engine/Source/Runtime/CoreUObject/Private/UObject/PropertyStruct.cpp
+++ b/Engine/Source/Runtime/CoreUObject/Private/UObject/PropertyStruct.cpp
@@ -543,7 +543,7 @@ bool FStructProperty::LoadTypeName(UE::FPropertyTypeName Type, const FPropertyTa
 		return false;
 	}
 
-	if (const FName Name = Type.GetTypeParameterName(0); !Name.IsNone())
+	if (const FName Name = Type.GetParameterName(0); !Name.IsNone())
 	{
 		TStringBuilder<256> NameString(InPlace, Name);
 		if (Struct = FindFirstObject<UScriptStruct>(*NameString, EFindFirstObjectOptions::NativeFirst); Struct)
@@ -565,13 +565,13 @@ void FStructProperty::SaveTypeName(UE::FPropertyTypeNameBuilder& Type) const
 	const UScriptStruct* LocalStruct = Struct;
 	check(LocalStruct);
 
-	Type.BeginTypeParameters();
-	Type.AddTypeName(OriginalType ? FName(**OriginalType) : LocalStruct->GetFName());
+	Type.BeginParameters();
+	Type.AddName(OriginalType ? FName(**OriginalType) : LocalStruct->GetFName());
 	if (const FGuid StructGuid = LocalStruct->GetCustomGuid(); StructGuid.IsValid())
 	{
-		Type.AddTypeName(FName(WriteToString<48>(StructGuid)));
+		Type.AddGuid(StructGuid);
 	}
-	Type.EndTypeParameters();
+	Type.EndParameters();
 }
 
 bool FStructProperty::CanSerializeFromTypeName(UE::FPropertyTypeName Type) const
@@ -584,13 +584,13 @@ bool FStructProperty::CanSerializeFromTypeName(UE::FPropertyTypeName Type) const
 	const UScriptStruct* LocalStruct = Struct;
 	check(LocalStruct);
 
-	const FName StructName = Type.GetTypeParameterName(0);
+	const FName StructName = Type.GetParameterName(0);
 	if (StructName == LocalStruct->GetFName())
 	{
 		return true;
 	}
 
-	const FName StructGuidName = Type.GetTypeParameterName(1);
+	const FName StructGuidName = Type.GetParameterName(1);
 	if (FGuid StructGuid; !StructGuidName.IsNone() && FGuid::Parse(StructGuidName.ToString(), StructGuid) && StructGuid.IsValid())
 	{
 		return StructGuid == LocalStruct->GetCustomGuid();
diff --git a/Engine/Source/Runtime/CoreUObject/Private/UObject/PropertyTypeName.cpp b/Engine/Source/Runtime/CoreUObject/Private/UObject/PropertyTypeName.cpp
--- a/Engine/Source/Runtime/CoreUObject/Private/UObject/PropertyTypeName.cpp
+++ b/Engine/Source/Runtime/CoreUObject/Private/UObject/PropertyTypeName.cpp
@@ -8,7 +8,9 @@
 #include "HAL/CriticalSection.h"
 #include "HAL/LowLevelMemTracker.h"
 #include "Misc/ScopeRWLock.h"
+#include "Misc/StringBuilder.h"
 #include "String/Find.h"
+#include "UObject/Class.h"
 
 namespace UE
 {
@@ -232,19 +234,19 @@ FPropertyTypeNameNode* FPropertyTypeNameTable::AllocateBlock(int32 BlockIndex)
 
 ///////////////////////////////////////////////////////////////////////////////////////////////////////////////////////
 
-FName FPropertyTypeName::GetTypeName() const
+FName FPropertyTypeName::GetName() const
 {
 	const FPropertyTypeNameNode* First = GPropertyTypeNameTable.ResolveByIndex(Index);
 	return First->Name;
 }
 
-int32 FPropertyTypeName::GetTypeParameterCount() const
+int32 FPropertyTypeName::GetParameterCount() const
 {
 	const FPropertyTypeNameNode* First = GPropertyTypeNameTable.ResolveByIndex(Index);
 	return First->InnerCount;
 }
 
-FPropertyTypeName FPropertyTypeName::GetTypeParameter(int32 ParamIndex) const
+FPropertyTypeName FPropertyTypeName::GetParameter(int32 ParamIndex) const
 {
 	const FPropertyTypeNameNode* First = GPropertyTypeNameTable.ResolveByIndex(Index);
 	if (UNLIKELY(ParamIndex < 0 || ParamIndex >= First->InnerCount))
@@ -392,18 +394,18 @@ FStringBuilderBase& operator<<(FStringBuilderBase& Builder, const FPropertyTypeN
 
 ///////////////////////////////////////////////////////////////////////////////////////////////////////////////////////
 
-void FPropertyTypeNameBuilder::BeginTypeParameters()
+void FPropertyTypeNameBuilder::BeginParameters()
 {
 	checkf(!Nodes.IsEmpty(), TEXT("A type name must be added prior to setting type parameters."));
 	ActiveIndex = Nodes.Num() - 1;
 }
 
-void FPropertyTypeNameBuilder::EndTypeParameters()
+void FPropertyTypeNameBuilder::EndParameters()
 {
 	ActiveIndex = OuterNodeIndex[ActiveIndex];
 }
 
-void FPropertyTypeNameBuilder::AddTypeName(FName Name)
+void FPropertyTypeNameBuilder::AddName(FName Name)
 {
 	checkf(ActiveIndex >= 0 || Nodes.IsEmpty(), TEXT("Only one type name may be added as the root node."));
 
@@ -418,17 +420,48 @@ void FPropertyTypeNameBuilder::AddTypeName(FName Name)
 	}
 }
 
-void FPropertyTypeNameBuilder::AddTypeName(FPropertyTypeName Name)
+void FPropertyTypeNameBuilder::AddGuid(const FGuid& Guid)
 {
-	AddTypeName(Name.GetTypeName());
-	if (const int32 Count = Name.GetTypeParameterCount())
+	AddName(FName(WriteToString<48>(Guid)));
+}
+
+void FPropertyTypeNameBuilder::AddPath(const UField* Field)
+{
+	if (!Field)
+	{
+		AddName(NAME_None);
+		return;
+	}
+
+	AddName(Field->GetFName());
+
+	TArray<UObject*, TInlineAllocator<8>> OuterChain;
+	for (UObject* Outer = Field->GetOuter(); Outer; Outer = Outer->GetOuter())
+	{
+		OuterChain.Add(Outer);
+	}
+	if (!OuterChain.IsEmpty())
+	{
+		BeginParameters();
+		for (UObject* Outer : ReverseIterate(OuterChain))
+		{
+			AddName(Outer->GetFName());
+		}
+		EndParameters();
+	}
+}
+
+void FPropertyTypeNameBuilder::AddType(FPropertyTypeName Name)
+{
+	AddName(Name.GetName());
+	if (const int32 Count = Name.GetParameterCount())
 	{
-		BeginTypeParameters();
+		BeginParameters();
 		for (int32 Index = 0; Index < Count; ++Index)
 		{
-			AddTypeName(Name.GetTypeParameter(Index));
+			AddType(Name.GetParameter(Index));
 		}
-		EndTypeParameters();
+		EndParameters();
 	}
 }
 
@@ -459,7 +492,7 @@ bool FPropertyTypeNameBuilder::TryParse(FStringView Name)
 					// Names must follow '<' or ',' or be at the root.
 					break;
 				}
-				AddTypeName(FName(Type));
+				AddName(FName(Type));
 				bAllowBegin = true;
 				bAllowName = false;
 			}
@@ -490,7 +523,7 @@ bool FPropertyTypeNameBuilder::TryParse(FStringView Name)
 				break;
 			}
 			++Depth;
-			BeginTypeParameters();
+			BeginParameters();
 			bAllowName = true;
 		}
 		else if (C == TEXT('>'))
@@ -500,7 +533,7 @@ bool FPropertyTypeNameBuilder::TryParse(FStringView Name)
 				// '>' must have a matching '<'.
 				break;
 			}
-			EndTypeParameters();
+			EndParameters();
 			--Depth;
 		}
 		else // if (C == TEXT(','))
diff --git a/Engine/Source/Runtime/CoreUObject/Public/UObject/PropertyPathName.h b/Engine/Source/Runtime/CoreUObject/Public/UObject/PropertyPathName.h
--- a/Engine/Source/Runtime/CoreUObject/Public/UObject/PropertyPathName.h
+++ b/Engine/Source/Runtime/CoreUObject/Public/UObject/PropertyPathName.h
@@ -101,15 +101,6 @@ public:
 		Segments.Pop(EAllowShrinking::No);
 	}
 
-	/** Sets the type of the last segment of the path. Ignored if the path is empty. */
-	inline void SetType(FPropertyTypeName Type)
-	{
-		if (!Segments.IsEmpty())
-		{
-			Segments.Last().Type = Type;
-		}
-	}
-
 	/** Sets the index of the last segment of the path. Ignored if the path is empty. */
 	inline void SetIndex(int32 Index)
 	{
diff --git a/Engine/Source/Runtime/CoreUObject/Public/UObject/PropertyTypeName.h b/Engine/Source/Runtime/CoreUObject/Public/UObject/PropertyTypeName.h
--- a/Engine/Source/Runtime/CoreUObject/Public/UObject/PropertyTypeName.h
+++ b/Engine/Source/Runtime/CoreUObject/Public/UObject/PropertyTypeName.h
@@ -11,6 +11,8 @@
 
 class FArchive;
 class FStructuredArchiveSlot;
+class UField;
+struct FGuid;
 
 namespace UE { class FPropertyTypeNameBuilder; }
 
@@ -26,11 +28,18 @@ struct FPropertyTypeNameNode
 /**
  * Represents the type name of a property, including any containers and underlying types.
  *
+ * The path name for a type is represented as the object name with the outer chain as type parameters.
+ * - /Script/CoreUObject.Outer0:Outer1.Outer2.StructName -> StructName</Script/CoreUObject,Outer0,Outer1,Outer2>
+ *
  * Examples:
  * - int32 -> IntProperty
  * - TArray<int32> -> ArrayProperty<IntProperty>
- * - TArray<FStructType> -> ArrayProperty<StructProperty<StructType>>
- * - TMap<FKeyStruct, EByteEnum> -> MapProperty<StructProperty<KeyStruct>,EnumProperty<ByteEnum,ByteProperty>>
+ * - TArray<FStructType> -> ArrayProperty<StructProperty<StructType</Script/Module>,e21f566f-7153-433a-959d-bfb3abed17e2>>
+ * - TMap<FKeyStruct, EByteEnum> -> MapProperty<StructProperty<KeyStruct</Script/Module>>,EnumProperty<ByteEnum</Script/Module>,ByteProperty>>
+ *
+ * Incomplete property types created from sources with incomplete type information. Consumers must support this.
+ * - ArrayProperty<StructProperty>
+ * - MapProperty<EnumProperty,StructProperty>
  */
 class FPropertyTypeName
 {
@@ -44,17 +53,17 @@ public:
 	 * Returns the type at the root of this property type name.
 	 *
 	 * Example: MapProperty<StructProperty<KeyStruct>,EnumProperty<ByteEnum,ByteProperty>>
-	 * - GetTypeName() -> MapProperty
+	 * - GetName() -> MapProperty
 	 */
-	UE_API FName GetTypeName() const;
+	UE_API FName GetName() const;
 
 	/**
 	 * Returns the number of type parameters under the root of this property type name.
 	 *
 	 * Example: MapProperty<StructProperty<KeyStruct>,EnumProperty<ByteEnum,ByteProperty>>
-	 * - GetTypeParameterCount() -> 2
+	 * - GetParameterCount() -> 2
 	 */
-	UE_API int32 GetTypeParameterCount() const;
+	UE_API int32 GetParameterCount() const;
 
 	/**
 	 * Returns the indexed parameter under the root of this property type name.
@@ -62,10 +71,10 @@ public:
 	 * An out-of-bounds index will return an empty type name.
 	 *
 	 * Example: MapProperty<StructProperty<KeyStruct>,EnumProperty<ByteEnum,ByteProperty>>
-	 * - GetTypeParameter(0) -> StructProperty<KeyStruct>
-	 * - GetTypeParameter(1) -> EnumProperty<ByteEnum,ByteProperty>
+	 * - GetParameter(0) -> StructProperty<KeyStruct>
+	 * - GetParameter(1) -> EnumProperty<ByteEnum,ByteProperty>
 	 */
-	UE_API FPropertyTypeName GetTypeParameter(int32 ParamIndex = 0) const;
+	UE_API FPropertyTypeName GetParameter(int32 ParamIndex = 0) const;
 
 	/**
 	 * Returns the indexed parameter type name under the root of this property type name.
@@ -73,12 +82,12 @@ public:
 	 * An out-of-bounds index will return a name of None.
 	 *
 	 * Example: MapProperty<StructProperty<KeyStruct>,EnumProperty<ByteEnum,ByteProperty>>
-	 * - GetTypeParameterName(0) -> StructProperty
-	 * - GetTypeParameterName(1) -> EnumProperty
+	 * - GetParameterName(0) -> StructProperty
+	 * - GetParameterName(1) -> EnumProperty
 	 */
-	inline FName GetTypeParameterName(int32 ParamIndex = 0) const
+	inline FName GetParameterName(int32 ParamIndex = 0) const
 	{
-		return GetTypeParameter(ParamIndex).GetTypeName();
+		return GetParameter(ParamIndex).GetName();
 	}
 
 	/**
@@ -108,37 +117,45 @@ private:
 /**
  * Builder for FPropertyTypeName.
  *
- * Example: MapProperty<StructProperty<KeyStruct>,EnumProperty<ByteEnum,ByteProperty>>
+ * Example: MapProperty<StructProperty<KeyStruct</Script/CoreUObject>,StructGuid>,EnumProperty<ByteEnum</Script/CoreUObject>,ByteProperty>>
  *
  * FPropertyTypeNameBuilder Builder;
- * Builder.AddTypeName(NAME_MapProperty));
- * Builder.BeginTypeParameters();
- *   Builder.AddTypeName(NAME_StructProperty);
- *   Builder.BeginTypeParameters();
- *     Builder.AddTypeName(TEXT("KeyStruct"));
- *   Builder.EndTypeParameters();
- *   Builder.AddTypeName(NAME_EnumProperty);
- *   Builder.BeginTypeParameters();
- *     Builder.AddTypeName(TEXT("ByteEnum"));
- *     Builder.AddTypeName(NAME_ByteProperty);
- *   Builder.EndTypeParameters();
- * Builder.EndTypeParameters();
+ * Builder.AddName(NAME_MapProperty));
+ * Builder.BeginParameters();
+ *   Builder.AddName(NAME_StructProperty);
+ *   Builder.BeginParameters();
+ *     Builder.AddPath(KeyStruct); // const UStruct*
+ *     if (const FGuid StructGuid = KeyStruct->GetCustomGuid(); StructGuid.IsValid())
+ *         Builder.AddGuid(StructGuid);
+ *   Builder.EndParameters();
+ *   Builder.AddName(NAME_EnumProperty);
+ *   Builder.BeginParameters();
+ *     Builder.AddPath(ByteEnum); // const UEnum*
+ *     Builder.AddName(NAME_ByteProperty);
+ *   Builder.EndParameters();
+ * Builder.EndParameters();
  * FPropertyTypeName Name = Builder.Build();
  */
 class FPropertyTypeNameBuilder
 {
 public:
-	/** Add a complete type name with its type parameters. */
-	UE_API void AddTypeName(FPropertyTypeName Name);
+	/** Add a name without any type parameters. */
+	UE_API void AddName(FName Name);
+
+	/** Add a guid in the format DigitsWithHyphensLower. */
+	UE_API void AddGuid(const FGuid& Guid);
+
+	/** Add a path in the format ObjectName</Path/Name,Outer0,Outer1,Outer2> */
+	UE_API void AddPath(const UField* Field);
 
-	/** Add a type name without any type parameters. */
-	UE_API void AddTypeName(FName Name);
+	/** Add a type name with its parameters. */
+	UE_API void AddType(FPropertyTypeName Name);
 
 	/** Mark the beginning of the type parameters for the last added type. */
-	UE_API void BeginTypeParameters();
+	UE_API void BeginParameters();
 
 	/** Mark the end of the type parameters for the matching begin. */
-	UE_API void EndTypeParameters();
+	UE_API void EndParameters();
 
 	/** Build from the types that have been added. */
 	UE_API FPropertyTypeName Build() const;
@@ -157,8 +174,8 @@ public:
 	UE_API bool TryParse(FStringView Name);
 
 private:
-	TArray<FPropertyTypeNameNode, TInlineAllocator<8>> Nodes;
-	TArray<int32, TInlineAllocator<8>> OuterNodeIndex;
+	TArray<FPropertyTypeNameNode, TInlineAllocator<10>> Nodes;
+	TArray<int32, TInlineAllocator<10>> OuterNodeIndex;
 	int32 ActiveIndex = INDEX_NONE;
 };
 
diff --git a/Engine/Source/Runtime/CoreUObject/Tests/UObject/PropertyPathNameTest.cpp b/Engine/Source/Runtime/CoreUObject/Tests/UObject/PropertyPathNameTest.cpp
--- a/Engine/Source/Runtime/CoreUObject/Tests/UObject/PropertyPathNameTest.cpp
+++ b/Engine/Source/Runtime/CoreUObject/Tests/UObject/PropertyPathNameTest.cpp
@@ -15,7 +15,7 @@ TEST_CASE_NAMED(FPropertyPathNameTest, "CoreUObject::PropertyPathName", "[Core][
 	const FName CountName(TEXTVIEW("Count"));
 
 	FPropertyTypeNameBuilder TypeBuilder;
-	TypeBuilder.AddTypeName(NAME_IntProperty);
+	TypeBuilder.AddName(NAME_IntProperty);
 	const FPropertyTypeName IntType = TypeBuilder.Build();
 
 	SECTION("Empty")
@@ -65,19 +65,19 @@ TEST_CASE_NAMED(FPropertyPathNameTest, "CoreUObject::PropertyPathName", "[Core][
 	SECTION("MultipleSegments")
 	{
 		FPropertyTypeNameBuilder ArrayTypeBuilder;
-		ArrayTypeBuilder.AddTypeName(NAME_ArrayProperty);
-		ArrayTypeBuilder.BeginTypeParameters();
-		ArrayTypeBuilder.AddTypeName(NAME_StructProperty);
-		ArrayTypeBuilder.BeginTypeParameters();
-		ArrayTypeBuilder.AddTypeName(TEXT("TestType"));
-		ArrayTypeBuilder.EndTypeParameters();
-		ArrayTypeBuilder.EndTypeParameters();
+		ArrayTypeBuilder.AddName(NAME_ArrayProperty);
+		ArrayTypeBuilder.BeginParameters();
+		ArrayTypeBuilder.AddName(NAME_StructProperty);
+		ArrayTypeBuilder.BeginParameters();
+		ArrayTypeBuilder.AddName(TEXT("TestType"));
+		ArrayTypeBuilder.EndParameters();
+		ArrayTypeBuilder.EndParameters();
 
 		FPropertyTypeNameBuilder VectorTypeBuilder;
-		VectorTypeBuilder.AddTypeName(NAME_StructProperty);
-		VectorTypeBuilder.BeginTypeParameters();
-		VectorTypeBuilder.AddTypeName(TEXT("IntVector"));
-		VectorTypeBuilder.EndTypeParameters();
+		VectorTypeBuilder.AddName(NAME_StructProperty);
+		VectorTypeBuilder.BeginParameters();
+		VectorTypeBuilder.AddName(TEXT("IntVector"));
+		VectorTypeBuilder.EndParameters();
 
 		FPropertyPathName PathName;
 		PathName.Push({TEXT("TestArray"), ArrayTypeBuilder.Build(), 7});
@@ -94,10 +94,6 @@ TEST_CASE_NAMED(FPropertyPathNameTest, "CoreUObject::PropertyPathName", "[Core][
 		PathName.Push({CountName});
 		PathName.Push({CountName});
 
-		PathName.SetType(IntType);
-		CHECK(PathName.GetSegment(0).Type.IsEmpty());
-		CHECK(PathName.GetSegment(1).Type == IntType);
-
 		PathName.SetIndex(7);
 		CHECK(PathName.GetSegment(0).Index == INDEX_NONE);
 		CHECK(PathName.GetSegment(1).Index == 7);
@@ -145,7 +141,7 @@ TEST_CASE_NAMED(FPropertyPathNameTest, "CoreUObject::PropertyPathName", "[Core][
 		const FName DepthName(TEXTVIEW("Depth"));
 		
 		TypeBuilder.Reset();
-		TypeBuilder.AddTypeName(NAME_BoolProperty);
+		TypeBuilder.AddName(NAME_BoolProperty);
 		const FPropertyTypeName BoolType = TypeBuilder.Build();
 
 		FPropertyPathName CountBool;
diff --git a/Engine/Source/Runtime/CoreUObject/Tests/UObject/PropertyTypeNameTest.cpp b/Engine/Source/Runtime/CoreUObject/Tests/UObject/PropertyTypeNameTest.cpp
--- a/Engine/Source/Runtime/CoreUObject/Tests/UObject/PropertyTypeNameTest.cpp
+++ b/Engine/Source/Runtime/CoreUObject/Tests/UObject/PropertyTypeNameTest.cpp
@@ -23,23 +23,23 @@ namespace PropertyTypeNameTest
 static FPropertyTypeName CreateInt()
 {
 	FPropertyTypeNameBuilder Builder;
-	Builder.AddTypeName(NAME_IntProperty);
+	Builder.AddName(NAME_IntProperty);
 	return Builder.Build();
 }
 
 static FPropertyTypeName CreateVectorArray()
 {
 	FPropertyTypeNameBuilder Builder;
-	Builder.AddTypeName(NAME_ArrayProperty);
+	Builder.AddName(NAME_ArrayProperty);
 	{
-		Builder.BeginTypeParameters();
-		Builder.AddTypeName(NAME_StructProperty);
+		Builder.BeginParameters();
+		Builder.AddName(NAME_StructProperty);
 		{
-			Builder.BeginTypeParameters();
-			Builder.AddTypeName(NAME_Vector);
-			Builder.EndTypeParameters();
+			Builder.BeginParameters();
+			Builder.AddName(NAME_Vector);
+			Builder.EndParameters();
 		}
-		Builder.EndTypeParameters();
+		Builder.EndParameters();
 	}
 	return Builder.Build();
 }
@@ -47,24 +47,24 @@ static FPropertyTypeName CreateVectorArray()
 static FPropertyTypeName CreateEnumMap()
 {
 	FPropertyTypeNameBuilder Builder;
-	Builder.AddTypeName(NAME_MapProperty);
+	Builder.AddName(NAME_MapProperty);
 	{
-		Builder.BeginTypeParameters();
-		Builder.AddTypeName(NAME_EnumProperty);
+		Builder.BeginParameters();
+		Builder.AddName(NAME_EnumProperty);
 		{
-			Builder.BeginTypeParameters();
-			Builder.AddTypeName(TEXT("Key"));
-			Builder.AddTypeName(NAME_ByteProperty);
-			Builder.EndTypeParameters();
+			Builder.BeginParameters();
+			Builder.AddName(TEXT("Key"));
+			Builder.AddName(NAME_ByteProperty);
+			Builder.EndParameters();
 		}
-		Builder.AddTypeName(NAME_EnumProperty);
+		Builder.AddName(NAME_EnumProperty);
 		{
-			Builder.BeginTypeParameters();
-			Builder.AddTypeName(TEXT("Value"));
-			Builder.AddTypeName(NAME_ByteProperty);
-			Builder.EndTypeParameters();
+			Builder.BeginParameters();
+			Builder.AddName(TEXT("Value"));
+			Builder.AddName(NAME_ByteProperty);
+			Builder.EndParameters();
 		}
-		Builder.EndTypeParameters();
+		Builder.EndParameters();
 	}
 	return Builder.Build();
 }
@@ -77,96 +77,96 @@ TEST_CASE_NAMED(FPropertyTypeNameSmokeTest, "CoreUObject::PropertyTypeName::Smok
 	{
 		FPropertyTypeName TypeName;
 		CHECK(TypeName.IsEmpty());
-		CHECK(TypeName.GetTypeName().IsNone());
-		CHECK(TypeName.GetTypeParameterCount() == 0);
-		CHECK(TypeName.GetTypeParameter(0).IsEmpty());
+		CHECK(TypeName.GetName().IsNone());
+		CHECK(TypeName.GetParameterCount() == 0);
+		CHECK(TypeName.GetParameter(0).IsEmpty());
 	}
 
 	SECTION("Numeric")
 	{
 		FPropertyTypeName TypeName = PropertyTypeNameTest::CreateInt();
 		CHECK(!TypeName.IsEmpty());
-		CHECK(TypeName.GetTypeName() == NAME_IntProperty);
-		CHECK(TypeName.GetTypeParameterCount() == 0);
-		CHECK(TypeName.GetTypeParameter(0).IsEmpty());
+		CHECK(TypeName.GetName() == NAME_IntProperty);
+		CHECK(TypeName.GetParameterCount() == 0);
+		CHECK(TypeName.GetParameter(0).IsEmpty());
 	}
 
 	SECTION("NumericArray")
 	{
 		FPropertyTypeNameBuilder Builder;
-		Builder.AddTypeName(NAME_ArrayProperty);
+		Builder.AddName(NAME_ArrayProperty);
 		{
-			Builder.BeginTypeParameters();
-			Builder.AddTypeName(NAME_IntProperty);
-			Builder.EndTypeParameters();
+			Builder.BeginParameters();
+			Builder.AddName(NAME_IntProperty);
+			Builder.EndParameters();
 		}
 
 		FPropertyTypeName TypeName = Builder.Build();
 		CHECK(!TypeName.IsEmpty());
-		CHECK(TypeName.GetTypeName() == NAME_ArrayProperty);
-		CHECK(TypeName.GetTypeParameterCount() == 1);
-		CHECK(TypeName.GetTypeParameter(0).GetTypeName() == NAME_IntProperty);
-		CHECK(TypeName.GetTypeParameter(0).GetTypeParameterCount() == 0);
-		CHECK(TypeName.GetTypeParameter(0).GetTypeParameter(0).IsEmpty());
-		CHECK(TypeName.GetTypeParameter(1).IsEmpty());
+		CHECK(TypeName.GetName() == NAME_ArrayProperty);
+		CHECK(TypeName.GetParameterCount() == 1);
+		CHECK(TypeName.GetParameter(0).GetName() == NAME_IntProperty);
+		CHECK(TypeName.GetParameter(0).GetParameterCount() == 0);
+		CHECK(TypeName.GetParameter(0).GetParameter(0).IsEmpty());
+		CHECK(TypeName.GetParameter(1).IsEmpty());
 	}
 
 	SECTION("StructArray")
 	{
 		FPropertyTypeName TypeName = PropertyTypeNameTest::CreateVectorArray();
 		CHECK(!TypeName.IsEmpty());
-		CHECK(TypeName.GetTypeName() == NAME_ArrayProperty);
-		CHECK(TypeName.GetTypeParameterCount() == 1);
-		CHECK(TypeName.GetTypeParameter(0).GetTypeName() == NAME_StructProperty);
-		CHECK(TypeName.GetTypeParameter(0).GetTypeParameterCount() == 1);
-		CHECK(TypeName.GetTypeParameter(0).GetTypeParameter(0).GetTypeName() == NAME_Vector);
-		CHECK(TypeName.GetTypeParameter(0).GetTypeParameter(0).GetTypeParameterCount() == 0);
-		CHECK(TypeName.GetTypeParameter(0).GetTypeParameter(0).GetTypeParameter(0).IsEmpty());
-		CHECK(TypeName.GetTypeParameter(0).GetTypeParameter(1).IsEmpty());
-		CHECK(TypeName.GetTypeParameter(1).IsEmpty());
+		CHECK(TypeName.GetName() == NAME_ArrayProperty);
+		CHECK(TypeName.GetParameterCount() == 1);
+		CHECK(TypeName.GetParameter(0).GetName() == NAME_StructProperty);
+		CHECK(TypeName.GetParameter(0).GetParameterCount() == 1);
+		CHECK(TypeName.GetParameter(0).GetParameter(0).GetName() == NAME_Vector);
+		CHECK(TypeName.GetParameter(0).GetParameter(0).GetParameterCount() == 0);
+		CHECK(TypeName.GetParameter(0).GetParameter(0).GetParameter(0).IsEmpty());
+		CHECK(TypeName.GetParameter(0).GetParameter(1).IsEmpty());
+		CHECK(TypeName.GetParameter(1).IsEmpty());
 	}
 
 	SECTION("NumericMap")
 	{
 		FPropertyTypeNameBuilder Builder;
-		Builder.AddTypeName(NAME_MapProperty);
+		Builder.AddName(NAME_MapProperty);
 		{
-			Builder.BeginTypeParameters();
-			Builder.AddTypeName(NAME_Int32Property);
-			Builder.AddTypeName(NAME_UInt32Property);
-			Builder.EndTypeParameters();
+			Builder.BeginParameters();
+			Builder.AddName(NAME_Int32Property);
+			Builder.AddName(NAME_UInt32Property);
+			Builder.EndParameters();
 		}
 
 		FPropertyTypeName TypeName = Builder.Build();
 		CHECK(!TypeName.IsEmpty());
-		CHECK(TypeName.GetTypeName() == NAME_MapProperty);
-		CHECK(TypeName.GetTypeParameterCount() == 2);
-		CHECK(TypeName.GetTypeParameter(0).GetTypeName() == NAME_Int32Property);
-		CHECK(TypeName.GetTypeParameter(0).GetTypeParameter(0).IsEmpty());
-		CHECK(TypeName.GetTypeParameter(1).GetTypeName() == NAME_UInt32Property);
-		CHECK(TypeName.GetTypeParameter(1).GetTypeParameter(0).IsEmpty());
-		CHECK(TypeName.GetTypeParameter(2).IsEmpty());
+		CHECK(TypeName.GetName() == NAME_MapProperty);
+		CHECK(TypeName.GetParameterCount() == 2);
+		CHECK(TypeName.GetParameter(0).GetName() == NAME_Int32Property);
+		CHECK(TypeName.GetParameter(0).GetParameter(0).IsEmpty());
+		CHECK(TypeName.GetParameter(1).GetName() == NAME_UInt32Property);
+		CHECK(TypeName.GetParameter(1).GetParameter(0).IsEmpty());
+		CHECK(TypeName.GetParameter(2).IsEmpty());
 	}
 
 	SECTION("EnumMap")
 	{
 		FPropertyTypeName TypeName = PropertyTypeNameTest::CreateEnumMap();
 		CHECK(!TypeName.IsEmpty());
-		CHECK(TypeName.GetTypeName() == NAME_MapProperty);
-		CHECK(TypeName.GetTypeParameterCount() == 2);
-		CHECK(TypeName.GetTypeParameter(0).GetTypeName() == NAME_EnumProperty);
-		CHECK(TypeName.GetTypeParameter(0).GetTypeParameterName(0) == TEXT("Key"));
-		CHECK(TypeName.GetTypeParameter(0).GetTypeParameterName(1) == NAME_ByteProperty);
-		CHECK(TypeName.GetTypeParameter(0).GetTypeParameter(0).GetTypeParameterCount() == 0);
-		CHECK(TypeName.GetTypeParameter(0).GetTypeParameter(1).GetTypeParameterCount() == 0);
-		CHECK(TypeName.GetTypeParameter(0).GetTypeParameter(2).IsEmpty());
-		CHECK(TypeName.GetTypeParameter(1).GetTypeName() == NAME_EnumProperty);
-		CHECK(TypeName.GetTypeParameter(1).GetTypeParameterName(0) == TEXT("Value"));
-		CHECK(TypeName.GetTypeParameter(1).GetTypeParameterName(1) == NAME_ByteProperty);
-		CHECK(TypeName.GetTypeParameter(1).GetTypeParameter(0).GetTypeParameterCount() == 0);
-		CHECK(TypeName.GetTypeParameter(1).GetTypeParameter(1).GetTypeParameterCount() == 0);
-		CHECK(TypeName.GetTypeParameter(1).GetTypeParameter(2).IsEmpty());
-		CHECK(TypeName.GetTypeParameter(2).IsEmpty());
+		CHECK(TypeName.GetName() == NAME_MapProperty);
+		CHECK(TypeName.GetParameterCount() == 2);
+		CHECK(TypeName.GetParameter(0).GetName() == NAME_EnumProperty);
+		CHECK(TypeName.GetParameter(0).GetParameterName(0) == TEXT("Key"));
+		CHECK(TypeName.GetParameter(0).GetParameterName(1) == NAME_ByteProperty);
+		CHECK(TypeName.GetParameter(0).GetParameter(0).GetParameterCount() == 0);
+		CHECK(TypeName.GetParameter(0).GetParameter(1).GetParameterCount() == 0);
+		CHECK(TypeName.GetParameter(0).GetParameter(2).IsEmpty());
+		CHECK(TypeName.GetParameter(1).GetName() == NAME_EnumProperty);
+		CHECK(TypeName.GetParameter(1).GetParameterName(0) == TEXT("Value"));
+		CHECK(TypeName.GetParameter(1).GetParameterName(1) == NAME_ByteProperty);
+		CHECK(TypeName.GetParameter(1).GetParameter(0).GetParameterCount() == 0);
+		CHECK(TypeName.GetParameter(1).GetParameter(1).GetParameterCount() == 0);
+		CHECK(TypeName.GetParameter(1).GetParameter(2).IsEmpty());
+		CHECK(TypeName.GetParameter(2).IsEmpty());
 	}
 
 	SECTION("Equals+Less+GetTypeHash")
@@ -175,33 +175,33 @@ TEST_CASE_NAMED(FPropertyTypeNameSmokeTest, "CoreUObject::PropertyTypeName::Smok
 		const FPropertyTypeName VectorArray = PropertyTypeNameTest::CreateVectorArray();
 
 		FPropertyTypeNameBuilder Builder;
-		Builder.AddTypeName(NAME_MapProperty);
-		Builder.BeginTypeParameters();
-		Builder.AddTypeName(NAME_IntProperty);
-		Builder.AddTypeName(NAME_IntProperty);
-		Builder.EndTypeParameters();
+		Builder.AddName(NAME_MapProperty);
+		Builder.BeginParameters();
+		Builder.AddName(NAME_IntProperty);
+		Builder.AddName(NAME_IntProperty);
+		Builder.EndParameters();
 		const FPropertyTypeName MapIntToInt = Builder.Build();
 
 		Builder.Reset();
-		Builder.AddTypeName(NAME_MapProperty);
-		Builder.BeginTypeParameters();
-		Builder.AddTypeName(NAME_IntProperty);
-		Builder.AddTypeName(NAME_MapProperty);
+		Builder.AddName(NAME_MapProperty);
+		Builder.BeginParameters();
+		Builder.AddName(NAME_IntProperty);
+		Builder.AddName(NAME_MapProperty);
 		{
-			Builder.BeginTypeParameters();
-			Builder.AddTypeName(NAME_IntProperty);
-			Builder.AddTypeName(NAME_IntProperty);
-			Builder.EndTypeParameters();
+			Builder.BeginParameters();
+			Builder.AddName(NAME_IntProperty);
+			Builder.AddName(NAME_IntProperty);
+			Builder.EndParameters();
 		}
-		Builder.EndTypeParameters();
+		Builder.EndParameters();
 		const FPropertyTypeName MapIntToIntToInt = Builder.Build();
 
 		Builder.Reset();
-		Builder.AddTypeName(NAME_MapProperty);
-		Builder.BeginTypeParameters();
-		Builder.AddTypeName(NAME_IntProperty);
-		Builder.AddTypeName(MapIntToInt);
-		Builder.EndTypeParameters();
+		Builder.AddName(NAME_MapProperty);
+		Builder.BeginParameters();
+		Builder.AddName(NAME_IntProperty);
+		Builder.AddType(MapIntToInt);
+		Builder.EndParameters();
 		const FPropertyTypeName MapIntToIntToIntAlternative = Builder.Build();
 
 		CHECK(FPropertyTypeName() == FPropertyTypeName());
@@ -219,7 +219,7 @@ TEST_CASE_NAMED(FPropertyTypeNameSmokeTest, "CoreUObject::PropertyTypeName::Smok
 		CHECK_FALSE(GetTypeHash(Int) == GetTypeHash(VectorArray));
 		CHECK_FALSE(GetTypeHash(Int) == GetTypeHash(FPropertyTypeName()));
 
-		CHECK(Int == MapIntToInt.GetTypeParameter(0));
+		CHECK(Int == MapIntToInt.GetParameter(0));
 		CHECK_FALSE(Int == MapIntToInt);
 		CHECK_FALSE(MapIntToInt == Int);
 		CHECK_FALSE(MapIntToInt < MapIntToInt);
@@ -227,18 +227,18 @@ TEST_CASE_NAMED(FPropertyTypeNameSmokeTest, "CoreUObject::PropertyTypeName::Smok
 		CHECK_FALSE(MapIntToInt < Int);
 		CHECK_FALSE(MapIntToInt < VectorArray);
 		CHECK(VectorArray < MapIntToInt);
-		CHECK(GetTypeHash(Int) == GetTypeHash(MapIntToInt.GetTypeParameter(0)));
+		CHECK(GetTypeHash(Int) == GetTypeHash(MapIntToInt.GetParameter(0)));
 		CHECK_FALSE(GetTypeHash(Int) == GetTypeHash(MapIntToInt));
 		CHECK_FALSE(GetTypeHash(MapIntToInt) == GetTypeHash(Int));
 
-		CHECK(MapIntToInt == MapIntToIntToInt.GetTypeParameter(1));
+		CHECK(MapIntToInt == MapIntToIntToInt.GetParameter(1));
 		CHECK_FALSE(MapIntToInt == MapIntToIntToInt);
 		CHECK_FALSE(MapIntToIntToInt < MapIntToIntToInt);
 		CHECK(MapIntToInt < MapIntToIntToInt);
 		CHECK_FALSE(MapIntToIntToInt < MapIntToInt);
 		CHECK_FALSE(MapIntToIntToInt < VectorArray);
 		CHECK(VectorArray < MapIntToIntToInt);
-		CHECK(GetTypeHash(MapIntToInt) == GetTypeHash(MapIntToIntToInt.GetTypeParameter(1)));
+		CHECK(GetTypeHash(MapIntToInt) == GetTypeHash(MapIntToIntToInt.GetParameter(1)));
 		CHECK_FALSE(GetTypeHash(MapIntToInt) == GetTypeHash(MapIntToIntToInt));
 
 		CHECK(MapIntToIntToInt == MapIntToIntToIntAlternative);
@@ -431,27 +431,27 @@ TEST_CASE_NAMED(FPropertyTypeNameConcurrencyTest, "CoreUObject::PropertyTypeName
 
 		if (Index % 8 == 0)
 		{
-			Builder.AddTypeName(NAME_OptionalProperty);
-			Builder.BeginTypeParameters();
+			Builder.AddName(NAME_OptionalProperty);
+			Builder.BeginParameters();
 		}
 		if (Index % 4 == 0)
 		{
-			Builder.AddTypeName(NAME_ArrayProperty);
-			Builder.BeginTypeParameters();
+			Builder.AddName(NAME_ArrayProperty);
+			Builder.BeginParameters();
 		}
 
-		Builder.AddTypeName(NAME_StructProperty);
-		Builder.BeginTypeParameters();
-		Builder.AddTypeName(FName(WriteToString<32>(TEXTVIEW("TestStruct"), Index)));
-		Builder.EndTypeParameters();
+		Builder.AddName(NAME_StructProperty);
+		Builder.BeginParameters();
+		Builder.AddName(FName(WriteToString<32>(TEXTVIEW("TestStruct"), Index)));
+		Builder.EndParameters();
 
 		if (Index % 4 == 0)
 		{
-			Builder.EndTypeParameters();
+			Builder.EndParameters();
 		}
 		if (Index % 8 == 0)
 		{
-			Builder.EndTypeParameters();
+			Builder.EndParameters();
 		}
 
 		Builder.Build();
