Fixed build error in PropertyBagTest

#rnx

diff --git a/Engine/Source/Runtime/CoreUObject/Tests/PropertyBagTest.cpp b/Engine/Source/Runtime/CoreUObject/Tests/PropertyBagTest.cpp
--- a/Engine/Source/Runtime/CoreUObject/Tests/PropertyBagTest.cpp
+++ b/Engine/Source/Runtime/CoreUObject/Tests/PropertyBagTest.cpp
@@ -1,4 +1,5 @@
 // Copyright Epic Games, Inc. All Rights Reserved.
+
 #if WITH_TESTS
 
 #include "PropertyBagTest.h"
@@ -49,20 +50,20 @@ namespace PropertyBagTestUtils
 		verify(I < (int)InNames.size());
 		verify(Names[I] != NAME_Exit);
 
-		Dst.AddTypeName(Names[I]);
+		Dst.AddName(Names[I]);
 
 		if (Names[I] == NAME_MapProperty || Names[I] == NAME_EnumProperty)
 		{
-			Dst.BeginTypeParameters();
+			Dst.BeginParameters();
 			check(0 < (I = BuildTypeName(InNames, I + 1, Dst))); // Key
 			check(0 < (I = BuildTypeName(InNames, I, Dst))); // Value
-			Dst.EndTypeParameters();
+			Dst.EndParameters();
 			return I;
 		}
 
 		if (Names[I] == NAME_StructProperty)
 		{
-			Dst.BeginTypeParameters();
+			Dst.BeginParameters();
 			while (I < (int)InNames.size())
 			{
 				if (Names[I] == NAME_Exit)
@@ -71,7 +72,7 @@ namespace PropertyBagTestUtils
 				}
 				check(0 < (I = BuildTypeName(InNames, I, Dst))); // Key
 			}
-			Dst.EndTypeParameters();
+			Dst.EndParameters();
 		}
 
 		return I + 1;
@@ -88,27 +89,27 @@ namespace PropertyBagTestUtils
 	{
 		if (FArrayProperty* ArrayProperty = CastField<FArrayProperty>(Property)) 
 		{
-			Builder.AddTypeName(NAME_ArrayProperty);
-			Builder.BeginTypeParameters();
+			Builder.AddName(NAME_ArrayProperty);
+			Builder.BeginParameters();
 			FProperty* ValueType = CastField<FProperty>(ArrayProperty->GetInnerFieldByName(Property->GetFName()));
 			BuildTypeName(Builder, ValueType);
-			Builder.EndTypeParameters();
+			Builder.EndParameters();
 
 			return;
 		}
 
 		if (FMapProperty* MapProperty = CastField<FMapProperty>(Property))
 		{
-			Builder.AddTypeName(NAME_MapProperty);
-			Builder.BeginTypeParameters();
+			Builder.AddName(NAME_MapProperty);
+			Builder.BeginParameters();
 			BuildTypeName(Builder, const_cast<FProperty*>(MapProperty->GetKeyProperty()));
 			BuildTypeName(Builder, const_cast<FProperty*>(MapProperty->GetValueProperty()));
-			Builder.EndTypeParameters();
+			Builder.EndParameters();
 
 			return;
 		}
 
-		Builder.AddTypeName(Property->GetClass()->GetFName());
+		Builder.AddName(Property->GetClass()->GetFName());
 	}
 
 	FPropertyTypeName BuildTypeName(FProperty* Property)
@@ -121,11 +122,11 @@ namespace PropertyBagTestUtils
 	FPropertyTypeName BuildPairTypeName(FProperty* Key, FProperty* Value)
 	{
 		FPropertyTypeNameBuilder Builder;
-		Builder.AddTypeName(NAME_StructProperty);
-		Builder.BeginTypeParameters();
+		Builder.AddName(NAME_StructProperty);
+		Builder.BeginParameters();
 		BuildTypeName(Builder, Key);
 		BuildTypeName(Builder, Value);
-		Builder.EndTypeParameters();
+		Builder.EndParameters();
 		return Builder.Build();
 	}
 
