Improved error reporting when loading an uncooked file causes us to try to load a cooked file which ultimately fails due to a missing export in the cooked file.

diff --git a/Engine/Source/Runtime/CoreUObject/Private/UObject/LinkerLoad.cpp b/Engine/Source/Runtime/CoreUObject/Private/UObject/LinkerLoad.cpp
--- a/Engine/Source/Runtime/CoreUObject/Private/UObject/LinkerLoad.cpp
+++ b/Engine/Source/Runtime/CoreUObject/Private/UObject/LinkerLoad.cpp
@@ -3351,7 +3351,7 @@ UPackage* LoadPackageInternal(UPackage* InOuter, const FPackagePath& PackagePath
  * @param InOutImportMap	The import table
  * @param InPackageImport	The package import index
  */
-void StaticFindAllImportObjects(TArray<FObjectImport>& InOutImportMap, FPackageIndex InPackageImport)
+void StaticFindAllImportObjects(TArray<FObjectImport>& InOutImportMap, FPackageIndex InPackageImport, const FPackagePath& PathOfPackageBeingLoaded)
 {
 	using FPackageIndexArray = TArray<FPackageIndex, TInlineAllocator<64>>;
 
@@ -3408,11 +3408,13 @@ void StaticFindAllImportObjects(TArray<FObjectImport>& InOutImportMap, FPackageI
 					}
 					else
 					{
-						UE_ASSET_LOG(LogLinker, Warning, Package, TEXT("Failed to resolve import '%s' ('%d') in outer '%s' ('%d') in cooked package"),
+						UE_ASSET_LOG(LogLinker, Warning, PathOfPackageBeingLoaded, TEXT("Failed to resolve import '%s' ('%d') in outer '%s' ('%d') within cooked package '%s'"),
 							*ObjectImport.ObjectName.ToString(),
 							Inner.ToImport(),
 							*OuterObject->GetName(),
-							Outer.ToImport());
+							Outer.ToImport(),
+							*FAssetMsg::FormatPathForAssetLog(Package)
+							);
 					}
 				}
 			}
@@ -3490,7 +3492,7 @@ bool FLinkerLoad::VerifyImportInner(const int32 ImportIndex, FString& WarningSuf
 				// Find the imports by name instead
 				// Note: The cooked package might not be marked as fully loaded at this stage, but we will have created and serialized all its exports
 				Import.XObject = Package;
-				StaticFindAllImportObjects(ImportMap, FPackageIndex::FromImport(ImportIndex));
+				StaticFindAllImportObjects(ImportMap, FPackageIndex::FromImport(ImportIndex), PackagePath);
 			}
 			return Package;
 		}
@@ -3543,7 +3545,7 @@ bool FLinkerLoad::VerifyImportInner(const int32 ImportIndex, FString& WarningSuf
 			// linkers. Static find all imported objects from this package.
 			check(Package->IsFullyLoaded());
 			Import.XObject = Package;
-			StaticFindAllImportObjects(ImportMap, FPackageIndex::FromImport(ImportIndex));
+			StaticFindAllImportObjects(ImportMap, FPackageIndex::FromImport(ImportIndex), PackagePath);
 		}
 #endif
 
