Reenable PropertyBagTests after blocking Mac builds.
Issue looks to be differences in memory allocation reuse stratagems between platforms.

diff --git a/Engine/Source/Runtime/CoreUObject/Tests/PropertyBagTest.cpp b/Engine/Source/Runtime/CoreUObject/Tests/PropertyBagTest.cpp
--- a/Engine/Source/Runtime/CoreUObject/Tests/PropertyBagTest.cpp
+++ b/Engine/Source/Runtime/CoreUObject/Tests/PropertyBagTest.cpp
@@ -1,5 +1,6 @@
 // Copyright Epic Games, Inc. All Rights Reserved.
-#if WITH_LOW_LEVEL_TESTS
+
+#if WITH_TESTS
 
 #include "PropertyBagTest.h"
 
@@ -292,7 +293,7 @@ namespace PropertyBagTestUtils
 	}
 } // PropertyBagTestUtils
 
-DISABLED_TEST_CASE(FPropertyBagTest_Add, "CoreUObject::PropertyBag::Add", "[Core][UObject][PropertyBag]")
+TEST_CASE_NAMED(FPropertyBagTest_Add, "CoreUObject::PropertyBag::Add", "[Core][UObject][PropertyBag]")
 {
 	FProperty* Property = CastFieldChecked<FProperty>(FField::Construct(NAME_IntProperty, {}, FName(TEXT("Tmp")), RF_NoFlags));
 	ON_SCOPE_EXIT
@@ -342,7 +343,6 @@ DISABLED_TEST_CASE(FPropertyBagTest_Add, "CoreUObject::PropertyBag::Add", "[Core
 
 		void* SecondStoredPointer = Dst.CreateConstIterator().GetValue();
 		const FProperty* Property2 = Dst.CreateConstIterator().GetProperty();
-		CHECK(SecondStoredPointer != StoredPointer); // Different Property, Different Storage for the same path
 		CHECK(*static_cast<int*>(SecondStoredPointer) == SecondValue); // It did store the value
 		CHECK(SecondProperty.Get() == Property2); // It did store the new property
 	}
@@ -402,7 +402,7 @@ DISABLED_TEST_CASE(FPropertyBagTest_Add, "CoreUObject::PropertyBag::Add", "[Core
 	}
 }
 
-DISABLED_TEST_CASE(FPropertyBagTest_Remove, "CoreUObject::PropertyBag::Remove", "[Core][UObject][PropertyBag]")
+TEST_CASE_NAMED(FPropertyBagTest_Remove, "CoreUObject::PropertyBag::Remove", "[Core][UObject][PropertyBag]")
 {
 	FProperty* Property = CastFieldChecked<FProperty>(FField::Construct(NAME_IntProperty, {}, FName(TEXT("Tmp")), RF_NoFlags));
 	ON_SCOPE_EXIT
@@ -483,7 +483,7 @@ DISABLED_TEST_CASE(FPropertyBagTest_Remove, "CoreUObject::PropertyBag::Remove",
 	}
 }
 
-DISABLED_TEST_CASE(FPropertyBagTest_Empty, "CoreUObject::PropertyBag::Empty", "[Core][UObject][PropertyBag]")
+TEST_CASE_NAMED(FPropertyBagTest_Empty, "CoreUObject::PropertyBag::Empty", "[Core][UObject][PropertyBag]")
 {
 	FProperty* Property = CastFieldChecked<FProperty>(FField::Construct(NAME_IntProperty, {}, FName(TEXT("Tmp")), RF_NoFlags));
 	ON_SCOPE_EXIT
@@ -519,7 +519,7 @@ DISABLED_TEST_CASE(FPropertyBagTest_Empty, "CoreUObject::PropertyBag::Empty", "[
 	}
 }
 
-DISABLED_TEST_CASE(FPropertyBagTest_LoadPropertyByTag, "CoreUObject::PropertyBag::LoadPropertyByTag", "[Core][UObject][PropertyBag]")
+TEST_CASE_NAMED(FPropertyBagTest_LoadPropertyByTag, "CoreUObject::PropertyBag::LoadPropertyByTag", "[Core][UObject][PropertyBag]")
 {
 	FPropertyTypeName IntType = PropertyBagTestUtils::BuildTypeName({ NAME_IntProperty });
 
@@ -576,7 +576,6 @@ DISABLED_TEST_CASE(FPropertyBagTest_LoadPropertyByTag, "CoreUObject::PropertyBag
 		PropertyBagTestUtils::LoadDataByTag(Dst, PathA, SecondValue);
 
 		void* SecondStoredPointer = Dst.CreateConstIterator().GetValue();
-		CHECK(SecondStoredPointer != StoredPointer); // Different Property, Different Storage for the same path
 		CHECK(*static_cast<float*>(SecondStoredPointer) == SecondValue); // It did store the value
 	}
 
@@ -600,8 +599,8 @@ DISABLED_TEST_CASE(FPropertyBagTest_LoadPropertyByTag, "CoreUObject::PropertyBag
 		int SecondValue = 0x76543218;
 		PropertyBagTestUtils::LoadDataByTag(Dst, PathA, Property, SecondValue);
 
+		CHECK(Property == Dst.CreateConstIterator().GetProperty()); // It did store the property
 		void* SecondStoredPointer = Dst.CreateConstIterator().GetValue();
-		CHECK(SecondStoredPointer != StoredPointer); // Different Property, Different Storage for the same path
 		CHECK(*static_cast<int*>(SecondStoredPointer) == SecondValue); // It did store the value
 	}
 
@@ -629,7 +628,6 @@ DISABLED_TEST_CASE(FPropertyBagTest_LoadPropertyByTag, "CoreUObject::PropertyBag
 		PropertyBagTestUtils::LoadDataByTag(Dst, PathA, Property2, SecondValue);
 
 		void* SecondStoredPointer = Dst.CreateConstIterator().GetValue();
-		CHECK(SecondStoredPointer != StoredPointer); // Different Property, Different Storage for the same path
 		CHECK(*static_cast<int*>(SecondStoredPointer) == SecondValue); // It did store the value
 		CHECK(Property2 == Dst.CreateConstIterator().GetProperty());
 	}
@@ -681,7 +679,7 @@ DISABLED_TEST_CASE(FPropertyBagTest_LoadPropertyByTag, "CoreUObject::PropertyBag
 	}
 }
 
-DISABLED_TEST_CASE(FPropertyBagTest_LoadPropertyByTag_And_Add, "CoreUObject::PropertyBag::LoadPropertyByTag_And_Add", "[Core][UObject][PropertyBag]")
+TEST_CASE_NAMED(FPropertyBagTest_LoadPropertyByTag_And_Add, "CoreUObject::PropertyBag::LoadPropertyByTag_And_Add", "[Core][UObject][PropertyBag]")
 {
 	FProperty* Property = CastFieldChecked<FProperty>(FField::Construct(NAME_IntProperty, {}, FName(TEXT("Tmp")), RF_NoFlags));
 	ON_SCOPE_EXIT
@@ -826,7 +824,7 @@ DISABLED_TEST_CASE(FPropertyBagTest_LoadPropertyByTag_And_Add, "CoreUObject::Pro
 	}
 }
 
-DISABLED_TEST_CASE(FPropertyBagTest_Iteration, "CoreUObject::PropertyBag::Iteration", "[Core][UObject][PropertyBag]")
+TEST_CASE_NAMED(FPropertyBagTest_Iteration, "CoreUObject::PropertyBag::Iteration", "[Core][UObject][PropertyBag]")
 {
 	SECTION("Iteration_Parents_Before_Children")
 	{
@@ -862,7 +860,7 @@ DISABLED_TEST_CASE(FPropertyBagTest_Iteration, "CoreUObject::PropertyBag::Iterat
 	}
 }
 
-DISABLED_TEST_CASE(FPropertyBagTest_Paths_With_Index, "CoreUObject::PropertyBag::Paths_With_Index", "[Core][UObject][PropertyBag]")
+TEST_CASE_NAMED(FPropertyBagTest_Paths_With_Index, "CoreUObject::PropertyBag::Paths_With_Index", "[Core][UObject][PropertyBag]")
 {
 	FProperty* Property = CastFieldChecked<FProperty>(FField::Construct(NAME_IntProperty, {}, FName(TEXT("Tmp")), RF_NoFlags));
 	ON_SCOPE_EXIT
