Removed deprecated FProperty constructors, and ExportTextItem and ImportText.

diff --git a/Engine/Source/Runtime/CoreUObject/Private/UObject/EnumProperty.cpp b/Engine/Source/Runtime/CoreUObject/Private/UObject/EnumProperty.cpp
--- a/Engine/Source/Runtime/CoreUObject/Private/UObject/EnumProperty.cpp
+++ b/Engine/Source/Runtime/CoreUObject/Private/UObject/EnumProperty.cpp
@@ -54,26 +54,6 @@ FEnumProperty::FEnumProperty(FFieldVariant InOwner, const FName& InName, EObject
 
 }
 
-FEnumProperty::FEnumProperty(FFieldVariant InOwner, const FName& InName, EObjectFlags InObjectFlags, UEnum* InEnum)
-PRAGMA_DISABLE_DEPRECATION_WARNINGS
-	: FProperty(InOwner, InName, InObjectFlags, 0, CPF_HasGetValueTypeHash)
-PRAGMA_ENABLE_DEPRECATION_WARNINGS
-	, Enum(InEnum)
-{
-	// This is expected to be set post-construction by AddCppProperty
-	UnderlyingProp = nullptr;
-}
-
-FEnumProperty::FEnumProperty(FFieldVariant InOwner, const FName& InName, EObjectFlags InObjectFlags, int32 InOffset, EPropertyFlags InFlags, UEnum* InEnum)
-PRAGMA_DISABLE_DEPRECATION_WARNINGS
-	: FProperty(InOwner, InName, InObjectFlags, InOffset, InFlags | CPF_HasGetValueTypeHash)
-PRAGMA_ENABLE_DEPRECATION_WARNINGS
-	, Enum(InEnum)
-{
-	// This is expected to be set post-construction by AddCppProperty
-	UnderlyingProp = nullptr;
-}
-
 FEnumProperty::FEnumProperty(FFieldVariant InOwner, const UECodeGen_Private::FEnumPropertyParams& Prop)
 	: FProperty(InOwner, (const UECodeGen_Private::FPropertyParamsBaseWithOffset&)Prop, CPF_HasGetValueTypeHash)
 {
diff --git a/Engine/Source/Runtime/CoreUObject/Private/UObject/Property.cpp b/Engine/Source/Runtime/CoreUObject/Private/UObject/Property.cpp
--- a/Engine/Source/Runtime/CoreUObject/Private/UObject/Property.cpp
+++ b/Engine/Source/Runtime/CoreUObject/Private/UObject/Property.cpp
@@ -741,22 +741,6 @@ FProperty::FProperty(FFieldVariant InOwner, const FName& InName, EObjectFlags In
 {
 }
 
-FProperty::FProperty(FFieldVariant InOwner, const FName& InName, EObjectFlags InObjectFlags, int32 InOffset, EPropertyFlags InFlags)
-	: FField(InOwner, InName, InObjectFlags)
-	, ArrayDim(1)
-	, ElementSize(0)
-	, PropertyFlags(InFlags)
-	, RepIndex(0)
-	, BlueprintReplicationCondition(COND_None)
-	, Offset_Internal(InOffset)
-	, PropertyLinkNext(nullptr)
-	, NextRef(nullptr)
-	, DestructorLinkNext(nullptr)
-	, PostConstructLinkNext(nullptr)
-{
-	Init();
-}
-
 FProperty::FProperty(FFieldVariant InOwner, const UECodeGen_Private::FPropertyParamsBaseWithOffset& Prop, EPropertyFlags AdditionalPropertyFlags /*= CPF_None*/)
 	: FField(InOwner, UTF8_TO_TCHAR(Prop.NameUTF8), Prop.ObjectFlags)
 	, ArrayDim(1)
diff --git a/Engine/Source/Runtime/CoreUObject/Private/UObject/PropertyBool.cpp b/Engine/Source/Runtime/CoreUObject/Private/UObject/PropertyBool.cpp
--- a/Engine/Source/Runtime/CoreUObject/Private/UObject/PropertyBool.cpp
+++ b/Engine/Source/Runtime/CoreUObject/Private/UObject/PropertyBool.cpp
@@ -25,18 +25,6 @@ FBoolProperty::FBoolProperty(FFieldVariant InOwner, const FName& InName, EObject
 	SetBoolSize(1, false, 1);
 }
 
-FBoolProperty::FBoolProperty(FFieldVariant InOwner, const FName& InName, EObjectFlags InObjectFlags, int32 InOffset, EPropertyFlags InFlags, uint32 InBitMask, uint32 InElementSize, bool bIsNativeBool)
-PRAGMA_DISABLE_DEPRECATION_WARNINGS
-	: FProperty(InOwner, InName, InObjectFlags, InOffset, InFlags | CPF_HasGetValueTypeHash)
-PRAGMA_ENABLE_DEPRECATION_WARNINGS
-	, FieldSize(0)
-	, ByteOffset(0)
-	, ByteMask(1)
-	, FieldMask(1)
-{
-	SetBoolSize(InElementSize, bIsNativeBool, InBitMask);
-}
-
 FBoolProperty::FBoolProperty(FFieldVariant InOwner, const UECodeGen_Private::FBoolPropertyParams& Prop)
 	: FProperty(InOwner, (const UECodeGen_Private::FPropertyParamsBaseWithoutOffset&)Prop, CPF_HasGetValueTypeHash)
 	, FieldSize(0)
diff --git a/Engine/Source/Runtime/CoreUObject/Private/UObject/PropertyMap.cpp b/Engine/Source/Runtime/CoreUObject/Private/UObject/PropertyMap.cpp
--- a/Engine/Source/Runtime/CoreUObject/Private/UObject/PropertyMap.cpp
+++ b/Engine/Source/Runtime/CoreUObject/Private/UObject/PropertyMap.cpp
@@ -156,18 +156,6 @@ FMapProperty::FMapProperty(FFieldVariant InOwner, const FName& InName, EObjectFl
 	MapFlags = InMapFlags;
 }
 
-FMapProperty::FMapProperty(FFieldVariant InOwner, const FName& InName, EObjectFlags InObjectFlags, int32 InOffset, EPropertyFlags InFlags, EMapPropertyFlags InMapFlags)
-PRAGMA_DISABLE_DEPRECATION_WARNINGS
-	: FMapProperty_Super(InOwner, InName, InObjectFlags, InOffset, InFlags)
-PRAGMA_ENABLE_DEPRECATION_WARNINGS
-{
-	// These are expected to be set post-construction by AddCppProperty
-	KeyProp   = nullptr;
-	ValueProp = nullptr;
-
-	MapFlags = InMapFlags;
-}
-
 FMapProperty::FMapProperty(FFieldVariant InOwner, const UECodeGen_Private::FMapPropertyParams& Prop)
 	: FMapProperty_Super(InOwner, (const UECodeGen_Private::FPropertyParamsBaseWithOffset&)Prop)
 {
diff --git a/Engine/Source/Runtime/CoreUObject/Private/UObject/PropertySet.cpp b/Engine/Source/Runtime/CoreUObject/Private/UObject/PropertySet.cpp
--- a/Engine/Source/Runtime/CoreUObject/Private/UObject/PropertySet.cpp
+++ b/Engine/Source/Runtime/CoreUObject/Private/UObject/PropertySet.cpp
@@ -135,15 +135,6 @@ FSetProperty::FSetProperty(FFieldVariant InOwner, const FName& InName, EObjectFl
 	ElementProp = nullptr;
 }
 
-FSetProperty::FSetProperty(FFieldVariant InOwner, const FName& InName, EObjectFlags InObjectFlags, int32 InOffset, EPropertyFlags InFlags)
-PRAGMA_DISABLE_DEPRECATION_WARNINGS
-: FSetProperty_Super(InOwner, InName, InObjectFlags, InOffset, InFlags)
-PRAGMA_ENABLE_DEPRECATION_WARNINGS
-{
-	// This is expected to be set post-construction by AddCppProperty
-	ElementProp = nullptr;
-}
-
 FSetProperty::FSetProperty(FFieldVariant InOwner, const UECodeGen_Private::FSetPropertyParams& Prop)
 	: FSetProperty_Super(InOwner, (const UECodeGen_Private::FPropertyParamsBaseWithOffset&)Prop)
 {
diff --git a/Engine/Source/Runtime/CoreUObject/Private/UObject/PropertyStruct.cpp b/Engine/Source/Runtime/CoreUObject/Private/UObject/PropertyStruct.cpp
--- a/Engine/Source/Runtime/CoreUObject/Private/UObject/PropertyStruct.cpp
+++ b/Engine/Source/Runtime/CoreUObject/Private/UObject/PropertyStruct.cpp
@@ -50,15 +50,6 @@ FStructProperty::FStructProperty(FFieldVariant InOwner, const FName& InName, EOb
 	ElementSize = 0;
 }
 
-FStructProperty::FStructProperty(FFieldVariant InOwner, const FName& InName, EObjectFlags InObjectFlags, int32 InOffset, EPropertyFlags InFlags, UScriptStruct* InStruct)
-PRAGMA_DISABLE_DEPRECATION_WARNINGS
-	: FProperty(InOwner, InName, InObjectFlags, InOffset, InStruct->GetCppStructOps() ? InStruct->GetCppStructOps()->GetComputedPropertyFlags() | InFlags : InFlags)
-PRAGMA_ENABLE_DEPRECATION_WARNINGS
-	, Struct( InStruct )
-{
-	ElementSize = Struct->PropertiesSize;
-}
-
 static EPropertyFlags GetStructComputedPropertyFlags(const UECodeGen_Private::FStructPropertyParams& Prop)
 {
 	EPropertyFlags ComputedPropertyFlags = CPF_None;
diff --git a/Engine/Source/Runtime/CoreUObject/Public/UObject/EnumProperty.h b/Engine/Source/Runtime/CoreUObject/Public/UObject/EnumProperty.h
--- a/Engine/Source/Runtime/CoreUObject/Public/UObject/EnumProperty.h
+++ b/Engine/Source/Runtime/CoreUObject/Public/UObject/EnumProperty.h
@@ -31,8 +31,6 @@ class FEnumProperty : public FProperty
 
 public:
 	COREUOBJECT_API FEnumProperty(FFieldVariant InOwner, const FName& InName, EObjectFlags InObjectFlags);
-	COREUOBJECT_API FEnumProperty(FFieldVariant InOwner, const FName& InName, EObjectFlags InObjectFlags, UEnum* InEnum);
-	COREUOBJECT_API FEnumProperty(FFieldVariant InOwner, const FName& InName, EObjectFlags InObjectFlags, int32 InOffset, EPropertyFlags InFlags, UEnum* InEnum);
 
 	/**
 	 * Constructor used for constructing compiled in properties
diff --git a/Engine/Source/Runtime/CoreUObject/Public/UObject/FieldPathProperty.h b/Engine/Source/Runtime/CoreUObject/Public/UObject/FieldPathProperty.h
--- a/Engine/Source/Runtime/CoreUObject/Public/UObject/FieldPathProperty.h
+++ b/Engine/Source/Runtime/CoreUObject/Public/UObject/FieldPathProperty.h
@@ -39,15 +39,6 @@ public:
 	{
 	}
 
-	UE_DEPRECATED(5.1, "Compiled-in property constructor is deprecated, use other constructors instead.")
-	COREUOBJECT_API FFieldPathProperty(FFieldVariant InOwner, const FName& InName, EObjectFlags InObjectFlags, int32 InOffset, EPropertyFlags InFlags, FFieldClass* InPropertyClass)
-PRAGMA_DISABLE_DEPRECATION_WARNINGS
-		: FFieldPathProperty_Super(InOwner, InName, InObjectFlags, InOffset, InFlags)
-PRAGMA_ENABLE_DEPRECATION_WARNINGS
-		, PropertyClass(InPropertyClass)
-	{
-	}
-
 	/**
 	 * Constructor used for constructing compiled in properties
 	 * @param InOwner Owner of the property
diff --git a/Engine/Source/Runtime/CoreUObject/Public/UObject/TextProperty.h b/Engine/Source/Runtime/CoreUObject/Public/UObject/TextProperty.h
--- a/Engine/Source/Runtime/CoreUObject/Public/UObject/TextProperty.h
+++ b/Engine/Source/Runtime/CoreUObject/Public/UObject/TextProperty.h
@@ -34,14 +34,6 @@ public:
 	{
 	}
 
-	UE_DEPRECATED(5.1, "Compiled-in property constructor is deprecated, use other constructors instead.")
-	COREUOBJECT_API FTextProperty(FFieldVariant InOwner, const FName& InName, EObjectFlags InObjectFlags, int32 InOffset, EPropertyFlags InFlags)
-PRAGMA_DISABLE_DEPRECATION_WARNINGS
-		: FTextProperty_Super(InOwner, InName, InObjectFlags, InOffset, InFlags)
-PRAGMA_ENABLE_DEPRECATION_WARNINGS
-	{
-	}
-
 	/**
 	 * Constructor used for constructing compiled in properties
 	 * @param InOwner Owner of the property
diff --git a/Engine/Source/Runtime/CoreUObject/Public/UObject/UnrealType.h b/Engine/Source/Runtime/CoreUObject/Public/UObject/UnrealType.h
--- a/Engine/Source/Runtime/CoreUObject/Public/UObject/UnrealType.h
+++ b/Engine/Source/Runtime/CoreUObject/Public/UObject/UnrealType.h
@@ -191,9 +191,6 @@ public:
 	// Constructors.
 	COREUOBJECT_API FProperty(FFieldVariant InOwner, const FName& InName, EObjectFlags InObjectFlags);
 
-	UE_DEPRECATED(5.1, "Compiled-in property constructor is deprecated, use other constructors instead.")
-	COREUOBJECT_API FProperty(FFieldVariant InOwner, const FName& InName, EObjectFlags InObjectFlags, int32 InOffset, EPropertyFlags InFlags);
-
 	/**
 	 * Constructor used for constructing compiled-in properties
 	 * @param InOwner Owner of the property
@@ -461,12 +458,6 @@ public:
 	COREUOBJECT_API virtual bool NetSerializeItem( FArchive& Ar, UPackageMap* Map, void* Data, TArray<uint8> * MetaData = NULL ) const;
 	COREUOBJECT_API virtual bool SupportsNetSharedSerialization() const;
 
-	UE_DEPRECATED(5.1, "Please use ExportTextItem_InContainer or ExportTextItem_Direct instead.")
-	virtual void ExportTextItem(FString& ValueStr, const void* PropertyValue, const void* DefaultValue, UObject* Parent, int32 PortFlags, UObject* ExportRootScope = nullptr) const
-	{
-		ExportTextItem_Direct(ValueStr, PropertyValue, DefaultValue, Parent, PortFlags, ExportRootScope);
-	}	
-
 	void ExportTextItem_Direct(FString& ValueStr, const void* PropertyValue, const void* DefaultValue, UObject* Parent, int32 PortFlags, UObject* ExportRootScope = nullptr) const
 	{
 		ExportText_Internal(ValueStr, PropertyValue, EPropertyPointerType::Direct, DefaultValue, Parent, PortFlags, ExportRootScope);
@@ -477,12 +468,6 @@ public:
 		ExportText_Internal(ValueStr, Container, EPropertyPointerType::Container, DefaultValue, Parent, PortFlags, ExportRootScope);
 	}
 
-	UE_DEPRECATED(5.1, "ImportText that takes a direct property pointer is deprecated.  Please use ImportText_Direct or ImportText_InContainer instead.")
-	const TCHAR* ImportText( const TCHAR* Buffer, void* Data, int32 PortFlags, UObject* OwnerObject, FOutputDevice* ErrorText = (FOutputDevice*)GWarn ) const
-	{
-		return ImportText_Direct(Buffer, Data, OwnerObject, PortFlags, ErrorText);
-	}
-
 	/**
 	 * Import a text value
 	 * @param Buffer		Text representing the property value
@@ -1352,15 +1337,6 @@ public:
 		SetElementSize();
 	}
 
-	UE_DEPRECATED(5.1, "Compiled-in property constructor is deprecated, use other constructors instead.")
-	TProperty(FFieldVariant InOwner, const FName& InName, EObjectFlags InObjectFlags, int32 InOffset, EPropertyFlags InFlags)
-PRAGMA_DISABLE_DEPRECATION_WARNINGS
-		: Super(InOwner, InName, InObjectFlags, InOffset, InFlags | TTypeFundamentals::GetComputedFlagsPropertyFlags())
-PRAGMA_ENABLE_DEPRECATION_WARNINGS
-	{
-		SetElementSize();
-	}
-
 	/**
 	 * Constructor used for constructing compiled in properties
 	 * @param InOwner Owner of the property
@@ -1496,14 +1472,6 @@ public:
 	{
 	}
 
-	UE_DEPRECATED(5.1, "Compiled-in property constructor is deprecated, use other constructors instead.")
-	TProperty_WithEqualityAndSerializer(FFieldVariant InOwner, const FName& InName, EObjectFlags InObjectFlags, int32 InOffset, EPropertyFlags InFlags)
-PRAGMA_DISABLE_DEPRECATION_WARNINGS
-		: Super(InOwner, InName, InObjectFlags, InOffset, InFlags)
-PRAGMA_ENABLE_DEPRECATION_WARNINGS
-	{
-	}
-
 	/**
 	 * Constructor used for constructing compiled in properties
 	 * @param InOwner Owner of the property
@@ -1544,14 +1512,6 @@ class COREUOBJECT_API FNumericProperty : public FProperty
 	{
 	}
 
-	UE_DEPRECATED(5.1, "Compiled-in property constructor is deprecated, use other constructors instead.")
-	FNumericProperty(FFieldVariant InOwner, const FName& InName, EObjectFlags InObjectFlags, int32 InOffset, EPropertyFlags InFlags)
-PRAGMA_DISABLE_DEPRECATION_WARNINGS
-		: FProperty(InOwner, InName, InObjectFlags, InOffset, InFlags)
-PRAGMA_ENABLE_DEPRECATION_WARNINGS
-	{
-	}
-
 	/**
 	 * Constructor used for constructing compiled in properties
 	 * @param InOwner Owner of the property
@@ -1695,14 +1655,6 @@ public:
 	{
 	}
 
-	UE_DEPRECATED(5.1, "Compiled-in property constructor is deprecated, use other constructors instead.")
-	TProperty_Numeric(FFieldVariant InOwner, const FName& InName, EObjectFlags InObjectFlags, int32 InOffset, EPropertyFlags InFlags)
-PRAGMA_DISABLE_DEPRECATION_WARNINGS
-		: Super(InOwner, InName, InObjectFlags, InOffset, InFlags)
-PRAGMA_ENABLE_DEPRECATION_WARNINGS
-	{
-	}
-
 	/**
 	 * Constructor used for constructing compiled in properties
 	 * @param InOwner Owner of the property
@@ -1961,15 +1913,6 @@ class COREUOBJECT_API FByteProperty : public TProperty_Numeric<uint8>
 	{
 	}
 
-	UE_DEPRECATED(5.1, "Compiled-in property constructor is deprecated, use other constructors instead.")
-	FByteProperty(FFieldVariant InOwner, const FName& InName, EObjectFlags InObjectFlags, int32 InOffset, EPropertyFlags InFlags, UEnum* InEnum = nullptr)
-PRAGMA_DISABLE_DEPRECATION_WARNINGS
-		: TProperty_Numeric(InOwner, InName, InObjectFlags, InOffset, InFlags)
-PRAGMA_ENABLE_DEPRECATION_WARNINGS
-	,	Enum( InEnum )
-	{
-	}
-
 	/**
 	 * Constructor used for constructing compiled in properties
 	 * @param InOwner Owner of the property
@@ -2036,14 +1979,6 @@ class COREUOBJECT_API FInt8Property : public TProperty_Numeric<int8>
 	{
 	}
 
-	UE_DEPRECATED(5.1, "Compiled-in property constructor is deprecated, use other constructors instead.")
-	FInt8Property(FFieldVariant InOwner, const FName& InName, EObjectFlags InObjectFlags, int32 InOffset, EPropertyFlags InFlags)
-PRAGMA_DISABLE_DEPRECATION_WARNINGS
-		: TProperty_Numeric(InOwner, InName, InObjectFlags, InOffset, InFlags)
-PRAGMA_ENABLE_DEPRECATION_WARNINGS
-	{
-	}
-
 	/**
 	 * Constructor used for constructing compiled in properties
 	 * @param InOwner Owner of the property
@@ -2075,14 +2010,6 @@ class COREUOBJECT_API FInt16Property : public TProperty_Numeric<int16>
 	{
 	}
 
-	UE_DEPRECATED(5.1, "Compiled-in property constructor is deprecated, use other constructors instead.")
-	FInt16Property(FFieldVariant InOwner, const FName& InName, EObjectFlags InObjectFlags, int32 InOffset, EPropertyFlags InFlags)
-PRAGMA_DISABLE_DEPRECATION_WARNINGS
-		: TProperty_Numeric(InOwner, InName, InObjectFlags, InOffset, InFlags)
-PRAGMA_ENABLE_DEPRECATION_WARNINGS
-	{
-	}
-
 	/**
 	 * Constructor used for constructing compiled in properties
 	 * @param InOwner Owner of the property
@@ -2115,14 +2042,6 @@ class COREUOBJECT_API FIntProperty : public TProperty_Numeric<int32>
 	{
 	}
 
-	UE_DEPRECATED(5.1, "Compiled-in property constructor is deprecated, use other constructors instead.")
-	FIntProperty(FFieldVariant InOwner, const FName& InName, EObjectFlags InObjectFlags, int32 InOffset, EPropertyFlags InFlags)
-PRAGMA_DISABLE_DEPRECATION_WARNINGS
-		: TProperty_Numeric(InOwner, InName, InObjectFlags, InOffset, InFlags)
-PRAGMA_ENABLE_DEPRECATION_WARNINGS
-	{
-	}
-
 	/**
 	 * Constructor used for constructing compiled in properties
 	 * @param InOwner Owner of the property
@@ -2154,14 +2073,6 @@ class COREUOBJECT_API FInt64Property : public TProperty_Numeric<int64>
 	{
 	}
 
-	UE_DEPRECATED(5.1, "Compiled-in property constructor is deprecated, use other constructors instead.")
-	FInt64Property(FFieldVariant InOwner, const FName& InName, EObjectFlags InObjectFlags, int32 InOffset, EPropertyFlags InFlags)
-PRAGMA_DISABLE_DEPRECATION_WARNINGS
-		: TProperty_Numeric(InOwner, InName, InObjectFlags, InOffset, InFlags)
-PRAGMA_ENABLE_DEPRECATION_WARNINGS
-	{
-	}
-
 	/**
 	 * Constructor used for constructing compiled in properties
 	 * @param InOwner Owner of the property
@@ -2193,14 +2104,6 @@ class COREUOBJECT_API FUInt16Property : public TProperty_Numeric<uint16>
 	{
 	}
 
-	UE_DEPRECATED(5.1, "Compiled-in property constructor is deprecated, use other constructors instead.")
-	FUInt16Property(FFieldVariant InOwner, const FName& InName, EObjectFlags InObjectFlags, int32 InOffset, EPropertyFlags InFlags)
-PRAGMA_DISABLE_DEPRECATION_WARNINGS
-		: TProperty_Numeric(InOwner, InName, InObjectFlags, InOffset, InFlags)
-PRAGMA_ENABLE_DEPRECATION_WARNINGS
-	{
-	}
-
 	/**
 	 * Constructor used for constructing compiled in properties
 	 * @param InOwner Owner of the property
@@ -2232,14 +2135,6 @@ class COREUOBJECT_API FUInt32Property : public TProperty_Numeric<uint32>
 	{
 	}
 
-	UE_DEPRECATED(5.1, "Compiled-in property constructor is deprecated, use other constructors instead.")
-	FUInt32Property(FFieldVariant InOwner, const FName& InName, EObjectFlags InObjectFlags, int32 InOffset, EPropertyFlags InFlags)
-PRAGMA_DISABLE_DEPRECATION_WARNINGS
-	:	TProperty_Numeric(InOwner, InName, InObjectFlags, InOffset, InFlags)
-PRAGMA_ENABLE_DEPRECATION_WARNINGS
-	{
-	}
-
 	/**
 	 * Constructor used for constructing compiled in properties
 	 * @param InOwner Owner of the property
@@ -2271,14 +2166,6 @@ class COREUOBJECT_API FUInt64Property : public TProperty_Numeric<uint64>
 	{
 	}
 
-	UE_DEPRECATED(5.1, "Compiled-in property constructor is deprecated, use other constructors instead.")
-	FUInt64Property(FFieldVariant InOwner, const FName& InName, EObjectFlags InObjectFlags, int32 InOffset, EPropertyFlags InFlags)
-PRAGMA_DISABLE_DEPRECATION_WARNINGS
-		: TProperty_Numeric(InOwner, InName, InObjectFlags, InOffset, InFlags)
-PRAGMA_ENABLE_DEPRECATION_WARNINGS
-	{
-	}
-
 	/**
 	 * Constructor used for constructing compiled in properties
 	 * @param InOwner Owner of the property
@@ -2311,14 +2198,6 @@ class COREUOBJECT_API FFloatProperty : public TProperty_Numeric<float>
 	{
 	}
 
-	UE_DEPRECATED(5.1, "Compiled-in property constructor is deprecated, use other constructors instead.")
-	FFloatProperty(FFieldVariant InOwner, const FName& InName, EObjectFlags InObjectFlags, int32 InOffset, EPropertyFlags InFlags)
-PRAGMA_DISABLE_DEPRECATION_WARNINGS
-		: TProperty_Numeric(InOwner, InName, InObjectFlags, InOffset, InFlags)
-PRAGMA_ENABLE_DEPRECATION_WARNINGS
-	{
-	}
-
 	/**
 	 * Constructor used for constructing compiled in properties
 	 * @param InOwner Owner of the property
@@ -2350,14 +2229,6 @@ class COREUOBJECT_API FDoubleProperty : public TProperty_Numeric<double>
 	{
 	}
 
-	UE_DEPRECATED(5.1, "Compiled-in property constructor is deprecated, use other constructors instead.")
-	FDoubleProperty(FFieldVariant InOwner, const FName& InName, EObjectFlags InObjectFlags, int32 InOffset, EPropertyFlags InFlags)
-PRAGMA_DISABLE_DEPRECATION_WARNINGS
-		: TProperty_Numeric(InOwner, InName, InObjectFlags, InOffset, InFlags)
-PRAGMA_ENABLE_DEPRECATION_WARNINGS
-	{
-	}
-
 	/**
 	 * Constructor used for constructing compiled in properties
 	 * @param InOwner Owner of the property
@@ -2411,20 +2282,6 @@ public:
 
 	FBoolProperty(FFieldVariant InOwner, const FName& InName, EObjectFlags InObjectFlags);
 
-	/**
-	 * Constructor.
-	 *
-	 * @param ECppProperty Unused.
-	 * @param InOffset Offset of the property.
-	 * @param InCategory Category of the property.
-	 * @param InFlags Property flags.
-	 * @param InBitMask Bitmask of the bitfield this property represents.
-	 * @param InElementSize Sizeof of the boolean type this property represents.
-	 * @param bIsNativeBool true if this property represents C++ bool type.
-	 */
-	UE_DEPRECATED(5.1, "Compiled-in property constructor is deprecated, use other constructors instead.")
-	FBoolProperty(FFieldVariant InOwner, const FName& InName, EObjectFlags InObjectFlags, int32 InOffset, EPropertyFlags InFlags, uint32 InBitMask, uint32 InElementSize, bool bIsNativeBool);
-
 	/**
 	 * Constructor used for constructing compiled in properties
 	 * @param InOwner Owner of the property
@@ -2559,14 +2416,6 @@ public:
 		, PropertyClass(nullptr)
 	{}
 
-	UE_DEPRECATED(5.1, "Compiled-in property constructor is deprecated, use other constructors instead.")
-	FObjectPropertyBase(FFieldVariant InOwner, const FName& InName, EObjectFlags InObjectFlags, int32 InOffset, EPropertyFlags InFlags, UClass* InClass = NULL)
-PRAGMA_DISABLE_DEPRECATION_WARNINGS
-		: FProperty(InOwner, InName, InObjectFlags, InOffset, InFlags)
-PRAGMA_ENABLE_DEPRECATION_WARNINGS
-		, PropertyClass(InClass)
-	{}
-
 	/**
 	 * Constructor used for constructing compiled-in properties
 	 * @param InOwner Owner of the property
@@ -2824,15 +2673,6 @@ public:
 		this->PropertyClass = nullptr;
 	}
 
-	UE_DEPRECATED(5.1, "Compiled-in property constructor is deprecated, use other constructors instead.")
-	TFObjectPropertyBase(FFieldVariant InOwner, const FName& InName, EObjectFlags InObjectFlags, int32 InOffset, EPropertyFlags InFlags, UClass* InClass)
-PRAGMA_DISABLE_DEPRECATION_WARNINGS
-		: Super(InOwner, InName, InObjectFlags, InOffset, InFlags)
-PRAGMA_ENABLE_DEPRECATION_WARNINGS
-	{
-		this->PropertyClass = InClass;
-	}
-
 	/**
 	 * Constructor used for constructing compiled-in properties
 	 * @param InOwner Owner of the property
@@ -2900,14 +2740,6 @@ class COREUOBJECT_API FObjectProperty : public TFObjectPropertyBase<TObjectPtr<U
 	{
 	}
 
-	UE_DEPRECATED(5.1, "Compiled-in property constructor is deprecated, use other constructors instead.")
-	FObjectProperty(FFieldVariant InOwner, const FName& InName, EObjectFlags InObjectFlags, int32 InOffset, EPropertyFlags InFlags, UClass* InClass)
-PRAGMA_DISABLE_DEPRECATION_WARNINGS
-		: TFObjectPropertyBase(InOwner, InName, InObjectFlags, InOffset, InFlags, InClass)
-PRAGMA_ENABLE_DEPRECATION_WARNINGS
-	{
-	}
-
 	/**
 	 * Constructor used for constructing compiled in properties
 	 * @param InOwner Owner of the property
@@ -2991,14 +2823,6 @@ class COREUOBJECT_API FWeakObjectProperty : public TFObjectPropertyBase<FWeakObj
 	{
 	}
 
-	UE_DEPRECATED(5.1, "Compiled-in property constructor is deprecated, use other constructors instead.")
-	FWeakObjectProperty(FFieldVariant InOwner, const FName& InName, EObjectFlags InObjectFlags, int32 InOffset, EPropertyFlags InFlags, UClass* InClass)
-PRAGMA_DISABLE_DEPRECATION_WARNINGS
-		: TFObjectPropertyBase(InOwner, InName, InObjectFlags, InOffset, InFlags, InClass)
-PRAGMA_ENABLE_DEPRECATION_WARNINGS
-	{
-	}
-	
 	/**
 	 * Constructor used for constructing compiled in properties
 	 * @param InOwner Owner of the property
@@ -3057,14 +2881,6 @@ class COREUOBJECT_API FLazyObjectProperty : public TFObjectPropertyBase<FLazyObj
 	{
 	}
 
-	UE_DEPRECATED(5.1, "Compiled-in property constructor is deprecated, use other constructors instead.")
-	FLazyObjectProperty(FFieldVariant InOwner, const FName& InName, EObjectFlags InObjectFlags, int32 InOffset, EPropertyFlags InFlags, UClass* InClass)
-PRAGMA_DISABLE_DEPRECATION_WARNINGS
-		: TFObjectPropertyBase(InOwner, InName, InObjectFlags, InOffset, InFlags, InClass)
-PRAGMA_ENABLE_DEPRECATION_WARNINGS
-	{
-	}
-
 	/**
 	 * Constructor used for constructing compiled in properties
 	 * @param InOwner Owner of the property
@@ -3125,14 +2941,6 @@ class COREUOBJECT_API FSoftObjectProperty : public TFObjectPropertyBase<FSoftObj
 	{
 	}
 
-	UE_DEPRECATED(5.1, "Compiled-in property constructor is deprecated, use other constructors instead.")
-	FSoftObjectProperty(FFieldVariant InOwner, const FName& InName, EObjectFlags InObjectFlags, int32 InOffset, EPropertyFlags InFlags, UClass* InClass)
-PRAGMA_DISABLE_DEPRECATION_WARNINGS
-		: TFObjectPropertyBase(InOwner, InName, InObjectFlags, InOffset, InFlags, InClass)
-PRAGMA_ENABLE_DEPRECATION_WARNINGS
-	{
-	}
-
 	/**
 	 * Constructor used for constructing compiled-in properties
 	 * @param InOwner Owner of the property
@@ -3214,15 +3022,6 @@ public:
 	{
 	}
 
-	UE_DEPRECATED(5.1, "Compiled-in property constructor is deprecated, use other constructors instead.")
-	FClassProperty(FFieldVariant InOwner, const FName& InName, EObjectFlags InObjectFlags, int32 InOffset, EPropertyFlags InFlags, UClass* InMetaClass, UClass* InClassType)
-PRAGMA_DISABLE_DEPRECATION_WARNINGS
-		: FObjectProperty(InOwner, InName, InObjectFlags, InOffset, InFlags, InClassType ? InClassType : UClass::StaticClass())
-PRAGMA_ENABLE_DEPRECATION_WARNINGS
-		, MetaClass(InMetaClass)
-	{
-	}
-
 	/**
 	 * Constructor used for constructing compiled in properties
 	 * @param InOwner Owner of the property
@@ -3298,14 +3097,6 @@ public:
 		, MetaClass(nullptr)
 	{}
 
-	UE_DEPRECATED(5.1, "Compiled-in property constructor is deprecated, use other constructors instead.")
-	FSoftClassProperty(FFieldVariant InOwner, const FName& InName, EObjectFlags InObjectFlags, int32 InOffset, EPropertyFlags InFlags, UClass* InMetaClass)
-PRAGMA_DISABLE_DEPRECATION_WARNINGS
-		: Super(InOwner, InName, InObjectFlags, InOffset, InFlags, UClass::StaticClass())
-PRAGMA_ENABLE_DEPRECATION_WARNINGS
-		, MetaClass(InMetaClass)
-	{}
-
 	/**
 	 * Constructor used for constructing compiled in properties
 	 * @param InOwner Owner of the property
@@ -3384,15 +3175,6 @@ public:
 	{
 	}
 
-	UE_DEPRECATED(5.1, "Compiled-in property constructor is deprecated, use other constructors instead.")
-	FInterfaceProperty(FFieldVariant InOwner, const FName& InName, EObjectFlags InObjectFlags, int32 InOffset, EPropertyFlags InFlags, UClass* InInterfaceClass)
-PRAGMA_DISABLE_DEPRECATION_WARNINGS
-		: FInterfaceProperty_Super(InOwner, InName, InObjectFlags, InOffset, (InFlags & ~CPF_InterfaceClearMask))
-PRAGMA_ENABLE_DEPRECATION_WARNINGS
-		, InterfaceClass(InInterfaceClass)
-	{
-	}
-
 	/**
 	 * Constructor used for constructing compiled in properties
 	 * @param InOwner Owner of the property
@@ -3477,14 +3259,6 @@ public:
 	{
 	}
 
-	UE_DEPRECATED(5.1, "Compiled-in property constructor is deprecated, use other constructors instead.")
-	FNameProperty(FFieldVariant InOwner, const FName& InName, EObjectFlags InObjectFlags, int32 InOffset, EPropertyFlags InFlags)
-PRAGMA_DISABLE_DEPRECATION_WARNINGS
-		: FNameProperty_Super(InOwner, InName, InObjectFlags, InOffset, InFlags)
-PRAGMA_ENABLE_DEPRECATION_WARNINGS
-	{
-	}
-
 	/**
 	 * Constructor used for constructing compiled in properties
 	 * @param InOwner Owner of the property
@@ -3534,14 +3308,6 @@ public:
 	{
 	}
 
-	UE_DEPRECATED(5.1, "Compiled-in property constructor is deprecated, use other constructors instead.")
-	FStrProperty(FFieldVariant InOwner, const FName& InName, EObjectFlags InObjectFlags, int32 InOffset, EPropertyFlags InFlags)
-PRAGMA_DISABLE_DEPRECATION_WARNINGS
-		: FStrProperty_Super(InOwner, InName, InObjectFlags, InOffset, InFlags)
-PRAGMA_ENABLE_DEPRECATION_WARNINGS
-	{
-	}
-
 	/**
 	 * Constructor used for constructing compiled in properties
 	 * @param InOwner Owner of the property
@@ -3615,17 +3381,6 @@ public:
 		SetElementSize();
 	}
 
-	UE_DEPRECATED(5.1, "Compiled-in property constructor is deprecated, use other constructors instead.")
-	FArrayProperty(FFieldVariant InOwner, const FName& InName, EObjectFlags InObjectFlags, int32 InOffset, EPropertyFlags InFlags, EArrayPropertyFlags InArrayPropertyFlags)
-PRAGMA_DISABLE_DEPRECATION_WARNINGS
-		: FArrayProperty_Super(InOwner, InName, InObjectFlags, InOffset, InFlags)
-PRAGMA_ENABLE_DEPRECATION_WARNINGS
-		, Inner(nullptr)
-	{
-		ArrayFlags = InArrayPropertyFlags;
-		SetElementSize();
-	}
-
 	/** 
 	 * Constructor used for constructing compiled in properties
 	 * @param InOwner Owner of the property
@@ -3768,9 +3523,6 @@ public:
 
 	FMapProperty(FFieldVariant InOwner, const FName& InName, EObjectFlags InObjectFlags, EMapPropertyFlags InMapFlags=EMapPropertyFlags::None);
 
-	UE_DEPRECATED(5.1, "Compiled-in property constructor is deprecated, use other constructors instead.")
-	FMapProperty(FFieldVariant InOwner, const FName& InName, EObjectFlags InObjectFlags, int32 InOffset, EPropertyFlags InFlags, EMapPropertyFlags InMapFlags);
-
 	/**
 	 * Constructor used for constructing compiled in properties
 	 * @param InOwner Owner of the property
@@ -3919,9 +3671,6 @@ public:
 
 	FSetProperty(FFieldVariant InOwner, const FName& InName, EObjectFlags InObjectFlags);
 
-	UE_DEPRECATED(5.1, "Compiled-in property constructor is deprecated, use other constructors instead.")
-	FSetProperty(FFieldVariant InOwner, const FName& InName, EObjectFlags InObjectFlags, int32 InOffset, EPropertyFlags InFlags);
-
 	/**
 	 * Constructor used for constructing compiled in properties
 	 * @param InOwner Owner of the property
@@ -6050,9 +5799,6 @@ class COREUOBJECT_API FStructProperty : public FProperty
 public:
 	FStructProperty(FFieldVariant InOwner, const FName& InName, EObjectFlags InObjectFlags);
 
-	UE_DEPRECATED(5.1, "Compiled-in property constructor is deprecated, use other constructors instead.")
-	FStructProperty(FFieldVariant InOwner, const FName& InName, EObjectFlags InObjectFlags, int32 InOffset, EPropertyFlags InFlags, UScriptStruct* InStruct);
-
 	/**
 	 * Constructor used for constructing compiled in properties
 	 * @param InOwner Owner of the property
@@ -6142,15 +5888,6 @@ public:
 	{
 	}
 
-	UE_DEPRECATED(5.1, "Compiled-in property constructor is deprecated, use other constructors instead.")
-	FDelegateProperty(FFieldVariant InOwner, const FName& InName, EObjectFlags InObjectFlags, int32 InOffset, EPropertyFlags InFlags, UFunction* InSignatureFunction = NULL)
-PRAGMA_DISABLE_DEPRECATION_WARNINGS
-		: FDelegateProperty_Super(InOwner, InName, InObjectFlags, InOffset, InFlags)
-PRAGMA_ENABLE_DEPRECATION_WARNINGS
-		, SignatureFunction(InSignatureFunction)
-	{
-	}
-
 	/**
 	 * Constructor used for constructing compiled in properties
 	 * @param InOwner Owner of the property
@@ -6214,15 +5951,6 @@ public:
 	{
 	}
 
-	UE_DEPRECATED(5.1, "Compiled-in property constructor is deprecated, use other constructors instead.")
-	FMulticastDelegateProperty(FFieldVariant InOwner, const FName& InName, EObjectFlags InObjectFlags, int32 InOffset, EPropertyFlags InFlags, UFunction* InSignatureFunction = NULL)
-PRAGMA_DISABLE_DEPRECATION_WARNINGS
-		: FProperty(InOwner, InName, InObjectFlags, InOffset, InFlags)
-PRAGMA_ENABLE_DEPRECATION_WARNINGS
-		, SignatureFunction(InSignatureFunction)
-	{
-	}
-
 	/**
 	 * Constructor used for constructing compiled in properties
 	 * @param InOwner Owner of the property
@@ -6297,15 +6025,6 @@ public:
 		this->SignatureFunction = nullptr;
 	}
 
-	UE_DEPRECATED(5.1, "Compiled-in property constructor is deprecated, use other constructors instead.")
-	TProperty_MulticastDelegate(FFieldVariant InOwner, const FName& InName, EObjectFlags InObjectFlags, int32 InOffset, EPropertyFlags InFlags, UFunction* InSignatureFunction = nullptr)
-PRAGMA_DISABLE_DEPRECATION_WARNINGS
-		: Super(InOwner, InName, InObjectFlags, InOffset, InFlags)
-PRAGMA_ENABLE_DEPRECATION_WARNINGS
-	{
-		this->SignatureFunction = InSignatureFunction;
-	}
-
 	TProperty_MulticastDelegate(EInternal InInernal, FFieldClass* InClass)
 		: Super(EC_InternalUseOnlyConstructor, InClass)
 	{
@@ -6348,14 +6067,6 @@ public:
 	{
 	}
 
-	UE_DEPRECATED(5.1, "Compiled-in property constructor is deprecated, use other constructors instead.")
-	FMulticastInlineDelegateProperty(FFieldVariant InOwner, const FName& InName, EObjectFlags InObjectFlags, int32 InOffset, EPropertyFlags InFlags, UFunction* InSignatureFunction = nullptr)
-PRAGMA_DISABLE_DEPRECATION_WARNINGS
-		: TProperty_MulticastDelegate(InOwner, InName, InObjectFlags, InOffset, InFlags, InSignatureFunction)
-PRAGMA_ENABLE_DEPRECATION_WARNINGS
-	{
-	}
-
 	/**
 	 * Constructor used for constructing compiled in properties
 	 * @param InOwner Owner of the property
@@ -6399,14 +6110,6 @@ public:
 	{
 	}
 
-	UE_DEPRECATED(5.1, "Compiled-in property constructor is deprecated, use other constructors instead.")
-	FMulticastSparseDelegateProperty(FFieldVariant InOwner, const FName& InName, EObjectFlags InObjectFlags, int32 InOffset, EPropertyFlags InFlags, UFunction* InSignatureFunction = nullptr)
-PRAGMA_DISABLE_DEPRECATION_WARNINGS
-		: TProperty_MulticastDelegate(InOwner, InName, InObjectFlags, InOffset, InFlags, InSignatureFunction)
-PRAGMA_ENABLE_DEPRECATION_WARNINGS
-	{
-	}
-
 	/**
 	 * Constructor used for constructing compiled in properties
 	 * @param InOwner Owner of the property
