DoesPackageExist AssetRegistry optimization: Re-enables optimization to avoid scanning on disk for corrected casing for FNames by changing the check to determine when the IoDispatcher is in use.

We expose the previously private, and unused HasScriptObjectsChunk method to the IODispatcher, as multiple codepaths are relying on this path to determine if the IoDispatcher is being used at runtime to load from chunks.

diff --git a/Engine/Source/Runtime/Core/Internal/IO/IoDispatcherInternal.h b/Engine/Source/Runtime/Core/Internal/IO/IoDispatcherInternal.h
new file mode 100644
--- /dev/null
+++ b/Engine/Source/Runtime/Core/Internal/IO/IoDispatcherInternal.h
@@ -0,0 +1,20 @@
+// Copyright Epic Games, Inc. All Rights Reserved.
+
+#pragma once
+#include "IO/IoDispatcher.h"
+#include "IO/IoChunkId.h"
+
+// This header is used for IoDispatcher functionality we don't want to expose outside of the engine
+
+class FIoDispatcherInternal final
+{
+public:
+	static bool HasPackageData()
+	{
+		// Checking for the global script objects chunk is currently the best means 
+		// to know if the IoDispatcher will be loading packaged data
+		const FIoDispatcher& Dispatcher = FIoDispatcher::Get();
+		static const bool HasScriptObjectsChunk = Dispatcher.DoesChunkExist(CreateIoChunkId(0, 0, EIoChunkType::ScriptObjects));
+		return HasScriptObjectsChunk;
+	}
+};
\ No newline at end of file
diff --git a/Engine/Source/Runtime/Core/Private/IO/IoDispatcher.cpp b/Engine/Source/Runtime/Core/Private/IO/IoDispatcher.cpp
--- a/Engine/Source/Runtime/Core/Private/IO/IoDispatcher.cpp
+++ b/Engine/Source/Runtime/Core/Private/IO/IoDispatcher.cpp
@@ -959,13 +959,6 @@ FIoDispatcher::OnSignatureError()
 	return Impl->OnSignatureError();
 }
 
-static bool
-HasScriptObjectsChunk(FIoDispatcher& Dispatcher)
-{
-	static bool bHasScriptObjectsChunk = Dispatcher.DoesChunkExist(CreateIoChunkId(0, 0, EIoChunkType::ScriptObjects));
-	return bHasScriptObjectsChunk;
-}
-
 bool
 FIoDispatcher::IsInitialized()
 {
diff --git a/Engine/Source/Runtime/CoreUObject/Private/Misc/PackageName.cpp b/Engine/Source/Runtime/CoreUObject/Private/Misc/PackageName.cpp
--- a/Engine/Source/Runtime/CoreUObject/Private/Misc/PackageName.cpp
+++ b/Engine/Source/Runtime/CoreUObject/Private/Misc/PackageName.cpp
@@ -18,6 +18,7 @@
 #include "Interfaces/IPluginManager.h"
 #include "Internationalization/PackageLocalizationManager.h"
 #include "IO/IoDispatcher.h"
+#include "IO/IoDispatcherInternal.h"
 #include "Misc/App.h"
 #include "Misc/AssetRegistryInterface.h"
 #include "Misc/AutomationTest.h"
@@ -1845,13 +1846,13 @@ FPackageName::EPackageLocationFilter FPackageName::DoesPackageExistEx(const FPac
 		return EPackageLocationFilter::None;
 	}
 
-#if 0 // WITH_EDITOR - Temporarily disabled until UE-206417 is resolved
+
+#if WITH_EDITOR
 	IAssetRegistryInterface* AssetRegistry = IAssetRegistryInterface::GetPtr();
 
 	// Todo: The AssetRegistry currently cannot determine if a package comes from the Filesystem 
 	// or cooked content so we avoid registry lookups since we can't provide a reliable Location
-	if (AssetRegistry && ((uint8)Filter & (uint8)EPackageLocationFilter::FileSystem) &&
-		(!FIoDispatcher::IsInitialized() || FIoDispatcher::Get().GetTotalLoaded() == 0))
+	if (AssetRegistry && ((uint8)Filter & (uint8)EPackageLocationFilter::FileSystem) && (!FIoDispatcher::IsInitialized() || !FIoDispatcherInternal::HasPackageData()))
 	{
 		FName PackageName = PackagePath.GetPackageFName();
 		FName CorrectCasePackageName;
diff --git a/Engine/Source/Runtime/CoreUObject/Private/Serialization/AsyncPackageLoader.cpp b/Engine/Source/Runtime/CoreUObject/Private/Serialization/AsyncPackageLoader.cpp
--- a/Engine/Source/Runtime/CoreUObject/Private/Serialization/AsyncPackageLoader.cpp
+++ b/Engine/Source/Runtime/CoreUObject/Private/Serialization/AsyncPackageLoader.cpp
@@ -12,6 +12,7 @@
 #include "Misc/PackageName.h"
 #include "Misc/PathViews.h"
 #include "IO/IoDispatcher.h"
+#include "IO/IoDispatcherInternal.h"
 #include "HAL/IConsoleManager.h"
 
 #define DO_TRACK_ASYNC_LOAD_REQUESTS (!UE_BUILD_SHIPPING && !UE_BUILD_TEST)
@@ -199,7 +200,7 @@ void InitAsyncThread()
 		bHasUseIoStoreParamInEditor = UE_FORCE_USE_IOSTORE || FParse::Param(FCommandLine::Get(), TEXT("UseIoStore"));
 #endif
 		FIoDispatcher& IoDispatcher = FIoDispatcher::Get();
-		bool bHasScriptObjectsChunk = IoDispatcher.DoesChunkExist(CreateIoChunkId(0, 0, EIoChunkType::ScriptObjects));
+		const bool bHasScriptObjectsChunk = FIoDispatcherInternal::HasPackageData();
 		if (!bCommandLineDisabled && (bSettingsEnabled || bCommandLineEnabled))
 		{
 			GPackageLoader.Reset(MakeAsyncPackageLoader2(IoDispatcher));
diff --git a/Engine/Source/Runtime/CoreUObject/Private/Serialization/EditorPackageLoader.cpp b/Engine/Source/Runtime/CoreUObject/Private/Serialization/EditorPackageLoader.cpp
--- a/Engine/Source/Runtime/CoreUObject/Private/Serialization/EditorPackageLoader.cpp
+++ b/Engine/Source/Runtime/CoreUObject/Private/Serialization/EditorPackageLoader.cpp
@@ -7,6 +7,7 @@
 #include "UObject/UObjectThreadContext.h"
 #include "Misc/PackageName.h"
 #include "IO/IoDispatcher.h"
+#include "IO/IoDispatcherInternal.h"
 
 DEFINE_LOG_CATEGORY_STATIC(LogEditorPackageLoader, Log, All);
 
@@ -32,7 +33,7 @@ public:
 
 	virtual void InitializeLoading() override
 	{
-		if (FIoDispatcher::Get().DoesChunkExist(CreateIoChunkId(0, 0, EIoChunkType::ScriptObjects)))
+		if (FIoDispatcherInternal::HasPackageData())
 		{
 			UE_LOG(LogEditorPackageLoader, Log, TEXT("Initializing Zen loader for cooked packages in editor startup"));
 			CookedPackageLoader->InitializeLoading();
diff --git a/Engine/Source/Runtime/RenderCore/Private/ShaderCodeLibrary.cpp b/Engine/Source/Runtime/RenderCore/Private/ShaderCodeLibrary.cpp
--- a/Engine/Source/Runtime/RenderCore/Private/ShaderCodeLibrary.cpp
+++ b/Engine/Source/Runtime/RenderCore/Private/ShaderCodeLibrary.cpp
@@ -40,6 +40,7 @@ ShaderCodeLibrary.cpp: Bound shader state cache implementation.
 #include "String/ParseTokens.h"
 #include "IO/IoChunkId.h"
 #include "IO/IoDispatcher.h"
+#include "IO/IoDispatcherInternal.h"
 
 #if WITH_EDITORONLY_DATA
 #include "Interfaces/IShaderFormat.h"
@@ -124,8 +125,7 @@ namespace UE
 			bool IsRunningWithIoStore()
 			{
 				static const bool bRunningWithIoStore =
-					FIoDispatcher::IsInitialized()
-					&& FIoDispatcher::Get().DoesChunkExist(CreateIoChunkId(0, 0, EIoChunkType::ScriptObjects));
+					FIoDispatcher::IsInitialized() && FIoDispatcherInternal::HasPackageData();
 				return bRunningWithIoStore;
 			}
 			bool IsRunningWithZenStore()
