Changed FPropertyTypeName to use () instead of <> to delimit parameters

We want to use these to contain object path names, for which () is invalid and <> is not technically invalid.

#rnx

diff --git a/Engine/Source/Runtime/CoreUObject/Private/UObject/PropertyTypeName.cpp b/Engine/Source/Runtime/CoreUObject/Private/UObject/PropertyTypeName.cpp
--- a/Engine/Source/Runtime/CoreUObject/Private/UObject/PropertyTypeName.cpp
+++ b/Engine/Source/Runtime/CoreUObject/Private/UObject/PropertyTypeName.cpp
@@ -53,14 +53,14 @@ inline const FPropertyTypeNameNode* AppendNode(FStringBuilderBase& Builder, cons
 
 	if (int32 Remaining = Node++->InnerCount)
 	{
-		Builder.AppendChar(TEXT('<'));
+		Builder.AppendChar(TEXT('('));
 		for (; Remaining > 0; --Remaining)
 		{
 			Node = AppendNode(Builder, Node);
 			Builder.AppendChar(TEXT(','));
 		}
 		Builder.RemoveSuffix(1);
-		Builder.AppendChar(TEXT('>'));
+		Builder.AppendChar(TEXT(')'));
 	}
 
 	return Node;
@@ -480,7 +480,7 @@ bool FPropertyTypeNameBuilder::TryParse(FStringView Name)
 	bool bAllowBegin = false;
 	for (FStringView Remaining = Name;; Remaining.RightChopInline(Index + 1))
 	{
-		Index = String::FindFirstOfAnyChar(Remaining, {TEXT('<'), TEXT(','), TEXT('>')});
+		Index = String::FindFirstOfAnyChar(Remaining, {TEXT('('), TEXT(','), TEXT(')')});
 
 		if (Index != 0)
 		{
@@ -489,7 +489,7 @@ bool FPropertyTypeNameBuilder::TryParse(FStringView Name)
 			{
 				if (!bAllowName)
 				{
-					// Names must follow '<' or ',' or be at the root.
+					// Names must follow '(' or ',' or be at the root.
 					break;
 				}
 				AddName(FName(Type));
@@ -504,7 +504,7 @@ bool FPropertyTypeNameBuilder::TryParse(FStringView Name)
 			{
 				return true;
 			}
-			// Missing a '>' and/or missing a name.
+			// Missing a ')' and/or missing a name.
 			break;
 		}
 
@@ -515,22 +515,22 @@ bool FPropertyTypeNameBuilder::TryParse(FStringView Name)
 		}
 
 		const TCHAR C = Remaining[Index];
-		if (C == TEXT('<'))
+		if (C == TEXT('('))
 		{
 			if (!bAllowBegin)
 			{
-				// '<' must follow a name.
+				// '(' must follow a name.
 				break;
 			}
 			++Depth;
 			BeginParameters();
 			bAllowName = true;
 		}
-		else if (C == TEXT('>'))
+		else if (C == TEXT(')'))
 		{
 			if (Depth <= 0)
 			{
-				// '>' must have a matching '<'.
+				// ')' must have a matching '('.
 				break;
 			}
 			EndParameters();
diff --git a/Engine/Source/Runtime/CoreUObject/Public/UObject/PropertyTypeName.h b/Engine/Source/Runtime/CoreUObject/Public/UObject/PropertyTypeName.h
--- a/Engine/Source/Runtime/CoreUObject/Public/UObject/PropertyTypeName.h
+++ b/Engine/Source/Runtime/CoreUObject/Public/UObject/PropertyTypeName.h
@@ -29,17 +29,17 @@ struct FPropertyTypeNameNode
  * Represents the type name of a property, including any containers and underlying types.
  *
  * The path name for a type is represented as the object name with the outer chain as type parameters.
- * - /Script/CoreUObject.Outer0:Outer1.Outer2.StructName -> StructName</Script/CoreUObject,Outer0,Outer1,Outer2>
+ * - /Script/CoreUObject.Outer0:Outer1.Outer2.StructName -> StructName(/Script/CoreUObject,Outer0,Outer1,Outer2)
  *
  * Examples:
  * - int32 -> IntProperty
- * - TArray<int32> -> ArrayProperty<IntProperty>
- * - TArray<FStructType> -> ArrayProperty<StructProperty<StructType</Script/Module>,e21f566f-7153-433a-959d-bfb3abed17e2>>
- * - TMap<FKeyStruct, EByteEnum> -> MapProperty<StructProperty<KeyStruct</Script/Module>>,EnumProperty<ByteEnum</Script/Module>,ByteProperty>>
+ * - TArray(int32) -> ArrayProperty(IntProperty)
+ * - TArray(FStructType) -> ArrayProperty(StructProperty(StructType(/Script/Module),e21f566f-7153-433a-959d-bfb3abed17e2))
+ * - TMap(FKeyStruct, EByteEnum) -> MapProperty(StructProperty(KeyStruct(/Script/Module)),EnumProperty(ByteEnum(/Script/Module),ByteProperty))
  *
  * Incomplete property types created from sources with incomplete type information. Consumers must support this.
- * - ArrayProperty<StructProperty>
- * - MapProperty<EnumProperty,StructProperty>
+ * - ArrayProperty(StructProperty)
+ * - MapProperty(EnumProperty,StructProperty)
  */
 class FPropertyTypeName
 {
@@ -52,7 +52,7 @@ public:
 	/**
 	 * Returns the type at the root of this property type name.
 	 *
-	 * Example: MapProperty<StructProperty<KeyStruct>,EnumProperty<ByteEnum,ByteProperty>>
+	 * Example: MapProperty(StructProperty(KeyStruct),EnumProperty(ByteEnum,ByteProperty))
 	 * - GetName() -> MapProperty
 	 */
 	UE_API FName GetName() const;
@@ -60,7 +60,7 @@ public:
 	/**
 	 * Returns the number of type parameters under the root of this property type name.
 	 *
-	 * Example: MapProperty<StructProperty<KeyStruct>,EnumProperty<ByteEnum,ByteProperty>>
+	 * Example: MapProperty(StructProperty(KeyStruct),EnumProperty(ByteEnum,ByteProperty))
 	 * - GetParameterCount() -> 2
 	 */
 	UE_API int32 GetParameterCount() const;
@@ -70,9 +70,9 @@ public:
 	 *
 	 * An out-of-bounds index will return an empty type name.
 	 *
-	 * Example: MapProperty<StructProperty<KeyStruct>,EnumProperty<ByteEnum,ByteProperty>>
-	 * - GetParameter(0) -> StructProperty<KeyStruct>
-	 * - GetParameter(1) -> EnumProperty<ByteEnum,ByteProperty>
+	 * Example: MapProperty(StructProperty(KeyStruct),EnumProperty(ByteEnum,ByteProperty))
+	 * - GetParameter(0) -> StructProperty(KeyStruct)
+	 * - GetParameter(1) -> EnumProperty(ByteEnum,ByteProperty)
 	 */
 	UE_API FPropertyTypeName GetParameter(int32 ParamIndex = 0) const;
 
@@ -81,7 +81,7 @@ public:
 	 *
 	 * An out-of-bounds index will return a name of None.
 	 *
-	 * Example: MapProperty<StructProperty<KeyStruct>,EnumProperty<ByteEnum,ByteProperty>>
+	 * Example: MapProperty(StructProperty(KeyStruct),EnumProperty(ByteEnum,ByteProperty))
 	 * - GetParameterName(0) -> StructProperty
 	 * - GetParameterName(1) -> EnumProperty
 	 */
@@ -117,7 +117,7 @@ private:
 /**
  * Builder for FPropertyTypeName.
  *
- * Example: MapProperty<StructProperty<KeyStruct</Script/CoreUObject>,StructGuid>,EnumProperty<ByteEnum</Script/CoreUObject>,ByteProperty>>
+ * Example: MapProperty(StructProperty(KeyStruct(/Script/CoreUObject),StructGuid),EnumProperty(ByteEnum(/Script/CoreUObject),ByteProperty))
  *
  * FPropertyTypeNameBuilder Builder;
  * Builder.AddName(NAME_MapProperty));
@@ -145,7 +145,7 @@ public:
 	/** Add a guid in the format DigitsWithHyphensLower. */
 	UE_API void AddGuid(const FGuid& Guid);
 
-	/** Add a path in the format ObjectName</Path/Name,Outer0,Outer1,Outer2> */
+	/** Add a path in the format ObjectName(/Path/Name,Outer0,Outer1,Outer2) */
 	UE_API void AddPath(const UField* Field);
 
 	/** Add a type name with its parameters. */
diff --git a/Engine/Source/Runtime/CoreUObject/Tests/UObject/PropertyPathNameTest.cpp b/Engine/Source/Runtime/CoreUObject/Tests/UObject/PropertyPathNameTest.cpp
--- a/Engine/Source/Runtime/CoreUObject/Tests/UObject/PropertyPathNameTest.cpp
+++ b/Engine/Source/Runtime/CoreUObject/Tests/UObject/PropertyPathNameTest.cpp
@@ -85,7 +85,7 @@ TEST_CASE_NAMED(FPropertyPathNameTest, "CoreUObject::PropertyPathName", "[Core][
 		PathName.Push({TEXT("X"), IntType});
 
 		CHECK(PathName.GetSegmentCount() == 3);
-		CHECK(TEXTVIEW("TestArray[7] (ArrayProperty<StructProperty<TestType>>) -> Position (StructProperty<IntVector>) -> X (IntProperty)").Equals(WriteToString<128>(PathName)));
+		CHECK(TEXTVIEW("TestArray[7] (ArrayProperty(StructProperty(TestType))) -> Position (StructProperty(IntVector)) -> X (IntProperty)").Equals(WriteToString<128>(PathName)));
 	}
 
 	SECTION("SetType/SetIndex")
diff --git a/Engine/Source/Runtime/CoreUObject/Tests/UObject/PropertyTypeNameTest.cpp b/Engine/Source/Runtime/CoreUObject/Tests/UObject/PropertyTypeNameTest.cpp
--- a/Engine/Source/Runtime/CoreUObject/Tests/UObject/PropertyTypeNameTest.cpp
+++ b/Engine/Source/Runtime/CoreUObject/Tests/UObject/PropertyTypeNameTest.cpp
@@ -386,8 +386,8 @@ TEST_CASE_NAMED(FPropertyTypeNameSmokeTest, "CoreUObject::PropertyTypeName::Smok
 
 		CHECK(TEXTVIEW("None").Equals(EmptyString));
 		CHECK(TEXTVIEW("IntProperty").Equals(IntString));
-		CHECK(TEXTVIEW("MapProperty<EnumProperty<Key,ByteProperty>,EnumProperty<Value,ByteProperty>>").Equals(EnumMapString));
-		CHECK(TEXTVIEW("ArrayProperty<StructProperty<Vector>>").Equals(VectorArrayString));
+		CHECK(TEXTVIEW("MapProperty(EnumProperty(Key,ByteProperty),EnumProperty(Value,ByteProperty))").Equals(EnumMapString));
+		CHECK(TEXTVIEW("ArrayProperty(StructProperty(Vector))").Equals(VectorArrayString));
 
 		CHECK(Empty == Parse(EmptyString));
 		CHECK(Int == Parse(IntString));
@@ -399,27 +399,27 @@ TEST_CASE_NAMED(FPropertyTypeNameSmokeTest, "CoreUObject::PropertyTypeName::Smok
 	{
 		// There are positive parsing tests in the String section above.
 		CHECK(Parse(TEXTVIEW(" \t IntProperty \t ")) == PropertyTypeNameTest::CreateInt());
-		CHECK(Parse(TEXTVIEW(" \t ArrayProperty < StructProperty\t<\tVector\t>\t > \t ")) == PropertyTypeNameTest::CreateVectorArray());
+		CHECK(Parse(TEXTVIEW(" \t ArrayProperty ( StructProperty\t(\tVector\t)\t ) \t ")) == PropertyTypeNameTest::CreateVectorArray());
 
 		CHECK_FALSE(TryParse(TEXTVIEW("")));
 		CHECK_FALSE(TryParse(TEXTVIEW(",")));
-		CHECK_FALSE(TryParse(TEXTVIEW("<>")));
+		CHECK_FALSE(TryParse(TEXTVIEW("()")));
 		CHECK_FALSE(TryParse(TEXTVIEW(",IntProperty")));
 		CHECK_FALSE(TryParse(TEXTVIEW("IntProperty,")));
 		CHECK_FALSE(TryParse(TEXTVIEW("IntProperty,IntProperty")));
-		CHECK_FALSE(TryParse(TEXTVIEW("IntProperty>")));
-		CHECK_FALSE(TryParse(TEXTVIEW("<IntProperty>")));
-		CHECK_FALSE(TryParse(TEXTVIEW("<IntProperty")));
-		CHECK_FALSE(TryParse(TEXTVIEW("ArrayProperty<IntProperty")));
-		CHECK_FALSE(TryParse(TEXTVIEW("ArrayProperty<IntProperty>>")));
-		CHECK_FALSE(TryParse(TEXTVIEW("ArrayProperty<<IntProperty>")));
-		CHECK_FALSE(TryParse(TEXTVIEW("ArrayProperty<IntProperty>IntProperty")));
-		CHECK_FALSE(TryParse(TEXTVIEW("MapProperty<>")));
-		CHECK_FALSE(TryParse(TEXTVIEW("MapProperty<,>")));
-		CHECK_FALSE(TryParse(TEXTVIEW("MapProperty<IntProperty,>")));
-		CHECK_FALSE(TryParse(TEXTVIEW("MapProperty<,IntProperty>")));
-		CHECK_FALSE(TryParse(TEXTVIEW("MapProperty<IntProperty,,IntProperty>")));
-		CHECK_FALSE(TryParse(TEXTVIEW("MapProperty<IntProperty><IntProperty>")));
+		CHECK_FALSE(TryParse(TEXTVIEW("IntProperty)")));
+		CHECK_FALSE(TryParse(TEXTVIEW("(IntProperty)")));
+		CHECK_FALSE(TryParse(TEXTVIEW("(IntProperty")));
+		CHECK_FALSE(TryParse(TEXTVIEW("ArrayProperty(IntProperty")));
+		CHECK_FALSE(TryParse(TEXTVIEW("ArrayProperty(IntProperty))")));
+		CHECK_FALSE(TryParse(TEXTVIEW("ArrayProperty((IntProperty)")));
+		CHECK_FALSE(TryParse(TEXTVIEW("ArrayProperty(IntProperty)IntProperty")));
+		CHECK_FALSE(TryParse(TEXTVIEW("MapProperty()")));
+		CHECK_FALSE(TryParse(TEXTVIEW("MapProperty(,)")));
+		CHECK_FALSE(TryParse(TEXTVIEW("MapProperty(IntProperty,)")));
+		CHECK_FALSE(TryParse(TEXTVIEW("MapProperty(,IntProperty)")));
+		CHECK_FALSE(TryParse(TEXTVIEW("MapProperty(IntProperty,,IntProperty)")));
+		CHECK_FALSE(TryParse(TEXTVIEW("MapProperty(IntProperty)(IntProperty)")));
 	}
 }
 
