Fix to properly create multiple export objects with the same type when placeholder types are enabled and the load class import cannot be resolved (editor only).

Also made a few minor API tweaks (mostly cosmetic) and restored IsPropertyBagPlaceholderType() since I chose not to add an explicit class flag for it.

#rnx

diff --git a/Engine/Source/Runtime/CoreUObject/Internal/UObject/PropertyBagRepository.h b/Engine/Source/Runtime/CoreUObject/Internal/UObject/PropertyBagRepository.h
--- a/Engine/Source/Runtime/CoreUObject/Internal/UObject/PropertyBagRepository.h
+++ b/Engine/Source/Runtime/CoreUObject/Internal/UObject/PropertyBagRepository.h
@@ -51,7 +51,7 @@ private:
 	TMap<const UObjectBase*, TObjectPtr<UObject>> Namespaces;
 
 	/** Internal registry that tracks the current set of types for property bag container objects instanced as placeholders for package exports that have invalid or missing class imports on load. */
-	TUniquePtr<class FPropertyBagTypeRegistry> PropertyBagTypeRegistry;
+	TUniquePtr<class FPropertyBagPlaceholderTypeRegistry> PropertyBagPlaceholderTypeRegistry;
 
 	FPropertyBagRepository();
 
@@ -113,13 +113,21 @@ public:
 	virtual FString GetReferencerName() const override;
 	// End FGCObject interface
 
-	// add a new placeholder type to swap in for a missing type on load; this will be associated with a property bag when instances are serialized so we don't lose its data
-	static void AddPropertyBagPlaceholderType(UClass* ClassType);
+	// query for whether or not the given struct/class is a placeholder type
+	static COREUOBJECT_API bool IsPropertyBagPlaceholderType(UStruct* Type);
 	// query for whether or not the given object was created as a placeholder type
 	static COREUOBJECT_API bool IsPropertyBagPlaceholderObject(UObject* Object);
 	// query for whether or not creating property bag placeholder objects should be allowed
 	static COREUOBJECT_API bool IsPropertyBagPlaceholderObjectSupportEnabled();
 
+	// create a new placeholder type object to swap in for a missing class/struct; this will be associated with a property bag when objects are serialized so we don't lose data
+	static COREUOBJECT_API UStruct* CreatePropertyBagPlaceholderType(UObject* Outer, UClass* Class, FName Name = NAME_None, EObjectFlags Flags = RF_NoFlags, UStruct* SuperStruct = nullptr);
+	template<typename T = UObject>
+	static UClass* CreatePropertyBagPlaceholderClass(UObject* Outer, UClass* Class, FName Name = NAME_None, EObjectFlags Flags = RF_NoFlags)
+	{
+		return Cast<UClass>(CreatePropertyBagPlaceholderType(Outer, Class, Name, Flags, T::StaticClass()));
+	}
+
 private:
 	void Lock() const { CriticalSection.Lock(); }
 	void Unlock() const { CriticalSection.Unlock(); }
diff --git a/Engine/Source/Runtime/CoreUObject/Private/UObject/LinkerLoad.cpp b/Engine/Source/Runtime/CoreUObject/Private/UObject/LinkerLoad.cpp
--- a/Engine/Source/Runtime/CoreUObject/Private/UObject/LinkerLoad.cpp
+++ b/Engine/Source/Runtime/CoreUObject/Private/UObject/LinkerLoad.cpp
@@ -4090,28 +4090,11 @@ UClass* FLinkerLoad::TryCreatePlaceholderTypeForExport(int32 ExportIndex)
 			{
 				if (UClass* LoadClassType = FindObjectFast<UClass>(LoadClassTypePackage, LoadClassImport.ClassName, /*bExactClass =*/ false))
 				{
-					// Create an opaque, non-native subtype that has no reflected properties. 
-					LoadClass = NewObject<UClass>(LoadClassParent, LoadClassType, LoadClassImport.ObjectName);
-					LoadClass->SetSuperStruct(UObject::StaticClass());
-					LoadClass->Bind();
-					LoadClass->StaticLink(/*bRelinkExistingProperties =*/ true);
-
-					// Create and configure its CDO as if it were loaded - for non-native class types, this is required.
-					UObject* LoadClassDefaults = LoadClass->GetDefaultObject();
-					LoadClass->PostLoadDefaultObject(LoadClassDefaults);
-
-					// This class is for internal use and should not be exposed for selection or instancing in the editor.
-					LoadClass->ClassFlags |= CLASS_Hidden | CLASS_HideDropDown;
+					// Create an opaque, non-native transient type object that has no reflected properties.
+					LoadClass = UE::FPropertyBagRepository::CreatePropertyBagPlaceholderClass(LoadClassParent, LoadClassType, LoadClassImport.ObjectName, RF_Transient);
 
 					// Patch it into the import table so that we resolve to this class for any future exports of this type.
 					LoadClassImport.XObject = LoadClass;
-
-					// Modify the export's object flags for instancing to indicate that its load class is a placeholder type.
-					Export.ObjectFlags |= RF_HasPlaceholderType;
-
-					// Use the property bag repository for now to manage property bag placeholder types (e.g. object lifetime).
-					// Note: The object lifetime of instances of this type will rely on existing references that are serialized.
-					UE::FPropertyBagRepository::AddPropertyBagPlaceholderType(LoadClass);
 				}
 			}
 		}
@@ -5090,6 +5073,11 @@ UObject* FLinkerLoad::CreateExport( int32 Index )
 		{
 			return nullptr;
 		}
+		else if (LoadClass && UE::FPropertyBagRepository::IsPropertyBagPlaceholderType(LoadClass))
+		{
+			// Modify the export's object flags for instancing to indicate that it has a placeholder type.
+			Export.ObjectFlags |= RF_HasPlaceholderType;
+		}
 #endif
 		if( !LoadClass )
 		{
@@ -6222,7 +6210,7 @@ FArchive& FLinkerLoad::operator<<(FObjectPtr& ObjectPtr)
 			// as also being a placeholder instance. That's because there are certain paths that need to be able
 			// to resolve the pointer (e.g. - the object initialization path during class construction). However,
 			// it also means we can't resolve other references to a placeholder CDO, as they may not be type-safe.
-			if (ResolvedObject->HasAnyFlags(RF_ClassDefaultObject))
+			if (UNLIKELY(ResolvedObject->HasAnyFlags(RF_ClassDefaultObject)))
 #endif
 			{
 				UE_LOG(LogLinker, Warning, TEXT("Serializing reference to \"%s\" as NULL to ensure type safety."), *ResolvedObject->GetPathName());
diff --git a/Engine/Source/Runtime/CoreUObject/Private/UObject/PropertyBagRepository.cpp b/Engine/Source/Runtime/CoreUObject/Private/UObject/PropertyBagRepository.cpp
--- a/Engine/Source/Runtime/CoreUObject/Private/UObject/PropertyBagRepository.cpp
+++ b/Engine/Source/Runtime/CoreUObject/Private/UObject/PropertyBagRepository.cpp
@@ -24,7 +24,7 @@ DEFINE_LOG_CATEGORY_STATIC(LogPropertyBagRepository, Log, All);
 namespace UE
 {
 
-class FPropertyBagTypeRegistry
+class FPropertyBagPlaceholderTypeRegistry
 {
 public:
 	void AddReferencedObjects(FReferenceCollector& Collector)
@@ -33,15 +33,15 @@ public:
 		Collector.AddReferencedObjects(PlaceholderTypes);
 	}
 
-	void Add(UClass* Class)
+	void Add(UStruct* Type)
 	{
-		PendingPlaceholderTypes.Enqueue(Class);
+		PendingPlaceholderTypes.Enqueue(Type);
 	}
 
-	bool Contains(UClass* Class)
+	bool Contains(UStruct* Type)
 	{
 		ConsumePendingPlaceholderTypes();
-		return PlaceholderTypes.Contains(Class);
+		return PlaceholderTypes.Contains(Type);
 	}
 
 protected:
@@ -51,10 +51,10 @@ protected:
 		{
 			FScopeLock ScopeLock(&CriticalSection);
 
-			TObjectPtr<UClass> Class;
-			while(PendingPlaceholderTypes.Dequeue(Class))
+			TObjectPtr<UStruct> PendingType;
+			while(PendingPlaceholderTypes.Dequeue(PendingType))
 			{
-				PlaceholderTypes.Add(Class);
+				PlaceholderTypes.Add(PendingType);
 			}
 		}
 	}
@@ -63,10 +63,10 @@ private:
 	FCriticalSection CriticalSection;
 
 	// List of types that have been registered.
-	TSet<TObjectPtr<UClass>> PlaceholderTypes;
+	TSet<TObjectPtr<UStruct>> PlaceholderTypes;
 
 	// Types that have been added but not yet registered. Utilizes a thread-safe queue so we can avoid race conditions during an async load.
-	TQueue<TObjectPtr<UClass>> PendingPlaceholderTypes;
+	TQueue<TObjectPtr<UStruct>> PendingPlaceholderTypes;
 };
 
 class FPropertyBagRepositoryLock
@@ -121,7 +121,7 @@ FPropertyBagRepository& FPropertyBagRepository::Get()
 
 FPropertyBagRepository::FPropertyBagRepository()
 {
-	PropertyBagTypeRegistry = MakeUnique<FPropertyBagTypeRegistry>();
+	PropertyBagPlaceholderTypeRegistry = MakeUnique<FPropertyBagPlaceholderTypeRegistry>();
 }
 
 void FPropertyBagRepository::ReassociateObjects(const TMap<UObject*, UObject*>& ReplacedObjects)
@@ -256,7 +256,7 @@ void FPropertyBagRepository::AddReferencedObjects(FReferenceCollector& Collector
 		Collector.AddReferencedObject(Element.Value);
 	}
 
-	PropertyBagTypeRegistry->AddReferencedObjects(Collector);
+	PropertyBagPlaceholderTypeRegistry->AddReferencedObjects(Collector);
 }
 
 FString FPropertyBagRepository::GetReferencerName() const
@@ -330,14 +330,14 @@ void FPropertyBagRepository::ShrinkMaps()
 	AssociatedData.Compact();
 }
 
-void FPropertyBagRepository::AddPropertyBagPlaceholderType(UClass* ClassType)
+bool FPropertyBagRepository::IsPropertyBagPlaceholderType(UStruct* Type)
 {
-	if (!ClassType)
+	if (!Type)
 	{
-		return;
+		return false;
 	}
 
-	FPropertyBagRepository::Get().PropertyBagTypeRegistry->Add(ClassType);
+	return FPropertyBagRepository::Get().PropertyBagPlaceholderTypeRegistry->Contains(Type);
 }
 
 bool FPropertyBagRepository::IsPropertyBagPlaceholderObject(UObject* Object)
@@ -348,7 +348,7 @@ bool FPropertyBagRepository::IsPropertyBagPlaceholderObject(UObject* Object)
 	}
 
 	return Object->HasAnyFlags(RF_HasPlaceholderType|RF_ClassDefaultObject)
-		&& FPropertyBagRepository::Get().PropertyBagTypeRegistry->Contains(Object->GetClass());
+		&& IsPropertyBagPlaceholderType(Object->GetClass());
 }
 
 namespace Private
@@ -380,4 +380,29 @@ bool FPropertyBagRepository::IsPropertyBagPlaceholderObjectSupportEnabled()
 #endif
 }
 
+UStruct* FPropertyBagRepository::CreatePropertyBagPlaceholderType(UObject* Outer, UClass* Class, FName Name, EObjectFlags Flags, UStruct* SuperStruct)
+{
+	UStruct* PlaceholderType = NewObject<UClass>(Outer, Class, Name, Flags);
+	PlaceholderType->SetSuperStruct(SuperStruct);
+	PlaceholderType->Bind();
+	PlaceholderType->StaticLink(/*bRelinkExistingProperties =*/ true);
+
+	// Extra configuration needed for class types.
+	if (UClass* PlaceholderTypeAsClass = Cast<UClass>(PlaceholderType))
+	{
+		// Create and configure its CDO as if it were loaded - for non-native class types, this is required.
+		UObject* PlaceholderClassDefaults = PlaceholderTypeAsClass->GetDefaultObject();
+		PlaceholderTypeAsClass->PostLoadDefaultObject(PlaceholderClassDefaults);
+
+		// This class is for internal use and should not be exposed for selection or instancing in the editor.
+		PlaceholderTypeAsClass->ClassFlags |= CLASS_Hidden | CLASS_HideDropDown;
+	}
+
+	// Use the property bag repository for now to manage property bag placeholder types (e.g. object lifetime).
+	// Note: The object lifetime of instances of this type will rely on existing references that are serialized.
+	FPropertyBagRepository::Get().PropertyBagPlaceholderTypeRegistry->Add(PlaceholderType);
+
+	return PlaceholderType;
+}
+
 } // UE
