Add StaticLoadAsset and LoadAssetAsync functions to UObjectGlobals which are wrappers around the existing package/object load functions.
Add FSoftObjectPath::LoadAsync which calls LoadAssetAsync with additional handling for subobject paths.
Fix some inaccurate comments and document which LoadObject parameters are deprecated.

diff --git a/Engine/Source/Runtime/CoreUObject/Private/UObject/SoftObjectPath.cpp b/Engine/Source/Runtime/CoreUObject/Private/UObject/SoftObjectPath.cpp
--- a/Engine/Source/Runtime/CoreUObject/Private/UObject/SoftObjectPath.cpp
+++ b/Engine/Source/Runtime/CoreUObject/Private/UObject/SoftObjectPath.cpp
@@ -557,6 +557,35 @@ UObject* FSoftObjectPath::TryLoad(FUObjectSerializeContext* InLoadContext) const
 	return LoadedObject;
 }
 
+int32 FSoftObjectPath::LoadAsync(FLoadSoftObjectPathAsyncDelegate InCompletionDelegate, FLoadAssetAsyncOptionalParams InOptionalParams)
+{
+	FSoftObjectPath RequestedPath = *this;
+	FSoftObjectPath PathToLoad = RequestedPath;
+#if WITH_EDITOR
+	if (GPlayInEditorID != INDEX_NONE)
+	{
+		// @todo: This logic may need updating to handle level instances properly and we may want to handle other fixups like CoreRedirects before requesting
+		PathToLoad.FixupForPIE();
+	}
+#endif
+
+	FLoadAssetAsyncDelegate WrapperDelegate = FLoadAssetAsyncDelegate::CreateLambda(
+		[RequestedPath, PathToLoad, CompletionDelegate = MoveTemp(InCompletionDelegate)](const FTopLevelAssetPath& InAssetPath, UObject* InLoadedObject, EAsyncLoadingResult::Type InResult) mutable
+		{
+			// If this isn't a subobject, InLoadedObject is already correct
+			if (PathToLoad.IsSubobject())
+			{
+				// Resolve the entire path, including the subobject
+				InLoadedObject = PathToLoad.ResolveObject();
+			}
+
+			// Call delegate with original requested path
+			CompletionDelegate.ExecuteIfBound(RequestedPath, InLoadedObject);
+		});
+
+	return LoadAssetAsync(PathToLoad.GetAssetPath(), MoveTemp(WrapperDelegate), MoveTemp(InOptionalParams));
+}
+
 UObject* FSoftObjectPath::ResolveObject() const
 {
 	// Don't try to resolve if we're saving a package because StaticFindObject can't be used here
diff --git a/Engine/Source/Runtime/CoreUObject/Private/UObject/UObjectGlobals.cpp b/Engine/Source/Runtime/CoreUObject/Private/UObject/UObjectGlobals.cpp
--- a/Engine/Source/Runtime/CoreUObject/Private/UObject/UObjectGlobals.cpp
+++ b/Engine/Source/Runtime/CoreUObject/Private/UObject/UObjectGlobals.cpp
@@ -1475,6 +1475,48 @@ UClass* StaticLoadClass( UClass* BaseClass, UObject* InOuter, const TCHAR* InNam
 	return Class;
 }
 
+UObject* StaticLoadAsset(UClass* Class, FTopLevelAssetPath InPath, uint32 LoadFlags, const FLinkerInstancingContext* InstancingContext)
+{
+	// @todo: This could call StaticLoadObjectInternal directly with some refactoring
+
+	TStringBuilder<256> ObjectNameString;
+	InPath.AppendString(ObjectNameString);
+	return StaticLoadObject(Class, nullptr, *ObjectNameString, nullptr, LoadFlags, nullptr, true, InstancingContext);
+}
+
+int32 LoadAssetAsync(FTopLevelAssetPath InAssetPath, FLoadAssetAsyncDelegate InCompletionDelegate, FLoadAssetAsyncOptionalParams InOptionalParams)
+{
+	// Asset paths should always have a valid package
+	FPackagePath PackagePath = FPackagePath::FromPackageNameChecked(InAssetPath.GetPackageName());
+	
+	FLoadPackageAsyncOptionalParams PackageParams{
+		.PackagePriority = InOptionalParams.PackagePriority,
+		.InstancingContext = InOptionalParams.InstancingContext,
+		.LoadFlags = InOptionalParams.LoadFlags
+	};
+
+	PackageParams.CompletionDelegate = MakeUnique<FLoadPackageAsyncDelegate>(FLoadPackageAsyncDelegate::CreateLambda(
+		[InAssetPath, CompletionDelegate = MoveTemp(InCompletionDelegate)](const FName& LoadedPackageName, UPackage* LoadedPackage, EAsyncLoadingResult::Type Result) mutable
+		{
+			UObject* LoadedObject = nullptr;
+			if (Result == EAsyncLoadingResult::Succeeded && LoadedPackage)
+			{
+				LoadedObject = StaticFindObjectFast(UObject::StaticClass(), LoadedPackage, InAssetPath.GetAssetName(), false);
+
+				// Package loaded but object was not found inside it, failure
+				if (!LoadedObject)
+				{
+					Result = EAsyncLoadingResult::Failed;
+				}
+			}
+
+			CompletionDelegate.ExecuteIfBound(InAssetPath, LoadedObject, Result);
+		}));
+
+	return LoadPackageAsync(PackagePath, MoveTemp(PackageParams));
+}
+
+
 #if WITH_EDITOR
 #include "Containers/StackTracker.h"
 class FDiffFileArchive : public FArchiveProxy
diff --git a/Engine/Source/Runtime/CoreUObject/Public/UObject/SoftObjectPath.h b/Engine/Source/Runtime/CoreUObject/Public/UObject/SoftObjectPath.h
--- a/Engine/Source/Runtime/CoreUObject/Public/UObject/SoftObjectPath.h
+++ b/Engine/Source/Runtime/CoreUObject/Public/UObject/SoftObjectPath.h
@@ -32,6 +32,9 @@ class FOutputDevice;
 struct FPropertyTag;
 struct FUObjectSerializeContext;
 
+/** Delegate called on completion of async loading a soft object. The UObject will be null if the load failed */
+DECLARE_DELEGATE_TwoParams(FLoadSoftObjectPathAsyncDelegate, const FSoftObjectPath&, UObject*);
+
 /**
  * A struct that contains a string reference to an object, either a package, a top level asset or a subobject.
  * This can be used to make soft references to assets that are loaded on demand.
@@ -212,6 +215,17 @@ PRAGMA_ENABLE_DEPRECATION_WARNINGS;
 	 */
 	COREUOBJECT_API UObject* TryLoad(FUObjectSerializeContext* InLoadContext = nullptr) const;
 
+	/**
+	 * Attempts to asynchronously load the object referenced by this path.
+	 * This will call LoadAssetAsync to load the top level asset and resolve the sub paths before calling the delegate.
+	 * FStreamableManager can be used for more control over callback timing and garbage collection.
+	 * 
+	 * @param	InCompletionDelegate	Delegate to be invoked when the async load finishes, this will execute on the game thread as soon as the load succeeds or fails
+	 * @param	InOptionalParams		Optional parameters for async loading the asset
+	 * @return Unique ID associated with this load request (the same object or package can be associated with multiple IDs).
+	 */
+	COREUOBJECT_API int32 LoadAsync(FLoadSoftObjectPathAsyncDelegate InCompletionDelegate, FLoadAssetAsyncOptionalParams InOptionalParams = FLoadAssetAsyncOptionalParams());
+
 	/**
 	 * Attempts to find a currently loaded object that matches this path
 	 *
diff --git a/Engine/Source/Runtime/CoreUObject/Public/UObject/UObjectGlobals.h b/Engine/Source/Runtime/CoreUObject/Public/UObject/UObjectGlobals.h
--- a/Engine/Source/Runtime/CoreUObject/Public/UObject/UObjectGlobals.h
+++ b/Engine/Source/Runtime/CoreUObject/Public/UObject/UObjectGlobals.h
@@ -537,14 +537,14 @@ inline bool ParseObject( const TCHAR* Stream, const TCHAR* Match, UClass* Class,
  * Find or load an object by string name with optional outer and filename specifications.
  * These are optional because the InName can contain all of the necessary information.
  *
- * @param ObjectClass	The class (or a superclass) of the object to be loaded.
+ * @param Class			The class (or a superclass) of the object to be loaded.
  * @param InOuter		An optional object to narrow where to find/load the object from
  * @param Name			String name of the object. If it's not fully qualified, InOuter and/or Filename will be needed
- * @param Filename		An optional file to load from (or find in the file's package object)
+ * @param Filename		An optional file to load from (Deprecated parameter)
  * @param LoadFlags		Flags controlling how to handle loading from disk, from the ELoadFlags enum
- * @param Sandbox		A list of packages to restrict the search for the object
- * @param bAllowObjectReconciliation	Whether to allow the object to be found via FindObject in the case of seek free loading
- * @param InstancingContext				InstancingContext used to remap imports when loading a packager under a new name
+ * @param Sandbox		A list of packages to restrict the search for the object (Deprecated parameter)
+ * @param bAllowObjectReconciliation	Whether to allow the object to be found via FindObject before forcing a load (Deprecated parameter)
+ * @param InstancingContext				InstancingContext used to remap imports when loading a package under a new name
  *
  * @return The object that was loaded or found. nullptr for a failure.
  */
@@ -553,6 +553,18 @@ COREUOBJECT_API UObject* StaticLoadObject( UClass* Class, UObject* InOuter, cons
 /** Version of StaticLoadObject() that will load classes */
 COREUOBJECT_API UClass* StaticLoadClass(UClass* BaseClass, UObject* InOuter, const TCHAR* Name, const TCHAR* Filename = nullptr, uint32 LoadFlags = LOAD_None, UPackageMap* Sandbox = nullptr);
 
+/**
+ * Find or load an object that is one of the top level assets in a package.
+ *
+ * @param	Class				The class (or a superclass) of the object to be loaded.
+ * @param	InPath				FName pair representing the outer package object and the inner top level object (asset)
+ * @param	LoadFlags			Flags controlling how to handle loading from disk, from the ELoadFlags enum
+ * @param	InstancingContext	InstancingContext used to remap imports when loading a package under a new name
+ *
+ * @return	Returns a pointer to the found object or nullptr if none could be found
+ */
+COREUOBJECT_API UObject* StaticLoadAsset(UClass* Class, FTopLevelAssetPath InPath, uint32 LoadFlags = LOAD_None, const FLinkerInstancingContext* InstancingContext = nullptr);
+
 /**
  * Create a new instance of an object.  The returned object will be fully initialized.  If InFlags contains RF_NeedsLoad (indicating that the object still needs to load its object data from disk), components
  * are not instanced (this will instead occur in PostLoad()).  The different between StaticConstructObject and StaticAllocateObject is that StaticConstructObject will also call the class constructor on the object
@@ -648,21 +660,21 @@ COREUOBJECT_API UPackage* LoadPackage( UPackage* InOuter, const TCHAR* InLongPac
  */
 COREUOBJECT_API UPackage* LoadPackage(UPackage* InOuter, const FPackagePath& InPackagePath, uint32 LoadFlags, FArchive* InReaderOverride = nullptr, const FLinkerInstancingContext* InstancingContext = nullptr, const FPackagePath* DiffPackagePath = nullptr);
 
-/** Async package loading result */
+/** Async package and object loading result */
 namespace EAsyncLoadingResult
 {
 	enum Type
 	{
-		/** Package failed to load */
+		/** Package or object failed to load */
 		Failed,
-		/** Package loaded successfully */
+		/** Package or object loaded successfully */
 		Succeeded,
 		/** Async loading was canceled */
 		Canceled
 	};
 }
 
-/** Async package loading result */
+/** Async loading progress for a specific package */
 enum class EAsyncLoadingProgress : uint32
 {
 	/** Package failed to load */
@@ -793,6 +805,40 @@ COREUOBJECT_API int32 LoadPackageAsync(const FString& InName, FLoadPackageAsyncD
  */
 COREUOBJECT_API int32 LoadPackageAsync(const FString& InName, FLoadPackageAsyncOptionalParams InOptionalParams);
 
+
+/**
+ * Delegate called on completion of async asset loading
+ * @param	AssetPath			Path of the asset we were trying to load
+ * @param	LoadedObject		Loaded object if successful, nullptr otherwise
+ * @param	Result				Result of async loading.
+ */
+DECLARE_DELEGATE_ThreeParams(FLoadAssetAsyncDelegate, const FTopLevelAssetPath& /*AssetPath*/, UObject* /*LoadedObject*/, EAsyncLoadingResult::Type /*Result*/)
+
+/**
+ * Optional parameters passed to the LoadAssetAsync function.
+ */
+struct FLoadAssetAsyncOptionalParams
+{
+	/** Loading priority. **/
+	int32 PackagePriority { 0 };
+	/** Additional context to map object names to their instanced counterpart when loading an instanced package. **/
+	const FLinkerInstancingContext* InstancingContext { nullptr };
+	/** Flags controlling loading behavior, from the ELoadFlags enum. */
+	uint32 LoadFlags { LOAD_None };
+};
+
+/**
+ * Asynchronously load a top level asset along with other objects in the same package. This is non-blocking and will call LoadPackageAsync.
+ * FSoftObjectPath::LoadAsync can be used to asynchronously load subobjects.
+ *
+ * @param	InAssetPath				Top level asset to load
+ * @param	InCompletionDelegate	Delegate to be invoked when the async load finishes, this will execute on the game thread as soon as the load succeeds or fails
+ * @param	InOptionalParams		Optional parameters 
+ * @return Unique ID associated with this load request (the same object or package can be associated with multiple IDs).
+ */
+COREUOBJECT_API int32 LoadAssetAsync(FTopLevelAssetPath InAssetPath, FLoadAssetAsyncDelegate InCompletionDelegate, FLoadAssetAsyncOptionalParams InOptionalParams = FLoadAssetAsyncOptionalParams());
+
+
 /**
 * Cancels all async package loading requests.
 */
