Added ApplyRedirectsToPropertyType to handle core redirects for FPropertyTypeName

#rnx

diff --git a/Engine/Source/Runtime/CoreUObject/Private/UObject/PropertyHelper.cpp b/Engine/Source/Runtime/CoreUObject/Private/UObject/PropertyHelper.cpp
--- a/Engine/Source/Runtime/CoreUObject/Private/UObject/PropertyHelper.cpp
+++ b/Engine/Source/Runtime/CoreUObject/Private/UObject/PropertyHelper.cpp
@@ -2,14 +2,19 @@
 
 #include "UObject/PropertyHelper.h"
 
-#include "CoreMinimal.h"
-#include "UObject/ObjectMacros.h"
-#include "UObject/Object.h"
+#include "String/Find.h"
+#include "String/ParseTokens.h"
 #include "UObject/Class.h"
 #include "UObject/CoreNet.h"
-#include "UObject/UnrealType.h"
+#include "UObject/CoreRedirects.h"
+#include "UObject/EnumProperty.h"
+#include "UObject/Object.h"
+#include "UObject/ObjectMacros.h"
 #include "UObject/Package.h"
+#include "UObject/PropertyOptional.h"
 #include "UObject/PropertyPortFlags.h"
+#include "UObject/PropertyTypeName.h"
+#include "UObject/UnrealType.h"
 
 void SkipWhitespace(const TCHAR*& Str)
 {
@@ -283,3 +288,203 @@ const TCHAR* DelegatePropertyTools::ImportDelegateFromText( FScriptDelegate& Del
 
 	return ( Func != nullptr && Object != nullptr ) ? Buffer : nullptr;
 }
+
+namespace UE
+{
+
+static FCoreRedirectObjectName BuildCoreRedirectObjectName(FPropertyTypeName TypeName)
+{
+	FName OuterName;
+	if (const int32 OuterCount = TypeName.GetParameterCount() - 1; OuterCount <= 1)
+	{
+		OuterName = TypeName.GetParameterName(1);
+	}
+	else
+	{
+		TStringBuilder<256> OuterChain;
+		for (int32 OuterIndex = 0; OuterIndex < OuterCount; ++OuterIndex)
+		{
+			OuterChain << TypeName.GetParameterName(OuterIndex + 1) << (OuterIndex == 0 ? TEXT(':') : TEXT('.'));
+		}
+		OuterChain.RemoveSuffix(1);
+		OuterName = FName(OuterChain);
+	}
+
+	const FName ObjectName = TypeName.GetName();
+	const FName PackageName = TypeName.GetParameterName(0);
+	checkfSlow(PackageName.IsNone() || String::FindFirstChar(WriteToString<256>(PackageName), TEXT('/')) != INDEX_NONE,
+		TEXT("PropertyTypeName with an outer name but not a package name is not supported yet."));
+	return FCoreRedirectObjectName(ObjectName, OuterName, PackageName);
+}
+
+static void BuildPropertyName(const FCoreRedirectObjectName& Redirect, FPropertyTypeNameBuilder& Builder)
+{
+	Builder.AddName(Redirect.ObjectName);
+	if (!Redirect.PackageName.IsNone() || !Redirect.OuterName.IsNone())
+	{
+		Builder.BeginParameters();
+		if (!Redirect.PackageName.IsNone())
+		{
+			Builder.AddName(Redirect.PackageName);
+		}
+		if (!Redirect.OuterName.IsNone())
+		{
+			TStringBuilder<256> OuterChain(InPlace, Redirect.OuterName);
+			String::ParseTokensMultiple(OuterChain, {TEXT('.'), TEXT(':')}, [&Builder](FStringView Token)
+			{
+				Builder.AddName(FName(Token));
+			});
+		}
+		Builder.EndParameters();
+	}
+}
+
+template <typename PropertyType>
+struct TFindRedirectForPropertyTraits;
+
+template <>
+struct TFindRedirectForPropertyTraits<FStructProperty>
+{
+	static const UField* GetField(const FStructProperty* Property) { return Property->Struct; }
+	static ECoreRedirectFlags GetFlags() { return ECoreRedirectFlags::Type_Struct; }
+};
+
+template <>
+struct TFindRedirectForPropertyTraits<FByteProperty>
+{
+	static const UField* GetField(const FByteProperty* Property) { return Property->Enum; }
+	static ECoreRedirectFlags GetFlags() { return ECoreRedirectFlags::Type_Enum; }
+};
+
+template <>
+struct TFindRedirectForPropertyTraits<FEnumProperty>
+{
+	static const UField* GetField(const FEnumProperty* Property) { return Property->GetEnum(); }
+	static ECoreRedirectFlags GetFlags() { return ECoreRedirectFlags::Type_Enum; }
+};
+
+template <typename PropertyType>
+static bool FindRedirectForProperty(FPropertyTypeName OldType, FPropertyTypeNameBuilder& NewType, const PropertyType* Property)
+{
+	using FTraits = TFindRedirectForPropertyTraits<PropertyType>;
+	const UField* Field = Property ? FTraits::GetField(Property) : nullptr;
+	FCoreRedirectObjectName OldNameRedirect = BuildCoreRedirectObjectName(OldType);
+	FCoreRedirectObjectName NewNameRedirect = FCoreRedirects::GetRedirectedName(FTraits::GetFlags(), OldNameRedirect,
+		Field ? ECoreRedirectMatchFlags::AllowPartialMatch : ECoreRedirectMatchFlags::None);
+
+	if (NewNameRedirect == OldNameRedirect)
+	{
+		return false;
+	}
+
+	// If a partial match does not match the property then repeat the lookup without allowing partial matches.
+	// This is necessary to handle redirects of the original type of a property that changed type.
+	if (Field && Field->GetFName() != NewNameRedirect.ObjectName)
+	{
+		NewNameRedirect = FCoreRedirects::GetRedirectedName(FTraits::GetFlags(), OldNameRedirect, ECoreRedirectMatchFlags::None);
+		if (OldNameRedirect == NewNameRedirect)
+		{
+			return false;
+		}
+	}
+
+	BuildPropertyName(NewNameRedirect, NewType);
+	return true;
+}
+
+static void BuildNewNameForPropertyType(FPropertyTypeName OldTypeName, FPropertyTypeNameBuilder& NewTypeName, const FProperty* Property, bool& bHasRedirect)
+{
+	const FName OldName = OldTypeName.GetName();
+	NewTypeName.AddName(OldName);
+
+	const int32 OldParamCount = OldTypeName.GetParameterCount();
+	if (OldParamCount == 0)
+	{
+		return;
+	}
+
+	int32 OldParamIndex = 0;
+	NewTypeName.BeginParameters();
+
+	if (const EName* OldEName = OldName.ToEName(); OldEName && OldName.GetNumber() == NAME_NO_NUMBER_INTERNAL)
+	{
+		switch (*OldEName)
+		{
+		case NAME_StructProperty:
+		{
+			if (FindRedirectForProperty(OldTypeName.GetParameter(OldParamIndex), NewTypeName, CastField<FStructProperty>(Property)))
+			{
+				++OldParamIndex;
+				bHasRedirect = true;
+			}
+			break;
+		}
+		case NAME_ByteProperty:
+		{
+			if (FindRedirectForProperty(OldTypeName.GetParameter(OldParamIndex), NewTypeName, CastField<FByteProperty>(Property)))
+			{
+				++OldParamIndex;
+				bHasRedirect = true;
+			}
+			break;
+		}
+		case NAME_EnumProperty:
+		{
+			if (FindRedirectForProperty(OldTypeName.GetParameter(OldParamIndex), NewTypeName, CastField<FEnumProperty>(Property)))
+			{
+				++OldParamIndex;
+				bHasRedirect = true;
+			}
+			break;
+		}
+		case NAME_ArrayProperty:
+		{
+			const FArrayProperty* ArrayProperty = CastField<FArrayProperty>(Property);
+			const FProperty* Inner = ArrayProperty ? ArrayProperty->Inner : nullptr;
+			BuildNewNameForPropertyType(OldTypeName.GetParameter(OldParamIndex++), NewTypeName, Inner, bHasRedirect);
+			break;
+		}
+		case NAME_OptionalProperty:
+		{
+			const FOptionalProperty* OptionalProperty = CastField<FOptionalProperty>(Property);
+			const FProperty* Value = OptionalProperty ? OptionalProperty->GetValueProperty() : nullptr;
+			BuildNewNameForPropertyType(OldTypeName.GetParameter(OldParamIndex++), NewTypeName, Value, bHasRedirect);
+			break;
+		}
+		case NAME_SetProperty:
+		{
+			const FSetProperty* SetProperty = CastField<FSetProperty>(Property);
+			const FProperty* Element = SetProperty ? SetProperty->ElementProp : nullptr;
+			BuildNewNameForPropertyType(OldTypeName.GetParameter(OldParamIndex++), NewTypeName, Element, bHasRedirect);
+			break;
+		}
+		case NAME_MapProperty:
+		{
+			const FMapProperty* MapProperty = CastField<FMapProperty>(Property);
+			const FProperty* Key = MapProperty ? MapProperty->KeyProp : nullptr;
+			const FProperty* Value = MapProperty ? MapProperty->ValueProp : nullptr;
+			BuildNewNameForPropertyType(OldTypeName.GetParameter(OldParamIndex++), NewTypeName, Key, bHasRedirect);
+			BuildNewNameForPropertyType(OldTypeName.GetParameter(OldParamIndex++), NewTypeName, Value, bHasRedirect);
+			break;
+		}
+		default:
+			break;
+		}
+	}
+
+	for (; OldParamIndex < OldParamCount; ++OldParamIndex)
+	{
+		BuildNewNameForPropertyType(OldTypeName.GetParameter(OldParamIndex), NewTypeName, nullptr, bHasRedirect);
+	}
+	NewTypeName.EndParameters();
+}
+
+FPropertyTypeName ApplyRedirectsToPropertyType(FPropertyTypeName OldTypeName, const FProperty* Property)
+{
+	bool bHasRedirect = false;
+	FPropertyTypeNameBuilder NewTypeName;
+	BuildNewNameForPropertyType(OldTypeName, NewTypeName, Property, bHasRedirect);
+	return bHasRedirect ? NewTypeName.Build() : FPropertyTypeName();
+}
+
+} // UE
diff --git a/Engine/Source/Runtime/CoreUObject/Private/UObject/PropertyHelper.h b/Engine/Source/Runtime/CoreUObject/Private/UObject/PropertyHelper.h
--- a/Engine/Source/Runtime/CoreUObject/Private/UObject/PropertyHelper.h
+++ b/Engine/Source/Runtime/CoreUObject/Private/UObject/PropertyHelper.h
@@ -8,6 +8,7 @@
 
 class UFunction;
 class UObject;
+namespace UE { class FPropertyTypeName; }
 
 DECLARE_LOG_CATEGORY_EXTERN(LogProperty, Log, All);
 
@@ -72,3 +73,20 @@ namespace DelegatePropertyTools
 	 */
 	COREUOBJECT_API const TCHAR* ImportDelegateFromText( FScriptDelegate& Delegate, const UFunction* SignatureFunction, const TCHAR* Buffer, UObject* Parent, FOutputDevice* ErrorText );
 }
+
+namespace UE
+{
+
+/**
+ * Applies core redirects to type names and paths within the property type.
+ *
+ * Allows partial matches only when a property is given and matches the layout of the type name.
+ * Example: StructProperty(OldName) will not have an OldName redirect applied if Property is an FIntProperty.
+ *
+ * @param OldTypeName   The type name to apply core redirects to.
+ * @param Property      Optional. The property that the type name might correspond to.
+ * @return The new type name, if redirects were applied, otherwise an empty type name.
+ */
+FPropertyTypeName ApplyRedirectsToPropertyType(FPropertyTypeName OldTypeName, const FProperty* Property = nullptr);
+
+} // UE
