Limit recursion depth for coroutines

diff --git a/Engine/Source/Runtime/CoreUObject/Public/UObject/Script.h b/Engine/Source/Runtime/CoreUObject/Public/UObject/Script.h
--- a/Engine/Source/Runtime/CoreUObject/Public/UObject/Script.h
+++ b/Engine/Source/Runtime/CoreUObject/Public/UObject/Script.h
@@ -17,6 +17,7 @@
 
 struct FFrame;
 struct FBlueprintExceptionInfo;
+namespace verse { class task; }
 
 // It's best to set only one of these, but strictly speaking you could set both.
 // The results will be confusing. Native time would be included only in a coarse 
@@ -543,6 +544,7 @@ private:
 	// Only FFrame can modify the stack
 	friend FFrame;
 	friend void ProcessLocalScriptFunction(UObject* Context, FFrame& Stack, RESULT_DECL);
+	friend verse::task;
 };
 
 #endif // DO_BLUEPRINT_GUARD
