[ZenLoader]
  - Load the proper package instead of keeping the renamed one when we detect a reference that might be held by the deferred package delete queue.
  - This could happen when streaming out a level and then streaming it back in before the deferred queue had time to be processed.

#rnx

diff --git a/Engine/Source/Runtime/CoreUObject/Private/Serialization/AsyncLoading2.cpp b/Engine/Source/Runtime/CoreUObject/Private/Serialization/AsyncLoading2.cpp
--- a/Engine/Source/Runtime/CoreUObject/Private/Serialization/AsyncLoading2.cpp
+++ b/Engine/Source/Runtime/CoreUObject/Private/Serialization/AsyncLoading2.cpp
@@ -1030,6 +1030,8 @@ public:
 class FGlobalImportStore
 {
 private:
+	// Reference back to our parent
+	FAsyncLoadingThread2& AsyncLoadingThread;
 	// Packages in active loading or completely loaded packages, with Desc.UPackageId as key.
 	// Does not track temp packages with custom UPackage names, since they are never imorted by other packages.
 	TMap<FPackageId, FLoadedPackageRef> Packages;
@@ -1038,8 +1040,11 @@ private:
 	// All currently loaded public export objects from any loaded package
 	TMap<int32, FPublicExportKey> ObjectIndexToPublicExport;
 
+	// Process all the deferred deletion for packages that are finished loading
+	void FlushDeferredDeletePackagesQueue();
 public:
-	FGlobalImportStore()
+	FGlobalImportStore(FAsyncLoadingThread2& InAsyncLoadingThread)
+		: AsyncLoadingThread(InAsyncLoadingThread)
 	{
 		Packages.Reserve(32768);
 		ScriptObjects.Reserve(32768);
@@ -1084,6 +1089,26 @@ public:
 		LLM_SCOPE_BYNAME(TEXT("AsyncLoadPackageStore"));
 
 		FLoadedPackageRef& PackageRef = Packages.FindOrAdd(PackageId);
+
+		// Detect package that have been renamed but are still referenced and attempt to get rid of the references
+		// by flushing the deferred delete packages queue so that we can load the proper package instead.
+		if (PackageRef.RefCount > 0)
+		{
+			if (UPackage* Package = PackageRef.GetPackage())
+			{
+				if (PackageRef.GetOriginalPackageName() != Package->GetFName())
+				{
+					FlushDeferredDeletePackagesQueue();
+
+					// Flush failed to clear the remaining references, warn since we want to know about this.
+					if (PackageRef.RefCount > 0)
+					{
+						UE_LOG(LogStreaming, Warning, TEXT("Package %s was renamed to %s but is unexpectedly still being referenced by other packages being loaded"), *PackageRef.GetOriginalPackageName().ToString(), *Package->GetFName().ToString());
+					}
+				}
+			}
+		}
+
 		// is this the first reference to a package that already exists?
 		if (PackageRef.RefCount == 0)
 		{
@@ -3641,11 +3666,17 @@ private:
 	void ConditionalBeginDeferredPostLoad(FAsyncLoadingThreadState2& ThreadState, FAsyncLoadingPostLoadGroup* DeferredPostLoadGroup);
 	void MergePostLoadGroups(FAsyncLoadingThreadState2& ThreadState, FAsyncLoadingPostLoadGroup* Target, FAsyncLoadingPostLoadGroup* Source, bool bUpdateSyncLoadContext = true);
 
+public:
 	bool ProcessDeferredDeletePackagesQueue(int32 MaxCount = MAX_int32)
 	{
+		// It is only safe to call from the async loading thread or when the GC is locked so the ALT is idle.
+		check(FAsyncLoadingThreadState2::Get()->bCanAccessAsyncLoadingThreadData || FGCCSyncObject::Get().IsGCLocked());
+
 		bool bDidSomething = false;
 		if (!DeferredDeletePackages.IsEmpty())
 		{
+			TRACE_CPUPROFILER_EVENT_SCOPE(ProcessDeferredDeletePackagesQueue);
+
 			FAsyncPackage2* Package = nullptr;
 			int32 Count = 0;
 			while (++Count <= MaxCount && DeferredDeletePackages.Dequeue(Package))
@@ -3657,6 +3688,7 @@ private:
 		return bDidSomething;
 	}
 
+private:
 	void OnPreGarbageCollect();
 
 	void CollectUnreachableObjects(TArrayView<FUObjectItem*> UnreachableObjectItems, FUnreachableObjects& OutUnreachableObjects);
@@ -7630,7 +7662,7 @@ void FAsyncLoadingThread2::UpdateSyncLoadContext(FAsyncLoadingThreadState2& Thre
 	}
 	if (SyncLoadContext->bHasFoundRequestedPackages.load(std::memory_order_acquire))
 	{
-		for (int32 i=0; i < SyncLoadContext->RequestIDs.Num(); ++i)
+		for (int32 i = 0; i < SyncLoadContext->RequestIDs.Num(); ++i)
 		{
 			int32 RequestID = SyncLoadContext->RequestIDs[i];
 			FAsyncPackage2* RequestedPackage = SyncLoadContext->RequestedPackages[i];
@@ -8074,6 +8106,7 @@ FAsyncLoadingThread2::FAsyncLoadingThread2(FIoDispatcher& InIoDispatcher, IAsync
 	, IoDispatcher(InIoDispatcher)
 	, UncookedPackageLoader(InUncookedPackageLoader)
 	, PackageStore(FPackageStore::Get())
+	, GlobalImportStore(*this)
 {
 	EventQueue.SetZenaphore(&AltZenaphore);
 
@@ -8107,7 +8140,7 @@ FAsyncLoadingThread2::FAsyncLoadingThread2(FIoDispatcher& InIoDispatcher, IAsync
 	EventQueue.SetOwnerThread(GameThreadState.Get());
 	MainThreadEventQueue.SetOwnerThread(GameThreadState.Get());
 	FAsyncLoadingThreadState2::Set(GameThreadState.Get());
-	
+
 	UE_LOG(LogStreaming, Display, TEXT("AsyncLoading2 - Created: Event Driven Loader: %s, Async Loading Thread: %s, Async Post Load: %s"),
 		GEventDrivenLoaderEnabled ? TEXT("true") : TEXT("false"),
 		FAsyncLoadingThreadSettings::Get().bAsyncLoadingThreadEnabled ? TEXT("true") : TEXT("false"),
@@ -9560,7 +9593,12 @@ EAsyncPackageState::Type FAsyncLoadingThread2::ProcessLoadingUntilCompleteFromGa
 	return bLoadingComplete ? EAsyncPackageState::Complete : EAsyncPackageState::TimeOut;
 }
 
+void FGlobalImportStore::FlushDeferredDeletePackagesQueue()
+{
+	AsyncLoadingThread.ProcessDeferredDeletePackagesQueue();
+}
+
 IAsyncPackageLoader* MakeAsyncPackageLoader2(FIoDispatcher& InIoDispatcher, IAsyncPackageLoader* InUncookedPackageLoader)
 {
 	return new FAsyncLoadingThread2(InIoDispatcher, InUncookedPackageLoader);
-}
+}
\ No newline at end of file
