Allow self references within override enabled prefabs

diff --git a/Engine/Source/Runtime/CoreUObject/Private/UObject/CoreNative.cpp b/Engine/Source/Runtime/CoreUObject/Private/UObject/CoreNative.cpp
--- a/Engine/Source/Runtime/CoreUObject/Private/UObject/CoreNative.cpp
+++ b/Engine/Source/Runtime/CoreUObject/Private/UObject/CoreNative.cpp
@@ -143,8 +143,10 @@ UObject* FObjectInstancingGraph::GetInstancedSubobject( UObject* SourceSubobject
 {
 	checkSlow(SourceSubobject);
 
-	bool bDoNotCreateNewInstance = !!(Flags & EInstancePropertyValueFlags::DoNotCreateNewInstance);
-	bool bAllowSelfReference     = !!(Flags & EInstancePropertyValueFlags::AllowSelfReference);
+	const bool bAreOverridesEnabled = SourceSubobject && FOverridableManager::Get().IsEnabled(*SourceSubobject);
+	const bool bDoNotCreateNewInstance = !!(Flags & EInstancePropertyValueFlags::DoNotCreateNewInstance);
+	const bool bAllowSelfReference     = !!(Flags & EInstancePropertyValueFlags::AllowSelfReference) || bAreOverridesEnabled;
+
 
 	UObject* InstancedSubobject = INVALID_OBJECT;
 
@@ -167,7 +169,7 @@ UObject* FObjectInstancingGraph::GetInstancedSubobject( UObject* SourceSubobject
 		if ( bShouldInstance )
 		{
 			// If the CurrentValue is within the SourceRoot, lets use it to instantiate as it must have come from the merge result of the serialization
-			if (FOverridableManager::Get().IsEnabled(*SourceSubobject) && SourceSubobject != CurrentValue && CurrentValue->IsIn(SourceRoot))
+			if (bAreOverridesEnabled && SourceSubobject != CurrentValue && CurrentValue->IsIn(SourceRoot))
 			{
 				SourceSubobject = CurrentValue;
 			}
