[TRACE] Fix region scope
#rnx
#tests replay

diff --git a/Engine/Source/Runtime/CoreUObject/Private/UObject/GarbageCollection.cpp b/Engine/Source/Runtime/CoreUObject/Private/UObject/GarbageCollection.cpp
--- a/Engine/Source/Runtime/CoreUObject/Private/UObject/GarbageCollection.cpp
+++ b/Engine/Source/Runtime/CoreUObject/Private/UObject/GarbageCollection.cpp
@@ -4601,6 +4601,7 @@ void IncrementalPurgeGarbage(bool bUseTimeLimit, double TimeLimit)
 		{
 			GGCStats.DumpToLog();
 		}
+		TRACE_END_REGION(TEXT("GarbageCollection"));
 	}
 }
 
@@ -5603,7 +5604,6 @@ void PostCollectGarbageImpl(EObjectFlags KeepFlags)
 	GTimingInfo.LastGCTime = PostCollectEndTime;
 	GGCStats.TotalTime += PostCollectEndTime - PostCollectStartTime;
 	STAT_ADD_CUSTOMMESSAGE_NAME(STAT_NamedMarker, TEXT("GarbageCollection - End"));
-	TRACE_END_REGION(TEXT("GarbageCollection"));
 
 	if (bPerformFullPurge)
 	{
@@ -5613,6 +5613,7 @@ void PostCollectGarbageImpl(EObjectFlags KeepFlags)
 		{
 			GGCStats.DumpToLog();
 		}
+		TRACE_END_REGION(TEXT("GarbageCollection"));
 	}
 }
 
