PR #11478: Updated FSoftObjectPath::ExportTextItem to respect PPF_SimpleObjectText, printing out just the asset name for soft object pointers and soft class pointers.

diff --git a/Engine/Source/Runtime/CoreUObject/Private/UObject/SoftObjectPath.cpp b/Engine/Source/Runtime/CoreUObject/Private/UObject/SoftObjectPath.cpp
--- a/Engine/Source/Runtime/CoreUObject/Private/UObject/SoftObjectPath.cpp
+++ b/Engine/Source/Runtime/CoreUObject/Private/UObject/SoftObjectPath.cpp
@@ -338,15 +338,17 @@ bool FSoftObjectPath::ExportTextItem(FString& ValueStr, FSoftObjectPath const& D
 		FSoftObjectPath Temp = *this;
 		Temp.PreSavePath();
 
+		const FString UndelimitedValue = (PortFlags & PPF_SimpleObjectText) ? Temp.GetAssetName() : Temp.ToString();
+
 		if (PortFlags & PPF_Delimited)
 		{
 			ValueStr += TEXT("\"");
-			ValueStr += Temp.ToString().ReplaceQuotesWithEscapedQuotes();
+			ValueStr += UndelimitedValue.ReplaceQuotesWithEscapedQuotes();
 			ValueStr += TEXT("\"");
 		}
 		else
 		{
-			ValueStr += Temp.ToString();
+			ValueStr += UndelimitedValue;
 		}
 	}
 	else
