Fix for entity pointing to CDO values

diff --git a/Engine/Source/Runtime/CoreUObject/Private/UObject/Class.cpp b/Engine/Source/Runtime/CoreUObject/Private/UObject/Class.cpp
--- a/Engine/Source/Runtime/CoreUObject/Private/UObject/Class.cpp
+++ b/Engine/Source/Runtime/CoreUObject/Private/UObject/Class.cpp
@@ -1431,6 +1431,7 @@ void UStruct::SerializeVersionedTaggedProperties(FStructuredArchive::FSlot Slot,
 			{
 				ControlContext.bEnableOverridableSerialization = true;
 				ControlContext.OverriddenProperties = &FOverridableManager::Get().SetOverriddenProperties(*(UObject*)Data, Operation);
+				ControlContext.OverriddenProperties->bNeedsSubobjectTemplateInstantiation = true;
 			}
 		}
 	}
diff --git a/Engine/Source/Runtime/CoreUObject/Private/UObject/Obj.cpp b/Engine/Source/Runtime/CoreUObject/Private/UObject/Obj.cpp
--- a/Engine/Source/Runtime/CoreUObject/Private/UObject/Obj.cpp
+++ b/Engine/Source/Runtime/CoreUObject/Private/UObject/Obj.cpp
@@ -1358,7 +1358,7 @@ void UObject::PostLoadSubobjects( FObjectInstancingGraph* OuterInstanceGraph/*=N
 		// 
 		// @todo_LoadPerf - Consider modifying to bypass this for cooked packages if the class doesn't include any instanced editor-only fields.
 		// Currently we do not have a flag that indicates this scenario, and it might be expensive to iterate over the linked property chain here.
-		if (!FPlatformProperties::RequiresCookedData())
+		if (!FPlatformProperties::RequiresCookedData() || FOverridableManager::Get().NeedSubObjectTemplateInstantiation(*this))
 		{
 			FObjectInstancingGraph CurrentInstanceGraph;
 
diff --git a/Engine/Source/Runtime/CoreUObject/Private/UObject/OverridableManager.cpp b/Engine/Source/Runtime/CoreUObject/Private/UObject/OverridableManager.cpp
--- a/Engine/Source/Runtime/CoreUObject/Private/UObject/OverridableManager.cpp
+++ b/Engine/Source/Runtime/CoreUObject/Private/UObject/OverridableManager.cpp
@@ -43,6 +43,15 @@ void FOverridableManager::InheritEnabledFrom(UObject& Object, const UObject* Def
 	}
 }
 
+bool FOverridableManager::NeedSubObjectTemplateInstantiation(const UObject& Object)
+{
+	if( const FOverriddenPropertySet* OverriddenProperties = OverriddenObjectAnnotations.Find(Object))
+	{
+		return OverriddenProperties->bNeedsSubobjectTemplateInstantiation;
+	}
+	return false;
+}
+
 FOverriddenPropertySet* FOverridableManager::GetOverriddenProperties(UObject& Object)
 {
 	return OverriddenObjectAnnotations.Find(Object);
diff --git a/Engine/Source/Runtime/CoreUObject/Public/UObject/OverridableManager.h b/Engine/Source/Runtime/CoreUObject/Public/UObject/OverridableManager.h
--- a/Engine/Source/Runtime/CoreUObject/Public/UObject/OverridableManager.h
+++ b/Engine/Source/Runtime/CoreUObject/Public/UObject/OverridableManager.h
@@ -106,6 +106,11 @@ public:
 	 * @param DefaultData its default object */
 	COREUOBJECT_API void InheritEnabledFrom(UObject& Object, const UObject* DefaultData);
 
+	/**
+	 * Return true if this object needs subobject template instantiation
+	 * @param Object to be querying about */
+	COREUOBJECT_API bool NeedSubObjectTemplateInstantiation(const UObject& Object);
+
 	/**
 	 * Retrieve the overridden properties for the specified object
 	 * @param Object to fetch the overridden properties
diff --git a/Engine/Source/Runtime/CoreUObject/Public/UObject/OverriddenPropertySet.h b/Engine/Source/Runtime/CoreUObject/Public/UObject/OverriddenPropertySet.h
--- a/Engine/Source/Runtime/CoreUObject/Public/UObject/OverriddenPropertySet.h
+++ b/Engine/Source/Runtime/CoreUObject/Public/UObject/OverriddenPropertySet.h
@@ -308,4 +308,7 @@ private:
 	TSet<FOverriddenPropertyNode> OverriddenPropertyNodes;
 
 	static inline FOverriddenPropertyNodeID RootNodeID = FOverriddenPropertyNodeID(FName(TEXT("root")));
+
+public:
+	bool bNeedsSubobjectTemplateInstantiation = false;
 };
