Skip looking for a type redirect If the type in the property tag matches the type of the property

This is necessary to handle redirects where the old type name continues to be used.

#rnx

diff --git a/Engine/Source/Runtime/CoreUObject/Private/UObject/PropertyHelper.cpp b/Engine/Source/Runtime/CoreUObject/Private/UObject/PropertyHelper.cpp
--- a/Engine/Source/Runtime/CoreUObject/Private/UObject/PropertyHelper.cpp
+++ b/Engine/Source/Runtime/CoreUObject/Private/UObject/PropertyHelper.cpp
@@ -368,16 +368,24 @@ static bool FindRedirectForProperty(FPropertyTypeName OldType, FPropertyTypeName
 {
 	using FTraits = TFindRedirectForPropertyTraits<PropertyType>;
 	const UField* Field = Property ? FTraits::GetField(Property) : nullptr;
+
+	// If the type in the tag matches the field then skip looking for redirects.
+	// This is necessary to handle redirects where the old name continues to be used.
 	FCoreRedirectObjectName OldNameRedirect = BuildCoreRedirectObjectName(OldType);
+	if (Field && Field->GetFName() == OldNameRedirect.ObjectName)
+	{
+		return false;
+	}
+
+	// Look for a partial match if there is a field to compare against, otherwise require a complete match.
 	FCoreRedirectObjectName NewNameRedirect = FCoreRedirects::GetRedirectedName(FTraits::GetFlags(), OldNameRedirect,
 		Field ? ECoreRedirectMatchFlags::AllowPartialMatch : ECoreRedirectMatchFlags::None);
-
 	if (NewNameRedirect == OldNameRedirect)
 	{
 		return false;
 	}
 
-	// If a partial match does not match the property then repeat the lookup without allowing partial matches.
+	// If a partial match does not match the field then repeat the lookup without allowing partial matches.
 	// This is necessary to handle redirects of the original type of a property that changed type.
 	if (Field && Field->GetFName() != NewNameRedirect.ObjectName)
 	{
