EditorBulkData add Insights scope around hashing, otherwise nop

diff --git a/Engine/Source/Runtime/CoreUObject/Private/Serialization/EditorBulkData.cpp b/Engine/Source/Runtime/CoreUObject/Private/Serialization/EditorBulkData.cpp
--- a/Engine/Source/Runtime/CoreUObject/Private/Serialization/EditorBulkData.cpp
+++ b/Engine/Source/Runtime/CoreUObject/Private/Serialization/EditorBulkData.cpp
@@ -264,6 +264,8 @@ static FIoHash HashPayload(const FSharedBuffer& InPayload)
 {
 	if (InPayload.GetSize() > 0)
 	{
+		TRACE_CPUPROFILER_EVENT_SCOPE(EditorBulkData.HashPayload);
+
 		return FIoHash::HashBuffer(InPayload);
 	}
 	else
