Remove categories on IDO properties and transferred DisplayName metadata of template class to IDO

diff --git a/Engine/Plugins/Experimental/InstanceDataObjectFixupTool/Source/Private/InstanceDataObjectFixupDetailCustomization.cpp b/Engine/Plugins/Experimental/InstanceDataObjectFixupTool/Source/Private/InstanceDataObjectFixupDetailCustomization.cpp
--- a/Engine/Plugins/Experimental/InstanceDataObjectFixupTool/Source/Private/InstanceDataObjectFixupDetailCustomization.cpp
+++ b/Engine/Plugins/Experimental/InstanceDataObjectFixupTool/Source/Private/InstanceDataObjectFixupDetailCustomization.cpp
@@ -458,7 +458,7 @@ void FHideLoosePropertiesCustomization::CustomizeDetails(IDetailLayoutBuilder& D
 	DetailBuilder.GetCategoryNames(CategoryNames);
 	for (const FName CategoryName : CategoryNames)
 	{
-		IDetailCategoryBuilder& Category = DetailBuilder.EditCategory(CategoryName);
+		IDetailCategoryBuilder& Category = DetailBuilder.EditCategoryAllowNone(CategoryName);
 		TArray<TSharedRef<IPropertyHandle>> Handles;
 		Category.GetDefaultProperties(Handles, true, true);
 		for (const TSharedRef<IPropertyHandle>& Handle : Handles)
diff --git a/Engine/Source/Runtime/CoreUObject/Private/UObject/InstanceDataObjectUtils.cpp b/Engine/Source/Runtime/CoreUObject/Private/UObject/InstanceDataObjectUtils.cpp
--- a/Engine/Source/Runtime/CoreUObject/Private/UObject/InstanceDataObjectUtils.cpp
+++ b/Engine/Source/Runtime/CoreUObject/Private/UObject/InstanceDataObjectUtils.cpp
@@ -16,7 +16,6 @@ namespace UE
 	static const FName NAME_StructOriginalTypeMetadata(TEXT("OriginalType"));
 	static const FName NAME_PresentAsTypeMetadata(TEXT("PresentAsType"));
 	static const FName NAME_IsLooseMetadata(TEXT("IsLoose"));
-	static const FName NAME_CategoryMetadata(TEXT("Category"));
 	static const FName NAME_VerseClass("VerseClass");
 	
 	struct ResolvePropertyPathNameHelperParams
@@ -314,13 +313,6 @@ namespace UE
 			ConvertToInstanceDataObjectProperty(AsMapProperty->ValueProp, Outer, LooseProperties, Path);
 			Path.Pop();
 		}
-		
-#if WITH_EDITORONLY_DATA
-		if (!Property->HasMetaData(NAME_CategoryMetadata))
-		{
-			Property->SetMetaData(NAME_CategoryMetadata, TEXT("Verse"));
-		}
-#endif
 	}
 	
 	// copy template property then convert it into an InstanceDataObject property by adding loose properties
@@ -405,7 +397,6 @@ namespace UE
 	{
 #if WITH_EDITORONLY_DATA
 		Property->SetMetaData(NAME_IsLooseMetadata, TEXT("True"));
-		Property->SetMetaData(NAME_CategoryMetadata, TEXT("Loose Properties"));
 #endif
 		Property->SetPropertyFlags(CPF_Edit | CPF_EditConst);
 		if (const FArrayProperty* AsArrayProperty = CastField<FArrayProperty>(Property))
@@ -461,12 +452,6 @@ namespace UE
 				Path.Push(CreateSegmentFromProperty(TemplateProperty));
 				FProperty* SuperProperty = CreateInstanceDataObjectProperty(TemplateProperty, Super, LooseProperties, Path);
 				
-#if WITH_EDITORONLY_DATA
-				if (!SuperProperty->HasMetaData(NAME_CategoryMetadata))
-				{
-					SuperProperty->SetMetaData(NAME_CategoryMetadata, TEXT("Verse"));
-				}
-#endif
 				Path.Pop();
 				SuperProperties.Add(SuperProperty);
 			}
@@ -559,6 +544,13 @@ namespace UE
 		const TMap<FWildcardPropertyPathName, TMap<FName, const FProperty*>> LooseProperties = GetWildcardedLooseProperties(PropertyBag);
 		FWildcardPropertyPathName ParentPath;
 		UClass* Result = CreateInstanceDataObjectStructRec<UClass>(OwnerClass, Outer, LooseProperties, ParentPath);
+#if WITH_EDITORONLY_DATA
+		const FString& DisplayName = OwnerClass->GetMetaData(TEXT("DisplayName"));
+		if (!DisplayName.IsEmpty())
+		{
+			Result->SetMetaData(TEXT("DisplayName"), *DisplayName);
+		}
+#endif
 
 		const UObject* OwnerCDO = OwnerClass->GetDefaultObject(true);
 		UObject* ResultCDO = Result->GetDefaultObject(true);
