Extend IDO support query to include subtypes (e.g. generated class assets).

diff --git a/Engine/Source/Runtime/CoreUObject/Private/UObject/InstanceDataObjectUtils.cpp b/Engine/Source/Runtime/CoreUObject/Private/UObject/InstanceDataObjectUtils.cpp
--- a/Engine/Source/Runtime/CoreUObject/Private/UObject/InstanceDataObjectUtils.cpp
+++ b/Engine/Source/Runtime/CoreUObject/Private/UObject/InstanceDataObjectUtils.cpp
@@ -37,9 +37,21 @@ namespace UE
 
 	bool IsInstanceDataObjectSupportEnabled(UObject* InObject)
 	{
-		return bEnableIDOSupport
+		// Note: NULL is a valid (default) input here; in that case we just return the enable flag.
+		bool bIsEnabled = bEnableIDOSupport;
+		if (bIsEnabled && InObject)
+		{
 			//@todo FH: change to check trait when available or use config object
-			&& (!InObject || InObject->GetClass()->GetClass()->GetFName() == NAME_VerseClass);
+			const UClass* ObjClass = InObject->GetClass();
+			while (ObjClass && ObjClass->GetClass()->GetFName() != NAME_VerseClass)
+			{
+				ObjClass = ObjClass->GetSuperClass();
+			}
+
+			bIsEnabled = !!ObjClass;
+		}
+
+		return bIsEnabled;
 	}
 
 	static void BuildSegmentTypeFromProperty(const FProperty* Property, FPropertyTypeNameBuilder& OutType)
