Remove the overridable serialization logic when using UPS

diff --git a/Engine/Source/Runtime/CoreUObject/Private/Serialization/UnversionedPropertySerialization.cpp b/Engine/Source/Runtime/CoreUObject/Private/Serialization/UnversionedPropertySerialization.cpp
--- a/Engine/Source/Runtime/CoreUObject/Private/Serialization/UnversionedPropertySerialization.cpp
+++ b/Engine/Source/Runtime/CoreUObject/Private/Serialization/UnversionedPropertySerialization.cpp
@@ -923,17 +923,6 @@ void SerializeUnversionedProperties(const UStruct* Struct, FStructuredArchive::F
 	}
 	else
 	{
-		bool bEnableOverridableSerialization = false;
-		FOverriddenPropertySet* OverriddenProperties = nullptr;
-		if (FOverriddenPropertySet* ObjectOverriddenProperties = Struct->IsA<UClass>() ? FOverridableManager::Get().GetOverriddenProperties(*(UObject*)Data) : nullptr)
-		{
-			bEnableOverridableSerialization = true;
-			OverriddenProperties = ObjectOverriddenProperties;
-		}
-
-		// Scope that enables the overridable serialization for this object
-		FEnableOverridableSerializationScope OverridableSerializationScope(bEnableOverridableSerialization, OverriddenProperties);
-
 		FUnversionedPropertyTestRunner TestRunner({Struct, Data, DefaultsStruct, DefaultsData});
 		FUnversionedPropertyTestCollector TestCollector;
 
diff --git a/Engine/Source/Runtime/CoreUObject/Private/UObject/PropertyArray.cpp b/Engine/Source/Runtime/CoreUObject/Private/UObject/PropertyArray.cpp
--- a/Engine/Source/Runtime/CoreUObject/Private/UObject/PropertyArray.cpp
+++ b/Engine/Source/Runtime/CoreUObject/Private/UObject/PropertyArray.cpp
@@ -131,7 +131,7 @@ void FArrayProperty::SerializeItem(FStructuredArchive::FSlot Slot, void* Value,
 	int32 ElementCount = ArrayHelper.Num();
 
 	// Custom branch for UPS to try and take advantage of bulk serialization
-	if (bUPS && !bExperimentalOverridableLogic)
+	if (bUPS)
 	{
 		checkf(!UnderlyingArchive.ArUseCustomPropertyList, TEXT("Custom property lists are not supported with UPS"));
 		checkf(!bIsTextFormat, TEXT("Text-based archives are not supported with UPS"));
@@ -212,7 +212,7 @@ void FArrayProperty::SerializeItem(FStructuredArchive::FSlot Slot, void* Value,
 	};
 
 	// Make sure the container is reloading accordingly to the value set in the property tag if any
-	if (!bUPS && UnderlyingArchive.IsLoading() && FPropertyTagScope::GetCurrentPropertyTag())
+	if (UnderlyingArchive.IsLoading() && FPropertyTagScope::GetCurrentPropertyTag())
 	{
 		bExperimentalOverridableLogic = FPropertyTagScope::GetCurrentPropertyTag()->bExperimentalOverridableLogic;
 	}
diff --git a/Engine/Source/Runtime/CoreUObject/Private/UObject/PropertyMap.cpp b/Engine/Source/Runtime/CoreUObject/Private/UObject/PropertyMap.cpp
--- a/Engine/Source/Runtime/CoreUObject/Private/UObject/PropertyMap.cpp
+++ b/Engine/Source/Runtime/CoreUObject/Private/UObject/PropertyMap.cpp
@@ -290,371 +290,374 @@ void FMapProperty::SerializeItem(FStructuredArchive::FSlot Slot, void* Value, co
 
 	FScriptMapHelper MapHelper(this, Value);
 
-	// Make sure the container is reloading accordingly to the value set in the property tag if any
-	if (!bUPS && UnderlyingArchive.IsLoading() && FPropertyTagScope::GetCurrentPropertyTag())
-	{
-		bExperimentalOverridableLogic = FPropertyTagScope::GetCurrentPropertyTag()->bExperimentalOverridableLogic;
-	}
-
 	// *** Experimental *** Special serialization path for map with overridable serialization logic
-	if (bExperimentalOverridableLogic)
+	if (!bUPS)
 	{
-		checkf(!UnderlyingArchive.ArUseCustomPropertyList, TEXT("Using custom property list is not supported by overridable serialization"));
-
-		auto GetIDFromKey = [&](uint8* KeyData) -> FOverriddenPropertyNodeID
+		// Make sure the container is reloading accordingly to the value set in the property tag if any
+		if (UnderlyingArchive.IsLoading() && FPropertyTagScope::GetCurrentPropertyTag())
 		{
-			if (FObjectProperty* KeyObjectProperty = CastField<FObjectProperty>(KeyProp))
-			{
-				if (const UObject* Object = KeyObjectProperty->GetObjectPropertyValue(KeyData))
-				{
-					return FOverriddenPropertyNodeID(*Object);
-				}
-			}
-			else
-			{
-				FString KeyString;
-				KeyProp->ExportTextItem_Direct(KeyString, KeyData, nullptr, nullptr, PPF_None);
-				return FOverriddenPropertyNodeID(FName(KeyString));
-			}
-	
-			checkf(false, TEXT("This case is not handled"))
-			return FOverriddenPropertyNodeID();
-		};
+			bExperimentalOverridableLogic = FPropertyTagScope::GetCurrentPropertyTag()->bExperimentalOverridableLogic;
+		}
 
-		if (UnderlyingArchive.IsLoading())
+		if (bExperimentalOverridableLogic)
 		{
-			int32 NumReplaced = 0;
-			FStructuredArchive::FArray ReplacedArray = Record.EnterArray(TEXT("Replaced"), NumReplaced);
-			if (NumReplaced != INDEX_NONE)
+			checkf(!UnderlyingArchive.ArUseCustomPropertyList, TEXT("Using custom property list is not supported by overridable serialization"));
+
+			auto GetIDFromKey = [&](uint8* KeyData) -> FOverriddenPropertyNodeID
 			{
-				MapHelper.EmptyValues(NumReplaced);
-				for (int32 i = 0; i < NumReplaced; i++)
+				if (FObjectProperty* KeyObjectProperty = CastField<FObjectProperty>(KeyProp))
 				{
-					FStructuredArchive::FRecord EntryRecord = ReplacedArray.EnterElement().EnterRecord();
-					int32 Index = MapHelper.AddDefaultValue_Invalid_NeedsRehash();
-					uint8* PairPtr = MapHelper.GetPairPtr(Index);
-
-					{
-						FSerializedPropertyScope SerializedProperty(UnderlyingArchive, KeyProp, this);
-						KeyProp->SerializeItem(EntryRecord.EnterField(TEXT("Key")), PairPtr);
-					}
+					if (const UObject* Object = KeyObjectProperty->GetObjectPropertyValue(KeyData))
 					{
-						FSerializedPropertyScope SerializedProperty(UnderlyingArchive, ValueProp, this);
-						ValueProp->SerializeItem(EntryRecord.EnterField(TEXT("Value")), PairPtr + MapLayout.ValueOffset);
+						return FOverriddenPropertyNodeID(*Object);
 					}
 				}
-				MapHelper.Rehash();
-			}
-			else
-			{
-				FOverriddenPropertySet* OverriddenProperties = FOverridableSerializationLogic::GetOverriddenProperties();
-
-				// This is not fully implemented yet and not a priority right now, so just trying to prevent it as the result could be random
-				checkf(!KeyProp->HasAnyPropertyFlags(CPF_PersistentInstance) || CastField<FClassProperty>(KeyProp) || !CastField<FObjectProperty>(KeyProp), TEXT("The key as an instanced sub object is NYI"));
-
-				uint8* TempKeyValueStorage = nullptr;
-				ON_SCOPE_EXIT
+				else
 				{
-					if (TempKeyValueStorage)
-					{
-						KeyProp->DestroyValue(TempKeyValueStorage);
-						ValueProp->DestroyValue(TempKeyValueStorage + MapLayout.ValueOffset);
-						FMemory::Free(TempKeyValueStorage);
-					}
-				};
+					FString KeyString;
+					KeyProp->ExportTextItem_Direct(KeyString, KeyData, /*DefaultValue*/nullptr, /*Parent*/nullptr, PPF_None);
+					return FOverriddenPropertyNodeID(FName(KeyString));
+				}
+		
+				checkf(false, TEXT("This case is not handled"))
+				return FOverriddenPropertyNodeID();
+			};
 
-				int32 NumRemoved = 0;
-				FStructuredArchive::FArray RemovedArray = Record.EnterArray(TEXT("Removed"), NumRemoved);
-				if (NumRemoved != 0)
+			if (UnderlyingArchive.IsLoading())
+			{
+				int32 NumReplaced = 0;
+				FStructuredArchive::FArray ReplacedArray = Record.EnterArray(TEXT("Replaced"), NumReplaced);
+				if (NumReplaced != INDEX_NONE)
 				{
-					TempKeyValueStorage = (uint8*)FMemory::Malloc(MapLayout.SetLayout.Size);
-					KeyProp->InitializeValue(TempKeyValueStorage);
-					ValueProp->InitializeValue(TempKeyValueStorage + MapLayout.ValueOffset);
-
-					for (int32 i = 0; i < NumRemoved; ++i)
+					MapHelper.EmptyValues(NumReplaced);
+					for (int32 i = 0; i < NumReplaced; i++)
 					{
+						FStructuredArchive::FRecord EntryRecord = ReplacedArray.EnterElement().EnterRecord();
+						int32 Index = MapHelper.AddDefaultValue_Invalid_NeedsRehash();
+						uint8* PairPtr = MapHelper.GetPairPtr(Index);
+
 						{
 							FSerializedPropertyScope SerializedProperty(UnderlyingArchive, KeyProp, this);
-							KeyProp->SerializeItem(RemovedArray.EnterElement().EnterRecord().EnterField(TEXT("Key")), TempKeyValueStorage);
+							KeyProp->SerializeItem(EntryRecord.EnterField(TEXT("Key")), PairPtr);
 						}
-
-						MapHelper.RemovePair(TempKeyValueStorage);
-
-						// Need to fetch the MapOverriddenPropertyNode every loop as the previous might have reallocated the node.
-						if (FOverriddenPropertyNode* MapOverriddenPropertyNode = OverriddenProperties ? OverriddenProperties->SetOverriddenPropertyOperation(EOverriddenPropertyOperation::Modified, UnderlyingArchive.GetSerializedPropertyChain(), /*Property*/nullptr) : nullptr)
 						{
-							// Rebuild the overridden info
-							FOverriddenPropertyNodeID RemovedKeyID = GetIDFromKey(TempKeyValueStorage);
-							OverriddenProperties->SetSubPropertyOperation(EOverriddenPropertyOperation::Remove, *MapOverriddenPropertyNode, RemovedKeyID);
+							FSerializedPropertyScope SerializedProperty(UnderlyingArchive, ValueProp, this);
+							ValueProp->SerializeItem(EntryRecord.EnterField(TEXT("Value")), PairPtr + MapLayout.ValueOffset);
 						}
 					}
+					MapHelper.Rehash();
 				}
-
-				int32 NumModified = 0;
-				FStructuredArchive::FArray ModifiedArray = Record.EnterArray(TEXT("Modified"), NumModified);
-				if (NumModified != 0)
+				else
 				{
-					if (!TempKeyValueStorage)
-					{
-						TempKeyValueStorage = (uint8*)FMemory::Malloc(MapLayout.SetLayout.Size);
-						KeyProp->InitializeValue(TempKeyValueStorage);
-						ValueProp->InitializeValue(TempKeyValueStorage + MapLayout.ValueOffset);
-					}
-					for (int32 i = 0; i < NumModified; ++i)
-					{
-						FStructuredArchive::FRecord EntryRecord = ModifiedArray.EnterElement().EnterRecord();
-
-						// Read key into temporary storage
-						{
-							FSerializedPropertyScope SerializedProperty(UnderlyingArchive, KeyProp, this);
-							KeyProp->SerializeItem(EntryRecord.EnterField(TEXT("Key")), TempKeyValueStorage);
-						}
+					FOverriddenPropertySet* OverriddenProperties = FOverridableSerializationLogic::GetOverriddenProperties();
 
-						const int32 Index = MapHelper.FindMapPairIndexFromHash(TempKeyValueStorage);
-						uint8* ValuePtr = Index != INDEX_NONE ? MapHelper.GetValuePtr(Index) : TempKeyValueStorage + MapLayout.ValueOffset;
+					// This is not fully implemented yet and not a priority right now, so just trying to prevent it as the result could be random
+					checkf(!KeyProp->HasAnyPropertyFlags(CPF_PersistentInstance) || CastField<FClassProperty>(KeyProp) || !CastField<FObjectProperty>(KeyProp), TEXT("The key as an instanced sub object is NYI"));
 
-						// Deserialize value into hash map-owned memory
+					uint8* TempKeyValueStorage = nullptr;
+					ON_SCOPE_EXIT
+					{
+						if (TempKeyValueStorage)
 						{
-							FSerializedPropertyScope SerializedProperty(UnderlyingArchive, ValueProp, this);
-							ValueProp->SerializeItem(EntryRecord.EnterField(TEXT("Value")), ValuePtr);
+							KeyProp->DestroyValue(TempKeyValueStorage);
+							ValueProp->DestroyValue(TempKeyValueStorage + MapLayout.ValueOffset);
+							FMemory::Free(TempKeyValueStorage);
 						}
-					}
-				}
+					};
 
-				int32 NumAdded = 0;
-				FStructuredArchive::FArray AddedArray = Record.EnterArray(TEXT("Added"), NumAdded);
-				if (NumAdded != 0)
-				{
-					if (!TempKeyValueStorage)
+					int32 NumRemoved = 0;
+					FStructuredArchive::FArray RemovedArray = Record.EnterArray(TEXT("Removed"), NumRemoved);
+					if (NumRemoved != 0)
 					{
 						TempKeyValueStorage = (uint8*)FMemory::Malloc(MapLayout.SetLayout.Size);
 						KeyProp->InitializeValue(TempKeyValueStorage);
 						ValueProp->InitializeValue(TempKeyValueStorage + MapLayout.ValueOffset);
+
+						for (int32 i = 0; i < NumRemoved; ++i)
+						{
+							{
+								FSerializedPropertyScope SerializedProperty(UnderlyingArchive, KeyProp, this);
+								KeyProp->SerializeItem(RemovedArray.EnterElement().EnterRecord().EnterField(TEXT("Key")), TempKeyValueStorage);
+							}
+
+							MapHelper.RemovePair(TempKeyValueStorage);
+
+							// Need to fetch the MapOverriddenPropertyNode every loop as the previous might have reallocated the node.
+							if (FOverriddenPropertyNode* MapOverriddenPropertyNode = OverriddenProperties ? OverriddenProperties->SetOverriddenPropertyOperation(EOverriddenPropertyOperation::Modified, UnderlyingArchive.GetSerializedPropertyChain(), /*Property*/nullptr) : nullptr)
+							{
+								// Rebuild the overridden info
+								FOverriddenPropertyNodeID RemovedKeyID = GetIDFromKey(TempKeyValueStorage);
+								OverriddenProperties->SetSubPropertyOperation(EOverriddenPropertyOperation::Remove, *MapOverriddenPropertyNode, RemovedKeyID);
+							}
+						}
 					}
 
-					for (int32 i = 0; i < NumAdded; ++i)
+					int32 NumModified = 0;
+					FStructuredArchive::FArray ModifiedArray = Record.EnterArray(TEXT("Modified"), NumModified);
+					if (NumModified != 0)
 					{
-						FStructuredArchive::FRecord EntryRecord = AddedArray.EnterElement().EnterRecord();
-
-						// Read key into temporary storage
+						if (!TempKeyValueStorage)
 						{
-							FSerializedPropertyScope SerializedProperty(UnderlyingArchive, KeyProp, this);
-							KeyProp->SerializeItem(EntryRecord.EnterField(TEXT("Key")), TempKeyValueStorage);
+							TempKeyValueStorage = (uint8*)FMemory::Malloc(MapLayout.SetLayout.Size);
+							KeyProp->InitializeValue(TempKeyValueStorage);
+							ValueProp->InitializeValue(TempKeyValueStorage + MapLayout.ValueOffset);
 						}
+						for (int32 i = 0; i < NumModified; ++i)
+						{
+							FStructuredArchive::FRecord EntryRecord = ModifiedArray.EnterElement().EnterRecord();
+
+							// Read key into temporary storage
+							{
+								FSerializedPropertyScope SerializedProperty(UnderlyingArchive, KeyProp, this);
+								KeyProp->SerializeItem(EntryRecord.EnterField(TEXT("Key")), TempKeyValueStorage);
+							}
 
-						void* ValuePtr = MapHelper.FindOrAdd(TempKeyValueStorage);
+							const int32 Index = MapHelper.FindMapPairIndexFromHash(TempKeyValueStorage);
+							uint8* ValuePtr = Index != INDEX_NONE ? MapHelper.GetValuePtr(Index) : TempKeyValueStorage + MapLayout.ValueOffset;
+
+							// Deserialize value into hash map-owned memory
+							{
+								FSerializedPropertyScope SerializedProperty(UnderlyingArchive, ValueProp, this);
+								ValueProp->SerializeItem(EntryRecord.EnterField(TEXT("Value")), ValuePtr);
+							}
+						}
+					}
 
-						// Deserialize value into hash map-owned memory
+					int32 NumAdded = 0;
+					FStructuredArchive::FArray AddedArray = Record.EnterArray(TEXT("Added"), NumAdded);
+					if (NumAdded != 0)
+					{
+						if (!TempKeyValueStorage)
 						{
-							FSerializedPropertyScope SerializedProperty(UnderlyingArchive, ValueProp, this);
-							ValueProp->SerializeItem(EntryRecord.EnterField(TEXT("Value")), ValuePtr);
+							TempKeyValueStorage = (uint8*)FMemory::Malloc(MapLayout.SetLayout.Size);
+							KeyProp->InitializeValue(TempKeyValueStorage);
+							ValueProp->InitializeValue(TempKeyValueStorage + MapLayout.ValueOffset);
 						}
 
-						// Need to fetch the MapOverriddenPropertyNode every loop as the previous might have reallocated the node.
-						if (FOverriddenPropertyNode* MapOverriddenPropertyNode = OverriddenProperties ? OverriddenProperties->SetOverriddenPropertyOperation(EOverriddenPropertyOperation::Modified, UnderlyingArchive.GetSerializedPropertyChain(), /*Property*/nullptr) : nullptr)
+						for (int32 i = 0; i < NumAdded; ++i)
 						{
-							// Rebuild the overridden info
-							FOverriddenPropertyNodeID AddedKeyID = GetIDFromKey(TempKeyValueStorage);
-							OverriddenProperties->SetSubPropertyOperation(EOverriddenPropertyOperation::Add, *MapOverriddenPropertyNode, AddedKeyID);
+							FStructuredArchive::FRecord EntryRecord = AddedArray.EnterElement().EnterRecord();
+
+							// Read key into temporary storage
+							{
+								FSerializedPropertyScope SerializedProperty(UnderlyingArchive, KeyProp, this);
+								KeyProp->SerializeItem(EntryRecord.EnterField(TEXT("Key")), TempKeyValueStorage);
+							}
+
+							void* ValuePtr = MapHelper.FindOrAdd(TempKeyValueStorage);
+
+							// Deserialize value into hash map-owned memory
+							{
+								FSerializedPropertyScope SerializedProperty(UnderlyingArchive, ValueProp, this);
+								ValueProp->SerializeItem(EntryRecord.EnterField(TEXT("Value")), ValuePtr);
+							}
+
+							// Need to fetch the MapOverriddenPropertyNode every loop as the previous might have reallocated the node.
+							if (FOverriddenPropertyNode* MapOverriddenPropertyNode = OverriddenProperties ? OverriddenProperties->SetOverriddenPropertyOperation(EOverriddenPropertyOperation::Modified, UnderlyingArchive.GetSerializedPropertyChain(), /*Property*/nullptr) : nullptr)
+							{
+								// Rebuild the overridden info
+								FOverriddenPropertyNodeID AddedKeyID = GetIDFromKey(TempKeyValueStorage);
+								OverriddenProperties->SetSubPropertyOperation(EOverriddenPropertyOperation::Add, *MapOverriddenPropertyNode, AddedKeyID);
+							}
 						}
 					}
 				}
 			}
-		}
-		else
-		{
-			auto FindKeyInternalIndex = [this](const FOverriddenPropertyNodeID& KeyIDToFind, FScriptMapHelper& MapHelper) -> int32
+			else
 			{
-				if (const FObjectProperty* KeyObjectProperty = CastField<FObjectProperty>(KeyProp))
+				auto FindKeyInternalIndex = [this](const FOverriddenPropertyNodeID& KeyIDToFind, FScriptMapHelper& MapHelper) -> int32
 				{
-					for (FScriptMapHelper::FIterator It(MapHelper); It; ++It) 
+					if (const FObjectProperty* KeyObjectProperty = CastField<FObjectProperty>(KeyProp))
 					{
-						if (UObject* CurrentObject = KeyObjectProperty->GetObjectPropertyValue(MapHelper.GetKeyPtr(It.GetInternalIndex())))
+						for (FScriptMapHelper::FIterator It(MapHelper); It; ++It) 
 						{
-							if (KeyIDToFind == FOverriddenPropertyNodeID(*CurrentObject))
+							if (UObject* CurrentObject = KeyObjectProperty->GetObjectPropertyValue(MapHelper.GetKeyPtr(It.GetInternalIndex())))
 							{
-								return It.GetInternalIndex();
+								if (KeyIDToFind == FOverriddenPropertyNodeID(*CurrentObject))
+								{
+									return It.GetInternalIndex();
+								}
 							}
 						}
 					}
-				}
-				else
-				{
-					void* TempKeyValueStorage = FMemory::Malloc(MapLayout.SetLayout.Size);
-					KeyProp->InitializeValue(TempKeyValueStorage);
+					else
+					{
+						void* TempKeyValueStorage = FMemory::Malloc(MapLayout.SetLayout.Size);
+						KeyProp->InitializeValue(TempKeyValueStorage);
 
-					FString KeyToFind(KeyIDToFind.ToString());
-					KeyProp->ImportText_Direct(*KeyToFind, TempKeyValueStorage, nullptr, PPF_None);
+						FString KeyToFind(KeyIDToFind.ToString());
+						KeyProp->ImportText_Direct(*KeyToFind, TempKeyValueStorage, nullptr, PPF_None);
 
-					const int32 InternalIndex = MapHelper.FindMapPairIndexFromHash(TempKeyValueStorage);
+						const int32 InternalIndex = MapHelper.FindMapPairIndexFromHash(TempKeyValueStorage);
 
-					KeyProp->DestroyValue(TempKeyValueStorage);
-					FMemory::Free(TempKeyValueStorage);
+						KeyProp->DestroyValue(TempKeyValueStorage);
+						FMemory::Free(TempKeyValueStorage);
 
-					return InternalIndex;
-				}
-				return INDEX_NONE;
-			};
+						return InternalIndex;
+					}
+					return INDEX_NONE;
+				};
 
-			// Container for temporarily tracking some indices
-			TArray<int32> RemovedIndices;
-			TArray<int32> AddedIndices;
-			TSet<int32> ModifiedIndices;
+				// Container for temporarily tracking some indices
+				TArray<int32> RemovedIndices;
+				TArray<int32> AddedIndices;
+				TSet<int32> ModifiedIndices;
 
-			bool bReplaceMap = false;
-			if (!Defaults || !UnderlyingArchive.DoDelta() || UnderlyingArchive.IsTransacting())
-			{
-				bReplaceMap = true;
-			}
-			else 
-			{
-				EOverriddenPropertyOperation MapOverrideOp = EOverriddenPropertyOperation::None;
-				FOverriddenPropertySet* OverriddenProperties = FOverridableSerializationLogic::GetOverriddenProperties();
-				if (OverriddenProperties)
+				bool bReplaceMap = false;
+				if (!Defaults || !UnderlyingArchive.DoDelta() || UnderlyingArchive.IsTransacting())
 				{
-					MapOverrideOp = OverriddenProperties->GetOverriddenPropertyOperation(UnderlyingArchive.GetSerializedPropertyChain(), /*Property*/nullptr);
-					bReplaceMap = MapOverrideOp == EOverriddenPropertyOperation::Replace;
+					bReplaceMap = true;
 				}
-				else
-				{
-					// Only instanced subobjects keys are not supported
-					bReplaceMap = CastField<FObjectProperty>(KeyProp) && KeyProp->HasAnyPropertyFlags(CPF_PersistentInstance) && !CastField<FClassProperty>(KeyProp); // Class property should never be instanced
-				}
-
-				if(!bReplaceMap)
+				else 
 				{
-					checkf(!KeyProp->HasAnyPropertyFlags(CPF_PersistentInstance) || CastField<FClassProperty>(KeyProp) || !CastField<FObjectProperty>(KeyProp), TEXT("The key as an instanced sub object is NYI"));
-
-					// For instanced subobject, the overridable is handled per object base, not only here, so we need to serialize the object ptr no matter what.
-					const bool bAreValuesInstancedSubObjects = ValueProp->HasAnyPropertyFlags(CPF_PersistentInstance) && CastField<FObjectProperty>(ValueProp);
-					if(bAreValuesInstancedSubObjects)
+					EOverriddenPropertyOperation MapOverrideOp = EOverriddenPropertyOperation::None;
+					FOverriddenPropertySet* OverriddenProperties = FOverridableSerializationLogic::GetOverriddenProperties();
+					if (OverriddenProperties)
 					{
-						for (FScriptMapHelper::FIterator It(MapHelper); It; ++It)
-						{
-							ModifiedIndices.Add(It.GetInternalIndex());
-						}
+						MapOverrideOp = OverriddenProperties->GetOverriddenPropertyOperation(UnderlyingArchive.GetSerializedPropertyChain(), /*Property*/nullptr);
+						bReplaceMap = MapOverrideOp == EOverriddenPropertyOperation::Replace;
+					}
+					else
+					{
+						// Only instanced subobjects keys are not supported
+						bReplaceMap = CastField<FObjectProperty>(KeyProp) && KeyProp->HasAnyPropertyFlags(CPF_PersistentInstance) && !CastField<FClassProperty>(KeyProp); // Class property should never be instanced
 					}
 
-					if (OverriddenProperties && MapOverrideOp != EOverriddenPropertyOperation::None)
+					if(!bReplaceMap)
 					{
-						checkf(Defaults, TEXT("Expecting overridable serialization to have defaults to compare to"));
-						FScriptMapHelper DefaultsMapHelper(this, Defaults);
+						checkf(!KeyProp->HasAnyPropertyFlags(CPF_PersistentInstance) || CastField<FClassProperty>(KeyProp) || !CastField<FObjectProperty>(KeyProp), TEXT("The key as an instanced sub object is NYI"));
+
+						// For instanced subobject, the overridable is handled per object base, not only here, so we need to serialize the object ptr no matter what.
+						const bool bAreValuesInstancedSubObjects = ValueProp->HasAnyPropertyFlags(CPF_PersistentInstance) && CastField<FObjectProperty>(ValueProp);
+						if(bAreValuesInstancedSubObjects)
+						{
+							for (FScriptMapHelper::FIterator It(MapHelper); It; ++It)
+							{
+								ModifiedIndices.Add(It.GetInternalIndex());
+							}
+						}
 
-						if (const FOverriddenPropertyNode* MapOverriddenPropertyNode = OverriddenProperties->GetOverriddenPropertyNode(UnderlyingArchive.GetSerializedPropertyChain()))
+						if (OverriddenProperties && MapOverrideOp != EOverriddenPropertyOperation::None)
 						{
-							// Figure out the modifications of the map
-							for (const auto& Pair : MapOverriddenPropertyNode->SubPropertyNodeKeys)
+							checkf(Defaults, TEXT("Expecting overridable serialization to have defaults to compare to"));
+							FScriptMapHelper DefaultsMapHelper(this, Defaults);
+
+							if (const FOverriddenPropertyNode* MapOverriddenPropertyNode = OverriddenProperties->GetOverriddenPropertyNode(UnderlyingArchive.GetSerializedPropertyChain()))
 							{
-								const EOverriddenPropertyOperation OverrideOp = OverriddenProperties->GetSubPropertyOperation(Pair.Value);
-								switch (OverrideOp)
+								// Figure out the modifications of the map
+								for (const auto& Pair : MapOverriddenPropertyNode->SubPropertyNodeKeys)
 								{
-								case EOverriddenPropertyOperation::Remove:
+									const EOverriddenPropertyOperation OverrideOp = OverriddenProperties->GetSubPropertyOperation(Pair.Value);
+									switch (OverrideOp)
 									{
-										const int32 InternalIndex = FindKeyInternalIndex(Pair.Key, DefaultsMapHelper);
-										if (InternalIndex != INDEX_NONE)
+									case EOverriddenPropertyOperation::Remove:
 										{
-											RemovedIndices.Add(InternalIndex);
+											const int32 InternalIndex = FindKeyInternalIndex(Pair.Key, DefaultsMapHelper);
+											if (InternalIndex != INDEX_NONE)
+											{
+												RemovedIndices.Add(InternalIndex);
+											}
+											break;
 										}
-										break;
-									}
-								case EOverriddenPropertyOperation::Add:
-									{
-										const int32 InternalIndex = FindKeyInternalIndex(Pair.Key, MapHelper);
-										if (InternalIndex != INDEX_NONE)
+									case EOverriddenPropertyOperation::Add:
 										{
-											AddedIndices.Add(InternalIndex);
-											ModifiedIndices.Remove(InternalIndex);
+											const int32 InternalIndex = FindKeyInternalIndex(Pair.Key, MapHelper);
+											if (InternalIndex != INDEX_NONE)
+											{
+												AddedIndices.Add(InternalIndex);
+												ModifiedIndices.Remove(InternalIndex);
+											}
+											break;
 										}
-										break;
-									}
-								case EOverriddenPropertyOperation::Modified:
-									if (!bAreValuesInstancedSubObjects)
-									{
-										const int32 InternalIndex = FindKeyInternalIndex(Pair.Key, MapHelper);
-										if (InternalIndex != INDEX_NONE)
+									case EOverriddenPropertyOperation::Modified:
+										if (!bAreValuesInstancedSubObjects)
 										{
-											ModifiedIndices.Add(InternalIndex);
+											const int32 InternalIndex = FindKeyInternalIndex(Pair.Key, MapHelper);
+											if (InternalIndex != INDEX_NONE)
+											{
+												ModifiedIndices.Add(InternalIndex);
+											}
 										}
+										break;
+									default:
+										checkf(false, TEXT("Unsupported map operation"));
+										break;
 									}
-									break;
-								default:
-									checkf(false, TEXT("Unsupported map operation"));
-									break;
 								}
 							}
 						}
 					}
 				}
-			}
 
-			int32 NumReplaced = bReplaceMap ? MapHelper.Num() : INDEX_NONE;
-			FStructuredArchive::FArray ReplacedArray = Record.EnterArray(TEXT("Replaced"), NumReplaced);
-			if (bReplaceMap)
-			{
-				for (FScriptMapHelper::FIterator It(MapHelper); It; ++It)
+				int32 NumReplaced = bReplaceMap ? MapHelper.Num() : INDEX_NONE;
+				FStructuredArchive::FArray ReplacedArray = Record.EnterArray(TEXT("Replaced"), NumReplaced);
+				if (bReplaceMap)
 				{
-					FStructuredArchive::FRecord EntryRecord = ReplacedArray.EnterElement().EnterRecord();
-					uint8* PairPtr = MapHelper.GetPairPtr(It.GetInternalIndex());
+					for (FScriptMapHelper::FIterator It(MapHelper); It; ++It)
 					{
-						FSerializedPropertyScope SerializedProperty(UnderlyingArchive, KeyProp, this);
-						KeyProp->SerializeItem(EntryRecord.EnterField(TEXT("Key")), PairPtr);
-					}
-					{
-						FSerializedPropertyScope SerializedProperty(UnderlyingArchive, ValueProp, this);
-						ValueProp->SerializeItem(EntryRecord.EnterField(TEXT("Value")), PairPtr + MapLayout.ValueOffset);
+						FStructuredArchive::FRecord EntryRecord = ReplacedArray.EnterElement().EnterRecord();
+						uint8* PairPtr = MapHelper.GetPairPtr(It.GetInternalIndex());
+						{
+							FSerializedPropertyScope SerializedProperty(UnderlyingArchive, KeyProp, this);
+							KeyProp->SerializeItem(EntryRecord.EnterField(TEXT("Key")), PairPtr);
+						}
+						{
+							FSerializedPropertyScope SerializedProperty(UnderlyingArchive, ValueProp, this);
+							ValueProp->SerializeItem(EntryRecord.EnterField(TEXT("Value")), PairPtr + MapLayout.ValueOffset);
+						}
 					}
 				}
-			}
-			else
-			{
-				checkf(Defaults, TEXT("Expecting overridable serialization to have defaults to compare to"));
-				FScriptMapHelper DefaultsMapHelper(this, Defaults);
-
-				int32 NumRemoved = bReplaceMap ? INDEX_NONE : RemovedIndices.Num();
-				FStructuredArchive::FArray RemovedArray = Record.EnterArray(TEXT("Removed"), NumRemoved);
-				for (int32 i : RemovedIndices)
+				else
 				{
-					FStructuredArchive::FRecord EntryRecord = RemovedArray.EnterElement().EnterRecord();
-					FSerializedPropertyScope SerializedProperty(UnderlyingArchive, KeyProp, this);
-					KeyProp->SerializeItem(EntryRecord.EnterField(TEXT("Key")), DefaultsMapHelper.GetKeyPtr(i));
-				}
+					checkf(Defaults, TEXT("Expecting overridable serialization to have defaults to compare to"));
+					FScriptMapHelper DefaultsMapHelper(this, Defaults);
 
-				int32 NumModified = ModifiedIndices.Num();
-				FStructuredArchive::FArray ModifiedArray = Record.EnterArray(TEXT("Modified"), NumModified);
-				for (int32 i : ModifiedIndices)
-				{
-					FStructuredArchive::FRecord EntryRecord = ModifiedArray.EnterElement().EnterRecord();
-					uint8* PairPtr = MapHelper.GetPairPtr(i);
+					int32 NumRemoved = bReplaceMap ? INDEX_NONE : RemovedIndices.Num();
+					FStructuredArchive::FArray RemovedArray = Record.EnterArray(TEXT("Removed"), NumRemoved);
+					for (int32 i : RemovedIndices)
 					{
+						FStructuredArchive::FRecord EntryRecord = RemovedArray.EnterElement().EnterRecord();
 						FSerializedPropertyScope SerializedProperty(UnderlyingArchive, KeyProp, this);
-						KeyProp->SerializeItem(EntryRecord.EnterField(TEXT("Key")), PairPtr);
-					}
-					{
-						FSerializedPropertyScope SerializedProperty(UnderlyingArchive, ValueProp, this);
-						ValueProp->SerializeItem(EntryRecord.EnterField(TEXT("Value")), PairPtr + MapLayout.ValueOffset);
+						KeyProp->SerializeItem(EntryRecord.EnterField(TEXT("Key")), DefaultsMapHelper.GetKeyPtr(i));
 					}
-				}
 
-				// Added keys
-				int32 NumAdded = AddedIndices.Num();
-				FStructuredArchive::FArray AddedArray = Record.EnterArray(TEXT("Added"), NumAdded);
-				for (int32 i : AddedIndices)
-				{
-					FStructuredArchive::FRecord EntryRecord = AddedArray.EnterElement().EnterRecord();
-					uint8* PairPtr = MapHelper.GetPairPtr(i);
+					int32 NumModified = ModifiedIndices.Num();
+					FStructuredArchive::FArray ModifiedArray = Record.EnterArray(TEXT("Modified"), NumModified);
+					for (int32 i : ModifiedIndices)
 					{
-						FSerializedPropertyScope SerializedProperty(UnderlyingArchive, KeyProp, this);
-						KeyProp->SerializeItem(EntryRecord.EnterField(TEXT("Key")), PairPtr);
+						FStructuredArchive::FRecord EntryRecord = ModifiedArray.EnterElement().EnterRecord();
+						uint8* PairPtr = MapHelper.GetPairPtr(i);
+						{
+							FSerializedPropertyScope SerializedProperty(UnderlyingArchive, KeyProp, this);
+							KeyProp->SerializeItem(EntryRecord.EnterField(TEXT("Key")), PairPtr);
+						}
+						{
+							FSerializedPropertyScope SerializedProperty(UnderlyingArchive, ValueProp, this);
+							ValueProp->SerializeItem(EntryRecord.EnterField(TEXT("Value")), PairPtr + MapLayout.ValueOffset);
+						}
 					}
+
+					// Added keys
+					int32 NumAdded = AddedIndices.Num();
+					FStructuredArchive::FArray AddedArray = Record.EnterArray(TEXT("Added"), NumAdded);
+					for (int32 i : AddedIndices)
 					{
-						FSerializedPropertyScope SerializedProperty(UnderlyingArchive, ValueProp, this);
-						ValueProp->SerializeItem(EntryRecord.EnterField(TEXT("Value")), PairPtr + MapLayout.ValueOffset);
+						FStructuredArchive::FRecord EntryRecord = AddedArray.EnterElement().EnterRecord();
+						uint8* PairPtr = MapHelper.GetPairPtr(i);
+						{
+							FSerializedPropertyScope SerializedProperty(UnderlyingArchive, KeyProp, this);
+							KeyProp->SerializeItem(EntryRecord.EnterField(TEXT("Key")), PairPtr);
+						}
+						{
+							FSerializedPropertyScope SerializedProperty(UnderlyingArchive, ValueProp, this);
+							ValueProp->SerializeItem(EntryRecord.EnterField(TEXT("Value")), PairPtr + MapLayout.ValueOffset);
+						}
 					}
 				}
 			}
-		}
 
-		return;
+			return;
+		}
 	}
 
 	if (UnderlyingArchive.IsLoading())
