Fix the low level FEditorBulkData tests from failing on runtime platforms.

#rnx

- The FEditorBulkData tests should only run on platforms with WITH_EDITORONLY_DATA enabled, this was lost when moving them from automation tests.
- Fixed up some poorly formatted includes.

diff --git a/Engine/Source/Runtime/CoreUObject/Tests/Serialization/EditorBulkDataTests.cpp b/Engine/Source/Runtime/CoreUObject/Tests/Serialization/EditorBulkDataTests.cpp
--- a/Engine/Source/Runtime/CoreUObject/Tests/Serialization/EditorBulkDataTests.cpp
+++ b/Engine/Source/Runtime/CoreUObject/Tests/Serialization/EditorBulkDataTests.cpp
@@ -1,6 +1,6 @@
 // Copyright Epic Games, Inc. All Rights Reserved.
 
-#if WITH_LOW_LEVEL_TESTS
+#if WITH_LOW_LEVEL_TESTS && WITH_EDITORONLY_DATA
 
 #include "Algo/AllOf.h"
 #include "Async/ParallelFor.h"
@@ -16,10 +16,10 @@
 #include "Serialization/MemoryWriter.h"
 #include "Tasks/Task.h"
 #include "Templates/UniquePtr.h"
-
+#include "TestCommon/Expectations.h"
 #include "TestHarness.h"
 #include "TestMacros/Assertions.h"
-#include "TestCommon/Expectations.h"
+
 #include <catch2/generators/catch_generators.hpp>
 
 namespace UE::Serialization
@@ -824,4 +824,4 @@ TEST_CASE("CoreUObject::Serialization::FEditorBulkData::Threading_Assignment", "
 
 } // namespace UE::Serialization
 
-#endif // WITH_LOW_LEVEL_TESTS
+#endif // WITH_LOW_LEVEL_TESTS && WITH_EDITORONLY_DATA
