Fixed SerializeItem use in FMapProperty::ConvertFromType

#rnx

diff --git a/Engine/Source/Runtime/CoreUObject/Private/UObject/PropertyMap.cpp b/Engine/Source/Runtime/CoreUObject/Private/UObject/PropertyMap.cpp
--- a/Engine/Source/Runtime/CoreUObject/Private/UObject/PropertyMap.cpp
+++ b/Engine/Source/Runtime/CoreUObject/Private/UObject/PropertyMap.cpp
@@ -1448,12 +1448,18 @@ EConvertFromTypeResult FMapProperty::ConvertFromType(const FPropertyTag& Tag, FS
 			const int64 StartOfProperty = InnerSlot.GetUnderlyingArchive().Tell();
 			switch (Inner->ConvertFromType(InnerTag, InnerSlot, InnerData, InnerDefaultsStruct, nullptr))
 			{
+				default:
+					checkNoEntry();
 				case EConvertFromTypeResult::Converted:
 				case EConvertFromTypeResult::Serialized:
 					return true;
 				case EConvertFromTypeResult::CannotConvert:
 					return false;
 				case EConvertFromTypeResult::UseSerializeItem:
+					if (InnerTag.Type != Inner->GetID())
+					{
+						return false;
+					}
 					// Fall through to default SerializeItem
 					break;
 			}
