- Fixed FLinkerInstancedPackageMap::BuildPackageMapping to be case insensitive
- Added optional SearchCase parameter to TStringView<>::Contains and TStringView<>::Find defaulted to ESearchCase::CaseSensitive for backward compatibility

diff --git a/Engine/Source/Runtime/Core/Public/Containers/StringView.h b/Engine/Source/Runtime/Core/Public/Containers/StringView.h
--- a/Engine/Source/Runtime/Core/Public/Containers/StringView.h
+++ b/Engine/Source/Runtime/Core/Public/Containers/StringView.h
@@ -309,19 +309,21 @@ public:
 	 *
 	 * @param Search          The string to search for. Comparison is lexicographic.
 	 * @param StartPosition   The character position to start searching from.
+	 * @param SearchCase      Indicates whether the search is case sensitive or not
 	 * @return The index of the first occurrence of the search string if found, otherwise INDEX_NONE.
 	 */
-	[[nodiscard]] inline int32 Find(ViewType Search, int32 StartPosition = 0) const;
+	[[nodiscard]] inline int32 Find(ViewType Search, int32 StartPosition = 0, ESearchCase::Type SearchCase = ESearchCase::CaseSensitive) const;
 
 	/**
 	 * Returns whether this view contains the specified substring.
 	 *
-	 * @param Search   Text to search for
+	 * @param Search          Text to search for
+	 * @param SearchCase      Indicates whether the search is case sensitive or not
 	 * @return True if the view contains the search string, otherwise false.
 	 */
-	[[nodiscard]] inline bool Contains(ViewType Search) const
+	[[nodiscard]] inline bool Contains(ViewType Search, ESearchCase::Type SearchCase = ESearchCase::CaseSensitive) const
 	{
-		return Find(Search) != INDEX_NONE;
+		return Find(Search, 0, SearchCase) != INDEX_NONE;
 	}
 
 	/**
@@ -713,9 +715,9 @@ inline bool TStringView<CharType>::EndsWith(ViewType Suffix, ESearchCase::Type S
 }
 
 template <typename CharType>
-inline int32 TStringView<CharType>::Find(const ViewType Search, const int32 StartPosition) const
+inline int32 TStringView<CharType>::Find(const ViewType Search, const int32 StartPosition, ESearchCase::Type SearchCase) const
 {
-	const int32 Index = UE::String::FindFirst(RightChop(StartPosition), Search);
+	const int32 Index = UE::String::FindFirst(RightChop(StartPosition), Search, SearchCase);
 	return Index == INDEX_NONE ? INDEX_NONE : Index + StartPosition;
 }
 
diff --git a/Engine/Source/Runtime/CoreUObject/Private/UObject/LinkerInstancingContext.cpp b/Engine/Source/Runtime/CoreUObject/Private/UObject/LinkerInstancingContext.cpp
--- a/Engine/Source/Runtime/CoreUObject/Private/UObject/LinkerInstancingContext.cpp
+++ b/Engine/Source/Runtime/CoreUObject/Private/UObject/LinkerInstancingContext.cpp
@@ -323,7 +323,7 @@ void FLinkerInstancedPackageMap::BuildPackageMapping(FName Original, FName Insta
 		FStringView OriginalView = TmpOriginal.ToView();
 		FStringView InstancedView = TmpInstanced.ToView();
 
-		const int32 Index = InstancedView.Find(OriginalView);
+		const int32 Index = InstancedView.Find(OriginalView, 0, ESearchCase::IgnoreCase);
 
 		// Stash the suffix used for this instance so we can also apply it to generated packages
 		if (Index != INDEX_NONE)
@@ -338,11 +338,11 @@ void FLinkerInstancedPackageMap::BuildPackageMapping(FName Original, FName Insta
 			const FStringView GeneratedFolderName = TEXTVIEW("/_Generated_/");
 			
 			// Does this package path include the generated folder?
-			if (const int32 GeneratedFolderStartIndex = OriginalView.Find(GeneratedFolderName); GeneratedFolderStartIndex != INDEX_NONE)
+			if (const int32 GeneratedFolderStartIndex = OriginalView.Find(GeneratedFolderName, 0, ESearchCase::IgnoreCase); GeneratedFolderStartIndex != INDEX_NONE)
 			{
 				// ... and is that generated folder immediately preceding the package name?
 				const int32 GeneratedFolderEndIndex = GeneratedFolderStartIndex + GeneratedFolderName.Len();
-				if (const int32 ExtraSlashIndex = OriginalView.Find(TEXTVIEW("/"), GeneratedFolderEndIndex); ExtraSlashIndex == INDEX_NONE)
+				if (const int32 ExtraSlashIndex = OriginalView.Find(TEXTVIEW("/"), GeneratedFolderEndIndex, ESearchCase::IgnoreCase); ExtraSlashIndex == INDEX_NONE)
 				{
 					GeneratedPackagesFolder = OriginalView.Left(GeneratedFolderEndIndex);
 					const FString PersistentSourcePackage(OriginalView.Left(GeneratedFolderStartIndex));
@@ -389,12 +389,12 @@ bool FLinkerInstancedPackageMap::FixupSoftObjectPath(FSoftObjectPath& InOutSoftO
 
 			// Does this package path start with the generated folder path?
 			FStringView TmpSoftObjectPathView = TmpSoftObjectPathBuilder.ToView();
-			if (const int32 GeneratedFolderIndex = TmpSoftObjectPathView.Find(GeneratedPackagesFolder); GeneratedFolderIndex != INDEX_NONE)
+			if (const int32 GeneratedFolderIndex = TmpSoftObjectPathView.Find(GeneratedPackagesFolder, 0, ESearchCase::IgnoreCase); GeneratedFolderIndex != INDEX_NONE)
 			{
 				check(GeneratedFolderIndex == 0 || (TmpSoftObjectPathView.StartsWith(InstancedPackagePrefix) && GeneratedFolderIndex == InstancedPackagePrefix.Len()));
 
 				// ... and is that generated folder path immediately preceding the package name?
-				if (const int32 ExtraSlashIndex = TmpSoftObjectPathView.Find(TEXTVIEW("/"), InstancedPackagePrefix.Len() + GeneratedPackagesFolder.Len()); ExtraSlashIndex == INDEX_NONE)
+				if (const int32 ExtraSlashIndex = TmpSoftObjectPathView.Find(TEXTVIEW("/"), InstancedPackagePrefix.Len() + GeneratedPackagesFolder.Len(), ESearchCase::IgnoreCase); ExtraSlashIndex == INDEX_NONE)
 				{
 					FNameBuilder PackageNameBuilder;
 
