[ZenLoader]
  - Add stack trace when logging invalid flushes error so we can more easily identify their provenance when running test builds
  - Add unittest to validate the error

diff --git a/Engine/Source/Runtime/CoreUObject/Private/Serialization/AsyncLoading2.cpp b/Engine/Source/Runtime/CoreUObject/Private/Serialization/AsyncLoading2.cpp
--- a/Engine/Source/Runtime/CoreUObject/Private/Serialization/AsyncLoading2.cpp
+++ b/Engine/Source/Runtime/CoreUObject/Private/Serialization/AsyncLoading2.cpp
@@ -13,6 +13,7 @@
 #include "HAL/PlatformFileManager.h"
 #include "HAL/FileManager.h"
 #include "HAL/Event.h"
+#include "HAL/ThreadManager.h"
 #include "HAL/RunnableThread.h"
 #include "HAL/PlatformMisc.h"
 #include "Misc/ScopeLock.h"
@@ -9376,11 +9377,21 @@ void FAsyncLoadingThread2::FlushLoading(TConstArrayView<int32> RequestIDs)
 #endif
 		ELoaderType LoaderType = GetLoaderType();
 
-		ensureMsgf(bIsFlushSupportedOnCurrentThread, TEXT("The current loader '%s' is unable to FlushAsyncLoading from the current thread."), LexToString(LoaderType));
-
 		if (!bIsFlushSupportedOnCurrentThread)
 		{
-			UE_LOG(LogStreaming, Error, TEXT("The current loader '%s' is unable to FlushAsyncLoading from the current thread. Flush will be ignored."), LexToString(LoaderType));
+			// As there is no ensure for test builds, use custom stacktrace logging so we can see these problems in playtest and fix them.
+			// But just return in shipping build to avoid crashing as the side effect of missing flushes are not always fatal.
+			// Don't forget to adjust unittest in FLoadingTests_InvalidFlush_FromWorker if changing this.
+#if !UE_BUILD_SHIPPING
+			const FString& ThreadName = FThreadManager::GetThreadName(FPlatformTLS::GetCurrentThreadId());
+			const FString Heading = 
+				FString::Printf(
+					TEXT("The current loader '%s' is unable to FlushAsyncLoading from the current thread '%s'. Flush will be ignored."), 
+					LexToString(LoaderType),
+					*ThreadName
+			);
+			FDebug::DumpStackTraceToLog(*Heading, ELogVerbosity::Error);
+#endif
 			return;
 		}
 
@@ -9601,4 +9612,4 @@ void FGlobalImportStore::FlushDeferredDeletePackagesQueue()
 IAsyncPackageLoader* MakeAsyncPackageLoader2(FIoDispatcher& InIoDispatcher, IAsyncPackageLoader* InUncookedPackageLoader)
 {
 	return new FAsyncLoadingThread2(InIoDispatcher, InUncookedPackageLoader);
-}
\ No newline at end of file
+}
diff --git a/Engine/Source/Runtime/Engine/Private/Tests/Loading/AsyncLoadingTests_InvalidFlush.cpp b/Engine/Source/Runtime/Engine/Private/Tests/Loading/AsyncLoadingTests_InvalidFlush.cpp
new file mode 100644
--- /dev/null
+++ b/Engine/Source/Runtime/Engine/Private/Tests/Loading/AsyncLoadingTests_InvalidFlush.cpp
@@ -0,0 +1,51 @@
+// Copyright Epic Games, Inc. All Rights Reserved.
+
+#include "AsyncLoadingTests_Shared.h"
+#include "Misc/AutomationTest.h"
+#include "Misc/PackageName.h"
+#include "Async/ManualResetEvent.h"
+#include "UObject/CoreRedirects.h"
+#include "UObject/SavePackage.h"
+#include "Tasks/Task.h"
+
+#if WITH_DEV_AUTOMATION_TESTS
+
+// All InvalidFlush tests should run on zenloader only as the other loaders are not compliant.
+typedef FLoadingTests_ZenLoaderOnly_Base FLoadingTests_InvalidFlush_Base;
+
+/**
+ * This test validates loading an object synchronously during serialize.
+ */
+IMPLEMENT_CUSTOM_SIMPLE_AUTOMATION_TEST(
+	FLoadingTests_InvalidFlush_FromWorker,
+	FLoadingTests_InvalidFlush_Base,
+	TEXT("System.Engine.Loading.InvalidFlush.FromWorker"),
+	EAutomationTestFlags::EditorContext | EAutomationTestFlags::EngineFilter
+)
+bool FLoadingTests_InvalidFlush_FromWorker::RunTest(const FString& Parameters)
+{
+	AddExpectedError(TEXT("is unable to FlushAsyncLoading from the current thread"), EAutomationExpectedErrorFlags::Contains);
+	AddExpectedError(TEXT("[Callstack]"), EAutomationExpectedErrorFlags::Contains, 0 /* At least 1 occurrence */);
+	
+	FLoadingTestsScope LoadingTestScope(this);
+
+	UAsyncLoadingTests_Shared::OnSerialize.BindLambda(
+		[this](FArchive& Ar, UAsyncLoadingTests_Shared* Object)
+		{
+			if (Ar.IsLoading())
+			{
+				// Use event instead of waiting on the task to prevent retraction as we really want that task
+				// to execute on a worker thread instead of being retracted in the serialize thread.
+				UE::FManualResetEvent Event;
+				UE::Tasks::Launch(TEXT("FlushAsyncLoading"), [&Event]() { FlushAsyncLoading(); Event.Notify(); });
+				Event.Wait();
+			}
+		}
+	);
+
+	LoadingTestScope.LoadObjects();
+
+	return true;
+}
+
+#endif // WITH_DEV_AUTOMATION_TESTS
