Better support of transaction using TPS for object that has overridable serialization enabled. Previously only supported UPS.

Repro :
1. Create an entity
2. Delete it
3. Undo the delete
4. Save the map
5. Reload the map
Result: Entity was reloaded without any components
Expected: Entity to reload with all its components.

diff --git a/Engine/Source/Runtime/CoreUObject/Private/UObject/Class.cpp b/Engine/Source/Runtime/CoreUObject/Private/UObject/Class.cpp
--- a/Engine/Source/Runtime/CoreUObject/Private/UObject/Class.cpp
+++ b/Engine/Source/Runtime/CoreUObject/Private/UObject/Class.cpp
@@ -1430,8 +1430,12 @@ void UStruct::SerializeVersionedTaggedProperties(FStructuredArchive::FSlot Slot,
 			if (UnderlyingArchive.IsLoading())
 			{
 				ControlContext.bEnableOverridableSerialization = true;
-				ControlContext.OverriddenProperties = &FOverridableManager::Get().SetOverriddenProperties(*(UObject*)Data, Operation);
-				ControlContext.OverriddenProperties->bNeedsSubobjectTemplateInstantiation = true;
+				// Overridden values are saved independently in transaction, so do no need to restore them here.
+				if (!UnderlyingArchive.IsTransacting())
+				{
+					ControlContext.OverriddenProperties = &FOverridableManager::Get().SetOverriddenProperties(*(UObject*)Data, Operation);
+					ControlContext.OverriddenProperties->bNeedsSubobjectTemplateInstantiation = true;
+				}
 			}
 		}
 	}
@@ -1661,12 +1665,16 @@ void UStruct::SerializeVersionedTaggedProperties(FStructuredArchive::FSlot Slot,
 						FStructuredArchive::FSlot ValueSlot = PropertyRecord.EnterField(TEXT("Value"));
 
 						// The operation was set part of the tag, now that we know the associated property, restore the overridden operation on the object
-						if (FOverriddenPropertySet* OverriddenProperties = FOverridableSerializationLogic::GetOverriddenProperties())
+						// No need to rebuild the overridden state in transaction as it was serialized is one chunk
+						if (!UnderlyingArchive.IsTransacting())
 						{
-							// No need to restore none operations
-							if (Tag.OverrideOperation != EOverriddenPropertyOperation::None)
+							if (FOverriddenPropertySet* OverriddenProperties = FOverridableSerializationLogic::GetOverriddenProperties())
 							{
-								OverriddenProperties->SetOverriddenPropertyOperation(Tag.OverrideOperation, UnderlyingArchive.GetSerializedPropertyChain(), Property);
+								// No need to restore none operations
+								if (Tag.OverrideOperation != EOverriddenPropertyOperation::None)
+								{
+									OverriddenProperties->SetOverriddenPropertyOperation(Tag.OverrideOperation, UnderlyingArchive.GetSerializedPropertyChain(), Property);
+								}
 							}
 						}
 
diff --git a/Engine/Source/Runtime/CoreUObject/Private/UObject/Obj.cpp b/Engine/Source/Runtime/CoreUObject/Private/UObject/Obj.cpp
--- a/Engine/Source/Runtime/CoreUObject/Private/UObject/Obj.cpp
+++ b/Engine/Source/Runtime/CoreUObject/Private/UObject/Obj.cpp
@@ -1648,8 +1648,8 @@ void UObject::Serialize(FStructuredArchive::FRecord Record)
 		if (ObjClass != UClass::StaticClass())
 		{
 			// Serializing the overriden properties if using UPS (like object duplication)
-			// @todo: We probably do not want this in cooked packages
-			if (UnderlyingArchive.UseUnversionedPropertySerialization())
+			// Also needed for transactions using TPS
+			if (UnderlyingArchive.UseUnversionedPropertySerialization() || UnderlyingArchive.IsTransacting())
 			{
 				FOverridableManager::Get().SerializeOverriddenProperties(*this, Record);
 			}
