Fix SubObject being saved only if it contained overrides but it would not account for sub sub objects that contained overrides, basically they would be skipped.

Added a unittest for it.

diff --git a/Engine/Source/Runtime/CoreUObject/Private/UObject/OverriddenPropertySet.cpp b/Engine/Source/Runtime/CoreUObject/Private/UObject/OverriddenPropertySet.cpp
--- a/Engine/Source/Runtime/CoreUObject/Private/UObject/OverriddenPropertySet.cpp
+++ b/Engine/Source/Runtime/CoreUObject/Private/UObject/OverriddenPropertySet.cpp
@@ -50,70 +50,6 @@ EOverriddenPropertyOperation FOverridableSerializationLogic::GetOverriddenProper
 				return 	EOverriddenPropertyOperation::Replace;
 			}
 		}
-
-		// Here we should just use CurrentProperty->ContainsInstancedObjectProperty() but somehow this CPF_InstanceReference seems to always be there on verse classes
-		// This should probably be fix in verse at some point.
-		else if (CurrentProperty->HasAnyPropertyFlags(CPF_PersistentInstance))
-		{
-			if (const FObjectProperty* ObjectProperty = CastField<FObjectProperty>(CurrentProperty))
-			{
-				if (UObject* SubObject = ObjectProperty->GetObjectPropertyValue(DataPtr))
-				{
-					if(const FOverriddenPropertySet* OverriddenSubObjectProperties = FOverridableManager::Get().GetOverriddenProperties(*SubObject))
-					{
-						const EOverriddenPropertyOperation SubObjectOperation = OverriddenSubObjectProperties->GetOverriddenPropertyOperation((FArchiveSerializedPropertyChain*)nullptr, (FProperty*)nullptr);
-						if (SubObjectOperation != EOverriddenPropertyOperation::None)
-						{
-							return EOverriddenPropertyOperation::Modified;
-						}
-					}
-				}
-			}
-		}
-		else if (const FArrayProperty* CurrentArrayProperty = CastField<FArrayProperty>(CurrentProperty))
-		{
-			if (const FObjectProperty* ObjectProperty = CurrentArrayProperty->Inner->HasAnyPropertyFlags(CPF_PersistentInstance) ? CastField<FObjectProperty>(CurrentArrayProperty->Inner) : nullptr)
-			{
-				FScriptArrayHelper ArrayHelper(CurrentArrayProperty, DataPtr);
-
-				for (int i = 0; i < ArrayHelper.Num(); ++i)
-				{
-					if (UObject* SubObject = ObjectProperty->GetObjectPropertyValue(ArrayHelper.GetElementPtr(i)))
-					{
-						if(const FOverriddenPropertySet* OverriddenSubObjectProperties = FOverridableManager::Get().GetOverriddenProperties(*SubObject))
-						{
-							const EOverriddenPropertyOperation SubObjectOperation = OverriddenSubObjectProperties->GetOverriddenPropertyOperation((FArchiveSerializedPropertyChain*)nullptr, (FProperty*)nullptr);
-							if (SubObjectOperation != EOverriddenPropertyOperation::None)
-							{
-								return EOverriddenPropertyOperation::Modified;
-							}
-						}
-					}
-				}
-			}
-		}
-		else if(const FMapProperty* CurrentMapProperty = CastField<FMapProperty>(CurrentProperty))
-		{
-			if (const FObjectProperty* ValueInstancedObjectProperty = CurrentMapProperty->ValueProp->HasAnyPropertyFlags(CPF_PersistentInstance) ? CastField<FObjectProperty>(CurrentMapProperty->ValueProp) : nullptr)
-			{
-				FScriptMapHelper MapHelper(CurrentMapProperty, DataPtr);
-
-				for (FScriptMapHelper::FIterator It(MapHelper); It; ++It)
-				{
-					if (UObject* SubObject = ValueInstancedObjectProperty->GetObjectPropertyValue(MapHelper.GetValuePtr(It.GetInternalIndex())))
-					{
-						if(const FOverriddenPropertySet* OverriddenSubObjectProperties = FOverridableManager::Get().GetOverriddenProperties(*SubObject))
-						{
-							const EOverriddenPropertyOperation SubObjectOperation = OverriddenSubObjectProperties->GetOverriddenPropertyOperation((FArchiveSerializedPropertyChain*)nullptr, (FProperty*)nullptr);
-							if (SubObjectOperation != EOverriddenPropertyOperation::None)
-							{
-								return EOverriddenPropertyOperation::Modified;
-							}
-						}
-					}
-				}
-			}
-		}
 	}
 	
 	return EOverriddenPropertyOperation::None;
@@ -1206,7 +1142,7 @@ FOverriddenPropertyNode* FOverriddenPropertySet::SetSubPropertyOperation(EOverri
 
 bool FOverriddenPropertySet::IsCDOOwningProperty(const FProperty& Property) const
 {
-	checkf(IsValid(Owner), TEXT("Expecting a valid overridable owner"));
+	checkf(Owner, TEXT("Expecting a valid overridable owner"));
 	if (!Owner->HasAnyFlags(RF_ClassDefaultObject))
 	{
 		return false;
diff --git a/Engine/Source/Runtime/CoreUObject/Private/UObject/SavePackage/PackageHarvester.cpp b/Engine/Source/Runtime/CoreUObject/Private/UObject/SavePackage/PackageHarvester.cpp
--- a/Engine/Source/Runtime/CoreUObject/Private/UObject/SavePackage/PackageHarvester.cpp
+++ b/Engine/Source/Runtime/CoreUObject/Private/UObject/SavePackage/PackageHarvester.cpp
@@ -3,7 +3,9 @@
 
 #include "UObject/SavePackage/PackageHarvester.h"
 
+#include "InstancedReferenceSubobjectHelper.h"
 #include "Interfaces/ITargetPlatform.h"
+#include "UObject/OverridableManager.h"
 #include "UObject/SavePackage/SaveContext.h"
 #include "UObject/SavePackage/SavePackageUtilities.h"
 #include "UObject/UObjectGlobals.h"
@@ -625,6 +627,18 @@ void FPackageHarvester::TryHarvestExportInternal(UObject* InObject)
 	}
 
 	HarvestExport(InObject);
+
+	// Objects that has overridable serialization enabled, needs to separately export their subobject independently of the property serializing it.
+	// This is because it needs to make a difference between the pointer that is overridden vs some properties inside the instanced sub object that are overridden.
+	if (FOverridableManager::Get().IsEnabled(*InObject))
+	{
+		TSet<UObject*> InstancedSubObjects;
+		FFindInstancedReferenceSubobjectHelper::GetInstancedSubObjects(InObject, InstancedSubObjects);
+		for (UObject* InstancedSubObject : InstancedSubObjects)
+		{
+			TryHarvestExportInternal(InstancedSubObject);
+		}
+	}
 }
 
 void FPackageHarvester::TryHarvestImport(TObjectPtr<UObject> InObject)
