Updated FPropertyTag to use FPropertyTypeName to represent the property type

This is a step toward deprecating the existing names used for these details and serializing the property type using its new TypeName member.

#rnx

diff --git a/Engine/Plugins/Experimental/InstanceDataObjectFixupTool/Source/Private/InstanceDataObjectFixupPanel.cpp b/Engine/Plugins/Experimental/InstanceDataObjectFixupTool/Source/Private/InstanceDataObjectFixupPanel.cpp
--- a/Engine/Plugins/Experimental/InstanceDataObjectFixupTool/Source/Private/InstanceDataObjectFixupPanel.cpp
+++ b/Engine/Plugins/Experimental/InstanceDataObjectFixupTool/Source/Private/InstanceDataObjectFixupPanel.cpp
@@ -538,7 +538,7 @@ bool FInstanceDataObjectFixupPanel::FTypeConverter::TryConvert(FProperty* Source
 	void* DestinationContainer = static_cast<uint8*>(DestinationData) - DestinationProperty->GetOffset_ForInternal();
 
 	// todo: handle static arrays
-	FPropertyTag SourceTag(MemoryReader, SourceProperty, 0, (uint8*)SourceData, nullptr);
+	FPropertyTag SourceTag(SourceProperty, 0, (uint8*)SourceData);
 
 	bool bResult = false;
 	switch(DestinationProperty->ConvertFromType(SourceTag, StructuredReader.GetSlot(), (uint8*)DestinationContainer, SourceProperty->GetOwnerStruct(), nullptr))
diff --git a/Engine/Source/Runtime/CoreUObject/Internal/UObject/PropertyBag.h b/Engine/Source/Runtime/CoreUObject/Internal/UObject/PropertyBag.h
--- a/Engine/Source/Runtime/CoreUObject/Internal/UObject/PropertyBag.h
+++ b/Engine/Source/Runtime/CoreUObject/Internal/UObject/PropertyBag.h
@@ -108,7 +108,7 @@ public:
 		inline bool operator!=(const FConstIterator& Rhs) const { return CurrentValue != Rhs.CurrentValue; }
 
 		inline const FPropertyPathName& GetPath() const { return CurrentPath; }
-		inline const FProperty* GetProperty() const { return CurrentValue->Tag.Prop; }
+		inline const FProperty* GetProperty() const { return CurrentValue->Tag.GetProperty(); }
 		inline void* GetValue() const { return CurrentValue->Data; }
 		inline int32 GetValueSize() const { return CurrentValue->GetSize(); }
 
diff --git a/Engine/Source/Runtime/CoreUObject/Private/UObject/Class.cpp b/Engine/Source/Runtime/CoreUObject/Private/UObject/Class.cpp
--- a/Engine/Source/Runtime/CoreUObject/Private/UObject/Class.cpp
+++ b/Engine/Source/Runtime/CoreUObject/Private/UObject/Class.cpp
@@ -1232,8 +1232,6 @@ void UStruct::LoadTaggedPropertiesFromText(FStructuredArchive::FSlot Slot, uint8
 
 		if (Property && Property->ShouldSerializeValue(UnderlyingArchive))
 		{
-			FName PropID = Property->GetID();
-
 			// Static arrays of tagged properties are special cases where the slot is always an array with no tag data attached. We currently have no TryEnterArray we can't 
 			// react based on what is in the file (yet) so we'll just have to assume that nobody converts a property from an array to a single value and go with whatever 
 			// the code property tells us.
@@ -1260,60 +1258,52 @@ void UStruct::LoadTaggedPropertiesFromText(FStructuredArchive::FSlot Slot, uint8
 
 				FPropertyTag Tag;
 				ItemSlot.GetValue() << Tag;
-				Tag.Prop = Property;
+				Tag.SetProperty(Property);
 				Tag.ArrayIndex = ItemIndex;
 				Tag.Name = PropertyName;
 
 				if (bUseRedirects)
 				{
-					if (Tag.Type == NAME_StructProperty && PropID == NAME_StructProperty)
+					if (UE::FPropertyTypeName NewTypeName = ApplyRedirectsToPropertyType(Tag.GetType(), Property); !NewTypeName.IsEmpty())
 					{
-						const FName NewName = FLinkerLoad::FindNewNameForStruct(Tag.StructName);
-						const FName StructName = CastFieldChecked<FStructProperty>(Property)->Struct->GetFName();
-						if (NewName == StructName)
-						{
-							Tag.StructName = NewName;
-						}
+						Tag.SetType(NewTypeName);
 					}
-					else if ((PropID == NAME_EnumProperty) && ((Tag.Type == NAME_EnumProperty) || (Tag.Type == NAME_ByteProperty)))
+				}
+
+				if (BreakRecursionIfFullyLoad && BreakRecursionIfFullyLoad->HasAllFlags(RF_LoadCompleted))
+				{
+					continue;
+				}
+
+				switch (Property->ConvertFromType(Tag, ItemSlot.GetValue(), Data, DefaultsStruct, Defaults))
+				{
+				case EConvertFromTypeResult::Converted:
+				case EConvertFromTypeResult::Serialized:
+					break;
+
+				case EConvertFromTypeResult::UseSerializeItem:
+					if (const FName PropID = Property->GetID(); Tag.Type != PropID)
 					{
-						const FName NewName = FLinkerLoad::FindNewNameForEnum(Tag.EnumName);
-						if (!NewName.IsNone())
-						{
-							Tag.EnumName = NewName;
-						}
+						UE_LOG(LogClass, Warning, TEXT("Type mismatch in %s of %s - Previous (%s) Current(%s) in package: %s"),
+							*WriteToString<32>(Tag.Name), *WriteToString<32>(GetFName()),
+							*WriteToString<32>(Tag.Type), *WriteToString<32>(PropID),
+							*UnderlyingArchive.GetArchiveName());
 					}
-
-					if (!(BreakRecursionIfFullyLoad && BreakRecursionIfFullyLoad->HasAllFlags(RF_LoadCompleted)))
+					else
 					{
-						switch (Property->ConvertFromType(Tag, ItemSlot.GetValue(), Data, DefaultsStruct, Defaults))
-						{
-						case EConvertFromTypeResult::Converted:
-						case EConvertFromTypeResult::Serialized:
-							break;
-
-						case EConvertFromTypeResult::UseSerializeItem:
-							if (Tag.Type != PropID)
-							{
-								UE_LOG(LogClass, Warning, TEXT("Type mismatch in %s of %s - Previous (%s) Current(%s) for package:  %s"), *Tag.Name.ToString(), *GetName(), *Tag.Type.ToString(), *PropID.ToString(), *UnderlyingArchive.GetArchiveName());
-							}
-							else
-							{
-								uint8* DestAddress = Property->ContainerPtrToValuePtr<uint8>(Data, Tag.ArrayIndex);
-								uint8* DefaultsFromParent = Property->ContainerPtrToValuePtrForDefaults<uint8>(DefaultsStruct, Defaults, Tag.ArrayIndex);
+						uint8* DestAddress = Property->ContainerPtrToValuePtr<uint8>(Data, Tag.ArrayIndex);
+						uint8* DefaultsFromParent = Property->ContainerPtrToValuePtrForDefaults<uint8>(DefaultsStruct, Defaults, Tag.ArrayIndex);
 
-								// This property is ok.
-								Tag.SerializeTaggedProperty(ItemSlot.GetValue(), Property, DestAddress, DefaultsFromParent);
-							}
-							break;
+						// This property is ok.
+						Tag.SerializeTaggedProperty(ItemSlot.GetValue(), Property, DestAddress, DefaultsFromParent);
+					}
+					break;
 
-						case EConvertFromTypeResult::CannotConvert:
-							break;
+				case EConvertFromTypeResult::CannotConvert:
+					break;
 
-						default:
-							check(false);
-						}
-					}
+				default:
+					check(false);
 				}
 			}
 		}
@@ -1574,71 +1564,25 @@ void UStruct::SerializeVersionedTaggedProperties(FStructuredArchive::FSlot Slot,
 					Property = CustomFindProperty(Tag.Name);
 				}
 
-				if (Property)
-				{
-					Property->AssignToTag(Tag);
-				}
+				Tag.SetProperty(Property);
 
-				if (SerializeContext && SerializeContext->bTrackSerializedPropertyPath)
+				if (bUseRedirects)
 				{
-					FPropertyTypeNameBuilder TypeBuilder;
-					TypeBuilder.AddName(Tag.Type);
-
-					if (!Tag.StructName.IsNone())
-					{
-						TypeBuilder.BeginParameters();
-						TypeBuilder.AddName(Tag.StructName);
-						TypeBuilder.EndParameters();
-					}
-					else if (!Tag.EnumName.IsNone())
+					if (UE::FPropertyTypeName NewTypeName = ApplyRedirectsToPropertyType(Tag.GetType(), Property); !NewTypeName.IsEmpty())
 					{
-						TypeBuilder.BeginParameters();
-						TypeBuilder.AddName(Tag.EnumName);
-						TypeBuilder.EndParameters();
-					}
-					else if (!Tag.InnerType.IsNone())
-					{
-						TypeBuilder.BeginParameters();
-						TypeBuilder.AddName(Tag.InnerType);
-						if (!Tag.ValueType.IsNone())
-						{
-							TypeBuilder.AddName(Tag.ValueType);
-						}
-						TypeBuilder.EndParameters();
+						Tag.SetType(NewTypeName);
 					}
+				}
 
+				if (SerializeContext && SerializeContext->bTrackSerializedPropertyPath)
+				{
 					const FName Name = Property ? Property->GetFName() : Tag.Name;
 					const int32 Index = Tag.ArrayIndex > 0 || (Property && Property->ArrayDim > 1) ? Tag.ArrayIndex : INDEX_NONE;
-					SerializeContext->SerializedPropertyPath.Push({Name, TypeBuilder.Build(), Index});
+					SerializeContext->SerializedPropertyPath.Push({Name, Tag.GetType(), Index});
 				}
 
 				if (Property)
 				{
-					FName PropID = Property->GetID();
-
-					// Check if this is a struct property and we have a redirector
-					// No need to check redirects on platforms where everything is cooked. Always check for save games
-					if (bUseRedirects)
-					{
-						if (Tag.Type == NAME_StructProperty && PropID == NAME_StructProperty)
-						{
-							const FName NewName = FLinkerLoad::FindNewNameForStruct(Tag.StructName);
-							const FName StructName = CastFieldChecked<FStructProperty>(Property)->Struct->GetFName();
-							if (NewName == StructName)
-							{
-								Tag.StructName = NewName;
-							}
-						}
-						else if ((PropID == NAME_EnumProperty) && ((Tag.Type == NAME_EnumProperty) || (Tag.Type == NAME_ByteProperty)))
-						{
-							const FName NewName = FLinkerLoad::FindNewNameForEnum(Tag.EnumName);
-							if (!NewName.IsNone())
-							{
-								Tag.EnumName = NewName;
-							}
-						}
-					}
-
 	#if WITH_EDITOR
 					if (BreakRecursionIfFullyLoad && BreakRecursionIfFullyLoad->HasAllFlags(RF_LoadCompleted))
 					{
@@ -1653,8 +1597,8 @@ void UStruct::SerializeVersionedTaggedProperties(FStructuredArchive::FSlot Slot,
 					// check for valid array index
 					else if (Tag.ArrayIndex >= Property->ArrayDim || Tag.ArrayIndex < 0)
 					{
-						UE_LOG(LogClass, Warning, TEXT("Array bound exceeded (var %s=%d, exceeds %s [0-%d] in package:  %s"),
-							*Tag.Name.ToString(), Tag.ArrayIndex, *GetName(), Property->ArrayDim - 1, *UnderlyingArchive.GetArchiveName());
+						UE_LOG(LogClass, Warning, TEXT("Array bound exceeded in %s of %s - %d exceeds [0-%d] in package: %s"),
+							*WriteToString<32>(Tag.Name), *WriteToString<32>(GetFName()), Tag.ArrayIndex, Property->ArrayDim - 1, *UnderlyingArchive.GetArchiveName());
 					}
 					else if (!Property->ShouldSerializeValue(UnderlyingArchive))
 					{
@@ -1686,22 +1630,25 @@ void UStruct::SerializeVersionedTaggedProperties(FStructuredArchive::FSlot Slot,
 								{
 									UnderlyingArchive.Seek(StartOfProperty);
 									FStructuredArchive::FSlot CopySlot = PropertyRecord.EnterField(TEXT("Value"));
-									Tag.Prop = nullptr;
+									Tag.SetProperty(nullptr);
 									PropertyBag->LoadPropertyByTag(SerializeContext->SerializedPropertyPath, Tag, CopySlot, Property->ContainerPtrToValuePtrForDefaults<uint8>(DefaultsStruct, Defaults, Tag.ArrayIndex));
 								}
 								break;
 
 							case EConvertFromTypeResult::Serialized:
-								bAdvanceProperty = true;
+								bAdvanceProperty = !UnderlyingArchive.IsCriticalError();
 								break;
 
 							case EConvertFromTypeResult::UseSerializeItem:
-								if (Tag.Type != PropID)
+								if (const FName PropID = Property->GetID(); Tag.Type != PropID)
 								{
-									UE_LOG(LogClass, Warning, TEXT("Type mismatch in %s of %s - Previous (%s) Current(%s) for package:  %s"), *Tag.Name.ToString(), *GetName(), *Tag.Type.ToString(), *PropID.ToString(), *UnderlyingArchive.GetArchiveName());
+									UE_LOG(LogClass, Warning, TEXT("Type mismatch in %s of %s - Previous (%s) Current(%s) in package: %s"),
+										*WriteToString<32>(Tag.Name), *WriteToString<32>(GetFName()),
+										*WriteToString<32>(Tag.Type), *WriteToString<32>(PropID),
+										*UnderlyingArchive.GetArchiveName());
 									if (FPropertyBag* PropertyBag = TryFindPropertyBag())
 									{
-										Tag.Prop = nullptr;
+										Tag.SetProperty(nullptr);
 										PropertyBag->LoadPropertyByTag(SerializeContext->SerializedPropertyPath, Tag, ValueSlot, Property->ContainerPtrToValuePtrForDefaults<uint8>(DefaultsStruct, Defaults, Tag.ArrayIndex));
 										bAdvanceProperty = !UnderlyingArchive.IsCriticalError();
 									}
@@ -1720,7 +1667,7 @@ void UStruct::SerializeVersionedTaggedProperties(FStructuredArchive::FSlot Slot,
 							case EConvertFromTypeResult::CannotConvert:
 								if (FPropertyBag* PropertyBag = TryFindPropertyBag())
 								{
-									Tag.Prop = nullptr;
+									Tag.SetProperty(nullptr);
 									PropertyBag->LoadPropertyByTag(SerializeContext->SerializedPropertyPath, Tag, ValueSlot, Property->ContainerPtrToValuePtrForDefaults<uint8>(DefaultsStruct, Defaults, Tag.ArrayIndex));
 									bAdvanceProperty = !UnderlyingArchive.IsCriticalError();
 								}
@@ -1753,7 +1700,9 @@ void UStruct::SerializeVersionedTaggedProperties(FStructuredArchive::FSlot Slot,
 				}
 				else
 				{
-					check(Tag.Size == Loaded);
+					checkf(Tag.Size == Loaded,
+						TEXT("Size mismatch in %s of %s of type %s. Loaded %" INT64_FMT " bytes but expected %d. Package: %s"),
+						*Tag.Name.ToString(), *GetName(), *WriteToString<64>(Tag.GetType()), Loaded, Tag.Size, *UnderlyingArchive.GetArchiveName());
 				}
 			}
 		}
@@ -1818,7 +1767,7 @@ void UStruct::SerializeVersionedTaggedProperties(FStructuredArchive::FSlot Slot,
 #endif
 						TestCollector.RecordSavedProperty(Property);
 
-						FPropertyTag Tag( UnderlyingArchive, Property, Idx, DataPtr, DefaultValue );
+						FPropertyTag Tag(Property, Idx, DataPtr);
 						// If available use the property guid from BlueprintGeneratedClasses, provided we aren't cooking data.
 						if (bArePropertyGuidsAvailable && !UnderlyingArchive.IsCooking())
 						{
diff --git a/Engine/Source/Runtime/CoreUObject/Private/UObject/EnumProperty.cpp b/Engine/Source/Runtime/CoreUObject/Private/UObject/EnumProperty.cpp
--- a/Engine/Source/Runtime/CoreUObject/Private/UObject/EnumProperty.cpp
+++ b/Engine/Source/Runtime/CoreUObject/Private/UObject/EnumProperty.cpp
@@ -408,16 +408,19 @@ EConvertFromTypeResult FEnumProperty::ConvertFromType(const FPropertyTag& Tag, F
 	case NAME_ByteProperty:
 	{
 		uint8 PreviousValue = 0;
-		if (Tag.EnumName == NAME_None)
+		if (Tag.GetType().GetParameterCount() == 0)
 		{
-			// If we're a nested property the EnumName tag got lost. Handle this case for backward compatibility reasons
-			FProperty* const PropertyOwner = GetOwner<FProperty>();
-
-			if (PropertyOwner)
+			// A nested property loses its enum name in the property tag. Handle this case for backward compatibility reasons.
+			if (GetOwner<FProperty>())
 			{
+				UE::FPropertyTypeNameBuilder TypeBuilder;
+				TypeBuilder.AddName(Tag.Type);
+				TypeBuilder.BeginParameters();
+				TypeBuilder.AddPath(Enum);
+				TypeBuilder.EndParameters();
+
 				FPropertyTag InnerPropertyTag;
-				InnerPropertyTag.Type = Tag.Type;
-				InnerPropertyTag.EnumName = Enum->GetFName();
+				InnerPropertyTag.SetType(TypeBuilder.Build());
 				InnerPropertyTag.ArrayIndex = 0;
 
 				PreviousValue = (uint8)FNumericProperty::ReadEnumAsInt64(Slot, DefaultsStruct, InnerPropertyTag);
@@ -612,12 +615,5 @@ bool FEnumProperty::CanSerializeFromTypeName(UE::FPropertyTypeName Type) const
 	}
 
 	const FName EnumName = Type.GetParameterName(0);
-	if (EnumName == LocalEnum->GetFName())
-	{
-		return true;
-	}
-
-	TStringBuilder<256> EnumNameString;
-	LocalEnum->GetPathName(nullptr, EnumNameString);
-	return EnumName == EnumNameString.ToView();
+	return EnumName == LocalEnum->GetFName();
 }
diff --git a/Engine/Source/Runtime/CoreUObject/Private/UObject/InstanceDataObjectUtils.cpp b/Engine/Source/Runtime/CoreUObject/Private/UObject/InstanceDataObjectUtils.cpp
--- a/Engine/Source/Runtime/CoreUObject/Private/UObject/InstanceDataObjectUtils.cpp
+++ b/Engine/Source/Runtime/CoreUObject/Private/UObject/InstanceDataObjectUtils.cpp
@@ -290,19 +290,34 @@ namespace UE
 		{
 			if (!AsStructProperty->Struct->UseNativeSerialization())
 			{
-				const FString* OriginalType = nullptr;
 #if WITH_EDITORONLY_DATA
 				//@note: Transfer existing metadata over as we build the InstanceDataObject from the struct or it owner, if any, this is useful for testing purposes
-				OriginalType = AsStructProperty->FindMetaData(NAME_StructOriginalTypeMetadata);
-				FField* OwnerField = OriginalType == nullptr ? AsStructProperty->Owner.ToField() : nullptr;
-				OriginalType = OwnerField ? OwnerField->FindMetaData(NAME_StructOriginalTypeMetadata) : OriginalType;
+				FString OriginalName;
+				if (const FString* OriginalType = AsStructProperty->FindMetaData(NAME_StructOriginalTypeMetadata))
+				{
+					OriginalName = *OriginalType;
+				}
+				//@note: To support metadata defined on array of struct in UPROPERTY for testing purposes
+				else if (FField* OwnerField = AsStructProperty->Owner.ToField())
+				{
+					if (const FString* OwnerOriginalType = OwnerField->FindMetaData(NAME_StructOriginalTypeMetadata))
+					{
+						OriginalName = *OwnerOriginalType;
+					}
+				}
+
+				if (OriginalName.IsEmpty())
+				{
+					UE::FPropertyTypeNameBuilder OriginalNameBuilder;
+					OriginalNameBuilder.AddPath(AsStructProperty->Struct);
+					OriginalName = WriteToString<256>(OriginalNameBuilder.Build()).ToView();
+				}
 #endif
-				const FName StructOriginalName = OriginalType ? FName(**OriginalType) : AsStructProperty->Struct->GetFName();
 				AsStructProperty->Struct = CreateInstanceDataObjectStructRec<UScriptStruct>(AsStructProperty->Struct, Outer, LooseProperties, Path);
 #if WITH_EDITORONLY_DATA
-				AsStructProperty->SetMetaData(NAME_StructOriginalTypeMetadata, StructOriginalName.ToString());
-				AsStructProperty->SetMetaData(NAME_PresentAsTypeMetadata, StructOriginalName.ToString());
-				AsStructProperty->Struct->SetMetaData(NAME_PresentAsTypeMetadata, *StructOriginalName.ToString());
+				AsStructProperty->SetMetaData(NAME_StructOriginalTypeMetadata, *OriginalName);
+				AsStructProperty->SetMetaData(NAME_PresentAsTypeMetadata, *OriginalName);
+				AsStructProperty->Struct->SetMetaData(NAME_PresentAsTypeMetadata, *OriginalName);
 #endif
 			}
 		}
diff --git a/Engine/Source/Runtime/CoreUObject/Private/UObject/Property.cpp b/Engine/Source/Runtime/CoreUObject/Private/UObject/Property.cpp
--- a/Engine/Source/Runtime/CoreUObject/Private/UObject/Property.cpp
+++ b/Engine/Source/Runtime/CoreUObject/Private/UObject/Property.cpp
@@ -2097,14 +2097,14 @@ bool FProperty::LoadFromTag(const FPropertyTag& Tag)
 
 void FProperty::SaveToTag(FPropertyTag& Tag)
 {
-	Tag.Prop = this;
+	Tag.SetProperty(this);
 	Tag.Type = GetID();
 	Tag.Name = GetFName();
 }
 
 void FProperty::AssignToTag(FPropertyTag& Tag)
 {
-	Tag.Prop = this;
+	Tag.SetProperty(this);
 }
 
 bool FProperty::LoadTypeName(UE::FPropertyTypeName Type, const FPropertyTag* Tag)
diff --git a/Engine/Source/Runtime/CoreUObject/Private/UObject/PropertyArray.cpp b/Engine/Source/Runtime/CoreUObject/Private/UObject/PropertyArray.cpp
--- a/Engine/Source/Runtime/CoreUObject/Private/UObject/PropertyArray.cpp
+++ b/Engine/Source/Runtime/CoreUObject/Private/UObject/PropertyArray.cpp
@@ -178,14 +178,27 @@ void FArrayProperty::SerializeItem(FStructuredArchive::FSlot Slot, void* Value,
 
 	if (bIsTextFormat && Inner->IsA<FStructProperty>())
 	{
-		MaybeInnerTag.Emplace(UnderlyingArchive, Inner, 0, (uint8*)Value, (uint8*)Defaults);
-		Slot << SA_ATTRIBUTE(TEXT("InnerStructName"), MaybeInnerTag.GetValue().StructName);
-		Slot << SA_OPTIONAL_ATTRIBUTE(TEXT("InnerStructGuid"), MaybeInnerTag.GetValue().StructGuid, FGuid());
+		MaybeInnerTag.Emplace(Inner, /*Index*/ 0, (uint8*)Value);
+
+		FName StructName;
+		FGuid StructGuid;
+		Slot << SA_ATTRIBUTE(TEXT("InnerStructName"), StructName);
+		Slot << SA_OPTIONAL_ATTRIBUTE(TEXT("InnerStructGuid"), StructGuid, FGuid());
+
+		UE::FPropertyTypeNameBuilder Builder;
+		Builder.AddName(NAME_StructProperty);
+		Builder.BeginParameters();
+		Builder.AddName(StructName);
+		if (StructGuid.IsValid())
+		{
+			Builder.AddGuid(StructGuid);
+		}
+		Builder.EndParameters();
+		MaybeInnerTag->SetType(Builder.Build());
 	}
 
 	TOptional<FPropertyTag> SerializeFromMismatchedTag;
 
-	// TODO: Should work for maps + sets too.
 	auto SerializeContainerItem = [this, &SerializeFromMismatchedTag](FStructuredArchiveSlot Slot, uint8* Item)
 	{
 		if (const FPropertyTag* Tag = SerializeFromMismatchedTag.GetPtrOrNull())
@@ -203,7 +216,7 @@ void FArrayProperty::SerializeItem(FStructuredArchive::FSlot Slot, void* Value,
 					Slot.GetUnderlyingArchive().Seek(StartOfProperty + Tag->Size);	// Skip this item
 					return;
 				case EConvertFromTypeResult::UseSerializeItem:
-					// Fall through to default serialize
+					// Fall through to default SerializeItem
 					break;
 			}
 		}
@@ -542,50 +555,42 @@ void FArrayProperty::SerializeItem(FStructuredArchive::FSlot Slot, void* Value,
 	}
 	ArrayHelper.CountBytes(UnderlyingArchive);
 
+	FUObjectSerializeContext* Context = FUObjectThreadContext::Get().GetSerializeContext();
+
 	// Serialize a PropertyTag for the inner property of this array, allows us to validate the inner struct to see if it has changed
-	if (UnderlyingArchive.UEVer() >= VER_UE4_INNER_ARRAY_TAG_INFO && Inner->IsA<FStructProperty>())
+	if (UnderlyingArchive.UEVer() >= VER_UE4_INNER_ARRAY_TAG_INFO &&
+		Inner->IsA<FStructProperty>())
 	{
 		if (!MaybeInnerTag)
 		{
-			MaybeInnerTag.Emplace(UnderlyingArchive, Inner, 0, (uint8*)Value, (uint8*)Defaults);
+			MaybeInnerTag.Emplace(Inner, 0, (uint8*)Value);
 			UnderlyingArchive << MaybeInnerTag.GetValue();
-			Inner->AssignToTag(MaybeInnerTag.GetValue());
 		}
 
 		FPropertyTag& InnerTag = MaybeInnerTag.GetValue();
 
 		if (UnderlyingArchive.IsLoading())
 		{
-			auto CanSerializeFromStructWithDifferentName = [](const FPropertyTag& PropertyTag, const FStructProperty* StructProperty)
+			if (UE::FPropertyTypeName NewTypeName = ApplyRedirectsToPropertyType(InnerTag.GetType(), Inner); !NewTypeName.IsEmpty())
 			{
-				return StructProperty
-					&& StructProperty->Struct
-					&& PropertyTag.StructGuid.IsValid()
-					&& PropertyTag.StructGuid == StructProperty->Struct->GetCustomGuid();
-			};
+				InnerTag.SetType(NewTypeName);
+			}
 
 			// Check if the Inner property can successfully serialize, the type may have changed
-			FStructProperty* StructProperty = CastFieldChecked<FStructProperty>(Inner);
-			// if check redirector to make sure if the name has changed
-			FName NewName = FLinkerLoad::FindNewNameForStruct(InnerTag.StructName);
-			FName StructName = StructProperty->Struct->GetFName();
-			if (NewName != NAME_None && NewName == StructName)
+			if (!Inner->CanSerializeFromTypeName(InnerTag.GetType()))
 			{
-				InnerTag.StructName = NewName;
-			}
+				FStructProperty* StructProperty = CastFieldChecked<FStructProperty>(Inner);
 
-			if (InnerTag.StructName != StructProperty->Struct->GetFName()
-				&& !CanSerializeFromStructWithDifferentName(InnerTag, StructProperty))
-			{
 				// Attempt mismatched tag serialization if available
-				if ((StructProperty->Struct->StructFlags & STRUCT_SerializeFromMismatchedTag) && (InnerTag.Type != NAME_StructProperty || (InnerTag.StructName != StructProperty->Struct->GetFName())))
+				const FName StructName = InnerTag.GetType().GetParameterName();
+				if ((StructProperty->Struct->StructFlags & STRUCT_SerializeFromMismatchedTag) && (InnerTag.Type != NAME_StructProperty || StructName != StructProperty->Struct->GetFName()))
 				{
 					SerializeFromMismatchedTag = InnerTag;
 				}
 				else
 				{
-					UE_LOG(LogClass, Warning, TEXT("Array Property %s of %s contains a struct type mismatch (tag %s != prop %s) in package:  %s. If that struct got renamed, add an entry to ActiveStructRedirects."),
-					*InnerTag.Name.ToString(), *GetName(), *InnerTag.StructName.ToString(), *CastFieldChecked<FStructProperty>(Inner)->Struct->GetName(), *UnderlyingArchive.GetArchiveName());
+					UE_LOG(LogClass, Warning, TEXT("Array Property %s contains a struct type mismatch (tag %s != prop %s) in package: %s. If that struct got renamed, add an entry to ActiveStructRedirects."),
+						*WriteToString<32>(InnerTag.Name), *WriteToString<32>(StructName), *WriteToString<32>(StructProperty->Struct->GetFName()), *UnderlyingArchive.GetArchiveName());
 
 #if WITH_EDITOR
 					// Ensure the structure is initialized
@@ -610,30 +615,26 @@ void FArrayProperty::SerializeItem(FStructuredArchive::FSlot Slot, void* Value,
 				}
 			}
 		}
-	}
 
-	FUObjectSerializeContext* Context = FUObjectThreadContext::Get().GetSerializeContext();
-	if (Context && Context->bTrackSerializedPropertyPath)
-	{
-		// Update the path with types from the inner tag if the outer tag is incomplete.
-		if (const int32 SegmentIndex = Context->SerializedPropertyPath.GetSegmentCount() - 1; SegmentIndex >= 0)
+		if (Context && Context->bTrackSerializedPropertyPath)
 		{
-			UE::FPropertyPathNameSegment Segment = Context->SerializedPropertyPath.GetSegment(SegmentIndex);
-			if (Segment.Type.GetName() == NAME_ArrayProperty)
+			// Update the path with types from the inner tag if the outer tag is incomplete.
+			if (const int32 SegmentIndex = Context->SerializedPropertyPath.GetSegmentCount() - 1; SegmentIndex >= 0)
 			{
-				const UE::FPropertyTypeName InnerTypeName = Segment.Type.GetParameter();
-				if (InnerTypeName.GetName() == NAME_StructProperty && InnerTypeName.GetParameterCount() == 0)
+				UE::FPropertyPathNameSegment Segment = Context->SerializedPropertyPath.GetSegment(SegmentIndex);
+				if (Segment.Type.GetName() == NAME_ArrayProperty)
 				{
-					UE::FPropertyTypeNameBuilder NewTypeBuilder;
-					NewTypeBuilder.AddName(NAME_ArrayProperty);
-					NewTypeBuilder.BeginParameters();
-					NewTypeBuilder.AddName(NAME_StructProperty);
-					NewTypeBuilder.BeginParameters();
-					NewTypeBuilder.AddName(MaybeInnerTag->StructName);
-					NewTypeBuilder.EndParameters();
-					NewTypeBuilder.EndParameters();
-					Segment.Type = NewTypeBuilder.Build();
-					Context->SerializedPropertyPath.SetSegment(SegmentIndex, Segment);
+					const UE::FPropertyTypeName InnerTypeName = Segment.Type.GetParameter();
+					if (InnerTypeName.GetName() == NAME_StructProperty && InnerTypeName.GetParameterCount() == 0)
+					{
+						UE::FPropertyTypeNameBuilder NewTypeBuilder;
+						NewTypeBuilder.AddName(NAME_ArrayProperty);
+						NewTypeBuilder.BeginParameters();
+						NewTypeBuilder.AddType(InnerTag.GetType());
+						NewTypeBuilder.EndParameters();
+						Segment.Type = NewTypeBuilder.Build();
+						Context->SerializedPropertyPath.SetSegment(SegmentIndex, Segment);
+					}
 				}
 			}
 		}
@@ -1162,66 +1163,71 @@ bool FArrayProperty::SameType(const FProperty* Other) const
 
 EConvertFromTypeResult FArrayProperty::ConvertFromType(const FPropertyTag& Tag, FStructuredArchive::FSlot Slot, uint8* Data, UStruct* DefaultsStruct, const uint8* Defaults)
 {
-	// TODO: The ArrayProperty Tag really doesn't have adequate information for
-	// many types. This should probably all be moved in to ::SerializeItem
+	if (Tag.Type != NAME_ArrayProperty)
+	{
+		return EConvertFromTypeResult::UseSerializeItem;
+	}
 
-	if (Tag.Type == NAME_ArrayProperty && Tag.InnerType != NAME_None && Tag.InnerType != Inner->GetID())
+	FArchive& UnderlyingArchive = Slot.GetUnderlyingArchive();
+
+	const FName InnerTypeName = Tag.GetType().GetParameterName();
+	if (InnerTypeName.IsNone() || InnerTypeName == Inner->GetID())
 	{
-		void* ArrayPropertyData = ContainerPtrToValuePtr<void>(Data);
+		return EConvertFromTypeResult::UseSerializeItem;
+	}
 
-		int32 ElementCount = 0;
+	if (Tag.bExperimentalOverridableLogic)
+	{
+		return EConvertFromTypeResult::CannotConvert;
+	}
 
-		if (Slot.GetUnderlyingArchive().IsTextFormat())
-		{
-			Slot.EnterArray(ElementCount);
-		}
-		else
-		{
-			Slot.GetUnderlyingArchive() << ElementCount;
-		}
+	int32 ElementCount = 0;
+	if (UnderlyingArchive.IsTextFormat())
+	{
+		Slot.EnterArray(ElementCount);
+	}
+	else
+	{
+		UnderlyingArchive << ElementCount;
+	}
 
-		FScriptArrayHelper ScriptArrayHelper(this, ArrayPropertyData);
-		ScriptArrayHelper.EmptyAndAddValues(ElementCount);
+	FScriptArrayHelper ScriptArrayHelper(this, ContainerPtrToValuePtr<void>(Data));
+	ScriptArrayHelper.EmptyAndAddValues(ElementCount);
 
-		FPropertyTag InnerPropertyTag;
-		InnerPropertyTag.Type = Tag.InnerType;
-		InnerPropertyTag.ArrayIndex = 0;
+	FPropertyTag InnerPropertyTag;
+	InnerPropertyTag.SetType(Tag.GetType().GetParameter());
+	InnerPropertyTag.ArrayIndex = 0;
 
-		if (Slot.GetArchiveState().UEVer() >= VER_UE4_INNER_ARRAY_TAG_INFO && Tag.InnerType == NAME_StructProperty)
-		{
-			Slot.GetUnderlyingArchive() << InnerPropertyTag;
-		}
+	if (UnderlyingArchive.UEVer() >= VER_UE4_INNER_ARRAY_TAG_INFO && InnerPropertyTag.Type == NAME_StructProperty)
+	{
+		UnderlyingArchive << InnerPropertyTag;
+	}
 
-		// Convert properties from old type to new type automatically if types are compatible (array case)
-		if (ElementCount > 0)
-		{
-			FStructuredArchive::FStream ValueStream = Slot.EnterStream();
+	if (ElementCount == 0)
+	{
+		return EConvertFromTypeResult::Converted;
+	}
 
-			EConvertFromTypeResult ConvertResult = Inner->ConvertFromType(InnerPropertyTag, ValueStream.EnterElement(), ScriptArrayHelper.GetRawPtr(0), DefaultsStruct, nullptr);
-			if (ConvertResult == EConvertFromTypeResult::Converted || ConvertResult == EConvertFromTypeResult::Serialized)
-			{
-				for (int32 i = 1; i < ElementCount; ++i)
-				{
-					ConvertResult = Inner->ConvertFromType(InnerPropertyTag, ValueStream.EnterElement(), ScriptArrayHelper.GetRawPtr(i), DefaultsStruct, nullptr);
-					check(ConvertResult == EConvertFromTypeResult::Converted || ConvertResult == EConvertFromTypeResult::Serialized);
-				}
+	FStructuredArchive::FStream ValueStream = Slot.EnterStream();
 
-				return EConvertFromTypeResult::Converted;
-			}
-			// TODO: Implement SerializeFromMismatchedTag handling for arrays of structs
-			else
-			{
-				UE_LOG(LogClass, Warning, TEXT("Array Inner Type mismatch in %s of %s - Previous (%s) Current(%s) for package:  %s"), *Tag.Name.ToString(), *GetName(), *Tag.InnerType.ToString(), *Inner->GetID().ToString(), *Slot.GetUnderlyingArchive().GetArchiveName());
-				return EConvertFromTypeResult::CannotConvert;
-			}
-		}
-		else
+	EConvertFromTypeResult ConvertResult = Inner->ConvertFromType(InnerPropertyTag, ValueStream.EnterElement(), ScriptArrayHelper.GetRawPtr(0), DefaultsStruct, nullptr);
+	if (ConvertResult == EConvertFromTypeResult::Converted || ConvertResult == EConvertFromTypeResult::Serialized)
+	{
+		for (int32 ElementIndex = 1; ElementIndex < ElementCount; ++ElementIndex)
 		{
-			return EConvertFromTypeResult::Converted;
+			ConvertResult = Inner->ConvertFromType(InnerPropertyTag, ValueStream.EnterElement(), ScriptArrayHelper.GetRawPtr(ElementIndex), DefaultsStruct, nullptr);
+			check(ConvertResult == EConvertFromTypeResult::Converted || ConvertResult == EConvertFromTypeResult::Serialized);
 		}
+		return ConvertResult;
+	}
+	else
+	{
+		UE::FPropertyTypeNameBuilder Builder;
+		Inner->SaveTypeName(Builder);
+		UE_LOG(LogClass, Warning, TEXT("Array Inner Type mismatch in %s - Previous (%s) Current(%s) in package: %s"),
+			*WriteToString<32>(Tag.Name), *WriteToString<32>(InnerPropertyTag.GetType()), *WriteToString<32>(Builder.Build()), *UnderlyingArchive.GetArchiveName());
+		return EConvertFromTypeResult::CannotConvert;
 	}
-
-	return EConvertFromTypeResult::UseSerializeItem;
 }
 
 FField* FArrayProperty::GetInnerFieldByName(const FName& InName)
diff --git a/Engine/Source/Runtime/CoreUObject/Private/UObject/PropertyBag.cpp b/Engine/Source/Runtime/CoreUObject/Private/UObject/PropertyBag.cpp
--- a/Engine/Source/Runtime/CoreUObject/Private/UObject/PropertyBag.cpp
+++ b/Engine/Source/Runtime/CoreUObject/Private/UObject/PropertyBag.cpp
@@ -23,15 +23,12 @@ void FPropertyBag::Add(const FPropertyPathName& Path, FProperty* Property, void*
 {
 	FValue& Value = FindOrCreateValue(Path);
 
-	const bool bPropertyChanged = Value.Tag.Prop != Property;
+	const bool bPropertyChanged = Value.Tag.GetProperty() != Property;
 
 	if (bPropertyChanged)
 	{
-		// TODO: NullAr is a workaround to FPropertyTag requiring an archive to assert on versioned property serialization.
-		FNullArchive NullAr;
-
 		Value.Destroy();
-		Value.Tag = FPropertyTag(NullAr, Property, INDEX_NONE, (uint8*)Data, nullptr);
+		Value.Tag = FPropertyTag(Property, INDEX_NONE, (uint8*)Data);
 	}
 
 	Value.AllocateAndInitializeValue();
@@ -83,14 +80,9 @@ void FPropertyBag::LoadPropertyByTag(const FPropertyPathName& Path, const FPrope
 	FValue& Value = FindOrCreateValue(Path);
 
 	const bool bPropertyChanged =
-		(Value.Tag.Prop != Tag.Prop && (Tag.Prop || !Value.bOwnsProperty)) ||
-		Value.Tag.Type != Tag.Type ||
+		(Value.Tag.GetProperty() != Tag.GetProperty() && (Tag.GetProperty() || !Value.bOwnsProperty)) ||
+		Value.Tag.GetType() != Tag.GetType() ||
 		Value.Tag.Name != Tag.Name ||
-		Value.Tag.StructName != Tag.StructName ||
-		Value.Tag.EnumName != Tag.EnumName ||
-		Value.Tag.InnerType != Tag.InnerType ||
-		Value.Tag.ValueType != Tag.ValueType ||
-		Value.Tag.StructGuid != Tag.StructGuid ||
 		Value.Tag.PropertyGuid != Tag.PropertyGuid;
 
 	if (bPropertyChanged)
@@ -101,7 +93,7 @@ void FPropertyBag::LoadPropertyByTag(const FPropertyPathName& Path, const FPrope
 	}
 
 	// Serialize the value using the existing property from the tag.
-	if (FProperty* Property = Value.Tag.Prop)
+	if (FProperty* Property = Value.Tag.GetProperty())
 	{
 		Value.AllocateAndInitializeValue();
 		Tag.SerializeTaggedProperty(ValueSlot, Property, (uint8*)Value.Data, (const uint8*)Defaults);
@@ -114,7 +106,7 @@ void FPropertyBag::LoadPropertyByTag(const FPropertyPathName& Path, const FPrope
 	{
 		Property->Link(UnderlyingArchive);
 		Value.bOwnsProperty = true;
-		Value.Tag.Prop = Property;
+		Value.Tag.SetProperty(Property);
 		Value.AllocateAndInitializeValue();
 		Tag.SerializeTaggedProperty(ValueSlot, Property, (uint8*)Value.Data, nullptr);
 		return;
@@ -168,7 +160,7 @@ void FPropertyBag::FValue::AllocateAndInitializeValue()
 		return;
 	}
 
-	if (const FProperty* Property = Tag.Prop)
+	if (const FProperty* Property = Tag.GetProperty())
 	{
 		// TODO: Need to allocate only one element for arrays.
 		Data = Property->AllocateAndInitializeValue();
@@ -181,7 +173,7 @@ void FPropertyBag::FValue::AllocateAndInitializeValue()
 
 void FPropertyBag::FValue::Destroy()
 {
-	if (const FProperty* Property = Tag.Prop)
+	if (const FProperty* Property = Tag.GetProperty())
 	{
 		if (Data)
 		{
@@ -194,7 +186,7 @@ void FPropertyBag::FValue::Destroy()
 		{
 			bOwnsProperty = false;
 			delete Property;
-			Tag.Prop = nullptr;
+			Tag.SetProperty(nullptr);
 		}
 	}
 	else
@@ -210,12 +202,12 @@ int32 FPropertyBag::FValue::GetSize() const
 	{
 		return 0;
 	}
-	if (!Tag.Prop)
+	if (!Tag.GetProperty())
 	{
 		return Tag.Size;
 	}
 	// TODO: Need to include the size of only one element for arrays.
-	return Tag.Prop->GetSize();
+	return Tag.GetProperty()->GetSize();
 }
 
 ///////////////////////////////////////////////////////////////////////////////////////////////////////////////////////
diff --git a/Engine/Source/Runtime/CoreUObject/Private/UObject/PropertyBool.cpp b/Engine/Source/Runtime/CoreUObject/Private/UObject/PropertyBool.cpp
--- a/Engine/Source/Runtime/CoreUObject/Private/UObject/PropertyBool.cpp
+++ b/Engine/Source/Runtime/CoreUObject/Private/UObject/PropertyBool.cpp
@@ -1,13 +1,10 @@
 // Copyright Epic Games, Inc. All Rights Reserved.
 
-#include "CoreMinimal.h"
+#include "UObject/UnrealType.h"
 
 #include "Hash/Blake3.h"
-#include "UObject/ObjectMacros.h"
-#include "UObject/UObjectGlobals.h"
-#include "UObject/UnrealType.h"
-#include "UObject/UnrealTypePrivate.h"
 #include "UObject/PropertyHelper.h"
+#include "UObject/UnrealTypePrivate.h"
 
 /*-----------------------------------------------------------------------------
 	FBoolProperty.
@@ -282,21 +279,19 @@ EConvertFromTypeResult FBoolProperty::ConvertFromType(const FPropertyTag& Tag, F
 	}
 	else if (Tag.Type == NAME_ByteProperty)
 	{
-		// if the byte property was an enum we won't allow a conversion to bool
-		if (Tag.EnumName == NAME_None)
+		// Disallow conversion of enum to bool.
+		if (Tag.GetType().GetParameterCount() > 0)
 		{
-			// If we're a nested property the EnumName tag got lost, don't allow this
-			if (GetOwner<FProperty>())
-			{
-				return EConvertFromTypeResult::UseSerializeItem;
-			}
-
-			LoadFromType<uint8>(this, Tag, Slot, Data);
+			return EConvertFromTypeResult::UseSerializeItem;
 		}
-		else
+
+		// Disallow a nested byte property because it is impossible to distinguish from a nested enum.
+		if (GetOwner<FProperty>())
 		{
 			return EConvertFromTypeResult::UseSerializeItem;
 		}
+
+		LoadFromType<uint8>(this, Tag, Slot, Data);
 	}
 	else if (Tag.Type == NAME_UInt16Property)
 	{
diff --git a/Engine/Source/Runtime/CoreUObject/Private/UObject/PropertyByte.cpp b/Engine/Source/Runtime/CoreUObject/Private/UObject/PropertyByte.cpp
--- a/Engine/Source/Runtime/CoreUObject/Private/UObject/PropertyByte.cpp
+++ b/Engine/Source/Runtime/CoreUObject/Private/UObject/PropertyByte.cpp
@@ -222,29 +222,27 @@ EConvertFromTypeResult FByteProperty::ConvertFromType(const FPropertyTag& Tag, F
 		return EConvertFromTypeResult::UseSerializeItem;
 	case NAME_ByteProperty:
 	{
-		if ((Tag.EnumName == NAME_None) != (Enum == nullptr))
+		if ((Tag.GetType().GetParameterCount() == 0) != (Enum == nullptr))
 		{
 			// A byte property gained or lost an enum.
 			uint8 PreviousValue = 0;
-			if (Tag.EnumName == NAME_None)
+			if (Enum)
 			{
-				// If we're a nested property the EnumName tag got lost. Fail to read in this case
-				FProperty* const PropertyOwner = GetOwner<FProperty>();
-				if (PropertyOwner)
+				// A nested property loses its enum name in the property tag. Handle this case for backward compatibility reasons.
+				if (GetOwner<FProperty>())
 				{
 					return EConvertFromTypeResult::UseSerializeItem;
 				}
 
-				// simply pretend the property still doesn't have an enum and serialize the single byte
+				// Read the byte and assume its value corresponds to a valid enumerator.
 				Slot << PreviousValue;
 			}
 			else
 			{
-				// attempt to find the old enum and get the byte value from the serialized enum name
+				// Attempt to find the enum from the tag and find the byte value from the enum.
 				PreviousValue = (uint8)ReadEnumAsInt64(Slot, DefaultsStruct, Tag);
 			}
 
-			// now copy the value into the object's address space
 			SetPropertyValue_InContainer(Data, PreviousValue, Tag.ArrayIndex);
 			return EConvertFromTypeResult::Converted;
 		}
@@ -252,17 +250,11 @@ EConvertFromTypeResult FByteProperty::ConvertFromType(const FPropertyTag& Tag, F
 	}
 	case NAME_EnumProperty:
 	{
-		if (Enum == nullptr || Tag.EnumName == Enum->GetFName() || Tag.EnumName.ToString() == Enum->GetPathName())
-		{
-			// an enum property became a byte
-			// attempt to find the old enum and get the byte value from the serialized enum name
-			uint8 PreviousValue = (uint8)ReadEnumAsInt64(Slot, DefaultsStruct, Tag);
+		// Attempt to find the enum from the tag and find the byte value from the enum.
+		uint8 PreviousValue = (uint8)ReadEnumAsInt64(Slot, DefaultsStruct, Tag);
 
-			// now copy the value into the object's address space
-			SetPropertyValue_InContainer(Data, PreviousValue, Tag.ArrayIndex);
-			return EConvertFromTypeResult::Converted;
-		}
-		return EConvertFromTypeResult::UseSerializeItem;
+		SetPropertyValue_InContainer(Data, PreviousValue, Tag.ArrayIndex);
+		return EConvertFromTypeResult::Converted;
 	}
 	case NAME_Int8Property:
 		if (Enum)
@@ -610,14 +602,7 @@ bool FByteProperty::CanSerializeFromTypeName(UE::FPropertyTypeName Type) const
 	const FName EnumName = Type.GetParameterName();
 	if (const UEnum* LocalEnum = Enum)
 	{
-		if (EnumName == LocalEnum->GetFName())
-		{
-			return true;
-		}
-
-		TStringBuilder<256> EnumNameString;
-		LocalEnum->GetPathName(nullptr, EnumNameString);
-		return EnumName == EnumNameString.ToView();
+		return EnumName == LocalEnum->GetFName();
 	}
 	return EnumName.IsNone();
 }
diff --git a/Engine/Source/Runtime/CoreUObject/Private/UObject/PropertyMap.cpp b/Engine/Source/Runtime/CoreUObject/Private/UObject/PropertyMap.cpp
--- a/Engine/Source/Runtime/CoreUObject/Private/UObject/PropertyMap.cpp
+++ b/Engine/Source/Runtime/CoreUObject/Private/UObject/PropertyMap.cpp
@@ -1434,18 +1434,26 @@ EConvertFromTypeResult FMapProperty::ConvertFromType(const FPropertyTag& Tag, FS
 		return EConvertFromTypeResult::UseSerializeItem;
 	}
 
-	const bool bKeyTypeMatches = (Tag.InnerType == KeyProp->GetID());
-	const bool bValueTypeMatches = (Tag.ValueType == ValueProp->GetID());
-	if ((bKeyTypeMatches || Tag.InnerType.IsNone()) && (bValueTypeMatches || Tag.ValueType.IsNone()))
+	const UE::FPropertyTypeName KeyType = Tag.GetType().GetParameter(0);
+	const UE::FPropertyTypeName ValueType = Tag.GetType().GetParameter(1);
+	const FName KeyTypeName = KeyType.GetName();
+	const FName ValueTypeName = ValueType.GetName();
+	const bool bCanSerializeKey = (KeyTypeName == KeyProp->GetID());
+	const bool bCanSerializeValue = (ValueTypeName == ValueProp->GetID());
+	if ((bCanSerializeKey || KeyTypeName.IsNone()) && (bCanSerializeValue || ValueTypeName.IsNone()))
 	{
 		return EConvertFromTypeResult::UseSerializeItem;
 	}
 
-	const auto SerializeOrConvert = [](bool bTypeMatches, FProperty* Inner, const FPropertyTag& InnerTag, FStructuredArchive::FSlot InnerSlot, uint8* InnerData, UStruct* InnerDefaultsStruct) -> bool
+	if (Tag.bExperimentalOverridableLogic)
 	{
-		if (!bTypeMatches)
+		return EConvertFromTypeResult::CannotConvert;
+	}
+
+	const auto SerializeOrConvert = [](bool bCanSerialize, FProperty* Inner, const FPropertyTag& InnerTag, FStructuredArchive::FSlot InnerSlot, uint8* InnerData, UStruct* InnerDefaultsStruct) -> bool
+	{
+		if (!bCanSerialize)
 		{
-			const int64 StartOfProperty = InnerSlot.GetUnderlyingArchive().Tell();
 			switch (Inner->ConvertFromType(InnerTag, InnerSlot, InnerData, InnerDefaultsStruct, nullptr))
 			{
 				default:
@@ -1483,11 +1491,11 @@ EConvertFromTypeResult FMapProperty::ConvertFromType(const FPropertyTag& Tag, FS
 	};
 
 	FPropertyTag KeyPropertyTag;
-	KeyPropertyTag.Type = Tag.InnerType;
+	KeyPropertyTag.SetType(KeyType);
 	KeyPropertyTag.ArrayIndex = 0;
 
 	FPropertyTag ValuePropertyTag;
-	ValuePropertyTag.Type = Tag.ValueType;
+	ValuePropertyTag.SetType(ValueType);
 	ValuePropertyTag.ArrayIndex = 0;
 
 	bool bConversionSucceeded = true;
@@ -1505,7 +1513,7 @@ EConvertFromTypeResult FMapProperty::ConvertFromType(const FPropertyTag& Tag, FS
 		TempKeyValueStorage = (uint8*)FMemory::Malloc(MapLayout.SetLayout.Size);
 		KeyProp->InitializeValue(TempKeyValueStorage);
 
-		if (SerializeOrConvert(bKeyTypeMatches, KeyProp, KeyPropertyTag, KeysToRemoveArray.EnterElement(), TempKeyValueStorage, DefaultsStruct))
+		if (SerializeOrConvert(bCanSerializeKey, KeyProp, KeyPropertyTag, KeysToRemoveArray.EnterElement(), TempKeyValueStorage, DefaultsStruct))
 		{
 			// If the key is in the map, remove it
 			int32 Found = MapHelper.FindMapIndexWithKey(TempKeyValueStorage);
@@ -1517,7 +1525,7 @@ EConvertFromTypeResult FMapProperty::ConvertFromType(const FPropertyTag& Tag, FS
 			// things are going fine, remove the rest of the keys:
 			for (int32 I = 1; I < NumKeysToRemove; ++I)
 			{
-				verify(SerializeOrConvert(bKeyTypeMatches, KeyProp, KeyPropertyTag, KeysToRemoveArray.EnterElement(), TempKeyValueStorage, DefaultsStruct));
+				verify(SerializeOrConvert(bCanSerializeKey, KeyProp, KeyPropertyTag, KeysToRemoveArray.EnterElement(), TempKeyValueStorage, DefaultsStruct));
 				Found = MapHelper.FindMapIndexWithKey(TempKeyValueStorage);
 				if (Found != INDEX_NONE)
 				{
@@ -1546,7 +1554,7 @@ EConvertFromTypeResult FMapProperty::ConvertFromType(const FPropertyTag& Tag, FS
 
 			FStructuredArchive::FRecord FirstPropertyRecord = EntriesArray.EnterElement().EnterRecord();
 
-			if (SerializeOrConvert(bKeyTypeMatches, KeyProp, KeyPropertyTag, FirstPropertyRecord.EnterField(TEXT("Key")), TempKeyValueStorage, DefaultsStruct))
+			if (SerializeOrConvert(bCanSerializeKey, KeyProp, KeyPropertyTag, FirstPropertyRecord.EnterField(TEXT("Key")), TempKeyValueStorage, DefaultsStruct))
 			{
 				// Add a new default value if the key doesn't currently exist in the map
 				bool bKeyAlreadyPresent = true;
@@ -1562,14 +1570,14 @@ EConvertFromTypeResult FMapProperty::ConvertFromType(const FPropertyTag& Tag, FS
 				KeyProp->CopyCompleteValue_InContainer(NextPairPtr, TempKeyValueStorage);
 
 				// Deserialize value
-				if (SerializeOrConvert(bValueTypeMatches, ValueProp, ValuePropertyTag, FirstPropertyRecord.EnterField(TEXT("Value")), NextPairPtr, DefaultsStruct))
+				if (SerializeOrConvert(bCanSerializeValue, ValueProp, ValuePropertyTag, FirstPropertyRecord.EnterField(TEXT("Value")), NextPairPtr, DefaultsStruct))
 				{
 					// first entry went fine, convert the rest:
 					for (int32 I = 1; I < Num; ++I)
 					{
 						FStructuredArchive::FRecord PropertyRecord = EntriesArray.EnterElement().EnterRecord();
 
-						verify(SerializeOrConvert(bKeyTypeMatches, KeyProp, KeyPropertyTag, PropertyRecord.EnterField(TEXT("Key")), TempKeyValueStorage, DefaultsStruct));
+						verify(SerializeOrConvert(bCanSerializeKey, KeyProp, KeyPropertyTag, PropertyRecord.EnterField(TEXT("Key")), TempKeyValueStorage, DefaultsStruct));
 						NextPairIndex = MapHelper.FindMapIndexWithKey(TempKeyValueStorage);
 						if (NextPairIndex == INDEX_NONE)
 						{
@@ -1579,7 +1587,7 @@ EConvertFromTypeResult FMapProperty::ConvertFromType(const FPropertyTag& Tag, FS
 						NextPairPtr = MapHelper.GetPairPtrWithoutCheck(NextPairIndex);
 						// This copy is unnecessary when the key was already in the map:
 						KeyProp->CopyCompleteValue_InContainer(NextPairPtr, TempKeyValueStorage);
-						verify(SerializeOrConvert(bValueTypeMatches, ValueProp, ValuePropertyTag, PropertyRecord.EnterField(TEXT("Value")), NextPairPtr, DefaultsStruct));
+						verify(SerializeOrConvert(bCanSerializeValue, ValueProp, ValuePropertyTag, PropertyRecord.EnterField(TEXT("Value")), NextPairPtr, DefaultsStruct));
 					}
 				}
 				else
@@ -1604,13 +1612,17 @@ EConvertFromTypeResult FMapProperty::ConvertFromType(const FPropertyTag& Tag, FS
 	// if we could not convert the property ourself, then indicate that calling code needs to advance the property
 	if (!bConversionSucceeded)
 	{
+		UE::FPropertyTypeNameBuilder KeyBuilder;
+		UE::FPropertyTypeNameBuilder ValueBuilder;
+		KeyProp->SaveTypeName(KeyBuilder);
+		ValueProp->SaveTypeName(ValueBuilder);
 		UE_LOG(LogClass, Warning,
 			TEXT("Map Element Type mismatch in %s - Previous (%s to %s) Current (%s to %s) for package: %s"),
 			*WriteToString<32>(Tag.Name),
-			*WriteToString<32>(KeyPropertyTag.Type),
-			*WriteToString<32>(ValuePropertyTag.Type),
-			*WriteToString<32>(KeyProp->GetID()),
-			*WriteToString<32>(ValueProp->GetID()),
+			*WriteToString<32>(KeyPropertyTag.GetType()),
+			*WriteToString<32>(ValuePropertyTag.GetType()),
+			*WriteToString<32>(KeyBuilder.Build()),
+			*WriteToString<32>(ValueBuilder.Build()),
 			*UnderlyingArchive.GetArchiveName());
 	}
 
diff --git a/Engine/Source/Runtime/CoreUObject/Private/UObject/PropertyNumeric.cpp b/Engine/Source/Runtime/CoreUObject/Private/UObject/PropertyNumeric.cpp
--- a/Engine/Source/Runtime/CoreUObject/Private/UObject/PropertyNumeric.cpp
+++ b/Engine/Source/Runtime/CoreUObject/Private/UObject/PropertyNumeric.cpp
@@ -1,10 +1,5 @@
 // Copyright Epic Games, Inc. All Rights Reserved.
 
-#include "CoreMinimal.h"
-#include "UObject/ObjectMacros.h"
-#include "UObject/Class.h"
-#include "UObject/Package.h"
-#include "Templates/Casts.h"
 #include "UObject/UnrealType.h"
 
 IMPLEMENT_FIELD(FNumericProperty)
@@ -30,16 +25,18 @@ int64 FNumericProperty::ReadEnumAsInt64(FStructuredArchive::FSlot Slot, UStruct*
 	FName EnumName;
 	Slot << EnumName;
 
-	UEnum* Enum = FindUField<UEnum>(dynamic_cast<UClass*>(DefaultsStruct) ? static_cast<UClass*>(DefaultsStruct) : DefaultsStruct->GetTypedOuter<UClass>(), Tag.EnumName);
+	FName EnumTypeName = Tag.GetType().GetParameterName();
+
+	UEnum* Enum = FindUField<UEnum>(dynamic_cast<UClass*>(DefaultsStruct) ? static_cast<UClass*>(DefaultsStruct) : DefaultsStruct->GetTypedOuter<UClass>(), EnumTypeName);
 	if (!Enum)
 	{
 		// Enums (at least native) are stored as short names (for now) so find the Tag enum by name
-		Enum = FindFirstObject<UEnum>(*Tag.EnumName.ToString(), EFindFirstObjectOptions::NativeFirst);
+		Enum = FindFirstObject<UEnum>(*WriteToString<128>(EnumTypeName), EFindFirstObjectOptions::NativeFirst);
 	}
 
 	if (!Enum)
 	{
-		UE_LOG(LogClass, Warning, TEXT("Failed to find enum '%s' when converting property '%s' during property loading - setting to 0"), *Tag.EnumName.ToString(), *Tag.Name.ToString());
+		UE_LOG(LogClass, Warning, TEXT("Failed to find enum '%s' when converting property '%s' during property loading - setting to 0"), *EnumTypeName.ToString(), *Tag.Name.ToString());
 		return 0;
 	}
 
diff --git a/Engine/Source/Runtime/CoreUObject/Private/UObject/PropertyOptional.cpp b/Engine/Source/Runtime/CoreUObject/Private/UObject/PropertyOptional.cpp
--- a/Engine/Source/Runtime/CoreUObject/Private/UObject/PropertyOptional.cpp
+++ b/Engine/Source/Runtime/CoreUObject/Private/UObject/PropertyOptional.cpp
@@ -264,7 +264,7 @@ void FOptionalProperty::SerializeItem(FStructuredArchive::FSlot Slot, void* Data
 			void* ValueData = GetValuePointerForReadOrReplace(Data);
 
 			// Construct and serialize a tag for the value.
-			FPropertyTag ValueTag(UnderlyingArchive, GetValueProperty(), 0, static_cast<uint8*>(ValueData), static_cast<const uint8*>(ValueDefaults));
+			FPropertyTag ValueTag(GetValueProperty(), 0, static_cast<uint8*>(ValueData));
 			ValueSlot << ValueTag;
 
 			// Serialize the value.
diff --git a/Engine/Source/Runtime/CoreUObject/Private/UObject/PropertySet.cpp b/Engine/Source/Runtime/CoreUObject/Private/UObject/PropertySet.cpp
--- a/Engine/Source/Runtime/CoreUObject/Private/UObject/PropertySet.cpp
+++ b/Engine/Source/Runtime/CoreUObject/Private/UObject/PropertySet.cpp
@@ -900,7 +900,8 @@ EConvertFromTypeResult FSetProperty::ConvertFromType(const FPropertyTag& Tag, FS
 		return EConvertFromTypeResult::UseSerializeItem;
 	}
 
-	if (Tag.InnerType.IsNone() || Tag.InnerType == ElementProp->GetID())
+	const FName InnerTypeName = Tag.GetType().GetParameterName();
+	if (InnerTypeName.IsNone() || InnerTypeName == ElementProp->GetID())
 	{
 		return EConvertFromTypeResult::UseSerializeItem;
 	}
@@ -917,7 +918,7 @@ EConvertFromTypeResult FSetProperty::ConvertFromType(const FPropertyTag& Tag, FS
 	};
 
 	FPropertyTag InnerPropertyTag;
-	InnerPropertyTag.Type = Tag.InnerType;
+	InnerPropertyTag.SetType(Tag.GetType().GetParameter());
 	InnerPropertyTag.ArrayIndex = 0;
 
 	bool bConversionSucceeded = true;
@@ -1018,8 +1019,10 @@ EConvertFromTypeResult FSetProperty::ConvertFromType(const FPropertyTag& Tag, FS
 	// if we could not convert the property ourself, then indicate that calling code needs to advance the property
 	if (!bConversionSucceeded)
 	{
+		UE::FPropertyTypeNameBuilder Builder;
+		ElementProp->SaveTypeName(Builder);
 		UE_LOG(LogClass, Warning, TEXT("Set Element Type mismatch in %s - Previous (%s) Current (%s) for package: %s"),
-			*WriteToString<32>(Tag.Name), *WriteToString<32>(InnerPropertyTag.Type), *WriteToString<32>(ElementProp->GetID()), *UnderlyingArchive.GetArchiveName());
+			*WriteToString<32>(Tag.Name), *WriteToString<32>(InnerPropertyTag.GetType()), *WriteToString<32>(Builder.Build()), *UnderlyingArchive.GetArchiveName());
 	}
 
 	return bConversionSucceeded ? EConvertFromTypeResult::Converted : EConvertFromTypeResult::CannotConvert;
diff --git a/Engine/Source/Runtime/CoreUObject/Private/UObject/PropertySoftObjectPtr.cpp b/Engine/Source/Runtime/CoreUObject/Private/UObject/PropertySoftObjectPtr.cpp
--- a/Engine/Source/Runtime/CoreUObject/Private/UObject/PropertySoftObjectPtr.cpp
+++ b/Engine/Source/Runtime/CoreUObject/Private/UObject/PropertySoftObjectPtr.cpp
@@ -1,11 +1,9 @@
 // Copyright Epic Games, Inc. All Rights Reserved.
 
-#include "CoreMinimal.h"
-#include "UObject/ObjectMacros.h"
-#include "UObject/SoftObjectPtr.h"
-#include "UObject/PropertyPortFlags.h"
 #include "UObject/UnrealType.h"
+
 #include "UObject/LinkerLoad.h"
+#include "UObject/SoftObjectPtr.h"
 #if WITH_EDITOR
 #include "Misc/EditorPathHelper.h"
 #endif
@@ -241,20 +239,24 @@ EConvertFromTypeResult FSoftObjectProperty::ConvertFromType(const FPropertyTag&
 		FSerializedPropertyScope SerializedProperty(Archive, this);
 		return PropertyValue->GetUniqueID().SerializeFromMismatchedTag(Tag, Slot) ? EConvertFromTypeResult::Converted : EConvertFromTypeResult::UseSerializeItem;
 	}
-	else if (Tag.Type == NAME_StructProperty && (Tag.StructName == NAME_SoftObjectPath || Tag.StructName == NAME_SoftClassPath || Tag.StructName == NAME_StringAssetReference || Tag.StructName == NAME_StringClassReference))
+	else if (Tag.Type == NAME_StructProperty)
 	{
-		// This property used to be a FSoftObjectPath but is now a TSoftObjectPtr<Foo>
-		FSoftObjectPath PreviousValue;
-		// explicitly call Serialize to ensure that the various delegates needed for cooking are fired
-		FSerializedPropertyScope SerializedProperty(Archive, this);
-		PreviousValue.Serialize(Slot);
-
-		// now copy the value into the object's address space
-		FSoftObjectPtr PreviousValueSoftObjectPtr;
-		PreviousValueSoftObjectPtr = PreviousValue;
-		SetPropertyValue_InContainer(Data, PreviousValueSoftObjectPtr, Tag.ArrayIndex);
-
-		return EConvertFromTypeResult::Converted;
+		const FName StructName = Tag.GetType().GetParameterName();
+		if (StructName == NAME_SoftObjectPath || StructName == NAME_SoftClassPath || StructName == NAME_StringAssetReference || StructName == NAME_StringClassReference)
+		{
+			// This property used to be a FSoftObjectPath but is now a TSoftObjectPtr<Foo>
+			FSoftObjectPath PreviousValue;
+			// explicitly call Serialize to ensure that the various delegates needed for cooking are fired
+			FSerializedPropertyScope SerializedProperty(Archive, this);
+			PreviousValue.Serialize(Slot);
+
+			// now copy the value into the object's address space
+			FSoftObjectPtr PreviousValueSoftObjectPtr;
+			PreviousValueSoftObjectPtr = PreviousValue;
+			SetPropertyValue_InContainer(Data, PreviousValueSoftObjectPtr, Tag.ArrayIndex);
+
+			return EConvertFromTypeResult::Converted;
+		}
 	}
 
 	return EConvertFromTypeResult::UseSerializeItem;
diff --git a/Engine/Source/Runtime/CoreUObject/Private/UObject/PropertyStruct.cpp b/Engine/Source/Runtime/CoreUObject/Private/UObject/PropertyStruct.cpp
--- a/Engine/Source/Runtime/CoreUObject/Private/UObject/PropertyStruct.cpp
+++ b/Engine/Source/Runtime/CoreUObject/Private/UObject/PropertyStruct.cpp
@@ -388,21 +388,12 @@ bool FStructProperty::SameType(const FProperty* Other) const
 
 EConvertFromTypeResult FStructProperty::ConvertFromType(const FPropertyTag& Tag, FStructuredArchive::FSlot Slot, uint8* Data, UStruct* DefaultsStruct, const uint8* Defaults)
 {
-	FArchive& UnderlyingArchive = Slot.GetUnderlyingArchive();
-
-	auto CanSerializeFromStructWithDifferentName = [](const FArchive& InAr, const FPropertyTag& PropertyTag, const FStructProperty* StructProperty)
-	{
-		if (InAr.UEVer() < VER_UE4_STRUCT_GUID_IN_PROPERTY_TAG)
-		{
-			// Old Implementation
-			return StructProperty && !StructProperty->UseBinaryOrNativeSerialization(InAr);
-		}
-		return PropertyTag.StructGuid.IsValid() && StructProperty && StructProperty->Struct && (PropertyTag.StructGuid == StructProperty->Struct->GetCustomGuid());
-	};
-
 	if (Struct)
 	{
-		if ((Struct->StructFlags & STRUCT_SerializeFromMismatchedTag) && (Tag.Type != NAME_StructProperty || (Tag.StructName != Struct->GetFName())))
+		FArchive& UnderlyingArchive = Slot.GetUnderlyingArchive();
+		const bool bCanSerialize = CanSerializeFromTypeName(Tag.GetType());
+
+		if ((Struct->StructFlags & STRUCT_SerializeFromMismatchedTag) && !bCanSerialize)
 		{
 			UScriptStruct::ICppStructOps* CppStructOps = Struct->GetCppStructOps();
 			check(CppStructOps && (CppStructOps->HasSerializeFromMismatchedTag() || CppStructOps->HasStructuredSerializeFromMismatchedTag())); // else should not have STRUCT_SerializeFromMismatchedTag
@@ -411,18 +402,17 @@ EConvertFromTypeResult FStructProperty::ConvertFromType(const FPropertyTag& Tag,
 			{
 				return EConvertFromTypeResult::Converted;
 			}
-			else 
+			else
 			{
 				FArchiveUObjectFromStructuredArchive Adapter(Slot);
-				FArchive& Ar = Adapter.GetArchive();
-				if (CppStructOps->HasSerializeFromMismatchedTag() && CppStructOps->SerializeFromMismatchedTag(Tag, Ar, DestAddress))
+				if (CppStructOps->HasSerializeFromMismatchedTag() && CppStructOps->SerializeFromMismatchedTag(Tag, Adapter.GetArchive(), DestAddress))
 				{
 					return EConvertFromTypeResult::Converted;
 				}
-				else if(((Struct->StructFlags & STRUCT_SerializeNative) == 0) && CppStructOps->HasSerializeFromMismatchedTag() && CppStructOps->IsUECoreVariant())
+				else if (((Struct->StructFlags & STRUCT_SerializeNative) == 0) && CppStructOps->HasSerializeFromMismatchedTag() && CppStructOps->IsUECoreVariant())
 				{
 					// Special case for Transform, as the f/d variants are immutable whilst the default is not, so we must call SerializeTaggedProperties directly to perform the conversion.
-					if(Tag.StructName == NAME_Transform)
+					if (Tag.GetType().GetParameterName(0) == NAME_Transform)
 					{
 						Struct->SerializeTaggedProperties(Slot, (uint8*)DestAddress, Struct, nullptr);
 						return EConvertFromTypeResult::Converted;
@@ -433,16 +423,19 @@ EConvertFromTypeResult FStructProperty::ConvertFromType(const FPropertyTag& Tag,
 				}
 				else
 				{
-					UE_LOG(LogClass, Warning, TEXT("SerializeFromMismatchedTag failed: Type mismatch in %s of %s - Previous (%s) Current(StructProperty) for package:  %s"), *Tag.Name.ToString(), *GetName(), *Tag.Type.ToString(), *UnderlyingArchive.GetArchiveName());
+					UE::FPropertyTypeNameBuilder Builder;
+					SaveTypeName(Builder);
+					UE_LOG(LogClass, Warning, TEXT("SerializeFromMismatchedTag failed: Type mismatch in %s - Previous (%s) Current(%s) in package: %s"),
+						*WriteToString<32>(Tag.Name), *WriteToString<32>(Tag.GetType()), *WriteToString<32>(Builder.Build()), *UnderlyingArchive.GetArchiveName());
 					return EConvertFromTypeResult::CannotConvert;
 				}
 			}
 		}
 
-		if (Tag.Type == NAME_StructProperty && Tag.StructName != Struct->GetFName() && !CanSerializeFromStructWithDifferentName(UnderlyingArchive, Tag, this))
+		if (Tag.Type == NAME_StructProperty && !bCanSerialize && (UnderlyingArchive.UEVer() >= VER_UE4_STRUCT_GUID_IN_PROPERTY_TAG || !UseBinaryOrNativeSerialization(UnderlyingArchive)))
 		{
-			UE_LOG(LogClass, Warning, TEXT("Property %s of %s has a struct type mismatch (tag %s != prop %s) in package:  %s. If that struct got renamed, add an entry to ActiveStructRedirects."),
-				*Tag.Name.ToString(), *GetName(), *Tag.StructName.ToString(), *Struct->GetName(), *UnderlyingArchive.GetArchiveName());
+			UE_LOG(LogClass, Warning, TEXT("Struct Property %s has a struct type mismatch (tag %s != prop %s) in package: %s. If that struct got renamed, add an entry to ActiveStructRedirects."),
+				*WriteToString<32>(Tag.Name), *WriteToString<32>(Tag.GetType().GetParameterName(0)), *WriteToString<32>(Struct->GetFName()), *UnderlyingArchive.GetArchiveName());
 			return EConvertFromTypeResult::CannotConvert;
 		}
 	}
@@ -463,7 +456,7 @@ void FStructProperty::AppendSchemaHash(FBlake3& Builder, bool bSkipEditorOnly) c
 
 static const FName NAME_StructOriginalType(ANSITEXTVIEW("OriginalType"));
 
-static const FString* FindOriginalType(const FStructProperty* Struct)
+static const FString* FindOriginalTypeName(const FStructProperty* Struct)
 {
 	FUObjectSerializeContext* Context = FUObjectThreadContext::Get().GetSerializeContext();
 	if (Context && Context->bImpersonateProperties)
@@ -507,7 +500,7 @@ void FStructProperty::SaveToTag(FPropertyTag& Tag)
 {
 	Super::SaveToTag(Tag);
 
-	const FString* OriginalType = FindOriginalType(this);
+	const FString* OriginalType = FindOriginalTypeName(this);
 	const UScriptStruct* LocalStruct = Struct;
 	check(LocalStruct);
 	Tag.StructName = OriginalType ? FName(**OriginalType) : LocalStruct->GetFName();
@@ -518,7 +511,7 @@ void FStructProperty::AssignToTag(FPropertyTag& Tag)
 {
 	Super::AssignToTag(Tag);
 
-	const FString* OriginalType = FindOriginalType(this);
+	const FString* OriginalType = FindOriginalTypeName(this);
 	const UScriptStruct* LocalStruct = Struct;
 	check(LocalStruct);
 	if (OriginalType && FName(**OriginalType) == Tag.StructName)
@@ -527,6 +520,33 @@ void FStructProperty::AssignToTag(FPropertyTag& Tag)
 	}
 }
 
+#if WITH_EDITORONLY_DATA
+static UE::FPropertyTypeName FindOriginalType(const FStructProperty* Struct)
+{
+	FUObjectSerializeContext* Context = FUObjectThreadContext::Get().GetSerializeContext();
+	if (Context && Context->bImpersonateProperties)
+	{
+		const FString* OriginalType = Struct->FindMetaData(NAME_StructOriginalType);
+		if (!OriginalType)
+		{
+			//@note: To support metadata defined on array of struct in UPROPERTY for testing purposes
+			if (FField* OwnerField = Struct->Owner.ToField())
+			{
+				OriginalType = OwnerField->FindMetaData(NAME_StructOriginalType);
+			}
+		}
+		if (OriginalType)
+		{
+			if (UE::FPropertyTypeNameBuilder Type; Type.TryParse(*OriginalType))
+			{
+				return Type.Build();
+			}
+		}
+	}
+	return {};
+}
+#endif // WITH_EDITORONLY_DATA
+
 bool FStructProperty::LoadTypeName(UE::FPropertyTypeName Type, const FPropertyTag* Tag)
 {
 	if (!Super::LoadTypeName(Type, Tag))
@@ -552,12 +572,20 @@ void FStructProperty::SaveTypeName(UE::FPropertyTypeNameBuilder& Type) const
 {
 	Super::SaveTypeName(Type);
 
-	const FString* OriginalType = FindOriginalType(this);
 	const UScriptStruct* LocalStruct = Struct;
 	check(LocalStruct);
 
 	Type.BeginParameters();
-	Type.AddName(OriginalType ? FName(**OriginalType) : LocalStruct->GetFName());
+#if WITH_EDITORONLY_DATA
+	if (const UE::FPropertyTypeName OriginalType = FindOriginalType(this); !OriginalType.IsEmpty())
+	{
+		Type.AddType(OriginalType);
+	}
+	else
+#endif // WITH_EDITORONLY_DATA
+	{
+		Type.AddPath(LocalStruct);
+	}
 	if (const FGuid StructGuid = LocalStruct->GetCustomGuid(); StructGuid.IsValid())
 	{
 		Type.AddGuid(StructGuid);
@@ -587,6 +615,12 @@ bool FStructProperty::CanSerializeFromTypeName(UE::FPropertyTypeName Type) const
 		return StructGuid == LocalStruct->GetCustomGuid();
 	}
 
-	const FString* OriginalType = FindOriginalType(this);
-	return StructName == (OriginalType ? FName(**OriginalType) : LocalStruct->GetFName());
+#if WITH_EDITORONLY_DATA
+	if (const UE::FPropertyTypeName OriginalType = FindOriginalType(this); !OriginalType.IsEmpty())
+	{
+		return StructName == OriginalType.GetName();
+	}
+#endif // WITH_EDITORONLY_DATA
+
+	return false;
 }
diff --git a/Engine/Source/Runtime/CoreUObject/Private/UObject/PropertyTag.cpp b/Engine/Source/Runtime/CoreUObject/Private/UObject/PropertyTag.cpp
--- a/Engine/Source/Runtime/CoreUObject/Private/UObject/PropertyTag.cpp
+++ b/Engine/Source/Runtime/CoreUObject/Private/UObject/PropertyTag.cpp
@@ -1,13 +1,37 @@
 // Copyright Epic Games, Inc. All Rights Reserved.
 
 #include "UObject/PropertyTag.h"
-#include "UObject/DebugSerializationFlags.h"
-#include "Serialization/SerializedPropertyScope.h"
+
+#include "Misc/AsciiSet.h"
 #include "Serialization/ArchiveUObjectFromStructuredArchive.h"
-#include "UObject/Package.h"
-#include "UObject/UnrealType.h"
+#include "Serialization/SerializedPropertyScope.h"
+#include "String/ParseTokens.h"
+#include "String/Split.h"
 #include "UObject/BlueprintsObjectVersion.h"
+#include "UObject/DebugSerializationFlags.h"
 #include "UObject/OverriddenPropertySet.h"
+#include "UObject/Package.h"
+#include "UObject/UnrealType.h"
+
+/**
+ * Enum flags that indicate that additional data was serialized for that property tag.
+ * Registered flags should be serialized in ascending order.
+ */
+enum class EPropertyTagExtension : uint8
+{
+	NoExtension					= 0x00,
+	ReserveForFutureUse			= 0x01, // Can be used to add a next group of extensions
+
+	////////////////////////////////////////////////
+	// First extension group
+	OverridableInformation		= 0x02,
+
+	//
+	// Add more extensions for the first group here
+	//
+};
+
+ENUM_CLASS_FLAGS(EPropertyTagExtension);
 
 thread_local const FPropertyTag* FPropertyTagScope::CurrentPropertyTag = nullptr;
 
@@ -15,33 +39,131 @@ thread_local const FPropertyTag* FPropertyTagScope::CurrentPropertyTag = nullptr
 FPropertyTag
 -----------------------------------------------------------------------------*/
 FPropertyTag::FPropertyTag()
-: OverrideOperation(EOverriddenPropertyOperation::None)
-{}
+	: OverrideOperation(EOverriddenPropertyOperation::None)
+{
+}
 
-FPropertyTag::FPropertyTag( FArchive& InSaveAr, FProperty* Property, int32 InIndex, uint8* Value, const uint8* Defaults )
-	: ArrayIndex(InIndex)
+FPropertyTag::FPropertyTag(FProperty* Property, int32 InIndex, uint8* Value)
+	: Prop(Property)
+	, Name(Property->GetFName())
+	, ArrayIndex(InIndex)
 	, OverrideOperation(EOverriddenPropertyOperation::None)
 {
-	check(!InSaveAr.GetArchiveState().UseUnversionedPropertySerialization());
+	UE::FPropertyTypeNameBuilder TypeBuilder;
+	Property->SaveTypeName(TypeBuilder);
+	SetType(TypeBuilder.Build());
+
 	Property->SaveToTag(*this);
+
 	if (FBoolProperty* Bool = CastField<FBoolProperty>(Property))
+	{
+		BoolVal = Bool->GetPropertyValue(Value);
+	}
+}
+
+void FPropertyTag::SetProperty(FProperty* Property)
+{
+	PRAGMA_DISABLE_DEPRECATION_WARNINGS;
+	Prop = Property;
+	PRAGMA_ENABLE_DEPRECATION_WARNINGS;
+}
+
+void FPropertyTag::SetPropertyGuid(const FGuid& InPropertyGuid)
+{
+	HasPropertyGuid = InPropertyGuid.IsValid();
+	PropertyGuid = InPropertyGuid;
+}
+
+void FPropertyTag::SetType(UE::FPropertyTypeName InFullType)
+{
+	TypeName = InFullType;
+	Type = TypeName.GetName();
+
+	if (const EName* TagType = Type.ToEName(); TagType && Type.GetNumber() == NAME_NO_NUMBER_INTERNAL)
+	{
+		switch (*TagType)
 		{
-			BoolVal = Bool->GetPropertyValue(Value);
+		case NAME_StructProperty:
+			StructName = TypeName.GetParameterName(0);
+			if (FName StructGuidName = TypeName.GetParameterName(1);
+				StructGuidName.IsNone() || !FGuid::Parse(StructGuidName.ToString(), StructGuid))
+			{
+				StructGuid.Invalidate();
+			}
+			break;
+		case NAME_ByteProperty:
+		case NAME_EnumProperty:
+			EnumName = TypeName.GetParameterName();
+			break;
+		case NAME_ArrayProperty:
+		case NAME_OptionalProperty:
+		case NAME_SetProperty:
+			InnerType = TypeName.GetParameterName();
+			break;
+		case NAME_MapProperty:
+			InnerType = TypeName.GetParameterName(0);
+			ValueType = TypeName.GetParameterName(1);
+			break;
+		default:
+			break;
 		}
+	}
 }
 
-// Set optional property guid
-void FPropertyTag::SetPropertyGuid(const FGuid& InPropertyGuid)
+static EPropertyTagExtension CalculatePropertyExtensionFlags(FArchive& UnderlyingArchive, FPropertyTag& Tag)
 {
-	if (InPropertyGuid.IsValid())
+	EPropertyTagExtension PropertyTagExtensions = EPropertyTagExtension::NoExtension;
+
+	if (UnderlyingArchive.IsSaving())
 	{
-		PropertyGuid = InPropertyGuid;
-		HasPropertyGuid = true;
+		// OverridableInformation
+		const FOverriddenPropertySet* OverriddenProperties = FOverridableSerializationLogic::GetOverriddenProperties();
+		Tag.bExperimentalOverridableLogic = Tag.GetProperty()->HasAnyPropertyFlags(CPF_ExperimentalOverridableLogic);
+		if (OverriddenProperties || Tag.bExperimentalOverridableLogic)
+		{
+			Tag.OverrideOperation = OverriddenProperties ? OverriddenProperties->GetOverriddenPropertyOperation(UnderlyingArchive.GetSerializedPropertyChain(), Tag.GetProperty()) : EOverriddenPropertyOperation::None;
+			PropertyTagExtensions |= EPropertyTagExtension::OverridableInformation;
+		}
+	}
+
+	return PropertyTagExtensions;
+}
+
+static void SerializePropertyExtensions(FStructuredArchive::FSlot Slot, EPropertyTagExtension PropertyTagExtensions, FPropertyTag& Tag)
+{
+	// Serialize tag extensions, consider doing an init function and context as `EClassSerializationControlExtension` if we add more extensions
+	Slot << SA_ATTRIBUTE(TEXT("PropertyExtensions"), PropertyTagExtensions);
+
+	// OverridableInformation
+	if (EnumHasAnyFlags(PropertyTagExtensions, EPropertyTagExtension::OverridableInformation))
+	{
+		Slot << SA_ATTRIBUTE(TEXT("OverriddenPropertyOperation"), Tag.OverrideOperation);
+		Slot << SA_ATTRIBUTE(TEXT("ExperimentalOverridableLogic"), Tag.bExperimentalOverridableLogic);
+	}
+}
+
+static void ParsePathName(FStringView Path, UE::FPropertyTypeNameBuilder& Builder)
+{
+	FStringView OuterChain;
+	FStringView ObjectName;
+	if (UE::String::SplitLastOfAnyChar(Path, {TEXT('.'), TEXT(':')}, OuterChain, ObjectName))
+	{
+		Builder.AddName(FName(ObjectName));
+		Builder.BeginParameters();
+		UE::String::ParseTokensMultiple(OuterChain, {TEXT('.'), TEXT(':')}, [&Builder](FStringView Outer)
+		{
+			Builder.AddName(FName(Outer));
+		});
+		Builder.EndParameters();
+	}
+	else
+	{
+		Builder.AddName(FName(Path));
 	}
 }
 
 // Serializer.
-FArchive& operator<<( FArchive& Ar, FPropertyTag& Tag )
+FArchive& operator<<(FArchive& Ar, FPropertyTag& Tag)
 {
 	FStructuredArchiveFromArchive(Ar).GetSlot() << Tag;
 	return Ar;
@@ -55,7 +177,20 @@ void operator<<(FStructuredArchive::FSlot Slot, FPropertyTag& Tag)
 	const FPackageFileVersion Version = UnderlyingArchive.UEVer();
 
 	check(!UnderlyingArchive.GetArchiveState().UseUnversionedPropertySerialization());
-	checkf(!UnderlyingArchive.IsSaving() || Tag.Prop, TEXT("FPropertyTag must be constructed with a valid property when used for saving data!"));
+	checkf(!UnderlyingArchive.IsSaving() || Tag.GetProperty(), TEXT("FPropertyTag must be constructed with a valid property when used for saving data!"));
+
+	const auto AddEnumPath = [](UE::FPropertyTypeNameBuilder& Builder, FName EnumName)
+	{
+		TStringBuilder<256> EnumPath(InPlace, EnumName);
+		if (FAsciiSet::HasNone(EnumPath.ToView(), ".:"))
+		{
+			Builder.AddName(EnumName);
+		}
+		else
+		{
+			ParsePathName(EnumPath, Builder);
+		}
+	};
 
 	if (!bIsTextFormat)
 	{
@@ -84,12 +219,18 @@ void operator<<(FStructuredArchive::FSlot Slot, FPropertyTag& Tag)
 
 	if (Tag.Type.GetNumber() == 0)
 	{
+		// Build Tag.TypeName from the partial type name that was saved in older versions.
+		UE::FPropertyTypeNameBuilder TypeBuilder;
+		TypeBuilder.AddName(Tag.Type);
+
 		FNameEntryId TagType = Tag.Type.GetComparisonIndex();
 
 		// only need to serialize this for structs
 		if (TagType == NAME_StructProperty)
 		{
+			TypeBuilder.BeginParameters();
 			Slot << SA_ATTRIBUTE(TEXT("StructName"), Tag.StructName);
+			TypeBuilder.AddName(Tag.StructName);
 			if (Version >= VER_UE4_STRUCT_GUID_IN_PROPERTY_TAG)
 			{
 				if (bIsTextFormat)
@@ -100,14 +241,19 @@ void operator<<(FStructuredArchive::FSlot Slot, FPropertyTag& Tag)
 				{
 					Slot << SA_ATTRIBUTE(TEXT("StructGuid"), Tag.StructGuid);
 				}
+				if (Tag.StructGuid.IsValid())
+				{
+					TypeBuilder.AddGuid(Tag.StructGuid);
+				}
 			}
+			TypeBuilder.EndParameters();
 		}
 		// only need to serialize this for bools
 		else if (TagType == NAME_BoolProperty && !UnderlyingArchive.IsTextFormat())
 		{
 			if (UnderlyingArchive.IsSaving())
 			{
-				FSerializedPropertyScope SerializedProperty(UnderlyingArchive, Tag.Prop);
+				FSerializedPropertyScope SerializedProperty(UnderlyingArchive, Tag.GetProperty());
 				Slot << SA_ATTRIBUTE(TEXT("BoolVal"), Tag.BoolVal);
 			}
 			else
@@ -126,10 +272,20 @@ void operator<<(FStructuredArchive::FSlot Slot, FPropertyTag& Tag)
 			{
 				Slot << SA_ATTRIBUTE(TEXT("EnumName"), Tag.EnumName);
 			}
+			if (!Tag.EnumName.IsNone())
+			{
+				TypeBuilder.BeginParameters();
+				AddEnumPath(TypeBuilder, Tag.EnumName);
+				TypeBuilder.EndParameters();
+			}
 		}
 		else if (TagType == NAME_EnumProperty)
 		{
 			Slot << SA_ATTRIBUTE(TEXT("EnumName"), Tag.EnumName);
+			TypeBuilder.BeginParameters();
+			AddEnumPath(TypeBuilder, Tag.EnumName);
+			TypeBuilder.AddName(NAME_ByteProperty);
+			TypeBuilder.EndParameters();
 		}
 		// need to serialize the InnerType for arrays
 		else if (TagType == NAME_ArrayProperty)
@@ -138,24 +294,42 @@ void operator<<(FStructuredArchive::FSlot Slot, FPropertyTag& Tag)
 			{
 				Slot << SA_ATTRIBUTE(TEXT("InnerType"), Tag.InnerType);
 			}
+			TypeBuilder.BeginParameters();
+			TypeBuilder.AddName(Tag.InnerType);
+			TypeBuilder.EndParameters();
 		}
 		// need to serialize the InnerType for optionals.
 		else if (TagType == NAME_OptionalProperty)
 		{
 			Slot << SA_ATTRIBUTE(TEXT("InnerType"), Tag.InnerType);
+			TypeBuilder.BeginParameters();
+			TypeBuilder.AddName(Tag.InnerType);
+			TypeBuilder.EndParameters();
 		}
 		else if (Version >= VER_UE4_PROPERTY_TAG_SET_MAP_SUPPORT)
 		{
 			if (TagType == NAME_SetProperty)
 			{
 				Slot << SA_ATTRIBUTE(TEXT("InnerType"), Tag.InnerType);
+				TypeBuilder.BeginParameters();
+				TypeBuilder.AddName(Tag.InnerType);
+				TypeBuilder.EndParameters();
 			}
 			else if (TagType == NAME_MapProperty)
 			{
 				Slot << SA_ATTRIBUTE(TEXT("InnerType"), Tag.InnerType);
 				Slot << SA_ATTRIBUTE(TEXT("ValueType"), Tag.ValueType);
+				TypeBuilder.BeginParameters();
+				TypeBuilder.AddName(Tag.InnerType);
+				TypeBuilder.AddName(Tag.ValueType);
+				TypeBuilder.EndParameters();
 			}
 		}
+
+		if (UnderlyingArchive.IsLoading())
+		{
+			Tag.TypeName = TypeBuilder.Build();
+		}
 	}
 
 	// Property tags to handle renamed blueprint properties effectively.
@@ -176,31 +350,10 @@ void operator<<(FStructuredArchive::FSlot Slot, FPropertyTag& Tag)
 		}
 	}
 
-	// Serialize tag extensions, consider doing a init function and context as `EClassSerializationControlExtension` if we add more extensions
-	if (UnderlyingArchive.UEVer() >= EUnrealEngineObjectUE5Version::PROPERTY_TAG_EXTENSION_AND_OVERRIDABLE_SERIALIZATION)
+	if (Version >= EUnrealEngineObjectUE5Version::PROPERTY_TAG_EXTENSION_AND_OVERRIDABLE_SERIALIZATION)
 	{
-		EPropertyTagExtension PropertyTagExtensions = EPropertyTagExtension::NoExtension;
-
-		if (UnderlyingArchive.IsSaving())
-		{
-			// Overridable information extension
-			const FOverriddenPropertySet* OverriddenProperties = FOverridableSerializationLogic::GetOverriddenProperties();
-			Tag.bExperimentalOverridableLogic = Tag.Prop->HasAnyPropertyFlags(CPF_ExperimentalOverridableLogic);
-			if (OverriddenProperties || Tag.bExperimentalOverridableLogic)
-			{
-				Tag.OverrideOperation = OverriddenProperties ? OverriddenProperties->GetOverriddenPropertyOperation(UnderlyingArchive.GetSerializedPropertyChain(), Tag.Prop) : EOverriddenPropertyOperation::None;
-				PropertyTagExtensions |= EPropertyTagExtension::OverridableInformation;
-			}
-		}
-
-		Slot << SA_ATTRIBUTE(TEXT("PropertyExtensions"), PropertyTagExtensions);
-
-		// Overridable information extension
-		if (EnumHasAnyFlags(PropertyTagExtensions,EPropertyTagExtension::OverridableInformation))
-		{
-			Slot << SA_ATTRIBUTE(TEXT("OverriddenPropertyOperation"), Tag.OverrideOperation);
-			Slot << SA_ATTRIBUTE(TEXT("ExperimentalOverridableLogic"), Tag.bExperimentalOverridableLogic);
-		}
+		EPropertyTagExtension PropertyTagExtensions = CalculatePropertyExtensionFlags(UnderlyingArchive, Tag);
+		SerializePropertyExtensions(Slot, PropertyTagExtensions, Tag);
 	}
 }
 
diff --git a/Engine/Source/Runtime/CoreUObject/Public/UObject/Class.h b/Engine/Source/Runtime/CoreUObject/Public/UObject/Class.h
--- a/Engine/Source/Runtime/CoreUObject/Public/UObject/Class.h
+++ b/Engine/Source/Runtime/CoreUObject/Public/UObject/Class.h
@@ -1457,7 +1457,12 @@ public:
 				if constexpr (TIsUECoreType<CPPSTRUCT>::Value)
 				{
 					// Custom version of SerializeFromMismatchedTag for core types, which don't have access to FPropertyTag.
-					return ((CPPSTRUCT*)Data)->SerializeFromMismatchedTag(Tag.StructName, Ar);
+					FName StructName;
+					if (Tag.GetType().GetName() == NAME_StructProperty && Tag.GetType().GetParameterCount() >= 1)
+					{
+						StructName = Tag.GetType().GetParameterName();
+					}
+					return ((CPPSTRUCT*)Data)->SerializeFromMismatchedTag(StructName, Ar);
 				}
 				else
 				{
@@ -1477,7 +1482,12 @@ public:
 				if constexpr (TIsUECoreType<CPPSTRUCT>::Value)
 				{
 					// Custom version of SerializeFromMismatchedTag for core types, which don't understand FPropertyTag.
-					return ((CPPSTRUCT*)Data)->SerializeFromMismatchedTag(Tag.StructName, Slot);
+					FName StructName;
+					if (Tag.GetType().GetName() == NAME_StructProperty && Tag.GetType().GetParameterCount() >= 1)
+					{
+						StructName = Tag.GetType().GetParameterName();
+					}
+					return ((CPPSTRUCT*)Data)->SerializeFromMismatchedTag(StructName, Slot);
 				}
 				else
 				{
diff --git a/Engine/Source/Runtime/CoreUObject/Public/UObject/PropertyTag.h b/Engine/Source/Runtime/CoreUObject/Public/UObject/PropertyTag.h
--- a/Engine/Source/Runtime/CoreUObject/Public/UObject/PropertyTag.h
+++ b/Engine/Source/Runtime/CoreUObject/Public/UObject/PropertyTag.h
@@ -10,30 +10,12 @@
 #include "Misc/Guid.h"
 #include "Serialization/StructuredArchive.h"
 #include "UObject/NameTypes.h"
+#include "UObject/PropertyTypeName.h"
 
 enum class EOverriddenPropertyOperation : uint8;
 class FArchive;
 class FProperty;
 
-/**
- *  Enum flags that indicate that additional data was serialized for that property tag
- *	Registered flag should be serialized in ascending order 
- */
-enum class EPropertyTagExtension : uint8
-{
-	NoExtension					= 0x00,
-	ReserveForFutureUse			= 0x01, // Can be use to add a next group of extension
-
-	////////////////////////////////////////////////
-	// First extension group
-	OverridableInformation		= 0x02,
-
-	//
-	// Add more extension for the first group here
-	//
-};
-ENUM_CLASS_FLAGS(EPropertyTagExtension);
-
 /**
  *  A tag describing a class property, to aid in serialization.
  */
@@ -43,39 +25,60 @@ struct FPropertyTag
 	FProperty* Prop = nullptr;
 
 	// Variables.
+private:
+	UE::FPropertyTypeName TypeName;
+
+public:
 	FName	Type;		// Type of property
-	uint8	BoolVal = 0;// a boolean property's value (never need to serialize data for bool properties except here)
 	FName	Name;		// Name of property.
 	FName	StructName;	// Struct name if FStructProperty.
 	FName	EnumName;	// Enum name if FByteProperty or FEnumProperty
-	FName	InnerType;	// Inner type if FArrayProperty, FSetProperty, or FMapProperty
+	FName	InnerType;	// Inner type if FArrayProperty, FSetProperty, FMapProperty, or OptionalProperty
 	FName	ValueType;	// Value type if UMapPropery
 	int32	Size = 0;   // Property size.
 	int32	ArrayIndex = INDEX_NONE; // Index if an array; else 0.
 	int64	SizeOffset = INDEX_NONE; // location in stream of tag size member
 	FGuid	StructGuid;
-	uint8	HasPropertyGuid = 0;
 	FGuid	PropertyGuid;
+	uint8	HasPropertyGuid = 0;
+	uint8	BoolVal = 0;// a boolean property's value (never need to serialize data for bool properties except here)
 	EOverriddenPropertyOperation OverrideOperation; // Overridable serialization state reconstruction 
 	bool	bExperimentalOverridableLogic = false; // Remember if property had CPF_ExperimentalOverridableLogic when saved
 
 	// Constructors.
 	FPropertyTag();
-	UE_INTERNAL COREUOBJECT_API FPropertyTag( FArchive& InSaveAr, FProperty* Property, int32 InIndex, uint8* Value, const uint8* Defaults );
+	UE_INTERNAL COREUOBJECT_API FPropertyTag(FProperty* Property, int32 InIndex, uint8* Value);
+
+	PRAGMA_DISABLE_DEPRECATION_WARNINGS
+	FPropertyTag(FPropertyTag&&) = default;
+	FPropertyTag(const FPropertyTag&) = default;
+	FPropertyTag& operator=(FPropertyTag&&) = default;
+	FPropertyTag& operator=(const FPropertyTag&) = default;
+	PRAGMA_ENABLE_DEPRECATION_WARNINGS
 
-	// Set optional property guid
+	inline FProperty* GetProperty() const
+	{
+		PRAGMA_DISABLE_DEPRECATION_WARNINGS
+		return Prop;
+		PRAGMA_ENABLE_DEPRECATION_WARNINGS
+	}
+
+	void SetProperty(FProperty* Property);
 	void SetPropertyGuid(const FGuid& InPropertyGuid);
 
+	inline UE::FPropertyTypeName GetType() const { return TypeName; }
+	void SetType(UE::FPropertyTypeName TypeName);
+
 	// Serializer.
-	friend FArchive& operator<<( FArchive& Ar, FPropertyTag& Tag );
+	friend FArchive& operator<<(FArchive& Ar, FPropertyTag& Tag);
 	friend void operator<<(FStructuredArchive::FSlot Slot, FPropertyTag& Tag);
 
 	// Property serializer.
-	void SerializeTaggedProperty( FArchive& Ar, FProperty* Property, uint8* Value, const uint8* Defaults ) const;
+	void SerializeTaggedProperty(FArchive& Ar, FProperty* Property, uint8* Value, const uint8* Defaults) const;
 	UE_INTERNAL COREUOBJECT_API void SerializeTaggedProperty(FStructuredArchive::FSlot Slot, FProperty* Property, uint8* Value, const uint8* Defaults) const;
 };
 
-struct FPropertyTagScope
+struct UE_INTERNAL FPropertyTagScope
 {
 	FPropertyTagScope(const FPropertyTag* InCurrentPropertyTag)
 	: PropertyTagToRestore(CurrentPropertyTag)
diff --git a/Engine/Source/Runtime/CoreUObject/Public/UObject/UnrealType.h b/Engine/Source/Runtime/CoreUObject/Public/UObject/UnrealType.h
--- a/Engine/Source/Runtime/CoreUObject/Public/UObject/UnrealType.h
+++ b/Engine/Source/Runtime/CoreUObject/Public/UObject/UnrealType.h
@@ -1072,7 +1072,7 @@ public:
 	COREUOBJECT_API virtual void SaveTypeName(UE::FPropertyTypeNameBuilder& Type) const;
 
 	/**
-	 * Returns true if this property can use SerializeItem to load from a property with this type name.
+	 * Returns true if the type of this property matches the type name.
 	 */
 	COREUOBJECT_API virtual bool CanSerializeFromTypeName(UE::FPropertyTypeName Type) const;
 
@@ -1751,7 +1751,7 @@ public:
 				return std::is_same_v<TCppType, int64> ? EConvertFromTypeResult::Serialized : EConvertFromTypeResult::Converted;
 
 			case NAME_ByteProperty:
-				if (!Tag.EnumName.IsNone())
+				if (Tag.GetType().GetParameterCount() >= 1)
 				{
 					int64 PreviousValue = this->ReadEnumAsInt64(Slot, DefaultsStruct, Tag);
 					this->SetPropertyValue_InContainer(Data, (TCppType)PreviousValue, Tag.ArrayIndex);
