Improve the support for notification that UObjects have been registered.  This includes the callback always being registered regardless of the target options.

diff --git a/Engine/Source/Programs/ReplicationSystemTest/Private/ReplicationSystemTest.cpp b/Engine/Source/Programs/ReplicationSystemTest/Private/ReplicationSystemTest.cpp
--- a/Engine/Source/Programs/ReplicationSystemTest/Private/ReplicationSystemTest.cpp
+++ b/Engine/Source/Programs/ReplicationSystemTest/Private/ReplicationSystemTest.cpp
@@ -197,9 +197,8 @@ static void LoadModules()
 {
 	// Always attempt to load CoreUObject. It requires additional pre-init which is called from its module's StartupModule method.
 #if WITH_COREUOBJECT
-#if USE_PER_MODULE_UOBJECT_BOOTSTRAP // otherwise do it later
+	// Always register the UObjects callback for VNI and general consistency with the callbacks ProcessNewlyLoadedUObjects calls.
 	FModuleManager::Get().OnProcessLoadedObjectsCallback().AddStatic(ProcessNewlyLoadedUObjects);
-#endif
 	FModuleManager::Get().LoadModule(TEXT("CoreUObject"));
 
 	FCoreDelegates::OnInit.Broadcast();
diff --git a/Engine/Source/Runtime/Core/Public/Modules/ModuleManager.h b/Engine/Source/Runtime/Core/Public/Modules/ModuleManager.h
--- a/Engine/Source/Runtime/Core/Public/Modules/ModuleManager.h
+++ b/Engine/Source/Runtime/Core/Public/Modules/ModuleManager.h
@@ -518,6 +518,9 @@ public:
 
 	/**
 	 * Gets a multicast delegate that is executed when any UObjects need processing after a module was loaded.
+	 * Do not use this delegate if you wish to monitor the registration of UObjects for a module. 
+	 * Use either CompiledInUObjectsRegisteredDelegate in CoreUObject to reliably receive notifications after UObjects 
+	 * have been registered.
 	 *
 	 * @return The delegate.
 	 */
diff --git a/Engine/Source/Runtime/CoreUObject/Private/UObject/Obj.cpp b/Engine/Source/Runtime/CoreUObject/Private/UObject/Obj.cpp
--- a/Engine/Source/Runtime/CoreUObject/Private/UObject/Obj.cpp
+++ b/Engine/Source/Runtime/CoreUObject/Private/UObject/Obj.cpp
@@ -5281,12 +5281,6 @@ void InitUObject()
 
 	FCoreDelegates::OnShutdownAfterError.AddStatic(StaticShutdownAfterError);
 	FCoreDelegates::OnExit.AddStatic(StaticExit);
-#if !USE_PER_MODULE_UOBJECT_BOOTSTRAP && !IS_MONOLITHIC
-	// Not necessary when USE_PER_MODULE_UOBJECT_BOOTSTRAP==0 since the callback gets installed elsewhere
-	// Also not necessary for monolithic builds as all pending registrants are available at once on app start
-	// so ProcessNewlyLoadedUObjects needs to only ever be invoked once, not for each module
-	FModuleManager::Get().OnProcessLoadedObjectsCallback().AddStatic(ProcessNewlyLoadedUObjects);
-#endif
 
 	struct Local
 	{
diff --git a/Engine/Source/Runtime/CoreUObject/Private/UObject/UObjectBase.cpp b/Engine/Source/Runtime/CoreUObject/Private/UObject/UObjectBase.cpp
--- a/Engine/Source/Runtime/CoreUObject/Private/UObject/UObjectBase.cpp
+++ b/Engine/Source/Runtime/CoreUObject/Private/UObject/UObjectBase.cpp
@@ -861,6 +861,7 @@ void ProcessNewlyLoadedUObjects(FName Package, bool bCanProcessNewlyLoadedObject
 #endif
 	if (!bCanProcessNewlyLoadedObjects)
 	{
+		FCoreUObjectDelegates::CompiledInUObjectsRegisteredDelegate.Broadcast(Package, ECompiledInUObjectsRegisteredStatus::Delayed);
 		return;
 	}
 	LLM_SCOPE(ELLMTag::UObject);
@@ -888,11 +889,13 @@ void ProcessNewlyLoadedUObjects(FName Package, bool bCanProcessNewlyLoadedObject
 		UObjectProcessRegistrants();
 		UObjectLoadAllCompiledInStructs();
 
-		FCoreUObjectDelegates::CompiledInUObjectsRegisteredDelegate.Broadcast(Package);
+		FCoreUObjectDelegates::CompiledInUObjectsRegisteredDelegate.Broadcast(Package, ECompiledInUObjectsRegisteredStatus::PreCDO);
 
 		UObjectLoadAllCompiledInDefaultProperties(AllNewClasses);
 	}
 
+	FCoreUObjectDelegates::CompiledInUObjectsRegisteredDelegate.Broadcast(Package, ECompiledInUObjectsRegisteredStatus::PostCDO);
+
 #if WITH_RELOAD
 	IReload* Reload = GetActiveReloadInterface();
 	if (Reload != nullptr)
diff --git a/Engine/Source/Runtime/CoreUObject/Public/UObject/UObjectGlobals.h b/Engine/Source/Runtime/CoreUObject/Public/UObject/UObjectGlobals.h
--- a/Engine/Source/Runtime/CoreUObject/Public/UObject/UObjectGlobals.h
+++ b/Engine/Source/Runtime/CoreUObject/Public/UObject/UObjectGlobals.h
@@ -3043,6 +3043,22 @@ struct FEndLoadPackageContext
 	bool bSynchronous;
 };
 
+enum class ECompiledInUObjectsRegisteredStatus
+{
+
+	// Registration for the given module has been delayed and will be registered with a layer notification where the package name is NAME_None.
+	// Calls to FindObject<> for any delayed registrations will return nullptr.
+	Delayed,
+
+	// All pending registrations have been done but CDOs have not been created.  This may be invoked multiple times for the same package name.
+	// Calls to FindObject<> for all pending registrations will return the UObject in question.
+	PreCDO,
+
+	// All pending registrations and CDO creation has been done.  Calls to FindObject<> for all pending registrations will return the UObject 
+	// in question.  The CDO will also be available for any registered UClasses.
+	PostCDO,
+};
+
 /**
  * Global CoreUObject delegates
  */
@@ -3171,8 +3187,14 @@ struct FCoreUObjectDelegates
 	UE_DEPRECATED(5.0, "ReinstanceHotReloadedClassesDelegate has been deprecated, use FReload for class re-instancing or ReloadReinstancingCompleteDelegate for notification")
 	static COREUOBJECT_API FReinstanceHotReloadedClassesDelegate ReinstanceHotReloadedClassesDelegate;
 
-	/** Delegate for catching when UClasses/UStructs/UEnums would be available via FindObject<>(), but before their CDOs would be constructed. */
-	DECLARE_MULTICAST_DELEGATE_OneParam(FCompiledInUObjectsRegisteredDelegate, FName /*Package*/);
+	/** 
+	 * Delegate invoked when requests are made to process pending UObject registrations.  This will be invoked during the loading of modules
+	 * and outside of module load to register any delayed registrations.  This callback will be invoked regardless of if the module contained
+	 * any UObject definitions.  The callback will also be invoked during module load in monolithic builds.
+	 * 
+	 * See ECompiledInUObjectsRegisteredStatus for more detailed information about the notifications.
+	 */
+	DECLARE_MULTICAST_DELEGATE_TwoParams(FCompiledInUObjectsRegisteredDelegate, FName /*Package*/, ECompiledInUObjectsRegisteredStatus Status);
 	static COREUOBJECT_API FCompiledInUObjectsRegisteredDelegate CompiledInUObjectsRegisteredDelegate;
 
 	/** Sent at the very beginning of LoadMap */
diff --git a/Engine/Source/Runtime/Launch/Private/LaunchEngineLoop.cpp b/Engine/Source/Runtime/Launch/Private/LaunchEngineLoop.cpp
--- a/Engine/Source/Runtime/Launch/Private/LaunchEngineLoop.cpp
+++ b/Engine/Source/Runtime/Launch/Private/LaunchEngineLoop.cpp
@@ -4398,9 +4398,8 @@ bool FEngineLoop::LoadCoreModules()
 {
 	// Always attempt to load CoreUObject. It requires additional pre-init which is called from its module's StartupModule method.
 #if WITH_COREUOBJECT
-#if USE_PER_MODULE_UOBJECT_BOOTSTRAP // otherwise do it later
+	// Always register the UObjects callback for VNI and general consistency with the callbacks ProcessNewlyLoadedUObjects calls.
 	FModuleManager::Get().OnProcessLoadedObjectsCallback().AddStatic(ProcessNewlyLoadedUObjects);
-#endif
 	return FModuleManager::Get().LoadModule(TEXT("CoreUObject")) != nullptr;
 #else
 	return true;
