Improves the performance of CoreRedirects

::Before::
- Boot -> PIE (6m3s)
- FCoreRedirects::RedirectNameAndValues: 18s

::After::
- Boot -> PIE (5m41s)
- FCoreRedirects::RedirectNameAndValues: 5.3s

- MatchWildcard=true support has been added to how CoreRedirects are defined in INI files. You may now define OldName="...SomeSuffix", OldName="SomePrefix..." or OldName="...SomeSubstring..." to perform string manipulation based matching for redirects, deprecating MatchSubstring=true
- Substring matches are very expensive compared to pure FName matching or even prefix and suffix matching. This expense is problematic since the presence of a substring matching rule can mean that all queries for a particular redirect type, such as package name, will need to perform the substring check even if a substring match is relatively uncommon. To combat this cost, we add a reimplementation of the PM-k, (PredictMatch-8) algorithm. This allows a redirect query to check for substring matches against _all_ possible substrings with a single pass over the redirect query name. The fuzzy matching is designed such that false positives may occur but not false negatives. When a fuzzy match is found we must perform the slow substring match to confirm. However fuzzy matching allows use to dramatically reduce our string scannin. In a case where there were 4 MatchSubstrings, from boot to PIE, we now perform slow substring matches < 0.8% of the time compared to previously.
- Replaces `FString` manipulation during redirect queries to instead use a small wrapper type `FNameUtf8String` which provides a small API to allow for FName comparisons, building and ultimately reducing copies. UTF8 is used since we need to ensure we have an encoding that we can keep the same for storing wildcard strings that will match how we query for wildcard matches. UTF8 is byte efficient for ANSI text, which improves fuzzy match prediction.

==========================================
Substring Stats
Substring Lookups:                24342736
Prediction Hit:                   72
Prediction Miss:                  213377
==========================================

diff --git a/Engine/Config/BaseEngine.ini b/Engine/Config/BaseEngine.ini
--- a/Engine/Config/BaseEngine.ini
+++ b/Engine/Config/BaseEngine.ini
@@ -335,9 +335,19 @@ DefaultWorldPartitionHLODLayer=/Engine/Maps/Templates/HLODs/HLODLayer_Instanced.
 ; You can also add options or add a value change map:
 ;
 ; +ClassRedirects=(OldName="OldClass", NewName="/Script/GameName.GameSpecificClass", bInstanceOnly=true)
-; +PackageRedirects=(OldName="/OldPlugin/", NewName="/NewPlugin/", MatchSubstring=true)
+; +PackageRedirects=(OldName="/OldPlugin/...", NewName="/NewPlugin/", MatchWildcard=true)
 ; +PackageRedirects=(OldName="/Game/DeletedContentPackage", Removed=true)
 ; +EnumRedirects=(OldName="/Script/Core.EEnumType", NewName="/Script/NewModule.ENewEnumType", ValueChanges=(("OldValue", "NewValue"), ("OldValue2", "NewValue2")) )
+; 
+; CoreRedirects can also specify matches by Wildcard.
+; You can match via substring, prefix, or suffix by specifing the "..." wildcard as part of OldName:
+;
+; +PackageRedirects=(OldName="/OldPrefix/...", NewName="/NewPluginPrefix/", MatchWildcard=true)				<-- Matches names starting with "/OldPrefix/"
+; +PackageRedirects=(OldName=".../OldSuffix", NewName="/NewPluginSuffix", MatchWildcard=true)				<-- Matches names ending with "/OldSuffix"
+; +PackageRedirects=(OldName=".../OldSubstring/...", NewName="/NewPluginSubstring/", MatchWildcard=true)	<-- Matches names containing "/OldSubstring/"
+;
+; NOTE, wildcard redirects are much slower than direct matches (non-wildcard matching) and should only be used temporarily! Each wildcard redirect defined
+; will slow down the CoreRedirect system, regardless of whether they can/will ever match. Use with caution.
 
 ; 4.10/4.11
 
diff --git a/Engine/Source/Runtime/CoreUObject/Internal/UObject/CoreRedirects/PM-k.h b/Engine/Source/Runtime/CoreUObject/Internal/UObject/CoreRedirects/PM-k.h
new file mode 100644
--- /dev/null
+++ b/Engine/Source/Runtime/CoreUObject/Internal/UObject/CoreRedirects/PM-k.h
@@ -0,0 +1,271 @@
+// ================================================================================
+// Software Name: pm4-bitap.c
+// Version: V1.0
+// URL: https://www.genivia.com/files/BSD-3.txt
+// ===========================================================================================
+//  BSD 3-Clause License
+// 
+// Copyright (c) 2023, Robert van Engelen, Genivia Inc.
+// All rights reserved.
+// 
+// Redistribution and use in source and binary forms, with or without
+// modification, are permitted provided that the following conditions are met:
+// 
+// 1. Redistributions of source code must retain the above copyright notice, this
+//    list of conditions and the following disclaimer.
+// 
+// 2. Redistributions in binary form must reproduce the above copyright notice,
+//    this list of conditions and the following disclaimer in the documentation
+//    and/or other materials provided with the distribution.
+// 
+// 3. Neither the name of the copyright holder nor the names of its
+//    contributors may be used to endorse or promote products derived from
+//    this software without specific prior written permission.
+// 
+// THIS SOFTWARE IS PROVIDED BY THE COPYRIGHT HOLDERS AND CONTRIBUTORS "AS IS"
+// AND ANY EXPRESS OR IMPLIED WARRANTIES, INCLUDING, BUT NOT LIMITED TO, THE
+// IMPLIED WARRANTIES OF MERCHANTABILITY AND FITNESS FOR A PARTICULAR PURPOSE ARE
+// DISCLAIMED. IN NO EVENT SHALL THE COPYRIGHT HOLDER OR CONTRIBUTORS BE LIABLE
+// FOR ANY DIRECT, INDIRECT, INCIDENTAL, SPECIAL, EXEMPLARY, OR CONSEQUENTIAL
+// DAMAGES (INCLUDING, BUT NOT LIMITED TO, PROCUREMENT OF SUBSTITUTE GOODS OR
+// SERVICES; LOSS OF USE, DATA, OR PROFITS; OR BUSINESS INTERRUPTION) HOWEVER
+// CAUSED AND ON ANY THEORY OF LIABILITY, WHETHER IN CONTRACT, STRICT LIABILITY,
+// OR TORT (INCLUDING NEGLIGENCE OR OTHERWISE) ARISING IN ANY WAY OUT OF THE USE
+// OF THIS SOFTWARE, EVEN IF ADVISED OF THE POSSIBILITY OF SUCH DAMAGE.
+
+#pragma once
+#include "CoreTypes.h"
+#include "HAL/UnrealMemory.h"
+
+/**
+	Predict Match (pm-k) is a fuzzy matching algorithm that enables looking for a potential match across
+	multiple substrings at the same time. We take two approaches, one is Bitap (Bit Approximation) which uses a
+	sliding window of the size of the smallest substring (small substrings will create more false positives). Bitap
+	is very quick for scanning potential matches. We then use pm-8, a sliding window of a fixed size
+	(in our case 8 bytes), to further look for potential matches. If we find a hit, we must defer to slower substring
+	matching approaches to confirm the match. PredictMatch relies on hashing to distribute the mapping of characters
+	in substrings at specific offsets. This hashing helps reduce overlaps between similar substrings which reduces false
+	positives, however this does mean that if your hash is poor, or the table size is too small, you may have more false
+	positives than preferred. 
+	
+	This implementation is written by Epic but based on the licensed pm4-bitap.c implementation found at https://www.genivia.com/ugrep.html
+*/
+struct alignas(PLATFORM_CACHE_LINE_SIZE) FPredictMatch8
+{
+	static constexpr uint16 TABLE_SIZE = 256;
+	static constexpr uint16 ALPHABET_SIZE = 256;
+
+	FPredictMatch8()
+	{
+		Reset();
+	}
+
+	uint16 GetMinimumWordLength() const
+	{
+		return MinimumWordLength;
+	}
+
+	void AddPredictionWord(const uint8* Data, const uint16 DataLen)
+	{
+		uint8 Bytes[8] = {};
+
+		check(DataLen);
+		switch (DataLen)
+		{
+		default:
+			Bytes[7] = Data[7];
+		case 7:
+			Bytes[6] = Data[6];
+		case 6:
+			Bytes[5] = Data[5];
+		case 5:
+			Bytes[4] = Data[4];
+		case 4:
+			Bytes[3] = Data[3];
+		case 3:
+			Bytes[2] = Data[2];
+		case 2:
+			Bytes[1] = Data[1];
+		case 1:
+			Bytes[0] = Data[0];
+		}
+
+		const uint16 Hash1 = HashFn(Bytes[0], Bytes[1]);
+		const uint16 Hash2 = HashFn(Hash1, Bytes[2]);
+		const uint16 Hash3 = HashFn(Hash2, Bytes[3]);
+		const uint16 Hash4 = HashFn(Hash3, Bytes[4]);
+		const uint16 Hash5 = HashFn(Hash4, Bytes[5]);
+		const uint16 Hash6 = HashFn(Hash5, Bytes[6]);
+		const uint16 Hash7 = HashFn(Hash6, Bytes[7]);
+
+		// For each character, store a pair of bits indicating a match (even value, b10)
+		// or an accept (00, even value < 2) into the table at increasing offsets into out 16-bit matching window, with each
+		// character stored into a different bucket (due to hashing). The last character for our added word is always
+		// 00 by definition.
+		// 
+		// e.g Adding 5-letter prediction word "apple"
+		//                character pos
+		//			      0  1  2  3  4  5  6  7
+		// bucket['a']    10 bb bb bb bb bb bb bb
+		// bucket[h('p')] bb 10 bb bb bb bb bb bb
+		// bucket[h('p')] bb bb 10 bb bb bb bb bb
+		// bucket[h('l')] bb bb bb 10 bb bb bb bb
+		// bucket[h('e')] bb bb bb bb 00 bb bb bb  <-- last character is match and accept
+
+		PredictMatchTable[Bytes[0]] &=				  (0x3FFF | ((DataLen > 1) << 15));
+		PredictMatchTable[Hash1]	&= (DataLen > 1 ? (0xCFFF | ((DataLen > 2) << 13)) : 0xFFFF);
+		PredictMatchTable[Hash2]	&= (DataLen > 2 ? (0xF3FF | ((DataLen > 3) << 11)) : 0xFFFF);
+		PredictMatchTable[Hash3]	&= (DataLen > 3 ? (0xFCFF | ((DataLen > 4) << 9))  : 0xFFFF);
+		PredictMatchTable[Hash4]	&= (DataLen > 4 ? (0xFF3F | ((DataLen > 5) << 7))  : 0xFFFF);
+		PredictMatchTable[Hash5]	&= (DataLen > 5 ? (0xFFCF | ((DataLen > 6) << 5))  : 0xFFFF);
+		PredictMatchTable[Hash6]	&= (DataLen > 6 ? (0xFFF3 | ((DataLen > 7) << 3))  : 0xFFFF);
+		PredictMatchTable[Hash7]	&= (DataLen > 7 ? (0xFFFC)                         : 0xFFFF);
+
+		MinimumWordLength = MinimumWordLength < DataLen ? MinimumWordLength : DataLen;
+		for (int i = 0; i < MinimumWordLength; ++i)
+		{
+			const uint8 Byte = Data[i];
+			BitApproxTable[Byte] &= ~(1 << i);
+		}
+	}
+
+	bool MatchApproximate(const uint8* Data, const uint32 DataLen) const
+	{
+		check(DataLen);
+
+		// Note this mask is not what is normally expected. 
+		// This mask is used to check if our sliding window of bits contains 
+		// a 0 in bit position MinimumWordLength which would indicate a potential match
+		const uint16 BitApproxMask = (uint16) (1 << (MinimumWordLength - 1));
+		const uint8* const DataEnd = Data + DataLen;
+
+		// Start with no matching bits (all 1)
+		uint16 Bits = ~0;
+		for (;;)
+		{
+			// Continue scanning until we either run out of bytes or we find an approximate bit match
+			// Shift left and OR our sliding window of potential matches. BitApproxTable[*Data] OR'd
+			// with Bits will either keep sliding a 0 bit left indicating we have a fuzzy match, or as we 
+			// OR values from BitApproxTable[*Data], we will stomp over the sliding 0 with a 1 indicating 
+			// no match and the matching process starts over by sliding in a 0 with the next left shift
+			while ((Data < DataEnd) && (((Bits = (uint16)(Bits << 1) | BitApproxTable[*Data]) & BitApproxMask) != 0))
+			{
+				Data++;
+			}
+
+			if (Data >= DataEnd)
+			{
+				return false;
+			}
+
+			// The Bitap scanning above has indicated we have a potential match, but now defer 
+			// to PredictMatch to further refine our prediction, since Bitap operates window sizes of 
+			// the smallest substring which might be quite small compared to the substrings we are searching for.
+			const uint8* PredictionStart = Data - MinimumWordLength + 1;
+			if (PredictMatch(PredictMatchTable, PredictionStart, (uint32)(DataEnd - PredictionStart)))
+			{
+				return true;
+			}
+			++Data;
+		}
+	}
+
+	void Reset()
+	{
+		// Can be no greater than the number of bits in a BitApproxTable element. This value will shrink if we are given a smaller substring
+		MinimumWordLength = sizeof(BitApproxTable[0]) * 8;
+
+		FMemory::Memset(BitApproxTable, 0xFF, sizeof(BitApproxTable));
+		FMemory::Memset(PredictMatchTable, 0xFF, sizeof(PredictMatchTable));
+	}
+
+private:
+	static uint16 HashFn(const uint16 A, const uint8 B)
+	{
+		return (uint16)((((uint16)(A) << 3) ^ (uint16)(B)) & (TABLE_SIZE - 1));
+	}
+
+	static bool PredictMatch(const uint16* RESTRICT PredictMatchTable, const uint8* RESTRICT Data, const uint32 DataLen)
+	{
+		uint8 Bytes[8] = {};
+
+		check(DataLen);
+		switch (DataLen)
+		{
+		default:
+			Bytes[7] = Data[7];
+		case 7:
+			Bytes[6] = Data[6];
+		case 6:
+			Bytes[5] = Data[5];
+		case 5:
+			Bytes[4] = Data[4];
+		case 4:
+			Bytes[3] = Data[3];
+		case 3:
+			Bytes[2] = Data[2];
+		case 2:
+			Bytes[1] = Data[1];
+		case 1:
+			Bytes[0] = Data[0];
+		}
+
+		/*
+		* Branchless implementation of the following logic:
+		* 
+		*	if Accept(Bytes[0], 0) then return TRUE			// If we have a substring that ends with Bytes[0] at position 0
+		*   if Match(Bytes[0], 0) then						// Otherwise, If we have a substring that has Bytes[0] at position 0
+		*       if Accept(Bytes[1], 1) then return TRUE
+		*       if Match(Bytes[1], 1) then
+		* 		...
+		*           if Accept(Bytes[6], 6) then return TRUE
+		*           if Match(Bytes[6], 6) then
+		*               // It's the last character in the window so no need to check Accept(Bytes[7], 7)
+		*               if Matchbit(Bytes[7], 7) then return TRUE
+		*/
+		const uint16 Hash1 = HashFn(Bytes[0],	Bytes[1]);
+		const uint16 Hash2 = HashFn(Hash1,		Bytes[2]);
+		const uint16 Hash3 = HashFn(Hash2,		Bytes[3]);
+		const uint16 Hash4 = HashFn(Hash3,		Bytes[4]);
+		const uint16 Hash5 = HashFn(Hash4,		Bytes[5]);
+		const uint16 Hash6 = HashFn(Hash5,		Bytes[6]);
+		const uint16 Hash7 = HashFn(Hash6,		Bytes[7]);
+
+		const uint16 AcceptBit0 = PredictMatchTable[Bytes[0]];
+		const uint16 AcceptBit1 = PredictMatchTable[Hash1];
+		const uint16 AcceptBit2 = PredictMatchTable[Hash2];
+		const uint16 AcceptBit3 = PredictMatchTable[Hash3];
+		const uint16 AcceptBit4 = PredictMatchTable[Hash4];
+		const uint16 AcceptBit5 = PredictMatchTable[Hash5];
+		const uint16 AcceptBit6 = PredictMatchTable[Hash6];
+		const uint16 AcceptBit7 = PredictMatchTable[Hash7];
+
+		const uint16 Bits =
+			(AcceptBit0 & 0xC000) | (AcceptBit1 & 0x3000) | (AcceptBit2 & 0x0C00) | (AcceptBit3 & 0x0300) |
+			(AcceptBit4 & 0x00C0) | (AcceptBit5 & 0x0030) | (AcceptBit6 & 0x000C) | (AcceptBit7 & 0x0003);
+		const uint16 MatchBits = ((((((((((((((Bits >> 2) | Bits) >> 2) | Bits) >> 2) | Bits) >> 2) | Bits) >> 2) | Bits) >> 2) | Bits) >> 1) | Bits);
+
+		// all two-bit pairs are set meaning no matches (no even valued pairs)
+		return MatchBits != 0xFFFF;
+	}
+
+	uint16	MinimumWordLength;
+
+	// Table describing for character 'x', BitApproxTable[x] returns a value (in this case 16 bits)
+	// where a 0 in the nth bit implies you may have a match if 'n' consecutive  possible matches have 
+	// been seen when scanning a string of characters.
+	uint16	BitApproxTable[ALPHABET_SIZE];
+
+	// Table encoding a 'match'and 'accept' value using two bits for each character in a window size of 8
+	// PredictMatch relies on two pieces of information:
+	// Match(x,n) == 1, meaning any of our prediction words has x as their nth character
+	// Accept(x, n) == 1, meaning any of our prediction words ends at n with x
+	// The PredictionMatchTable[x] returns 16-bits for an 8-character window of characters where 
+	// for each two bit pair, we provide the answer to Match(x, n) == 1 (true if the bit pair is even), 
+	// and  Accept(x, n) == 1 (true if the bit pair is < 2).
+	// Since hashing is used to distribute match and accept data for prediction words, the values stored 
+	// at PredictionMatchTable[x] might not represent the window for character 'x' across all prediction words
+	// so that predictions can ideally produce fewer false positives.
+	// For more information it's highly recommended reading the reference link above
+	uint16	PredictMatchTable[TABLE_SIZE];
+};
diff --git a/Engine/Source/Runtime/CoreUObject/Private/UObject/CoreRedirects.cpp b/Engine/Source/Runtime/CoreUObject/Private/UObject/CoreRedirects.cpp
--- a/Engine/Source/Runtime/CoreUObject/Private/UObject/CoreRedirects.cpp
+++ b/Engine/Source/Runtime/CoreUObject/Private/UObject/CoreRedirects.cpp
@@ -24,6 +24,207 @@ DEFINE_LOG_CATEGORY(LogCoreRedirects);
 #	define UE_WITH_CORE_REDIRECTS 1
 #endif
 
+#define COREREDIRECT_STATS 0
+#if COREREDIRECT_STATS
+static int sSubstringLookups = 0;
+static int sSubstringPredictHit = 0;
+static int sSubstringPredictMiss = 0;
+FAutoConsoleCommand StaticLoadTrackerDump(TEXT("CoreRedirects.SubstringPrediction.StatsDump"), TEXT(""), FConsoleCommandDelegate::CreateLambda([]()
+	{
+		UE_LOG(LogCoreRedirects, Log, TEXT(
+			"==========================================\n"
+			"Dumping Substring Stats\n\n"
+			"Lookups:\t\t\t\t%d\n"
+			"Prediction Hit:\t\t%d\n"
+			"Prediction Miss:\t\t%d\n"
+			"=========================================="
+		), sSubstringLookups, sSubstringPredictHit, sSubstringPredictMiss);
+	}));
+
+#define COREREDIRECT_STATS_UPDATE_PREDICTION_LOOKUP() do{ sSubstringLookups++; }while(false)
+#define COREREDIRECT_STATS_UPDATE_PREDICTION_RESULT(bFound) do{ if(bFound){ sSubstringPredictHit++; } else { sSubstringPredictMiss++; }}while(false)
+#else
+#define COREREDIRECT_STATS_UPDATE_PREDICTION_LOOKUP()
+#define COREREDIRECT_STATS_UPDATE_PREDICTION_RESULT(bFound)
+#endif
+
+/* 
+* Small helper class to deal with the string representation of FNames while trying to minimize copies 
+* without relying on internal details of FName. Due to the sheer volume of CoreRedirects, it's 
+* significant to avoid the copies where possible. We maintain UTF8 encoding for FNames when dealing with Wildcard
+* redirects since UTF8 overlays ASCII allowing pure ASCII strings to be compact, greatly improving PredictMatch efficiency.
+* We also need a common encoding for how we store substring redirects and how we query them to ensure we do not generate
+* false negatives in our fuzzy searchs using PredictMatch; UTF8 will allow us to still predict non-ASCII substrings.
+*/
+namespace UE::CoreRedirects::Private
+{
+	struct FNameUtf8String
+	{
+		FNameUtf8String() : UTF8() {}
+		~FNameUtf8String() {}
+		FNameUtf8String(FName Name)
+			: UTF8()
+		{
+			// FName will append "None" when FName.IsNone. Instead treat empty string as None
+			if (!Name.IsNone())
+			{
+				Name.AppendString(UTF8);
+			}
+		}
+
+		operator FUtf8StringView() const
+		{ 
+			return UTF8; 
+		}
+
+		bool IsNone() const
+		{
+			return Len() == 0;
+		}
+
+		const UTF8CHAR* GetData() const
+		{
+			return UTF8.GetData();
+		}
+
+		int32 Len() const
+		{
+			return UTF8.Len();
+		}
+
+		int32 ByteLength() const
+		{
+			return UTF8.Len() * sizeof(UTF8CHAR);
+		}
+
+		bool StartsWith(const FUtf8StringView InPrefix, ESearchCase::Type InSearchCase) const
+		{
+			return UTF8.ToView().StartsWith(InPrefix, InSearchCase);
+		}
+
+		bool EndsWith(const FUtf8StringView InSuffix, ESearchCase::Type InSearchCase) const
+		{
+			return UTF8.ToView().EndsWith(InSuffix, InSearchCase);
+		}
+
+		bool Contains(const FUtf8StringView InSubstring, ESearchCase::Type InSearchCase) const
+		{
+			return (INDEX_NONE != UE::String::FindFirst(UTF8, InSubstring, InSearchCase));
+		}
+
+		void ReplaceAt(int32 Pos, int32 RemoveLen, FNameUtf8String& Str)
+		{
+			UTF8.ReplaceAt(Pos, RemoveLen, Str);
+		}
+
+		void RightChopInline(int Len)
+		{
+			UTF8.RemoveAt(0, Len);
+		}
+
+		TUtf8StringBuilder<FName::StringBufferSize> UTF8;
+	};
+
+	struct FCoreRedirectObjectUtf8Name
+	{
+		FCoreRedirectObjectUtf8Name(const FCoreRedirectObjectName& InName)
+			: ObjectName(InName.ObjectName)
+			, OuterName(InName.OuterName)
+			, PackageName(InName.PackageName)
+		{
+		}
+
+		FNameUtf8String ObjectName;
+		FNameUtf8String OuterName;
+		FNameUtf8String PackageName;
+	};
+
+	struct SubstringMatcher
+	{
+		static bool Matches(const FName LHS, const FNameUtf8String& RHSString, bool bPartialLHS, bool bPartialRHS)
+		{
+			const bool bLHSIsNone = LHS.IsNone();
+			const bool bRHSIsNone = RHSString.IsNone();
+			if (bLHSIsNone || bRHSIsNone)
+			{
+				return bLHSIsNone == bRHSIsNone || ((!bLHSIsNone || bPartialLHS) && (!bRHSIsNone || bPartialRHS));
+			}
+
+			FNameUtf8String LHSString(LHS);
+			return RHSString.Contains(LHSString, ESearchCase::IgnoreCase);
+		}
+	};
+
+	struct PrefixMatcher
+	{
+		static bool Matches(const FName LHS, const FNameUtf8String& RHSString, bool bPartialLHS, bool bPartialRHS)
+		{
+			const bool bLHSIsNone = LHS.IsNone();
+			const bool bRHSIsNone = RHSString.IsNone();
+			if (bLHSIsNone || bRHSIsNone)
+			{
+				return bLHSIsNone == bRHSIsNone || ((!bLHSIsNone || bPartialLHS) && (!bRHSIsNone || bPartialRHS));
+			}
+
+			FNameUtf8String LHSString(LHS);
+			return RHSString.StartsWith(LHSString, ESearchCase::IgnoreCase);
+		}
+	};
+
+	struct SuffixMatcher
+	{
+		static bool Matches(const FName LHS, const FNameUtf8String& RHSString, bool bPartialLHS, bool bPartialRHS)
+		{
+			const bool bLHSIsNone = LHS.IsNone();
+			const bool bRHSIsNone = RHSString.IsNone();
+			if (bLHSIsNone || bRHSIsNone)
+			{
+				return bLHSIsNone == bRHSIsNone || ((!bLHSIsNone || bPartialLHS) && (!bRHSIsNone || bPartialRHS));
+			}
+
+			FNameUtf8String LHSString(LHS);
+			return RHSString.EndsWith(LHSString, ESearchCase::IgnoreCase);
+		}
+	};
+
+	template <typename Matcher>
+	static bool MatchWildcardRedirect(const FCoreRedirectObjectName& InRedirect, 
+		const FCoreRedirectObjectUtf8Name& InUtf8Name, bool bPartialLHS, bool bPartialRHS)
+	{
+		if (!Matcher::Matches(InRedirect.ObjectName, InUtf8Name.ObjectName, bPartialLHS, bPartialRHS))
+		{
+			return false;
+		}
+		if (!Matcher::Matches(InRedirect.OuterName, InUtf8Name.OuterName, bPartialLHS, bPartialRHS))
+		{
+			return false;
+		}
+		if (!Matcher::Matches(InRedirect.PackageName, InUtf8Name.PackageName, bPartialLHS, bPartialRHS))
+		{
+			return false;
+		}
+		return true;
+	}
+
+	static bool MatchSubstring(const FCoreRedirectObjectName& InSubstringRedirect, const FCoreRedirectObjectUtf8Name& InUtf8Name,
+		bool bPartialLHS, bool bPartialRHS)
+	{
+		return MatchWildcardRedirect<SubstringMatcher>(InSubstringRedirect, InUtf8Name, bPartialLHS, bPartialRHS);
+	}
+
+	static bool MatchPrefix(const FCoreRedirectObjectName& InPrefixRedirect, const FCoreRedirectObjectUtf8Name& InUtf8Name,
+		bool bPartialLHS, bool bPartialRHS)
+	{
+		return MatchWildcardRedirect<PrefixMatcher>(InPrefixRedirect, InUtf8Name, bPartialLHS, bPartialRHS);
+	}
+
+	static bool MatchSuffix(const FCoreRedirectObjectName& InSuffixRedirect, const FCoreRedirectObjectUtf8Name& InUtf8Name,
+		bool bPartialLHS, bool bPartialRHS)
+	{
+		return MatchWildcardRedirect<SuffixMatcher>(InSuffixRedirect, InUtf8Name, bPartialLHS, bPartialRHS);
+	}
+}
+
 FCoreRedirectObjectName::FCoreRedirectObjectName(const FTopLevelAssetPath& TopLevelAssetPath)
 	: FCoreRedirectObjectName(TopLevelAssetPath.GetAssetName(), NAME_None, TopLevelAssetPath.GetPackageName())
 {
@@ -86,33 +287,202 @@ void FCoreRedirectObjectName::Reset()
 	ObjectName = OuterName = PackageName = NAME_None;
 }
 
-bool FCoreRedirectObjectName::Matches(const FCoreRedirectObjectName& Other, bool bCheckSubstring) const
+void FCoreRedirects::FWildcardData::Add(const FCoreRedirect& Redirect)
 {
-	return Matches(Other, EMatchFlags::CheckSubString);
+	if (Redirect.IsSubstringMatch())
+	{
+		Substrings.Add(Redirect);
+		AddPredictionWords(Redirect);
+	}
+	else if (Redirect.IsPrefixMatch())
+	{
+		Prefixes.Add(Redirect);
+	}
+	else
+	{
+		check(Redirect.IsSuffixMatch());
+		Suffixes.Add(Redirect);
+	}
 }
 
-bool FCoreRedirectObjectName::Matches(const FCoreRedirectObjectName& Other, EMatchFlags MatchFlags) const
+void FCoreRedirects::FWildcardData::AddPredictionWords(const FCoreRedirect& Redirect)
 {
-	bool bPartialLHS = !EnumHasAnyFlags(MatchFlags, EMatchFlags::DisallowPartialLHSMatch);
-	bool bPartialRHS = EnumHasAnyFlags(MatchFlags, EMatchFlags::AllowPartialRHSMatch);
-	bool bSubString = EnumHasAllFlags(MatchFlags, EMatchFlags::CheckSubString);
+	using namespace UE::CoreRedirects::Private;
 
-	auto FieldMatches = [bPartialLHS, bPartialRHS, bSubString](FName LHS, FName RHS)
+	const FName Names[3] = { Redirect.OldName.ObjectName, Redirect.OldName.OuterName, Redirect.OldName.PackageName };
+	for (const FName& Name : Names)
 	{
-		if (LHS.IsNone() || RHS.IsNone())
+		if (Name.IsNone())
 		{
-			return LHS == RHS || ((!LHS.IsNone() || bPartialLHS) && (!RHS.IsNone() || bPartialRHS));
+			continue;
 		}
-		else if (bSubString)
+
+		FNameUtf8String NameView(Name);
+
+		// Since we only predict based on a small window of characters (8), 
+		// try to use characters that will be more unique by removing common prefixes.
+		// Note, this does not affect the actual substring matching, only
+		// our prediction and more specifically our false positive rate.
+		FUtf8StringView CommonPrefixes[] = 
+		{ "/", "Script/", "Temp/", "Extra/", "Memory/", "Config/", "Game/", "Engine/", "Transient/", "Niagara/"};
+
+		for (FUtf8StringView Prefix : CommonPrefixes)
 		{
-			// Much slower, see if RHS contains the substring LHS.
-			return INDEX_NONE != UE::String::FindFirst(WriteToString<FName::StringBufferSize>(RHS), WriteToString<FName::StringBufferSize>(LHS), ESearchCase::IgnoreCase);
+			// Only remove the prefix if doing so won't leave the string empty.
+			// To consider: To keep entropy high, we might want to consider only performing the 
+			// removal if we leave 'n' number of characers for our prediction.
+			if (NameView.Len() > Prefix.Len() && NameView.StartsWith(Prefix, ESearchCase::IgnoreCase))
+			{
+				NameView.RightChopInline(Prefix.Len());
+			}
 		}
-		else
+
+		// Predict match doesn't interpret character data, however in order to avoid false negatives
+		// we need to ensure all PredictionWords are added in the same encoding we will be using
+		// for comparisons in MatchSubstringApproximate
+		PredictMatch.AddPredictionWord((uint8*) NameView.GetData(), (uint16) NameView.ByteLength());
+	}
+}
+
+bool FCoreRedirects::FWildcardData::MatchSubstringApproximate(const UE::CoreRedirects::Private::FCoreRedirectObjectUtf8Name& InUtf8Name) const
+{
+	using namespace UE::CoreRedirects::Private;
+
+	const FNameUtf8String* Names[3] = { &InUtf8Name.ObjectName, &InUtf8Name.OuterName, &InUtf8Name.PackageName };
+	for (const FNameUtf8String* Name : Names)
+	{
+		if (Name->IsNone())
+		{
+			continue;
+		}
+
+		if (PredictMatch.MatchApproximate((const uint8*) Name->GetData(), (uint16) Name->ByteLength()))
+		{
+			return true;
+		}
+	}
+	return false;
+}
+
+void FCoreRedirects::FWildcardData::Rebuild()
+{
+	PredictMatch.Reset();
+
+	for (const FCoreRedirect& Redirect : Substrings)
+	{
+		AddPredictionWords(Redirect);
+	}
+}
+
+bool FCoreRedirects::FWildcardData::Matches(ECoreRedirectFlags InFlags, const FCoreRedirectObjectName& InName, 
+	ECoreRedirectMatchFlags InMatchFlags, TArray<const FCoreRedirect*>& OutFoundRedirects) const
+{
+	using namespace UE::CoreRedirects::Private;
+
+	bool bFound = false;
+	const bool bPartialRHS = EnumHasAnyFlags(InMatchFlags, ECoreRedirectMatchFlags::AllowPartialMatch);
+
+	// Substring implies prefix and suffix matches. For wildcard matches we must search all three
+	// types as each can be defined distinctly
+	const bool bSubstringMatch = EnumHasAllFlags(InFlags, ECoreRedirectFlags::Option_MatchSubstring);
+	const bool bPrefixMatch = EnumHasAllFlags(InFlags, ECoreRedirectFlags::Option_MatchPrefix);
+	const bool bSuffixMatch = EnumHasAllFlags(InFlags, ECoreRedirectFlags::Option_MatchSuffix);
+
+	// Creating the string form of InName is expensive so early out if we know we can
+	if ((!bSubstringMatch && !bPrefixMatch && !bSuffixMatch) ||
+		(Prefixes.IsEmpty() && Suffixes.IsEmpty() && Substrings.IsEmpty()))
+	{
+		return false;
+	}
+
+	const FCoreRedirectObjectUtf8Name ObjectUtf8Name(InName);
+	if (bPrefixMatch)
+	{
+		for (const FCoreRedirect& CheckRedirect : Prefixes)
+		{
+			if (MatchPrefix(CheckRedirect.OldName, ObjectUtf8Name, true /*bPartialLHS*/, bPartialRHS))
+			{
+				bFound = true;
+				OutFoundRedirects.Add(&CheckRedirect);
+			}
+		}
+	}
+
+	if (bSuffixMatch)
+	{
+		for (const FCoreRedirect& CheckRedirect : Suffixes)
 		{
-			return LHS == RHS;
+			if (MatchSuffix(CheckRedirect.OldName, ObjectUtf8Name, true /*bPartialLHS*/, bPartialRHS))
+			{
+				bFound = true;
+				OutFoundRedirects.Add(&CheckRedirect);
+			}
 		}
+	}
+	
+	if (bSubstringMatch && !Substrings.IsEmpty())
+	{
+		COREREDIRECT_STATS_UPDATE_PREDICTION_LOOKUP();
+		// Perform a fuzzy match against all known substrings. We may have false positives
+		// so we must still perform a stronger check on a fuzzy match. However we can eliminate
+		// the common failure case quickly with a fast fuzzy match.
+		if (MatchSubstringApproximate(ObjectUtf8Name))
+		{
+			// We found a possible substring match, so now do a slow match to find if we really match
+			for (const FCoreRedirect& CheckRedirect : Substrings)
+			{
+				if (MatchSubstring(CheckRedirect.OldName, ObjectUtf8Name, true /*bPartialLHS*/, bPartialRHS))
+				{
+					bFound = true;
+					OutFoundRedirects.Add(&CheckRedirect);
+				}
+			}
+			COREREDIRECT_STATS_UPDATE_PREDICTION_RESULT(bFound);
+		}
+	}
+
+	return bFound;
+}
+
+bool FCoreRedirectObjectName::Matches(const FCoreRedirectObjectName& Other, EMatchFlags MatchFlags) const
+{
+	using namespace UE::CoreRedirects::Private;
+
+	bool bPartialLHS = !EnumHasAnyFlags(MatchFlags, EMatchFlags::DisallowPartialLHSMatch);
+	bool bPartialRHS = EnumHasAnyFlags(MatchFlags, EMatchFlags::AllowPartialRHSMatch);
+	bool bSubstring = EnumHasAllFlags(MatchFlags, EMatchFlags::CheckSubString);
+	bool bPrefix = EnumHasAllFlags(MatchFlags, EMatchFlags::CheckPrefix);
+	bool bSuffix = EnumHasAllFlags(MatchFlags, EMatchFlags::CheckSuffix);
+
+	// Substring implies prefix and suffix so we must check it first
+	if (bSubstring)
+	{
+		return MatchSubstring(*this, Other, bPartialLHS, bPartialRHS);
+	}
+	else if(bPrefix)
+	{
+		return MatchPrefix(*this, Other, bPartialLHS, bPartialRHS);
+	}
+	else if (bSuffix)
+	{
+		return MatchSuffix(*this, Other, bPartialLHS, bPartialRHS);
+	}
+
+	auto FieldMatches = [bPartialLHS, bPartialRHS, bSubstring, bPrefix, bSuffix](FName LHS, FName RHS)
+	{
+		if (LHS == RHS)
+			return true;
+
+		const bool bLHSIsNone = LHS.IsNone();
+		const bool bRHSIsNone = RHS.IsNone();
+		if (bLHSIsNone || bRHSIsNone)
+		{
+			return (!bLHSIsNone || bPartialLHS) && (!bRHSIsNone || bPartialRHS);
+		}
+
+		return false;
 	};
+
 	if (!FieldMatches(ObjectName, Other.ObjectName))
 	{
 		return false;
@@ -193,13 +563,16 @@ bool FCoreRedirectObjectName::HasValidCharacters(ECoreRedirectFlags Type) const
 	static FString InvalidRedirectCharacters = TEXT("\"',|&!~\n\r\t@#(){}[]=;^%$`");
 	static FString InvalidObjectRedirectCharacters = TEXT(".\n\r\t"); // Object and field names in Blueprints are very permissive...
 
-	return ObjectName.IsValidXName(EnumHasAnyFlags(Type, ECoreRedirectFlags::Type_Object | ECoreRedirectFlags::Type_Property | ECoreRedirectFlags::Type_Function) ? InvalidObjectRedirectCharacters : InvalidRedirectCharacters)
+	return ObjectName.IsValidXName(
+		EnumHasAnyFlags(Type, ECoreRedirectFlags::Type_Object | ECoreRedirectFlags::Type_Property| ECoreRedirectFlags::Type_Function)
+		? InvalidObjectRedirectCharacters : InvalidRedirectCharacters)
 		&& OuterName.IsValidXName(InvalidRedirectCharacters) && PackageName.IsValidXName(InvalidRedirectCharacters);
 }
 
-bool FCoreRedirectObjectName::ExpandNames(const FString& InString, FName& OutName, FName& OutOuter, FName &OutPackage)
+bool FCoreRedirectObjectName::ExpandNames(const FStringView InStringView, FName& OutName, FName& OutOuter, FName &OutPackage)
 {
-	FString FullString = InString.TrimStartAndEnd();
+	FStringView FullStringView = InStringView.TrimStartAndEnd();
+	FullStringView.TrimStartAndEndInline();
 
 	// Parse (/path.)?(outerchain.)?(name) where path and outerchain are optional
 	// We also need to support (/path.)?(singleouter:)?(name) because the second delimiter in a chain is : for historical reasons
@@ -208,11 +581,9 @@ bool FCoreRedirectObjectName::ExpandNames(const FString& InString, FName& OutNam
 	int32 LastPeriodIndex = INDEX_NONE;
 	int32 FirstColonIndex = INDEX_NONE;
 	int32 LastColonIndex = INDEX_NONE;
-
-	FullString.FindChar(TEXT('/'), SlashIndex);
-
-	FullString.FindChar(TEXT('.'), FirstPeriodIndex);
-	FullString.FindChar(TEXT(':'), FirstColonIndex);
+	FullStringView.FindChar(TEXT('/'), SlashIndex);
+	FullStringView.FindChar(TEXT('.'), FirstPeriodIndex);
+	FullStringView.FindChar(TEXT(':'), FirstColonIndex);
 
 	if (FirstColonIndex != INDEX_NONE && (FirstPeriodIndex == INDEX_NONE || FirstColonIndex < FirstPeriodIndex))
 	{
@@ -225,17 +596,17 @@ bool FCoreRedirectObjectName::ExpandNames(const FString& InString, FName& OutNam
 		// If start with /, fill in package name, otherwise name
 		if (SlashIndex != INDEX_NONE)
 		{
-			OutPackage = FName(*FullString);
+			OutPackage = FName(FullStringView);
 		}
 		else
 		{
-			OutName = FName(*FullString);
+			OutName = FName(FullStringView);
 		}
 		return true;
 	}
 
-	FullString.FindLastChar(TEXT('.'), LastPeriodIndex);
-	FullString.FindLastChar(TEXT(':'), LastColonIndex);
+	FullStringView.FindLastChar(TEXT('.'), LastPeriodIndex);
+	FullStringView.FindLastChar(TEXT(':'), LastColonIndex);
 
 	if (LastColonIndex != INDEX_NONE && (LastPeriodIndex == INDEX_NONE || LastColonIndex > LastPeriodIndex))
 	{
@@ -246,19 +617,19 @@ bool FCoreRedirectObjectName::ExpandNames(const FString& InString, FName& OutNam
 	if (SlashIndex == INDEX_NONE)
 	{
 		// No /, so start from beginning. There must be an outer if we got this far
-		OutOuter = FName(*FullString.Mid(0, LastPeriodIndex));
+		OutOuter = FName(FullStringView.Mid(0, LastPeriodIndex));
 	}
 	else
 	{
-		OutPackage = FName(*FullString.Left(FirstPeriodIndex));
+		OutPackage = FName(FullStringView.Left(FirstPeriodIndex));
 		if (FirstPeriodIndex != LastPeriodIndex)
 		{
 			// Extract Outer between periods
-			OutOuter = FName(*FullString.Mid(FirstPeriodIndex + 1, LastPeriodIndex - FirstPeriodIndex - 1));
+			OutOuter = FName(FullStringView.Mid(FirstPeriodIndex + 1, LastPeriodIndex - FirstPeriodIndex - 1));
 		}
 	}
 
-	OutName = FName(*FullString.Mid(LastPeriodIndex + 1));
+	OutName = FName(FullStringView.Mid(LastPeriodIndex + 1));
 
 	return true;
 }
@@ -271,7 +642,6 @@ FString FCoreRedirectObjectName::CombineNames(FName NewName, FName NewOuter, FNa
 	{
 		// If Outer is simple, need to use : instead of . because : is used for second delimiter only
 		FString OuterString = NewOuter.ToString();
-
 		int32 DelimIndex = INDEX_NONE;
 
 		if (OuterString.FindChar(TEXT('.'), DelimIndex) || OuterString.FindChar(TEXT(':'), DelimIndex))
@@ -392,7 +762,7 @@ const TCHAR* FCoreRedirect::ParseValueChanges(const TCHAR* Buffer)
 	}
 }
 
-static FORCEINLINE bool CheckRedirectFlagsMatch(ECoreRedirectFlags FlagsA, ECoreRedirectFlags FlagsB)
+static bool CheckRedirectFlagsMatch(ECoreRedirectFlags FlagsA, ECoreRedirectFlags FlagsB)
 {
 	// For type, check it includes the matching type
 	const bool bTypesOverlap = !!((FlagsA & FlagsB) & ECoreRedirectFlags::Type_AllMask);
@@ -401,7 +771,6 @@ static FORCEINLINE bool CheckRedirectFlagsMatch(ECoreRedirectFlags FlagsA, ECore
 	const bool bCategoriesMatch = (FlagsA & ECoreRedirectFlags::Category_AllMask) == (FlagsB & ECoreRedirectFlags::Category_AllMask);
 
 	// Options are not considered in this function; Flags will match at this point regardless of their bits in Options_AllMask
-
 	return bTypesOverlap && bCategoriesMatch;
 }
 
@@ -413,21 +782,32 @@ bool FCoreRedirect::Matches(ECoreRedirectFlags InFlags, const FCoreRedirectObjec
 	{
 		return false;
 	}
+
 	return Matches(InName, MatchFlags);
 }
 
-
 bool FCoreRedirect::Matches(const FCoreRedirectObjectName& InName, ECoreRedirectMatchFlags MatchFlags) const
 {
 	FCoreRedirectObjectName::EMatchFlags NameMatchFlags = FCoreRedirectObjectName::EMatchFlags::None;
+
+	if (EnumHasAllFlags(MatchFlags, ECoreRedirectMatchFlags::AllowPartialMatch))
+	{
+		NameMatchFlags |= FCoreRedirectObjectName::EMatchFlags::AllowPartialRHSMatch;
+	}
+
 	if (IsSubstringMatch())
 	{
 		NameMatchFlags |= FCoreRedirectObjectName::EMatchFlags::CheckSubString;
 	}
-	if (EnumHasAllFlags(MatchFlags, ECoreRedirectMatchFlags::AllowPartialMatch))
+	else if (IsPrefixMatch())
 	{
-		NameMatchFlags |= FCoreRedirectObjectName::EMatchFlags::AllowPartialRHSMatch;
+		NameMatchFlags |= FCoreRedirectObjectName::EMatchFlags::CheckPrefix;
+	}
+	else if (IsSuffixMatch())
+	{
+		NameMatchFlags |= FCoreRedirectObjectName::EMatchFlags::CheckSuffix;
 	}
+
 	return OldName.Matches(InName, NameMatchFlags);
 }
 
@@ -438,60 +818,118 @@ bool FCoreRedirect::HasValueChanges() const
 
 bool FCoreRedirect::IsSubstringMatch() const
 {
-	return !!(RedirectFlags & ECoreRedirectFlags::Option_MatchSubstring);
+	return EnumHasAllFlags(RedirectFlags, ECoreRedirectFlags::Option_MatchSubstring);
 }
 
-FCoreRedirectObjectName FCoreRedirect::RedirectName(const FCoreRedirectObjectName& OldObjectName) const
+FCoreRedirectObjectName FCoreRedirect::RedirectName(const FCoreRedirectObjectName& OldObjectName, bool bIsKnownToMatch) const
 {
-	FCoreRedirectObjectName ModifyName(OldObjectName);
+	using namespace UE::CoreRedirects::Private;
 
-	// Convert names that are different and non empty
-	if (OldName.ObjectName != NewName.ObjectName && !ModifyName.ObjectName.IsNone())
-	{
-		if (IsSubstringMatch())
-		{
-			ModifyName.ObjectName = FName(*OldObjectName.ObjectName.ToString().Replace(*OldName.ObjectName.ToString(), *NewName.ObjectName.ToString()));
-		}
-		else
+	const bool bIsSubstringMatch = IsSubstringMatch();
+	const bool bIsPrefixMatch = IsPrefixMatch();
+	const bool bIsSuffixMatch = IsSuffixMatch();
+
+	auto ConvertNames = [bIsSubstringMatch, bIsPrefixMatch, bIsSuffixMatch, bIsKnownToMatch](FName CurrentName,
+		FName RedirectOldName, FName RedirectNewName)
 		{
-			ModifyName.ObjectName = NewName.ObjectName;
-		}
-	}
-	// If package name and object name are specified, copy outer also it was set to null explicitly
-	if ((OldName.OuterName != NewName.OuterName || (NewName.PackageName != NAME_None && NewName.ObjectName != NAME_None)) && !ModifyName.OuterName.IsNone())
+			if ((RedirectOldName != RedirectNewName) && !CurrentName.IsNone())
+			{
+				if (RedirectOldName.IsNone())
+				{
+					return RedirectNewName;
+				}
+
+				if (bIsSubstringMatch)
+				{
+					FNameUtf8String OutName(CurrentName);
+					FNameUtf8String Substring(RedirectOldName);
+					FNameUtf8String ReplacementName(RedirectNewName);
+					int ReplacePos = UE::String::FindFirst(OutName, Substring, ESearchCase::IgnoreCase);
+
+					if (!bIsKnownToMatch && ReplacePos == INDEX_NONE)
+					{
+						return CurrentName;
+					}
+					check(ReplacePos != INDEX_NONE);
+
+					OutName.ReplaceAt(ReplacePos, Substring.Len(), ReplacementName);
+					return FName(OutName);
+				}
+				else if (bIsPrefixMatch)
+				{
+					FNameUtf8String OutName(CurrentName);
+					FNameUtf8String Prefix(RedirectOldName);
+					FNameUtf8String ReplacementName(RedirectNewName);
+
+					if (!bIsKnownToMatch && !OutName.StartsWith(Prefix, ESearchCase::IgnoreCase))
+					{
+						return CurrentName;
+					}
+					// bKnownToMatch implies FPrefixMatcher::Matches(CurrentName, RedirectOldName), 
+					// so OutName.StartsWith(Prefix) is true, so we skip the stringcompare
+					check(OutName.Len() >= Prefix.Len());
+
+					OutName.ReplaceAt(0, Prefix.Len(), ReplacementName);
+					return FName(OutName);
+				}
+				else if (bIsSuffixMatch)
+				{
+					FNameUtf8String OutName(CurrentName);
+					FNameUtf8String Suffix(RedirectOldName);
+					FNameUtf8String ReplacementName(RedirectNewName);
+
+					if (!bIsKnownToMatch && !OutName.EndsWith(Suffix, ESearchCase::IgnoreCase))
+					{
+						return CurrentName;
+					}
+					// bKnownToMatch implies FSuffixMatcher::Matches(CurrentName, RedirectOldName), 
+					// so OutName.EndsWith(Suffix) is true, so we skip the stringcompare
+					check(OutName.Len() >= Suffix.Len());
+
+					OutName.ReplaceAt(OutName.Len() - Suffix.Len(), Suffix.Len(), ReplacementName);
+					return FName(OutName);
+				}
+				else
+				{
+					return RedirectNewName;
+				}
+			}
+			return CurrentName;
+		};
+
+	// Convert names that are different and non empty
+	FCoreRedirectObjectName ModifyName(OldObjectName);
+	ModifyName.ObjectName = ConvertNames(OldObjectName.ObjectName, OldName.ObjectName, NewName.ObjectName);
+
+	if (OldName.OuterName == NewName.OuterName)
 	{
-		if (IsSubstringMatch())
-		{
-			ModifyName.OuterName = FName(*OldObjectName.OuterName.ToString().Replace(*OldName.OuterName.ToString(), *NewName.OuterName.ToString()));
-		}
-		else
+		// If package name and object name are specified, overwrite outer also since it was set to null explicitly
+		if (OldName.OuterName.IsNone() && !NewName.PackageName.IsNone() && !NewName.ObjectName.IsNone() && !ModifyName.OuterName.IsNone())
 		{
 			ModifyName.OuterName = NewName.OuterName;
 		}
 	}
-	if (OldName.PackageName != NewName.PackageName && !ModifyName.PackageName.IsNone())
+	else
 	{
-		if (IsSubstringMatch())
-		{
-			ModifyName.PackageName = FName(*OldObjectName.PackageName.ToString().Replace(*OldName.PackageName.ToString(), *NewName.PackageName.ToString()));
-		}
-		else
-		{
-			ModifyName.PackageName = NewName.PackageName;
-		}
+		ModifyName.OuterName = ConvertNames(OldObjectName.OuterName, OldName.OuterName, NewName.OuterName);
 	}
+	
+	ModifyName.PackageName = ConvertNames(OldObjectName.PackageName, OldName.PackageName, NewName.PackageName);
 
 	return ModifyName;
 }
 
+FCoreRedirectObjectName FCoreRedirect::RedirectName(const FCoreRedirectObjectName& OldObjectName) const
+{
+	return RedirectName(OldObjectName, false /* bIsKnownToMatch */);
+}
+
 bool FCoreRedirect::IdenticalMatchRules(const FCoreRedirect& Other) const
 {
 	// All types now use the full path
-	ECoreRedirectFlags TypeFlags = RedirectFlags & ECoreRedirectFlags::Type_AllMask;
 	return RedirectFlags == Other.RedirectFlags && OldName == Other.OldName;
 }
 
-
 bool FCoreRedirects::bInitialized = false;
 bool FCoreRedirects::bInDebugMode = false;
 bool FCoreRedirects::bValidatedOnce = false;
@@ -512,8 +950,9 @@ FCoreRedirects::FRedirectNameMap& FCoreRedirects::FRedirectTypeMap::FindOrAdd(EC
 	}
 
 	TPair<ECoreRedirectFlags, FRedirectNameMap>* OldData = FastIterable.GetData();
-	TPair<ECoreRedirectFlags, FRedirectNameMap>& NewPair =
-		FastIterable.Emplace_GetRef(Key, FRedirectNameMap());
+	TPair<ECoreRedirectFlags, FRedirectNameMap>& NewPair = FastIterable.Emplace_GetRef(Key, FRedirectNameMap());
+
+	// Check to ensure FastIterable did not reallocate after emplacement
 	if (FastIterable.GetData() == OldData)
 	{
 		NameMap = &NewPair.Value;
@@ -526,6 +965,12 @@ FCoreRedirects::FRedirectNameMap& FCoreRedirects::FRedirectTypeMap::FindOrAdd(EC
 			Map.Add(Pair.Key, &Pair.Value);
 		}
 	}
+
+	if (EnumHasAnyFlags(Key, ECoreRedirectFlags::Option_MatchWildcardMask))
+	{
+		NewPair.Value.Wildcards.Reset(new FWildcardData());
+	}
+	
 	return NewPair.Value;
 }
 
@@ -606,52 +1051,63 @@ void FCoreRedirects::EnterMultithreadedPhase()
 bool FCoreRedirects::RedirectNameAndValues(ECoreRedirectFlags Type, const FCoreRedirectObjectName& OldObjectName,
 	FCoreRedirectObjectName& NewObjectName, const FCoreRedirect** FoundValueRedirect, ECoreRedirectMatchFlags MatchFlags)
 {
+	auto ProcessRedirect = [&FoundValueRedirect, OldObjectName](const FCoreRedirect* Redirect, const FCoreRedirectObjectName& NewObjectName)
+	{
+		if (FoundValueRedirect && (Redirect->HasValueChanges() || Redirect->OverrideClassName.IsValid()))
+		{
+			if (*FoundValueRedirect)
+			{
+				if ((*FoundValueRedirect)->ValueChanges.OrderIndependentCompareEqual(Redirect->ValueChanges) == false)
+				{
+					UE_LOG(LogCoreRedirects, Error, TEXT("RedirectNameAndValues(%s) found multiple conflicting value redirects, %s and %s!"),
+						*OldObjectName.ToString(), *(*FoundValueRedirect)->OldName.ToString(), *Redirect->OldName.ToString());
+				}
+			}
+			else
+			{
+				// Set value redirects for processing outside
+				*FoundValueRedirect = Redirect;
+			}
+		}
+
+		return Redirect->RedirectName(NewObjectName, true /* bIsKnownToMatch */);
+	};
+
 	NewObjectName = OldObjectName;
 	TArray<const FCoreRedirect*> FoundRedirects;
-
 	if (GetMatchingRedirects(Type, OldObjectName, FoundRedirects, MatchFlags))
 	{
-		// Sort them based on match
-		FoundRedirects.Sort([&OldObjectName](const FCoreRedirect& A, const FCoreRedirect& B) { return A.OldName.MatchScore(OldObjectName) > B.OldName.MatchScore(OldObjectName); });
-
-		// Apply in order
-		for (int32 i = 0; i < FoundRedirects.Num(); i++)
+		if (FoundRedirects.Num() > 1)
 		{
-			const FCoreRedirect* Redirect = FoundRedirects[i];
+			// Sort them based on match
+			FoundRedirects.Sort([&OldObjectName](const FCoreRedirect& A, const FCoreRedirect& B) 
+				{ return A.OldName.MatchScore(OldObjectName) > B.OldName.MatchScore(OldObjectName); });
+			NewObjectName = ProcessRedirect(FoundRedirects[0], NewObjectName);
 
-			if (!Redirect)
+			// Apply in order
+			for (int32 i = 1; i < FoundRedirects.Num(); i++)
 			{
-				continue;
-			}
+				const FCoreRedirect* Redirect = FoundRedirects[i];
 
-			// Only apply if name match is still valid, if it already renamed part of it it may not apply any more. Don't want to check flags as those were checked in the gather step
-			if (Redirect->Matches(NewObjectName, MatchFlags))
-			{
-				if (FoundValueRedirect && (Redirect->HasValueChanges() || Redirect->OverrideClassName.IsValid()))
+				// Only apply if name match is still valid, if it already renamed part of it it may not apply any more. 
+				// Don't want to check flags as those were checked in the gather step
+				if (Redirect->Matches(NewObjectName, MatchFlags))
 				{
-					if (*FoundValueRedirect)
-					{
-						if ((*FoundValueRedirect)->ValueChanges.OrderIndependentCompareEqual(Redirect->ValueChanges) == false)
-						{
-							UE_LOG(LogCoreRedirects, Error, TEXT("RedirectNameAndValues(%s) found multiple conflicting value redirects, %s and %s!"), *OldObjectName.ToString(), *(*FoundValueRedirect)->OldName.ToString(), *Redirect->OldName.ToString());
-						}
-					}
-					else
-					{
-						// Set value redirects for processing outside
-						*FoundValueRedirect = Redirect;
-					}
+					NewObjectName = ProcessRedirect(Redirect, NewObjectName);
 				}
-
-				NewObjectName = Redirect->RedirectName(NewObjectName);
 			}
 		}
+		else
+		{
+			const FCoreRedirect* Redirect = FoundRedirects[0];
+			NewObjectName = ProcessRedirect(Redirect, NewObjectName);
+		}
 	}
 
-	UE_CLOG(bInDebugMode && NewObjectName != OldObjectName, LogCoreRedirects, Verbose, 
+	const bool bDidRedirect = NewObjectName != OldObjectName;
+	UE_CLOG(bInDebugMode && bDidRedirect, LogCoreRedirects, Verbose,
 		TEXT("RedirectNameAndValues(%s) replaced by %s"), *OldObjectName.ToString(), *NewObjectName.ToString());
-
-	return NewObjectName != OldObjectName;
+	return bDidRedirect;
 }
 
 FCoreRedirectObjectName FCoreRedirects::GetRedirectedName(ECoreRedirectFlags Type,
@@ -674,7 +1130,8 @@ const TMap<FString, FString>* FCoreRedirects::GetValueRedirects(ECoreRedirectFla
 
 	if (FoundRedirect && FoundRedirect->ValueChanges.Num() > 0)
 	{
-		UE_CLOG(bInDebugMode, LogCoreRedirects, VeryVerbose, TEXT("GetValueRedirects found %d matches for %s"), FoundRedirect->ValueChanges.Num(), *OldObjectName.ToString());
+		UE_CLOG(bInDebugMode, LogCoreRedirects, VeryVerbose, TEXT("GetValueRedirects found %d matches for %s"),
+			FoundRedirect->ValueChanges.Num(), *OldObjectName.ToString());
 
 		return &FoundRedirect->ValueChanges;
 	}
@@ -687,12 +1144,12 @@ bool FCoreRedirects::GetMatchingRedirects(ECoreRedirectFlags SearchFlags, const
 {
 	// Look for all redirects that match the given names and flags
 	bool bFound = false;
-	
+
 	// If we're not explicitly searching for packages, and not looking for removed things, and not searching for partial matches
 	// based on ObjectName only, add the implicit (Type=Package,Category=None) redirects
 	const bool bSearchPackageRedirects = !(SearchFlags & ECoreRedirectFlags::Type_Package) &&
 		!(SearchFlags & ECoreRedirectFlags::Category_Removed) &&
-			(!(MatchFlags & ECoreRedirectMatchFlags::AllowPartialMatch) || !OldObjectName.PackageName.IsNone());
+		(!(MatchFlags & ECoreRedirectMatchFlags::AllowPartialMatch) || !OldObjectName.PackageName.IsNone());
 
 	// Determine list of maps to look over, need to handle being passed multiple types in a bit mask
 	for (const TPair<ECoreRedirectFlags, FRedirectNameMap>& Pair : RedirectTypeMap)
@@ -700,18 +1157,27 @@ bool FCoreRedirects::GetMatchingRedirects(ECoreRedirectFlags SearchFlags, const
 		ECoreRedirectFlags PairFlags = Pair.Key;
 
 		// We need to check all maps that match the search or package flags
-		if (CheckRedirectFlagsMatch(PairFlags, SearchFlags) || (bSearchPackageRedirects && CheckRedirectFlagsMatch(PairFlags, ECoreRedirectFlags::Type_Package)))
+		if (CheckRedirectFlagsMatch(PairFlags, SearchFlags) || 
+			(bSearchPackageRedirects && CheckRedirectFlagsMatch(PairFlags, ECoreRedirectFlags::Type_Package)))
 		{
-			const TArray<FCoreRedirect>* RedirectsForName = Pair.Value.RedirectMap.Find(OldObjectName.GetSearchKey(PairFlags));
-
-			if (RedirectsForName)
+			if (EnumHasAnyFlags(PairFlags, ECoreRedirectFlags::Option_MatchWildcardMask))
+			{
+				const FWildcardData* Wildcards = Pair.Value.Wildcards.Get();
+				check(Wildcards);
+				bFound |= Wildcards->Matches(PairFlags, OldObjectName, MatchFlags, FoundRedirects);
+			}
+			else
 			{
-				for (const FCoreRedirect& CheckRedirect : *RedirectsForName)
+				const TArray<FCoreRedirect>* RedirectsForName = Pair.Value.RedirectMap.Find(OldObjectName.GetSearchKey(PairFlags));
+				if (RedirectsForName)
 				{
-					if (CheckRedirect.Matches(PairFlags, OldObjectName, MatchFlags))
+					for (const FCoreRedirect& CheckRedirect : *RedirectsForName)
 					{
-						bFound = true;
-						FoundRedirects.Add(&CheckRedirect);
+						if (CheckRedirect.Matches(PairFlags, OldObjectName, MatchFlags))
+						{
+							bFound = true;
+							FoundRedirects.Add(&CheckRedirect);
+						}
 					}
 				}
 			}
@@ -729,7 +1195,23 @@ bool FCoreRedirects::FindPreviousNames(ECoreRedirectFlags SearchFlags, const FCo
 	// If we're not explicitly searching for packages or looking for removed things, add the implicit (Type=Package,Category=None) redirects
 	const bool bSearchPackageRedirects = !(SearchFlags & ECoreRedirectFlags::Type_Package) && !(SearchFlags & ECoreRedirectFlags::Category_Removed);
 
-	// Determine list of maps to look over, need to handle being passed multiple Flagss in a bit mask
+	auto TryReverseRedirect = [](const FCoreRedirect& Redirect, const FCoreRedirectObjectName& NewObjectName, TArray<FCoreRedirectObjectName>& PreviousNames)
+		{
+			FCoreRedirect ReverseRedirect = FCoreRedirect(Redirect);
+			ReverseRedirect.OldName = Redirect.NewName;
+			ReverseRedirect.NewName = Redirect.OldName;
+
+			FCoreRedirectObjectName OldName = ReverseRedirect.RedirectName(NewObjectName, true /* bIsKnownToMatch */);
+
+			if (OldName != NewObjectName)
+			{
+				PreviousNames.AddUnique(OldName);
+				return true;
+			}
+			return false;
+		};
+
+	// Determine list of maps to look over, need to handle being passed multiple Flags in a bit mask
 	for (const TPair<ECoreRedirectFlags, FRedirectNameMap>& Pair : RedirectTypeMap)
 	{
 		ECoreRedirectFlags PairFlags = Pair.Key;
@@ -737,28 +1219,48 @@ bool FCoreRedirects::FindPreviousNames(ECoreRedirectFlags SearchFlags, const FCo
 		// We need to check all maps that match the search or package flags
 		if (CheckRedirectFlagsMatch(PairFlags, SearchFlags) || (bSearchPackageRedirects && CheckRedirectFlagsMatch(PairFlags, ECoreRedirectFlags::Type_Package)))
 		{
-			for (const TPair<FName, TArray<FCoreRedirect>>& RedirectPair : Pair.Value.RedirectMap)
+			if (EnumHasAnyFlags(PairFlags, ECoreRedirectFlags::Option_MatchWildcardMask))
 			{
-				for (const FCoreRedirect& Redirect : RedirectPair.Value)
+				const FWildcardData* Wildcards = Pair.Value.Wildcards.Get();
+				check(Wildcards);
+
+				FCoreRedirectObjectName::EMatchFlags MatchFlags = FCoreRedirectObjectName::EMatchFlags::None;
+				const TArray<FCoreRedirect>* WildcardRedirects = nullptr;
+
+				if (EnumHasAllFlags(PairFlags, ECoreRedirectFlags::Option_MatchSubstring))
 				{
-					FCoreRedirectObjectName::EMatchFlags MatchFlags = FCoreRedirectObjectName::EMatchFlags::None;
-					if (Redirect.IsSubstringMatch())
+					WildcardRedirects = &Wildcards->Substrings;
+					MatchFlags |= FCoreRedirectObjectName::EMatchFlags::CheckSubString;
+				}
+				else if (EnumHasAllFlags(PairFlags, ECoreRedirectFlags::Option_MatchPrefix))
+				{
+					WildcardRedirects = &Wildcards->Prefixes;
+					MatchFlags |= FCoreRedirectObjectName::EMatchFlags::CheckPrefix;
+				}
+				else if (EnumHasAllFlags(PairFlags, ECoreRedirectFlags::Option_MatchSuffix))
+				{
+					WildcardRedirects = &Wildcards->Suffixes;
+					MatchFlags |= FCoreRedirectObjectName::EMatchFlags::CheckSuffix;
+				}
+				check(WildcardRedirects);
+
+				for (const FCoreRedirect& Redirect : *WildcardRedirects)
+				{
+					if (Redirect.NewName.Matches(NewObjectName, MatchFlags))
 					{
-						MatchFlags |= FCoreRedirectObjectName::EMatchFlags::CheckSubString;
+						bFound |= TryReverseRedirect(Redirect, NewObjectName, PreviousNames);
 					}
-					if (Redirect.NewName.Matches(NewObjectName, MatchFlags))
+				}
+			}
+			else
+			{
+				for (const TPair<FName, TArray<FCoreRedirect>>& RedirectPair : Pair.Value.RedirectMap)
+				{
+					for (const FCoreRedirect& Redirect : RedirectPair.Value)
 					{
-						// Construct a reverse redirect
-						FCoreRedirect ReverseRedirect = FCoreRedirect(Redirect);
-						ReverseRedirect.OldName = Redirect.NewName;
-						ReverseRedirect.NewName = Redirect.OldName;
-
-						FCoreRedirectObjectName OldName = ReverseRedirect.RedirectName(NewObjectName);
-
-						if (OldName != NewObjectName)
+						if (Redirect.NewName.Matches(NewObjectName, FCoreRedirectObjectName::EMatchFlags::None))
 						{
-							bFound = true;
-							PreviousNames.AddUnique(OldName);
+							bFound |= TryReverseRedirect(Redirect, NewObjectName, PreviousNames);
 						}
 					}
 				}
@@ -815,6 +1317,7 @@ void FCoreRedirects::ClearKnownMissing(ECoreRedirectFlags Type, ECoreRedirectFla
 bool FCoreRedirects::RunTests()
 {
 	bool bSuccess = true;
+
 	FRedirectTypeMap BackupMap = MoveTemp(RedirectTypeMap);
 #if WITH_COREREDIRECTS_MULTITHREAD_WARNING
 	bool BackupIsInMultithreadedPhase = bIsInMultithreadedPhase;
@@ -833,6 +1336,14 @@ bool FCoreRedirects::RunTests()
 	NewRedirects.Emplace(ECoreRedirectFlags::Type_Class | ECoreRedirectFlags::Category_InstanceOnly, TEXT("/Game/Package.Class"), TEXT("/Game/Package.ClassInstance"));
 	NewRedirects.Emplace(ECoreRedirectFlags::Type_Package, TEXT("/Game/Package"), TEXT("/Game/Package2"));
 	NewRedirects.Emplace(ECoreRedirectFlags::Type_Package | ECoreRedirectFlags::Option_MatchSubstring, TEXT("/oldgame"), TEXT("/newgame"));
+	NewRedirects.Emplace(ECoreRedirectFlags::Type_Package | ECoreRedirectFlags::Option_MatchSubstring, TEXT("/古いゲーム"), TEXT("/新しいゲーム"));
+	NewRedirects.Emplace(ECoreRedirectFlags::Type_Package | ECoreRedirectFlags::Option_MatchSubstring, TEXT("/混合部分文字列"), TEXT("/mixed_substring"));
+	NewRedirects.Emplace(ECoreRedirectFlags::Type_Package | ECoreRedirectFlags::Option_MatchPrefix, TEXT("/oldprefix"), TEXT("/newprefix"));
+	NewRedirects.Emplace(ECoreRedirectFlags::Type_Package | ECoreRedirectFlags::Option_MatchPrefix, TEXT("/古いプレフィックス"), TEXT("/新しいプレフィックス"));
+	NewRedirects.Emplace(ECoreRedirectFlags::Type_Package | ECoreRedirectFlags::Option_MatchPrefix, TEXT("/混合接頭辞"), TEXT("/mixed_prefix"));
+	NewRedirects.Emplace(ECoreRedirectFlags::Type_Package | ECoreRedirectFlags::Option_MatchSuffix, TEXT("/oldsuffix"), TEXT("/newsuffix"));
+	NewRedirects.Emplace(ECoreRedirectFlags::Type_Package | ECoreRedirectFlags::Option_MatchSuffix, TEXT("/古い接尾辞"), TEXT("/新しい接尾辞"));
+	NewRedirects.Emplace(ECoreRedirectFlags::Type_Package | ECoreRedirectFlags::Option_MatchSuffix, TEXT("/混合接尾辞"), TEXT("/mixed_suffix"));
 	NewRedirects.Emplace(ECoreRedirectFlags::Type_Package | ECoreRedirectFlags::Category_Removed, TEXT("/Game/RemovedPackage"), TEXT("/Game/RemovedPackage"));
 	NewRedirects.Emplace(ECoreRedirectFlags::Type_Package | ECoreRedirectFlags::Category_Removed | ECoreRedirectFlags::Option_MissingLoad, TEXT("/Game/MissingLoadPackage"), TEXT("/Game/MissingLoadPackage"));
 
@@ -867,8 +1378,22 @@ bool FCoreRedirects::RunTests()
 	Tests.Emplace(TEXT("/Game/PackageOther.Class"), TEXT("/Game/PackageOther.Class2"), ECoreRedirectFlags::Type_Class);
 	// Check instance option
 	Tests.Emplace(TEXT("/Game/Package.Class"), TEXT("/Game/Package2.ClassInstance"), ECoreRedirectFlags::Type_Class | ECoreRedirectFlags::Category_InstanceOnly);
-	// Substring test
-	Tests.Emplace(TEXT("/oldgame/Package.DefaultClass"), TEXT("/newgame/Package.DefaultClass"), ECoreRedirectFlags::Type_Class);
+
+	// String manipulation tests. While we use TCHAR for input, FNames will coerce storage to ASCII where possible 
+	// so we validate we can support truly wide strings as well as mixed string manipulation operations.
+
+	// Substring tests
+	Tests.Emplace(TEXT("/oldgame/Package.DefaultClass"), TEXT("/newgame/Package.DefaultClass"), ECoreRedirectFlags::Type_Package);
+	Tests.Emplace(TEXT("/古いゲーム/Package.DefaultClass"), TEXT("/新しいゲーム/Package.DefaultClass"), ECoreRedirectFlags::Type_Package);
+	Tests.Emplace(TEXT("/混合部分文字列/Package.DefaultClass"), TEXT("/mixed_substring/Package.DefaultClass"), ECoreRedirectFlags::Type_Package);
+	// Prefix tests
+	Tests.Emplace(TEXT("/oldprefix_SomeGame/Package.DefaultClass"), TEXT("/newprefix_SomeGame/Package.DefaultClass"), ECoreRedirectFlags::Type_Package);
+	Tests.Emplace(TEXT("/古いプレフィックス_SomeGame/Package.DefaultClass"), TEXT("/新しいプレフィックス_SomeGame/Package.DefaultClass"), ECoreRedirectFlags::Type_Package);
+	Tests.Emplace(TEXT("/混合接頭辞_SomeGame/Package.DefaultClass"), TEXT("/mixed_prefix_SomeGame/Package.DefaultClass"), ECoreRedirectFlags::Type_Package);
+	// Suffix tests
+	Tests.Emplace(TEXT("/Game/Package/oldsuffix"), TEXT("/Game/Package/newsuffix"), ECoreRedirectFlags::Type_Package);
+	Tests.Emplace(TEXT("/Game/Package/古い接尾辞"), TEXT("/Game/Package/新しい接尾辞"), ECoreRedirectFlags::Type_Package);
+	Tests.Emplace(TEXT("/Game/Package/混合接尾辞"), TEXT("/Game/Package/mixed_suffix"), ECoreRedirectFlags::Type_Package);
 
 	for (FRedirectTest& Test : Tests)
 	{
@@ -886,12 +1411,39 @@ bool FCoreRedirects::RunTests()
 	TArray<FCoreRedirectObjectName> OldNames;
 
 	FindPreviousNames(ECoreRedirectFlags::Type_Class, FCoreRedirectObjectName(TEXT("/Game/PackageOther.Class2")), OldNames);
-
 	if (OldNames.Num() != 1 || OldNames[0].ToString() != TEXT("/Game/PackageOther.Class"))
 	{
 		bSuccess = false;
-		UE_LOG(LogCoreRedirects, Error, TEXT("FCoreRedirect Test Failed: ReverseLookup!"));
+		UE_LOG(LogCoreRedirects, Error, TEXT("FCoreRedirect Test Failed: ReverseLookup (direct matching)!"));
+	}
+	OldNames.Empty();
+
+	// Check reverse lookup - substring
+	FindPreviousNames(ECoreRedirectFlags::Type_Package, FCoreRedirectObjectName(TEXT("/newgame/TestPackage")), OldNames);
+	if (OldNames.Num() != 1 || OldNames[0].ToString() != TEXT("/oldgame/TestPackage"))
+	{
+		bSuccess = false;
+		UE_LOG(LogCoreRedirects, Error, TEXT("FCoreRedirect Test Failed: ReverseLookup (substring matching)!"));
 	}
+	OldNames.Empty();
+
+	// Check reverse lookup - prefix
+	FindPreviousNames(ECoreRedirectFlags::Type_Package, FCoreRedirectObjectName(TEXT("/newprefix_SomeGame/TestPackage")), OldNames);
+	if (OldNames.Num() != 1 || OldNames[0].ToString() != TEXT("/oldprefix_SomeGame/TestPackage"))
+	{
+		bSuccess = false;
+		UE_LOG(LogCoreRedirects, Error, TEXT("FCoreRedirect Test Failed: ReverseLookup (prefix matching)!"));
+	}
+	OldNames.Empty();
+
+	// Check reverse lookup - suffix
+	FindPreviousNames(ECoreRedirectFlags::Type_Package, FCoreRedirectObjectName(TEXT("/TestGame/newsuffix")), OldNames);
+	if (OldNames.Num() != 1 || OldNames[0].ToString() != TEXT("/TestGame/oldsuffix"))
+	{
+		bSuccess = false;
+		UE_LOG(LogCoreRedirects, Error, TEXT("FCoreRedirect Test Failed: ReverseLookup (suffix matching)!"));
+	}
+	OldNames.Empty();
 
 	// Check removed
 	if (!IsKnownMissing(ECoreRedirectFlags::Type_Package, FCoreRedirectObjectName(TEXT("/Game/RemovedPackage"))))
@@ -966,6 +1518,7 @@ bool FCoreRedirects::RunTests()
 	bIsInMultithreadedPhase = BackupIsInMultithreadedPhase;
 #endif
 
+	UE_LOG(LogCoreRedirects, Log, TEXT("FCoreRedirect Test %s!"), (bSuccess ? TEXT("Passed") : TEXT("Failed")));
 	return bSuccess;
 }
 
@@ -993,12 +1546,14 @@ bool FCoreRedirects::ReadRedirectsFromIni(const FString& IniName)
 				bool bInstanceOnly = false;
 				bool bRemoved = false;
 				bool bMatchSubstring = false;
+				bool bMatchWildcard = false;
 
 				const FString& ValueString = It.Value().GetValue();
 
 				FParse::Bool(*ValueString, TEXT("InstanceOnly="), bInstanceOnly);
 				FParse::Bool(*ValueString, TEXT("Removed="), bRemoved);
 				FParse::Bool(*ValueString, TEXT("MatchSubstring="), bMatchSubstring);
+				FParse::Bool(*ValueString, TEXT("MatchWildcard="), bMatchWildcard);
 
 				FParse::Value(*ValueString, TEXT("OldName="), OldName);
 				FParse::Value(*ValueString, TEXT("NewName="), NewName);
@@ -1022,9 +1577,48 @@ bool FCoreRedirects::ReadRedirectsFromIni(const FString& IniName)
 						NewFlags |= ECoreRedirectFlags::Category_Removed;
 					}
 
-					if (bMatchSubstring)
+					if (bMatchWildcard || bMatchSubstring)
 					{
-						NewFlags |= ECoreRedirectFlags::Option_MatchSubstring;
+						// Enable once we deprecate MatchSubstring across the engine
+						//UE_CLOG(bMatchSubstring, LogCoreRedirects, Warning, TEXT("ReadRedirectsFromIni(%s) 'MatchSubstring=' is deprecated. "
+						//	"Please prefer `MatchWildcard=' instead for redirect %s!"), *IniName, *ValueString);
+
+						constexpr FStringView WildcardMarker(TEXTVIEW("..."));
+						const bool bMatchPrefix = OldName.EndsWith(WildcardMarker, ESearchCase::CaseSensitive);
+						const bool bMatchSuffix = OldName.StartsWith(WildcardMarker, ESearchCase::CaseSensitive);
+						bMatchSubstring = bMatchSubstring || (bMatchPrefix && bMatchSuffix);
+
+						// Count how many '...' there are to ensure the OldName is not malformed
+						int WildcardCount = 0;
+						int StartPos = 0;
+						FStringView OldNameView(OldName);
+						while ((StartPos = OldNameView.Find(WildcardMarker, StartPos)) != INDEX_NONE)
+						{
+							StartPos += WildcardMarker.Len();
+							WildcardCount++;
+						}
+
+						if ((!bMatchPrefix && !bMatchSuffix && !bMatchSubstring) // No wildcards found and not handling MatchSubstring (no '...')
+							|| (WildcardCount > (bMatchPrefix + bMatchSuffix)))	 // Ensure we don't have more wildcards than necessary
+						{
+							UE_LOG(LogCoreRedirects, Error, TEXT("ReadRedirectsFromIni(%s) failed to parse OldName for wildcard redirect %s! "
+								"OldName must be of the form 'PrefixName...', '...SubstringName...' or '...SuffixName'. For more information refer to the documentation in Engine/Config/BaseEngine.ini."), *IniName, *ValueString);
+							continue;
+						}
+
+						if (bMatchPrefix)
+						{
+							NewFlags |= ECoreRedirectFlags::Option_MatchPrefix;
+							OldName.LeftChopInline(WildcardMarker.Len(), EAllowShrinking::No);
+						}
+
+						if (bMatchSuffix)
+						{
+							NewFlags |= ECoreRedirectFlags::Option_MatchSuffix;
+							OldName.RightChopInline(WildcardMarker.Len(), EAllowShrinking::No);
+						}
+
+						NewFlags |= bMatchSubstring ? ECoreRedirectFlags::Option_MatchSubstring : ECoreRedirectFlags::None;
 					}
 
 					FCoreRedirect* Redirect = &NewRedirects.Emplace_GetRef(NewFlags, FCoreRedirectObjectName(OldName), FCoreRedirectObjectName(NewName));
@@ -1035,7 +1629,6 @@ bool FCoreRedirects::ReadRedirectsFromIni(const FString& IniName)
 					}
 
 					int32 ValueChangesIndex = ValueString.Find(TEXT("ValueChanges="));
-				
 					if (ValueChangesIndex != INDEX_NONE)
 					{
 						// Look for first (
@@ -1102,9 +1695,9 @@ bool FCoreRedirects::AddRedirectList(TArrayView<const FCoreRedirect> Redirects,
 			continue;
 		}
 
-		if (NewRedirect.IsSubstringMatch())
+		if (NewRedirect.IsWildcardMatch())
 		{
-			UE_LOG(LogCoreRedirects, Log, TEXT("AddRedirect(%s) has substring redirect %s, these are very slow and should be resolved as soon as possible!"), *SourceString, *NewRedirect.OldName.ToString());
+			UE_LOG(LogCoreRedirects, Log, TEXT("AddRedirect(%s) has wildcard redirect %s, these are very slow and should be resolved as soon as possible! Please refer to the documentation in Engine/Config/BaseEngine.ini."), *SourceString, *NewRedirect.OldName.ToString());
 		}
 		
 		if (AddSingleRedirect(NewRedirect, SourceString))
@@ -1127,7 +1720,9 @@ bool FCoreRedirects::AddRedirectList(TArrayView<const FCoreRedirect> Redirects,
 
 bool FCoreRedirects::AddSingleRedirect(const FCoreRedirect& NewRedirect, const FString& SourceString)
 {
+	const bool bIsWildcardMatch = NewRedirect.IsWildcardMatch();
 	FRedirectNameMap* ExistingNameMap;
+
 #if WITH_COREREDIRECTS_MULTITHREAD_WARNING
 	if (bIsInMultithreadedPhase)
 	{
@@ -1139,12 +1734,32 @@ bool FCoreRedirects::AddSingleRedirect(const FCoreRedirect& NewRedirect, const F
 	{
 		ExistingNameMap = &RedirectTypeMap.FindOrAdd(NewRedirect.RedirectFlags);
 	}
-	TArray<FCoreRedirect>& ExistingRedirects = ExistingNameMap->RedirectMap.FindOrAdd(NewRedirect.GetSearchKey());
 
-	bool bFoundDuplicate = false;
+	TArray<FCoreRedirect>* ExistingRedirects = nullptr;
+	if (bIsWildcardMatch)
+	{
+		if (NewRedirect.IsSubstringMatch())
+		{
+			ExistingRedirects = &ExistingNameMap->Wildcards->Substrings;
+		}
+		else if (NewRedirect.IsPrefixMatch())
+		{
+			ExistingRedirects = &ExistingNameMap->Wildcards->Prefixes;
+		}
+		else if (NewRedirect.IsSuffixMatch())
+		{
+			ExistingRedirects = &ExistingNameMap->Wildcards->Suffixes;
+		}
+		check(ExistingRedirects);
+	}
+	else
+	{
+		ExistingRedirects = &ExistingNameMap->RedirectMap.FindOrAdd(NewRedirect.GetSearchKey());
+	}
 
 	// Check for duplicate
-	for (FCoreRedirect& ExistingRedirect : ExistingRedirects)
+	bool bFoundDuplicate = false;
+	for (FCoreRedirect& ExistingRedirect : *ExistingRedirects)
 	{
 		if (ExistingRedirect.IdenticalMatchRules(NewRedirect))
 		{
@@ -1195,7 +1810,15 @@ bool FCoreRedirects::AddSingleRedirect(const FCoreRedirect& NewRedirect, const F
 		return false;
 	}
 
-	ExistingRedirects.Add(NewRedirect);
+	if (bIsWildcardMatch)
+	{
+		ExistingNameMap->Wildcards->Add(NewRedirect);
+	}
+	else
+	{
+		ExistingRedirects->Add(NewRedirect);
+	}
+
 	return true;
 }
 
@@ -1230,9 +1853,9 @@ bool FCoreRedirects::RemoveRedirectList(TArrayView<const FCoreRedirect> Redirect
 			continue;
 		}
 
-		if (RedirectToRemove.IsSubstringMatch())
+		if (RedirectToRemove.IsWildcardMatch())
 		{
-			UE_LOG(LogCoreRedirects, Log, TEXT("RemoveRedirect(%s) has substring redirect %s, these are very slow and should be resolved as soon as possible!"), *SourceString, *RedirectToRemove.OldName.ToString());
+			UE_LOG(LogCoreRedirects, Log, TEXT("RemoveRedirect(%s) has wildcard redirect %s, these are very slow and should be resolved as soon as possible! Please refer to the documentation in Engine/Config/BaseEngine.ini."), *SourceString, *RedirectToRemove.OldName.ToString());
 		}
 
 		bRemovedAny |= RemoveSingleRedirect(RedirectToRemove, SourceString);
@@ -1243,12 +1866,35 @@ bool FCoreRedirects::RemoveRedirectList(TArrayView<const FCoreRedirect> Redirect
 
 bool FCoreRedirects::RemoveSingleRedirect(const FCoreRedirect& RedirectToRemove, const FString& SourceString)
 {
+	const bool bIsWildcardMatch = RedirectToRemove.IsWildcardMatch();
 	FRedirectNameMap* ExistingNameMap = RedirectTypeMap.Find(RedirectToRemove.RedirectFlags);
+
 	if (!ExistingNameMap)
 	{
 		return false;
 	}
-	TArray<FCoreRedirect>* ExistingRedirects = ExistingNameMap->RedirectMap.Find(RedirectToRemove.GetSearchKey());
+
+	TArray<FCoreRedirect>* ExistingRedirects = nullptr;
+	if (bIsWildcardMatch)
+	{
+		if (RedirectToRemove.IsSubstringMatch())
+		{
+			ExistingRedirects = &ExistingNameMap->Wildcards->Substrings;
+		}
+		else if (RedirectToRemove.IsPrefixMatch())
+		{
+			ExistingRedirects = &ExistingNameMap->Wildcards->Prefixes;
+		}
+		else if (RedirectToRemove.IsSuffixMatch())
+		{
+			ExistingRedirects = &ExistingNameMap->Wildcards->Suffixes;
+		}
+	}
+	else
+	{
+		ExistingRedirects = ExistingNameMap->RedirectMap.Find(RedirectToRemove.GetSearchKey());
+	}
+
 	if (!ExistingRedirects)
 	{
 		return false;
@@ -1273,6 +1919,13 @@ bool FCoreRedirects::RemoveSingleRedirect(const FCoreRedirect& RedirectToRemove,
 		}
 	}
 
+	if (bIsWildcardMatch && bRemovedRedirect)
+	{
+		// We removed a substring redirect so we need to regenerate our prediction 
+		// tables to avoid unnecessary false positives
+		ExistingNameMap->Wildcards->Rebuild();
+	}
+
 	return bRemovedRedirect;
 }
 
@@ -2196,4 +2849,4 @@ void FCoreRedirects::RegisterNativeRedirects()
 void FCoreRedirects::RegisterNativeRedirects()
 {
 }
-#endif // UE_WITH_CORE_REDIRECTS
+#endif // UE_WITH_CORE_REDIRECTS
\ No newline at end of file
diff --git a/Engine/Source/Runtime/CoreUObject/Public/UObject/CoreRedirects.h b/Engine/Source/Runtime/CoreUObject/Public/UObject/CoreRedirects.h
--- a/Engine/Source/Runtime/CoreUObject/Public/UObject/CoreRedirects.h
+++ b/Engine/Source/Runtime/CoreUObject/Public/UObject/CoreRedirects.h
@@ -17,10 +17,12 @@
 #include "Misc/EnumClassFlags.h"
 #include "UObject/NameTypes.h"
 #include "UObject/UnrealNames.h"
+#include "UObject/CoreRedirects/PM-k.h"
 
 class IPakFile;
 class UClass;
 struct FTopLevelAssetPath;
+namespace UE::CoreRedirects::Private { struct FCoreRedirectObjectUtf8Name; }
 
 #define WITH_COREREDIRECTS_MULTITHREAD_WARNING !UE_BUILD_SHIPPING && !IS_PROGRAM && !WITH_EDITOR
 
@@ -34,24 +36,28 @@ enum class ECoreRedirectFlags : uint32
 	None = 0,
 
 	// Core type of the Thing being redirected, multiple can be set.  A Query will only find Redirects that have at least one of the same Type bits set.
-	Type_Object =			0x00000001, // UObject
-	Type_Class =			0x00000002, // UClass
-	Type_Struct =			0x00000004, // UStruct
-	Type_Enum =				0x00000008, // UEnum
-	Type_Function =			0x00000010, // UFunction
-	Type_Property =			0x00000020, // FProperty
-	Type_Package =			0x00000040, // UPackage
-	Type_AllMask =			0x0000FFFF, // Bit mask of all possible Types
+	Type_Object =				0x00000001, // UObject
+	Type_Class =				0x00000002, // UClass
+	Type_Struct =				0x00000004, // UStruct
+	Type_Enum =					0x00000008, // UEnum
+	Type_Function =				0x00000010, // UFunction
+	Type_Property =				0x00000020, // FProperty
+	Type_Package =				0x00000040, // UPackage
+	Type_AllMask =				0x0000FFFF, // Bit mask of all possible Types
 
 	// Category flags.  A Query will only match Redirects that have the same value for every category bit.
-	Category_InstanceOnly = 0x00010000, // Only redirect instances of this type, not the type itself
-	Category_Removed =		0x00020000, // This type was explicitly removed, new name isn't valid
-	Category_AllMask =		0x00FF0000, // Bit mask of all possible Categories
+	Category_InstanceOnly =		0x00010000, // Only redirect instances of this type, not the type itself
+	Category_Removed =			0x00020000, // This type was explicitly removed, new name isn't valid
+	Category_AllMask =			0x00FF0000, // Bit mask of all possible Categories
 
 	// Option flags.  Does not behave as a bit-match between Queries and Redirects.  Each one specifies a custom rule for how FCoreRedirects handles the Redirect.
-	Option_MatchSubstring = 0x01000000, // Does a slow substring match
-	Option_MissingLoad =	0x02000000, // An automatically-created redirect that was created in response to a missing Thing during load. Redirect will be removed if and when the Thing is loaded.
-	Option_AllMask =		0xFF000000, // Bit mask of all possible Options
+	Option_MatchPrefix =		0x01000000, // Does a prefix string match
+	Option_MatchSuffix =		0x02000000, // Does a suffix string match
+	Option_MatchSubstring =		Option_MatchPrefix | Option_MatchSuffix, // Does a slow substring match
+	Option_MatchWildcardMask =	Option_MatchSubstring, // Bit mask of all possible wildcards
+
+	Option_MissingLoad =		0x04000000, // An automatically-created redirect that was created in response to a missing Thing during load. Redirect will be removed if and when the Thing is loaded.
+	Option_AllMask =			0xFF000000, // Bit mask of all possible Options
 };
 ENUM_CLASS_FLAGS(ECoreRedirectFlags);
 
@@ -90,7 +96,6 @@ struct FCoreRedirectObjectName
 
 	COREUOBJECT_API FCoreRedirectObjectName(const FTopLevelAssetPath& TopLevelAssetPath);
 
-	/** Construct from a path string, this handles full paths with packages, or partial paths without */
 	COREUOBJECT_API FCoreRedirectObjectName(const FString& InString);
 
 	/** Construct from object in memory */
@@ -123,17 +128,29 @@ struct FCoreRedirectObjectName
 		AllowPartialRHSMatch = (1 << 1),
 		/**
 		 * LHS fields (aka *this) are searchstrings; RHS (aka Other) fields are searched for that substring.
-		 * Default is to require a complete string match LHS == RHS.
+		 * Without this flag a Match returns true if and only if the complete string matches: LHS == RHS.
+		 * With this flag a Match returns true if and only if RHS.Contains(LHS).
 		 * This flag makes the match more expensive and should be avoided when possible.
 		 */
 		CheckSubString = (1 << 2),
+		/**
+		 * LHS fields (aka *this) are searchstrings; RHS (aka Other) fields are searched for that prefix.
+		 * Without this flag a Match returns true if and only if the complete string matches: LHS == RHS.
+		 * With this flag a Match returns true if and only if RHS.StartsWith(LHS).
+		 * This flag makes the match more expensive and should be avoided when possible.
+		 */
+		CheckPrefix = (1 << 3),
+		/**
+		 * LHS fields (aka *this) are searchstrings; RHS (aka Other) fields are searched for that suffix.
+		 * Without this flag a Match returns true if and only if the complete string matches: LHS == RHS.
+		 * With this flag a Match returns true if and only if RHS.EndsWith(LHS).
+		 * This flag makes the match more expensive and should be avoided when possible.
+		 */
+		CheckSuffix = (1 << 4),
 	};
 	/** Returns true if the passed in name matches requirements. */
 	COREUOBJECT_API bool Matches(const FCoreRedirectObjectName& Other, EMatchFlags MatchFlags = EMatchFlags::None) const;
 
-	UE_DEPRECATED(5.1, "Use EMatchFlags::CheckSubString to pass in bCheckSubstring=true.")
-	COREUOBJECT_API bool Matches(const FCoreRedirectObjectName& Other, bool bCheckSubstring) const;
-
 	/** Returns integer of degree of match. 0 if doesn't match at all, higher integer for better matches */
 	COREUOBJECT_API int32 MatchScore(const FCoreRedirectObjectName& Other) const;
 
@@ -143,14 +160,6 @@ struct FCoreRedirectObjectName
 	/** Returns the name used as the key into the acceleration map */
 	FName GetSearchKey(ECoreRedirectFlags Type) const
 	{
-		if ((Type & ECoreRedirectFlags::Option_MatchSubstring) == ECoreRedirectFlags::Option_MatchSubstring)
-		{
-			static FName SubstringName = FName(TEXT("*SUBSTRING*"));
-
-			// All substring matches pass initial test as they need to be manually checked
-			return SubstringName;
-		}
-
 		if ((Type & ECoreRedirectFlags::Type_Package) == ECoreRedirectFlags::Type_Package)
 		{
 			return PackageName;
@@ -169,7 +178,7 @@ struct FCoreRedirectObjectName
 	COREUOBJECT_API bool HasValidCharacters(ECoreRedirectFlags Type) const;
 
 	/** Expand OldName/NewName as needed */
-	static COREUOBJECT_API bool ExpandNames(const FString& FullString, FName& OutName, FName& OutOuter, FName &OutPackage);
+	static COREUOBJECT_API bool ExpandNames(const FStringView FullString, FName& OutName, FName& OutOuter, FName& OutPackage);
 
 	/** Turn it back into an FString */
 	static COREUOBJECT_API FString CombineNames(FName NewName, FName NewOuter, FName NewPackage);
@@ -240,6 +249,31 @@ struct FCoreRedirect
 	{
 		return OldName.GetSearchKey(RedirectFlags);
 	}
+
+private:
+	friend struct FCoreRedirects;
+
+	/* Returns the updated name after redirection. If bIsKnownToMatch is true, OldObjectName must have 
+	been validated previously to be acceptable for redirection */
+	FCoreRedirectObjectName RedirectName(const FCoreRedirectObjectName& OldObjectName, bool bIsKnownToMatch) const;
+
+	/** Returns true if this is a Wildcard match (substring, prefix or suffix) */
+	bool IsWildcardMatch() const
+	{
+		return EnumHasAnyFlags(RedirectFlags, ECoreRedirectFlags::Option_MatchWildcardMask);
+	}
+
+	/** Returns true if this is a prefix match */
+	bool IsPrefixMatch() const
+	{
+		return EnumHasAllFlags(RedirectFlags, ECoreRedirectFlags::Option_MatchPrefix);
+	}
+
+	/** Returns true if this is a prefix match */
+	bool IsSuffixMatch() const
+	{
+		return EnumHasAllFlags(RedirectFlags, ECoreRedirectFlags::Option_MatchSuffix);
+	}
 };
 
 /**
@@ -346,11 +380,32 @@ private:
 	static COREUOBJECT_API void EnterMultithreadedPhase();
 #endif
 
+	/** Container for managing Wildcard redirects (substrings, prefixes, suffixes) */
+	struct FWildcardData
+	{
+		void Add(const FCoreRedirect& Redirect);
+
+		void Rebuild();
+		bool Matches(ECoreRedirectFlags InFlags, const FCoreRedirectObjectName& InName, ECoreRedirectMatchFlags InMatchFlags, TArray<const FCoreRedirect*>& OutFoundRedirects) const;
+
+		TArray<FCoreRedirect> Substrings;
+		TArray<FCoreRedirect> Prefixes;
+		TArray<FCoreRedirect> Suffixes;
+	private:
+		/** This function may return false positives, but will not return false negatives */
+		bool MatchSubstringApproximate(const UE::CoreRedirects::Private::FCoreRedirectObjectUtf8Name& RedirectName) const;
+		void AddPredictionWords(const FCoreRedirect& Redirect);
+
+		FPredictMatch8 PredictMatch;
+	};
+
 	/** There is one of these for each registered set of redirect flags */
 	struct FRedirectNameMap
 	{
 		/** Map from name of thing being mapped to full list. List must be filtered further */
-		TMap<FName, TArray<FCoreRedirect> > RedirectMap;
+		TMap<FName, TArray<FCoreRedirect>> RedirectMap;
+		/** Used to manage wildcard data and accelerate wildcard queries */
+		TUniquePtr<FWildcardData> Wildcards;
 	};
 
 	/** Whether this has been initialized at least once */
