Take 2 at fixing edit property chain in overridable manager when changing values on a subobject

diff --git a/Engine/Source/Runtime/CoreUObject/Private/UObject/OverridableManager.cpp b/Engine/Source/Runtime/CoreUObject/Private/UObject/OverridableManager.cpp
--- a/Engine/Source/Runtime/CoreUObject/Private/UObject/OverridableManager.cpp
+++ b/Engine/Source/Runtime/CoreUObject/Private/UObject/OverridableManager.cpp
@@ -110,8 +110,8 @@ void FOverridableManager::OverrideObject(UObject& Object)
 {
 	if (FOverriddenPropertySet* ThisObjectOverriddenProperties = OverriddenObjectAnnotations.Find(Object))
 	{
-		ThisObjectOverriddenProperties->NotifyPropertyChange(EPropertyNotificationType::PreEdit, FPropertyChangedEvent(nullptr), nullptr, &Object);
-		ThisObjectOverriddenProperties->NotifyPropertyChange(EPropertyNotificationType::PostEdit, FPropertyChangedEvent(nullptr), nullptr, &Object);
+		// Passing no property node means we are overriding the object itself
+		ThisObjectOverriddenProperties->OverrideProperty(FPropertyChangedEvent(nullptr), /*PropertyNode*/nullptr, &Object);
 	}
 }
 
@@ -138,8 +138,7 @@ void FOverridableManager::OverrideProperty(UObject& Object, const FPropertyChang
 {
 	if (FOverriddenPropertySet* ThisObjectOverriddenProperties = OverriddenObjectAnnotations.Find(Object))
 	{
-		ThisObjectOverriddenProperties->NotifyPropertyChange(EPropertyNotificationType::PreEdit, FPropertyChangedEvent(nullptr), PropertyChain.GetHead(), &Object);
-		ThisObjectOverriddenProperties->NotifyPropertyChange(EPropertyNotificationType::PostEdit, PropertyEvent, PropertyChain.GetHead(), &Object);
+		ThisObjectOverriddenProperties->OverrideProperty(PropertyEvent, PropertyChain.GetActiveMemberNode() ? PropertyChain.GetActiveMemberNode() : PropertyChain.GetHead(), &Object);
 	}
 }
 
@@ -156,7 +155,7 @@ void FOverridableManager::PreOverrideProperty(UObject& Object, const FEditProper
 {
 	if (FOverriddenPropertySet* ThisObjectOverriddenProperties = OverriddenObjectAnnotations.Find(Object))
 	{
-		ThisObjectOverriddenProperties->NotifyPropertyChange(EPropertyNotificationType::PreEdit, FPropertyChangedEvent(nullptr), PropertyChain.GetHead(), &Object);
+		ThisObjectOverriddenProperties->NotifyPropertyChange(EPropertyNotificationType::PreEdit, FPropertyChangedEvent(nullptr), PropertyChain.GetActiveMemberNode() ? PropertyChain.GetActiveMemberNode() : PropertyChain.GetHead(), &Object);
 	}
 }
 
@@ -164,7 +163,7 @@ void FOverridableManager::PostOverrideProperty(UObject& Object, const FPropertyC
 {
 	if (FOverriddenPropertySet* ThisObjectOverriddenProperties = OverriddenObjectAnnotations.Find(Object))
 	{
-		ThisObjectOverriddenProperties->NotifyPropertyChange(EPropertyNotificationType::PostEdit, PropertyEvent, PropertyChain.GetHead(), &Object);
+		ThisObjectOverriddenProperties->NotifyPropertyChange(EPropertyNotificationType::PostEdit, PropertyEvent, PropertyChain.GetActiveMemberNode() ? PropertyChain.GetActiveMemberNode() : PropertyChain.GetHead(), &Object);
 	}
 }
 
diff --git a/Engine/Source/Runtime/CoreUObject/Private/UObject/OverriddenPropertySet.cpp b/Engine/Source/Runtime/CoreUObject/Private/UObject/OverriddenPropertySet.cpp
--- a/Engine/Source/Runtime/CoreUObject/Private/UObject/OverriddenPropertySet.cpp
+++ b/Engine/Source/Runtime/CoreUObject/Private/UObject/OverriddenPropertySet.cpp
@@ -1023,6 +1023,13 @@ bool FOverriddenPropertySet::ClearOverriddenProperty(const FPropertyChangedEvent
 	return false;
 }
 
+void FOverriddenPropertySet::OverrideProperty(const FPropertyChangedEvent& PropertyEvent, const FEditPropertyChain::TDoubleLinkedListNode* PropertyNode, const void* Data)
+{
+	FOverriddenPropertyNode& RootPropertyNode = OverriddenPropertyNodes.FindOrAddByHash(GetTypeHash(RootNodeID), RootNodeID);
+	NotifyPropertyChange(&RootPropertyNode, EPropertyNotificationType::PreEdit, FPropertyChangedEvent(nullptr), PropertyNode, Data);
+	NotifyPropertyChange(&RootPropertyNode, EPropertyNotificationType::PostEdit, PropertyEvent, PropertyNode, Data);
+}
+
 void FOverriddenPropertySet::NotifyPropertyChange(const EPropertyNotificationType Notification, const FPropertyChangedEvent& PropertyEvent, const FEditPropertyChain::TDoubleLinkedListNode* PropertyNode, const void* Data)
 {
 	NotifyPropertyChange(&OverriddenPropertyNodes.FindOrAddByHash(GetTypeHash(RootNodeID), RootNodeID), Notification, PropertyEvent, PropertyNode, Data);
diff --git a/Engine/Source/Runtime/CoreUObject/Public/UObject/OverridableManager.h b/Engine/Source/Runtime/CoreUObject/Public/UObject/OverridableManager.h
--- a/Engine/Source/Runtime/CoreUObject/Public/UObject/OverridableManager.h
+++ b/Engine/Source/Runtime/CoreUObject/Public/UObject/OverridableManager.h
@@ -184,7 +184,7 @@ public:
 	 * @return true if the property was successfully cleared. */
 	FORCEINLINE bool ClearOverriddenProperty(UObject& Object, const FPropertyChangedEvent& PropertyEvent, const FEditPropertyChain& PropertyChain)
 	{
-		return ClearOverriddenProperty(Object, PropertyEvent, PropertyChain.GetHead());
+		return ClearOverriddenProperty(Object, PropertyEvent, PropertyChain.GetActiveMemberNode() ? PropertyChain.GetActiveMemberNode() : PropertyChain.GetHead());
 	}
 
 	/**
@@ -228,7 +228,7 @@ public:
 	 * @return the current type of override operation on the property */
 	FORCEINLINE EOverriddenPropertyOperation GetOverriddenPropertyOperation(UObject& Object, const FPropertyChangedEvent& PropertyEvent, const FEditPropertyChain& PropertyChain, bool* bOutInheritedOperation = nullptr)
 	{
-		return GetOverriddenPropertyOperation(Object, PropertyEvent, PropertyChain.GetHead(), bOutInheritedOperation);
+		return GetOverriddenPropertyOperation(Object, PropertyEvent, PropertyChain.GetActiveMemberNode() ? PropertyChain.GetActiveMemberNode() : PropertyChain.GetHead(), bOutInheritedOperation);
 	}
 
 	/**
diff --git a/Engine/Source/Runtime/CoreUObject/Public/UObject/OverriddenPropertySet.h b/Engine/Source/Runtime/CoreUObject/Public/UObject/OverriddenPropertySet.h
--- a/Engine/Source/Runtime/CoreUObject/Public/UObject/OverriddenPropertySet.h
+++ b/Engine/Source/Runtime/CoreUObject/Public/UObject/OverriddenPropertySet.h
@@ -229,6 +229,13 @@ public:
 	 * @return if the operation was successful */
 	bool ClearOverriddenProperty(const FPropertyChangedEvent& PropertyEvent, const FEditPropertyChain::TDoubleLinkedListNode* PropertyNode);
 
+	/**
+	 * Utility methods that call NotifyPropertyChange(Pre/PostEdit)
+	 * @param PropertyEvent information about the type of change
+	 * @param PropertyNode leading to the property that is changing, null means it is the object itself that is changing
+	 * @param Data memory of the current property */
+	void OverrideProperty(const FPropertyChangedEvent& PropertyEvent, const FEditPropertyChain::TDoubleLinkedListNode* PropertyNode, const void* Data);
+
 	/**
 	 * Handling and storing modification on a property of an object
 	 * @param Notification type either pre/post property overridden
