Added functions to FPropertyTypeName to query for structs and enums

These are convenience functions that cover most usage of FPropertyTypeName outside of CoreUObject.

#rnx

diff --git a/Engine/Source/Runtime/CoreUObject/Private/UObject/PropertyTypeName.cpp b/Engine/Source/Runtime/CoreUObject/Private/UObject/PropertyTypeName.cpp
--- a/Engine/Source/Runtime/CoreUObject/Private/UObject/PropertyTypeName.cpp
+++ b/Engine/Source/Runtime/CoreUObject/Private/UObject/PropertyTypeName.cpp
@@ -263,6 +263,18 @@ FPropertyTypeName FPropertyTypeName::GetParameter(int32 ParamIndex) const
 	return Type;
 }
 
+bool FPropertyTypeName::IsStruct(FName StructName) const
+{
+	const FPropertyTypeNameNode* First = GPropertyTypeNameTable.ResolveByIndex(Index);
+	return First->InnerCount > 0 && First->Name == NAME_StructProperty && First[1].Name == StructName;
+}
+
+bool FPropertyTypeName::IsEnum(FName EnumName) const
+{
+	const FPropertyTypeNameNode* First = GPropertyTypeNameTable.ResolveByIndex(Index);
+	return First->InnerCount > 0 && (First->Name == NAME_EnumProperty || First->Name == NAME_ByteProperty) && First[1].Name == EnumName;
+}
+
 uint32 GetTypeHash(const FPropertyTypeName& TypeName)
 {
 	const FPropertyTypeNameNode* First = GPropertyTypeNameTable.ResolveByIndex(TypeName.Index);
diff --git a/Engine/Source/Runtime/CoreUObject/Public/UObject/PropertyTypeName.h b/Engine/Source/Runtime/CoreUObject/Public/UObject/PropertyTypeName.h
--- a/Engine/Source/Runtime/CoreUObject/Public/UObject/PropertyTypeName.h
+++ b/Engine/Source/Runtime/CoreUObject/Public/UObject/PropertyTypeName.h
@@ -98,6 +98,11 @@ public:
 		Index = 0;
 	}
 
+	/** Returns true if this is StructProperty with a first parameter of StructName. */
+	UE_API bool IsStruct(FName StructName) const;
+	/** Returns true if this is EnumProperty or ByteProperty with a first parameter of EnumName. */
+	UE_API bool IsEnum(FName EnumName) const;
+
 private:
 	UE_API friend uint32 GetTypeHash(const FPropertyTypeName& TypeName);
 
diff --git a/Engine/Source/Runtime/CoreUObject/Tests/UObject/PropertyTypeNameTest.cpp b/Engine/Source/Runtime/CoreUObject/Tests/UObject/PropertyTypeNameTest.cpp
--- a/Engine/Source/Runtime/CoreUObject/Tests/UObject/PropertyTypeNameTest.cpp
+++ b/Engine/Source/Runtime/CoreUObject/Tests/UObject/PropertyTypeNameTest.cpp
@@ -169,6 +169,38 @@ TEST_CASE_NAMED(FPropertyTypeNameSmokeTest, "CoreUObject::PropertyTypeName::Smok
 		CHECK(TypeName.GetParameter(2).IsEmpty());
 	}
 
+	SECTION("IsStruct")
+	{
+		FPropertyTypeNameBuilder Builder;
+		Builder.AddName(NAME_StructProperty);
+		CHECK_FALSE(Builder.Build().IsStruct(NAME_Vector));
+		CHECK_FALSE(FPropertyTypeName().IsStruct(NAME_Vector));
+		CHECK(PropertyTypeNameTest::CreateVectorArray().GetParameter(0).IsStruct(NAME_Vector));
+	}
+
+	SECTION("IsEnum(EnumProperty)")
+	{
+		const FName EnumName = TEXT("Key");
+		FPropertyTypeNameBuilder Builder;
+		Builder.AddName(NAME_EnumProperty);
+		CHECK_FALSE(Builder.Build().IsEnum(EnumName));
+		CHECK_FALSE(FPropertyTypeName().IsEnum(EnumName));
+		CHECK(PropertyTypeNameTest::CreateEnumMap().GetParameter(0).IsEnum(EnumName));
+	}
+
+	SECTION("IsEnum(ByteProperty)")
+	{
+		const FName EnumName = TEXT("Key");
+		FPropertyTypeNameBuilder Builder;
+		Builder.AddName(NAME_ByteProperty);
+		CHECK_FALSE(Builder.Build().IsEnum(EnumName));
+		CHECK_FALSE(FPropertyTypeName().IsEnum(EnumName));
+		Builder.BeginParameters();
+		Builder.AddName(EnumName);
+		Builder.EndParameters();
+		CHECK(Builder.Build().IsEnum(EnumName));
+	}
+
 	SECTION("Equals+Less+GetTypeHash")
 	{
 		const FPropertyTypeName Int = PropertyTypeNameTest::CreateInt();
