Deprecated fields on FPropertyTag that are now represented by its TypeName

#rnx

diff --git a/Engine/Plugins/Experimental/Avalanche/Source/AvalancheTag/Private/AvaTagHandleContainer.cpp b/Engine/Plugins/Experimental/Avalanche/Source/AvalancheTag/Private/AvaTagHandleContainer.cpp
--- a/Engine/Plugins/Experimental/Avalanche/Source/AvalancheTag/Private/AvaTagHandleContainer.cpp
+++ b/Engine/Plugins/Experimental/Avalanche/Source/AvalancheTag/Private/AvaTagHandleContainer.cpp
@@ -85,7 +85,7 @@ bool FAvaTagHandleContainer::SerializeFromMismatchedTag(const FPropertyTag& InPr
 {
 	static const FName TagHandleContextName = FAvaTagHandle::StaticStruct()->GetFName();
 
-	if (InPropertyTag.Type == NAME_StructProperty && InPropertyTag.StructName == TagHandleContextName)
+	if (InPropertyTag.GetType().IsStruct(TagHandleContextName))
 	{
 		FAvaTagHandle TagHandle;
 		FAvaTagHandle::StaticStruct()->SerializeItem(InSlot, &TagHandle, nullptr);
diff --git a/Engine/Plugins/Experimental/StructUtils/Source/StructUtils/Private/InstancedStruct.cpp b/Engine/Plugins/Experimental/StructUtils/Source/StructUtils/Private/InstancedStruct.cpp
--- a/Engine/Plugins/Experimental/StructUtils/Source/StructUtils/Private/InstancedStruct.cpp
+++ b/Engine/Plugins/Experimental/StructUtils/Source/StructUtils/Private/InstancedStruct.cpp
@@ -321,7 +321,7 @@ bool FInstancedStruct::ImportTextItem(const TCHAR*& Buffer, int32 PortFlags, UOb
 bool FInstancedStruct::SerializeFromMismatchedTag(const FPropertyTag& Tag, FStructuredArchive::FSlot Slot)
 {
 	static const FName NAME_StructVariant = "StructVariant";
-	if (Tag.Type == NAME_StructProperty && Tag.StructName == NAME_StructVariant)
+	if (Tag.GetType().IsStruct(NAME_StructVariant))
 	{
 		auto SerializeStructVariant = [this, &Tag, &Slot]()
 		{
diff --git a/Engine/Plugins/FX/Niagara/Source/Niagara/Private/NiagaraTypes.cpp b/Engine/Plugins/FX/Niagara/Source/Niagara/Private/NiagaraTypes.cpp
--- a/Engine/Plugins/FX/Niagara/Source/Niagara/Private/NiagaraTypes.cpp
+++ b/Engine/Plugins/FX/Niagara/Source/Niagara/Private/NiagaraTypes.cpp
@@ -58,7 +58,7 @@ bool FNiagaraVariableBase::ReplaceRootNamespace(const FStringView& ExpectedNames
 
 bool FNiagaraVariableBase::SerializeFromMismatchedTag(const struct FPropertyTag& Tag, FStructuredArchive::FSlot Slot)
 {
-	if (Tag.Type == NAME_StructProperty && Tag.StructName == FNiagaraVariable::StaticStruct()->GetFName())
+	if (Tag.GetType().IsStruct(FNiagaraVariable::StaticStruct()->GetFName()))
 	{
 		FNiagaraVariable Var;
 		FNiagaraVariable::StaticStruct()->SerializeItem(Slot, &Var, nullptr);
diff --git a/Engine/Plugins/FX/Niagara/Source/Niagara/Private/NiagaraUserRedirectionParameterStore.cpp b/Engine/Plugins/FX/Niagara/Source/Niagara/Private/NiagaraUserRedirectionParameterStore.cpp
--- a/Engine/Plugins/FX/Niagara/Source/Niagara/Private/NiagaraUserRedirectionParameterStore.cpp
+++ b/Engine/Plugins/FX/Niagara/Source/Niagara/Private/NiagaraUserRedirectionParameterStore.cpp
@@ -180,7 +180,7 @@ void FNiagaraUserRedirectionParameterStore::ConvertParameterType(const FNiagaraV
 bool FNiagaraUserRedirectionParameterStore::SerializeFromMismatchedTag(const FPropertyTag & Tag, FStructuredArchive::FSlot Slot)
 {
 	static FName StoreDataName("NiagaraParameterStore");
-	if (Tag.Type == NAME_StructProperty && Tag.StructName == StoreDataName)
+	if (Tag.GetType().IsStruct(StoreDataName))
 	{
 		FNiagaraParameterStore OldStore;
 		FNiagaraParameterStore::StaticStruct()->SerializeItem(Slot, &OldStore, nullptr);
diff --git a/Engine/Plugins/PCG/Source/PCG/Private/Metadata/PCGAttributePropertySelector.cpp b/Engine/Plugins/PCG/Source/PCG/Private/Metadata/PCGAttributePropertySelector.cpp
--- a/Engine/Plugins/PCG/Source/PCG/Private/Metadata/PCGAttributePropertySelector.cpp
+++ b/Engine/Plugins/PCG/Source/PCG/Private/Metadata/PCGAttributePropertySelector.cpp
@@ -332,8 +332,7 @@ FPCGAttributePropertyInputSelector FPCGAttributePropertyInputSelector::CopyAndFi
 
 bool FPCGAttributePropertyInputSelector::SerializeFromMismatchedTag(const struct FPropertyTag& Tag, FStructuredArchive::FSlot Slot)
 {
-	if (Tag.Type == NAME_StructProperty
-		&& Tag.StructName == FPCGAttributePropertySelector::StaticStruct()->GetFName())
+	if (Tag.GetType().IsStruct(FPCGAttributePropertySelector::StaticStruct()->GetFName()))
 	{
 		FPCGAttributePropertyInputSelector::StaticStruct()->SerializeItem(Slot, this, nullptr);
 		return true;
@@ -393,8 +392,7 @@ FPCGAttributePropertyOutputSelector FPCGAttributePropertyOutputSelector::CopyAnd
 
 bool FPCGAttributePropertyOutputSelector::SerializeFromMismatchedTag(const struct FPropertyTag& Tag, FStructuredArchive::FSlot Slot)
 {
-	if (Tag.Type == NAME_StructProperty
-		&& Tag.StructName == FPCGAttributePropertySelector::StaticStruct()->GetFName())
+	if (Tag.GetType().IsStruct(FPCGAttributePropertySelector::StaticStruct()->GetFName()))
 	{
 		FPCGAttributePropertyOutputSelector::StaticStruct()->SerializeItem(Slot, this, nullptr);
 		return true;
diff --git a/Engine/Plugins/Runtime/RigVM/Source/RigVM/Private/RigVMCore/RigVMExecuteContext.cpp b/Engine/Plugins/Runtime/RigVM/Source/RigVM/Private/RigVMCore/RigVMExecuteContext.cpp
--- a/Engine/Plugins/Runtime/RigVM/Source/RigVM/Private/RigVMCore/RigVMExecuteContext.cpp
+++ b/Engine/Plugins/Runtime/RigVM/Source/RigVM/Private/RigVMCore/RigVMExecuteContext.cpp
@@ -37,7 +37,7 @@ void FRigVMExecuteContext::SetWorld(const UWorld* InWorld)
 bool FRigVMExecuteContext::SerializeFromMismatchedTag(const FPropertyTag& Tag, FStructuredArchive::FSlot Slot)
 {
 	static const FName ControlRigExecuteContextName("ControlRigExecuteContext");
-	if (Tag.Type == NAME_StructProperty && Tag.StructName == ControlRigExecuteContextName)
+	if (Tag.GetType().IsStruct(ControlRigExecuteContextName))
 	{
 		static const FString CRExecuteContextPath = TEXT("/Script/ControlRig.ControlRigExecuteContext");
 		UScriptStruct* OldStruct = FindFirstObject<UScriptStruct>(*CRExecuteContextPath, EFindFirstObjectOptions::NativeFirst | EFindFirstObjectOptions::EnsureIfAmbiguous);
diff --git a/Engine/Plugins/Runtime/StateTree/Source/StateTreeModule/Public/StateTreeIndexTypes.cpp b/Engine/Plugins/Runtime/StateTree/Source/StateTreeModule/Public/StateTreeIndexTypes.cpp
--- a/Engine/Plugins/Runtime/StateTree/Source/StateTreeModule/Public/StateTreeIndexTypes.cpp
+++ b/Engine/Plugins/Runtime/StateTree/Source/StateTreeModule/Public/StateTreeIndexTypes.cpp
@@ -1,12 +1,11 @@
-﻿// Copyright Epic Games, Inc. All Rights Reserved.
+// Copyright Epic Games, Inc. All Rights Reserved.
 
 #include "StateTreeIndexTypes.h"
 
 bool FStateTreeIndex16::SerializeFromMismatchedTag(const FPropertyTag& Tag, FStructuredArchive::FSlot Slot)
 {
 	// Support loading from Index8.
-	if (Tag.Type == NAME_StructProperty
-		&& Tag.StructName == FStateTreeIndex8::StaticStruct()->GetFName())
+	if (Tag.GetType().IsStruct(FStateTreeIndex8::StaticStruct()->GetFName()))
 	{
 		FStateTreeIndex8 OldValue;
 		FStateTreeIndex8::StaticStruct()->SerializeItem(Slot, &OldValue, nullptr);
@@ -29,8 +28,7 @@ bool FStateTreeIndex16::SerializeFromMismatchedTag(const FPropertyTag& Tag, FStr
 bool FStateTreeIndex8::SerializeFromMismatchedTag(const FPropertyTag& Tag, FStructuredArchive::FSlot Slot)
 {
 	// Support loading from Index16.
-	if (Tag.Type == NAME_StructProperty
-		&& Tag.StructName == FStateTreeIndex16::StaticStruct()->GetFName())
+	if (Tag.GetType().IsStruct(FStateTreeIndex16::StaticStruct()->GetFName()))
 	{
 		FStateTreeIndex16 OldValue;
 		FStateTreeIndex16::StaticStruct()->SerializeItem(Slot, &OldValue, nullptr);
diff --git a/Engine/Source/Runtime/CoreUObject/Private/UObject/EnumProperty.cpp b/Engine/Source/Runtime/CoreUObject/Private/UObject/EnumProperty.cpp
--- a/Engine/Source/Runtime/CoreUObject/Private/UObject/EnumProperty.cpp
+++ b/Engine/Source/Runtime/CoreUObject/Private/UObject/EnumProperty.cpp
@@ -526,7 +526,9 @@ bool FEnumProperty::LoadFromTag(const FPropertyTag& Tag)
 	}
 
 	// Update FByteProperty when making changes here.
+	PRAGMA_DISABLE_DEPRECATION_WARNINGS;
 	TStringBuilder<256> EnumName(InPlace, Tag.EnumName);
+	PRAGMA_ENABLE_DEPRECATION_WARNINGS;
 	if (UEnum* LocalEnum = FindFirstObject<UEnum>(*EnumName, EFindFirstObjectOptions::NativeFirst))
 	{
 		Enum = LocalEnum;
@@ -545,6 +547,7 @@ void FEnumProperty::SaveToTag(FPropertyTag& Tag)
 
 	if (const UEnum* LocalEnum = Enum)
 	{
+		PRAGMA_DISABLE_DEPRECATION_WARNINGS;
 		// RobM: Ugly hack so that we can avoid content changes in most of the packages
 		// Update FByteProperty when making changes here.
 		if (LocalEnum->GetPackage()->HasAnyPackageFlags(PKG_CompiledIn))
@@ -555,6 +558,7 @@ void FEnumProperty::SaveToTag(FPropertyTag& Tag)
 		{
 			Tag.EnumName = FName(*LocalEnum->GetPathName());
 		}
+		PRAGMA_ENABLE_DEPRECATION_WARNINGS;
 	}
 }
 
diff --git a/Engine/Source/Runtime/CoreUObject/Private/UObject/PropertyArray.cpp b/Engine/Source/Runtime/CoreUObject/Private/UObject/PropertyArray.cpp
--- a/Engine/Source/Runtime/CoreUObject/Private/UObject/PropertyArray.cpp
+++ b/Engine/Source/Runtime/CoreUObject/Private/UObject/PropertyArray.cpp
@@ -1281,6 +1281,7 @@ bool FArrayProperty::LoadFromTag(const FPropertyTag& Tag)
 		return false;
 	}
 
+	PRAGMA_DISABLE_DEPRECATION_WARNINGS;
 	FField* Field = FField::TryConstruct(Tag.InnerType, {}, Tag.Name, RF_NoFlags);
 	if (FProperty* Property = CastField<FProperty>(Field))
 	{
@@ -1299,6 +1300,7 @@ bool FArrayProperty::LoadFromTag(const FPropertyTag& Tag)
 			return true;
 		}
 	}
+	PRAGMA_ENABLE_DEPRECATION_WARNINGS;
 	delete Field;
 	return false;
 }
@@ -1309,7 +1311,9 @@ void FArrayProperty::SaveToTag(FPropertyTag& Tag)
 
 	const FProperty* LocalInner = Inner;
 	check(LocalInner);
+	PRAGMA_DISABLE_DEPRECATION_WARNINGS;
 	Tag.InnerType = LocalInner->GetID();
+	PRAGMA_ENABLE_DEPRECATION_WARNINGS;
 }
 
 bool FArrayProperty::LoadTypeName(UE::FPropertyTypeName Type, const FPropertyTag* Tag)
diff --git a/Engine/Source/Runtime/CoreUObject/Private/UObject/PropertyByte.cpp b/Engine/Source/Runtime/CoreUObject/Private/UObject/PropertyByte.cpp
--- a/Engine/Source/Runtime/CoreUObject/Private/UObject/PropertyByte.cpp
+++ b/Engine/Source/Runtime/CoreUObject/Private/UObject/PropertyByte.cpp
@@ -523,6 +523,7 @@ bool FByteProperty::LoadFromTag(const FPropertyTag& Tag)
 		return false;
 	}
 
+	PRAGMA_DISABLE_DEPRECATION_WARNINGS;
 	if (Tag.EnumName.IsNone())
 	{
 		return true;
@@ -530,6 +531,7 @@ bool FByteProperty::LoadFromTag(const FPropertyTag& Tag)
 
 	// Update FEnumProperty when making changes here.
 	TStringBuilder<256> EnumName(InPlace, Tag.EnumName);
+	PRAGMA_ENABLE_DEPRECATION_WARNINGS;
 	if (UEnum* LocalEnum = FindFirstObject<UEnum>(*EnumName, EFindFirstObjectOptions::NativeFirst))
 	{
 		Enum = LocalEnum;
@@ -545,6 +547,7 @@ void FByteProperty::SaveToTag(FPropertyTag& Tag)
 
 	if (const UEnum* LocalEnum = Enum)
 	{
+		PRAGMA_DISABLE_DEPRECATION_WARNINGS;
 		// RobM: Ugly hack so that we can avoid content changes in most of the packages
 		// Update FEnumProperty when making changes here.
 		if (LocalEnum->GetPackage()->HasAnyPackageFlags(PKG_CompiledIn))
@@ -555,6 +558,7 @@ void FByteProperty::SaveToTag(FPropertyTag& Tag)
 		{
 			Tag.EnumName = FName(*LocalEnum->GetPathName());
 		}
+		PRAGMA_ENABLE_DEPRECATION_WARNINGS;
 	}
 }
 
diff --git a/Engine/Source/Runtime/CoreUObject/Private/UObject/PropertyMap.cpp b/Engine/Source/Runtime/CoreUObject/Private/UObject/PropertyMap.cpp
--- a/Engine/Source/Runtime/CoreUObject/Private/UObject/PropertyMap.cpp
+++ b/Engine/Source/Runtime/CoreUObject/Private/UObject/PropertyMap.cpp
@@ -1732,8 +1732,10 @@ void FMapProperty::SaveToTag(FPropertyTag& Tag)
 	const FProperty* LocalValueProp = ValueProp;
 	check(LocalKeyProp);
 	check(LocalValueProp);
+	PRAGMA_DISABLE_DEPRECATION_WARNINGS;
 	Tag.InnerType = LocalKeyProp->GetID();
 	Tag.ValueType = LocalValueProp->GetID();
+	PRAGMA_ENABLE_DEPRECATION_WARNINGS;
 }
 
 bool FMapProperty::LoadTypeName(UE::FPropertyTypeName Type, const FPropertyTag* Tag)
diff --git a/Engine/Source/Runtime/CoreUObject/Private/UObject/PropertyOptional.cpp b/Engine/Source/Runtime/CoreUObject/Private/UObject/PropertyOptional.cpp
--- a/Engine/Source/Runtime/CoreUObject/Private/UObject/PropertyOptional.cpp
+++ b/Engine/Source/Runtime/CoreUObject/Private/UObject/PropertyOptional.cpp
@@ -583,7 +583,9 @@ void FOptionalProperty::SaveToTag(FPropertyTag& Tag)
 
 	const FProperty* LocalValueProperty = ValueProperty;
 	check(LocalValueProperty);
+	PRAGMA_DISABLE_DEPRECATION_WARNINGS;
 	Tag.InnerType = LocalValueProperty->GetID();
+	PRAGMA_ENABLE_DEPRECATION_WARNINGS;
 }
 
 bool FOptionalProperty::LoadTypeName(UE::FPropertyTypeName Type, const FPropertyTag* Tag)
diff --git a/Engine/Source/Runtime/CoreUObject/Private/UObject/PropertySet.cpp b/Engine/Source/Runtime/CoreUObject/Private/UObject/PropertySet.cpp
--- a/Engine/Source/Runtime/CoreUObject/Private/UObject/PropertySet.cpp
+++ b/Engine/Source/Runtime/CoreUObject/Private/UObject/PropertySet.cpp
@@ -1107,7 +1107,9 @@ void FSetProperty::SaveToTag(FPropertyTag& Tag)
 
 	const FProperty* LocalElementProp = ElementProp;
 	check(LocalElementProp);
+	PRAGMA_DISABLE_DEPRECATION_WARNINGS;
 	Tag.InnerType = LocalElementProp->GetID();
+	PRAGMA_ENABLE_DEPRECATION_WARNINGS;
 }
 
 bool FSetProperty::LoadTypeName(UE::FPropertyTypeName Type, const FPropertyTag* Tag)
diff --git a/Engine/Source/Runtime/CoreUObject/Private/UObject/PropertyStruct.cpp b/Engine/Source/Runtime/CoreUObject/Private/UObject/PropertyStruct.cpp
--- a/Engine/Source/Runtime/CoreUObject/Private/UObject/PropertyStruct.cpp
+++ b/Engine/Source/Runtime/CoreUObject/Private/UObject/PropertyStruct.cpp
@@ -483,6 +483,7 @@ bool FStructProperty::LoadFromTag(const FPropertyTag& Tag)
 		return false;
 	}
 
+	PRAGMA_DISABLE_DEPRECATION_WARNINGS;
 	if (!Tag.StructName.IsNone())
 	{
 		TStringBuilder<256> StructName(InPlace, Tag.StructName);
@@ -492,6 +493,7 @@ bool FStructProperty::LoadFromTag(const FPropertyTag& Tag)
 			return true;
 		}
 	}
+	PRAGMA_ENABLE_DEPRECATION_WARNINGS;
 
 	return false;
 }
@@ -503,8 +505,10 @@ void FStructProperty::SaveToTag(FPropertyTag& Tag)
 	const FString* OriginalType = FindOriginalTypeName(this);
 	const UScriptStruct* LocalStruct = Struct;
 	check(LocalStruct);
+	PRAGMA_DISABLE_DEPRECATION_WARNINGS;
 	Tag.StructName = OriginalType ? FName(**OriginalType) : LocalStruct->GetFName();
 	Tag.StructGuid = LocalStruct->GetCustomGuid();
+	PRAGMA_ENABLE_DEPRECATION_WARNINGS;
 }
 
 void FStructProperty::AssignToTag(FPropertyTag& Tag)
@@ -514,10 +518,12 @@ void FStructProperty::AssignToTag(FPropertyTag& Tag)
 	const FString* OriginalType = FindOriginalTypeName(this);
 	const UScriptStruct* LocalStruct = Struct;
 	check(LocalStruct);
+	PRAGMA_DISABLE_DEPRECATION_WARNINGS;
 	if (OriginalType && FName(**OriginalType) == Tag.StructName)
 	{
 		Tag.StructName = LocalStruct->GetFName();
 	}
+	PRAGMA_ENABLE_DEPRECATION_WARNINGS;
 }
 
 #if WITH_EDITORONLY_DATA
diff --git a/Engine/Source/Runtime/CoreUObject/Private/UObject/PropertyTag.cpp b/Engine/Source/Runtime/CoreUObject/Private/UObject/PropertyTag.cpp
--- a/Engine/Source/Runtime/CoreUObject/Private/UObject/PropertyTag.cpp
+++ b/Engine/Source/Runtime/CoreUObject/Private/UObject/PropertyTag.cpp
@@ -63,9 +63,9 @@ FPropertyTag::FPropertyTag(FProperty* Property, int32 InIndex, uint8* Value)
 
 void FPropertyTag::SetProperty(FProperty* Property)
 {
-	PRAGMA_DISABLE_DEPRECATION_WARNINGS;
+	PRAGMA_DISABLE_DEPRECATION_WARNINGS
 	Prop = Property;
-	PRAGMA_ENABLE_DEPRECATION_WARNINGS;
+	PRAGMA_ENABLE_DEPRECATION_WARNINGS
 }
 
 void FPropertyTag::SetPropertyGuid(const FGuid& InPropertyGuid)
@@ -79,6 +79,7 @@ void FPropertyTag::SetType(UE::FPropertyTypeName InFullType)
 	TypeName = InFullType;
 	Type = TypeName.GetName();
 
+	PRAGMA_DISABLE_DEPRECATION_WARNINGS
 	if (const EName* TagType = Type.ToEName(); TagType && Type.GetNumber() == NAME_NO_NUMBER_INTERNAL)
 	{
 		switch (*TagType)
@@ -108,6 +109,7 @@ void FPropertyTag::SetType(UE::FPropertyTypeName InFullType)
 			break;
 		}
 	}
+	PRAGMA_ENABLE_DEPRECATION_WARNINGS
 }
 
 static EPropertyTagExtension CalculatePropertyExtensionFlags(FArchive& UnderlyingArchive, FPropertyTag& Tag)
@@ -229,24 +231,32 @@ void operator<<(FStructuredArchive::FSlot Slot, FPropertyTag& Tag)
 		if (TagType == NAME_StructProperty)
 		{
 			TypeBuilder.BeginParameters();
-			Slot << SA_ATTRIBUTE(TEXT("StructName"), Tag.StructName);
-			TypeBuilder.AddName(Tag.StructName);
+			PRAGMA_DISABLE_DEPRECATION_WARNINGS
+			FName StructName = Tag.StructName;
+			FGuid StructGuid = Tag.StructGuid;
+			PRAGMA_ENABLE_DEPRECATION_WARNINGS
+			Slot << SA_ATTRIBUTE(TEXT("StructName"), StructName);
+			TypeBuilder.AddName(StructName);
 			if (Version >= VER_UE4_STRUCT_GUID_IN_PROPERTY_TAG)
 			{
 				if (bIsTextFormat)
 				{
-					Slot << SA_OPTIONAL_ATTRIBUTE(TEXT("StructGuid"), Tag.StructGuid, FGuid());
+					Slot << SA_OPTIONAL_ATTRIBUTE(TEXT("StructGuid"), StructGuid, FGuid());
 				}
 				else
 				{
-					Slot << SA_ATTRIBUTE(TEXT("StructGuid"), Tag.StructGuid);
+					Slot << SA_ATTRIBUTE(TEXT("StructGuid"), StructGuid);
 				}
-				if (Tag.StructGuid.IsValid())
+				if (StructGuid.IsValid())
 				{
-					TypeBuilder.AddGuid(Tag.StructGuid);
+					TypeBuilder.AddGuid(StructGuid);
 				}
 			}
 			TypeBuilder.EndParameters();
+			PRAGMA_DISABLE_DEPRECATION_WARNINGS
+			Tag.StructName = StructName;
+			Tag.StructGuid = StructGuid;
+			PRAGMA_ENABLE_DEPRECATION_WARNINGS
 		}
 		// only need to serialize this for bools
 		else if (TagType == NAME_BoolProperty && !UnderlyingArchive.IsTextFormat())
@@ -264,65 +274,103 @@ void operator<<(FStructuredArchive::FSlot Slot, FPropertyTag& Tag)
 		// only need to serialize this for bytes/enums
 		else if (TagType == NAME_ByteProperty)
 		{
+			PRAGMA_DISABLE_DEPRECATION_WARNINGS
+			FName EnumName = Tag.EnumName;
+			PRAGMA_ENABLE_DEPRECATION_WARNINGS
 			if (UnderlyingArchive.IsTextFormat())
 			{
-				Slot << SA_OPTIONAL_ATTRIBUTE(TEXT("EnumName"), Tag.EnumName, NAME_None);
+				Slot << SA_OPTIONAL_ATTRIBUTE(TEXT("EnumName"), EnumName, NAME_None);
 			}
 			else
 			{
-				Slot << SA_ATTRIBUTE(TEXT("EnumName"), Tag.EnumName);
+				Slot << SA_ATTRIBUTE(TEXT("EnumName"), EnumName);
 			}
-			if (!Tag.EnumName.IsNone())
+			if (!EnumName.IsNone())
 			{
 				TypeBuilder.BeginParameters();
-				AddEnumPath(TypeBuilder, Tag.EnumName);
+				AddEnumPath(TypeBuilder, EnumName);
 				TypeBuilder.EndParameters();
 			}
+			PRAGMA_DISABLE_DEPRECATION_WARNINGS
+			Tag.EnumName = EnumName;
+			PRAGMA_ENABLE_DEPRECATION_WARNINGS
 		}
 		else if (TagType == NAME_EnumProperty)
 		{
-			Slot << SA_ATTRIBUTE(TEXT("EnumName"), Tag.EnumName);
+			PRAGMA_DISABLE_DEPRECATION_WARNINGS
+			FName EnumName = Tag.EnumName;
+			PRAGMA_ENABLE_DEPRECATION_WARNINGS
+			Slot << SA_ATTRIBUTE(TEXT("EnumName"), EnumName);
 			TypeBuilder.BeginParameters();
-			AddEnumPath(TypeBuilder, Tag.EnumName);
+			AddEnumPath(TypeBuilder, EnumName);
 			TypeBuilder.AddName(NAME_ByteProperty);
 			TypeBuilder.EndParameters();
+			PRAGMA_DISABLE_DEPRECATION_WARNINGS
+			Tag.EnumName = EnumName;
+			PRAGMA_ENABLE_DEPRECATION_WARNINGS
 		}
 		// need to serialize the InnerType for arrays
 		else if (TagType == NAME_ArrayProperty)
 		{
+			PRAGMA_DISABLE_DEPRECATION_WARNINGS
+			FName InnerType = Tag.InnerType;
+			PRAGMA_ENABLE_DEPRECATION_WARNINGS
 			if (Version >= VAR_UE4_ARRAY_PROPERTY_INNER_TAGS)
 			{
-				Slot << SA_ATTRIBUTE(TEXT("InnerType"), Tag.InnerType);
+				Slot << SA_ATTRIBUTE(TEXT("InnerType"), InnerType);
 			}
 			TypeBuilder.BeginParameters();
-			TypeBuilder.AddName(Tag.InnerType);
+			TypeBuilder.AddName(InnerType);
 			TypeBuilder.EndParameters();
+			PRAGMA_DISABLE_DEPRECATION_WARNINGS
+			Tag.InnerType = InnerType;
+			PRAGMA_ENABLE_DEPRECATION_WARNINGS
 		}
 		// need to serialize the InnerType for optionals.
 		else if (TagType == NAME_OptionalProperty)
 		{
-			Slot << SA_ATTRIBUTE(TEXT("InnerType"), Tag.InnerType);
+			PRAGMA_DISABLE_DEPRECATION_WARNINGS
+			FName InnerType = Tag.InnerType;
+			PRAGMA_ENABLE_DEPRECATION_WARNINGS
+			Slot << SA_ATTRIBUTE(TEXT("InnerType"), InnerType);
 			TypeBuilder.BeginParameters();
-			TypeBuilder.AddName(Tag.InnerType);
+			TypeBuilder.AddName(InnerType);
 			TypeBuilder.EndParameters();
+			PRAGMA_DISABLE_DEPRECATION_WARNINGS
+			Tag.InnerType = InnerType;
+			PRAGMA_ENABLE_DEPRECATION_WARNINGS
 		}
 		else if (Version >= VER_UE4_PROPERTY_TAG_SET_MAP_SUPPORT)
 		{
 			if (TagType == NAME_SetProperty)
 			{
-				Slot << SA_ATTRIBUTE(TEXT("InnerType"), Tag.InnerType);
+				PRAGMA_DISABLE_DEPRECATION_WARNINGS
+				FName InnerType = Tag.InnerType;
+				PRAGMA_ENABLE_DEPRECATION_WARNINGS
+				Slot << SA_ATTRIBUTE(TEXT("InnerType"), InnerType);
 				TypeBuilder.BeginParameters();
-				TypeBuilder.AddName(Tag.InnerType);
+				TypeBuilder.AddName(InnerType);
 				TypeBuilder.EndParameters();
+				PRAGMA_DISABLE_DEPRECATION_WARNINGS
+				Tag.InnerType = InnerType;
+				PRAGMA_ENABLE_DEPRECATION_WARNINGS
 			}
 			else if (TagType == NAME_MapProperty)
 			{
-				Slot << SA_ATTRIBUTE(TEXT("InnerType"), Tag.InnerType);
-				Slot << SA_ATTRIBUTE(TEXT("ValueType"), Tag.ValueType);
+				PRAGMA_DISABLE_DEPRECATION_WARNINGS
+				FName InnerType = Tag.InnerType;
+				FName ValueType = Tag.ValueType;
+				PRAGMA_ENABLE_DEPRECATION_WARNINGS
+				Slot << SA_ATTRIBUTE(TEXT("InnerType"), InnerType);
+				Slot << SA_ATTRIBUTE(TEXT("ValueType"), ValueType);
 				TypeBuilder.BeginParameters();
-				TypeBuilder.AddName(Tag.InnerType);
-				TypeBuilder.AddName(Tag.ValueType);
+				TypeBuilder.AddName(InnerType);
+				TypeBuilder.AddName(ValueType);
 				TypeBuilder.EndParameters();
+				PRAGMA_DISABLE_DEPRECATION_WARNINGS
+				Tag.InnerType = InnerType;
+				Tag.ValueType = ValueType;
+				PRAGMA_ENABLE_DEPRECATION_WARNINGS
 			}
 		}
 
diff --git a/Engine/Source/Runtime/CoreUObject/Public/UObject/PropertyTag.h b/Engine/Source/Runtime/CoreUObject/Public/UObject/PropertyTag.h
--- a/Engine/Source/Runtime/CoreUObject/Public/UObject/PropertyTag.h
+++ b/Engine/Source/Runtime/CoreUObject/Public/UObject/PropertyTag.h
@@ -22,6 +22,7 @@ class FProperty;
 struct FPropertyTag
 {
 	// Transient.
+	UE_DEPRECATED(5.4, "Prop has been deprecated. Use GetProperty()/SetProperty().")
 	FProperty* Prop = nullptr;
 
 	// Variables.
@@ -31,13 +32,18 @@ private:
 public:
 	FName	Type;		// Type of property
 	FName	Name;		// Name of property.
+	UE_DEPRECATED(5.4, "StructName has been deprecated and replaced by GetType(). Use GetType().IsStruct(StructName) to test, otherwise if Type is StructProperty then StructName is GetType().GetParameterName(0).")
 	FName	StructName;	// Struct name if FStructProperty.
+	UE_DEPRECATED(5.4, "EnumName has been deprecated and replaced by GetType(). Use GetType().IsEnum(EnumName) to test, otherwise if Type is EnumProperty or ByteProperty then EnumName is GetType().GetParameterName(0).")
 	FName	EnumName;	// Enum name if FByteProperty or FEnumProperty
+	UE_DEPRECATED(5.4, "InnerName has been deprecated and replaced by GetType(). If Type is ArrayProperty/SetProperty/MapProperty/OptionalProperty then InnerType is GetType().GetParameterName(0).")
 	FName	InnerType;	// Inner type if FArrayProperty, FSetProperty, FMapProperty, or OptionalProperty
+	UE_DEPRECATED(5.4, "ValueName has been deprecated and replaced by GetType(). If Type is MapProperty then ValueType is GetType().GetParameterName(1).")
 	FName	ValueType;	// Value type if UMapPropery
 	int32	Size = 0;   // Property size.
 	int32	ArrayIndex = INDEX_NONE; // Index if an array; else 0.
 	int64	SizeOffset = INDEX_NONE; // location in stream of tag size member
+	UE_DEPRECATED(5.4, "StructGuid has been deprecated and replaced by GetType(). If Type is StructProperty then StructGuid is GetType().GetParameterName(1).")
 	FGuid	StructGuid;
 	FGuid	PropertyGuid;
 	uint8	HasPropertyGuid = 0;
diff --git a/Engine/Source/Runtime/Engine/Private/Animation/AnimNode_StateResult.cpp b/Engine/Source/Runtime/Engine/Private/Animation/AnimNode_StateResult.cpp
--- a/Engine/Source/Runtime/Engine/Private/Animation/AnimNode_StateResult.cpp
+++ b/Engine/Source/Runtime/Engine/Private/Animation/AnimNode_StateResult.cpp
@@ -6,7 +6,7 @@
 
 bool FAnimNode_StateResult::SerializeFromMismatchedTag(const FPropertyTag& Tag, FStructuredArchive::FSlot Slot)
 {
-	if(Tag.Type == NAME_StructProperty && Tag.StructName == FAnimNode_Root::StaticStruct()->GetFName())
+	if (Tag.GetType().IsStruct(FAnimNode_Root::StaticStruct()->GetFName()))
 	{
 		FAnimNode_Root OldValue;
 		FAnimNode_Root::StaticStruct()->SerializeItem(Slot, &OldValue, nullptr);
diff --git a/Engine/Source/Runtime/Engine/Private/EngineTypes.cpp b/Engine/Source/Runtime/Engine/Private/EngineTypes.cpp
--- a/Engine/Source/Runtime/Engine/Private/EngineTypes.cpp
+++ b/Engine/Source/Runtime/Engine/Private/EngineTypes.cpp
@@ -195,7 +195,7 @@ UActorComponent* FSoftComponentReference::GetComponent(AActor* OwningActor) cons
 bool FSoftComponentReference::SerializeFromMismatchedTag(const FPropertyTag& Tag, FStructuredArchive::FSlot Slot)
 {
 	static const FName ComponentReferenceContextName("ComponentReference");
-	if (Tag.Type == NAME_StructProperty && Tag.StructName == ComponentReferenceContextName)
+	if (Tag.GetType().IsStruct(ComponentReferenceContextName))
 	{
 		FComponentReference Reference;
 		FComponentReference::StaticStruct()->SerializeItem(Slot, &Reference, nullptr);
diff --git a/Engine/Source/Runtime/Engine/Private/Materials/MaterialExpressions.cpp b/Engine/Source/Runtime/Engine/Private/Materials/MaterialExpressions.cpp
--- a/Engine/Source/Runtime/Engine/Private/Materials/MaterialExpressions.cpp
+++ b/Engine/Source/Runtime/Engine/Private/Materials/MaterialExpressions.cpp
@@ -16261,8 +16261,7 @@ bool FMaterialLayersFunctionsRuntimeData::SerializeFromMismatchedTag(const FProp
 #if WITH_EDITORONLY_DATA
 	static const FName MaterialLayersFunctionsName("MaterialLayersFunctions");
 	static const FName MaterialLayersPropertyName("MaterialLayers");
-	if (Tag.Type == NAME_StructProperty &&
-		Tag.StructName == MaterialLayersFunctionsName &&
+	if (Tag.GetType().IsStruct(MaterialLayersFunctionsName) &&
 		Tag.Name == MaterialLayersPropertyName)
 	{
 		FMaterialLayersFunctions LocalMaterialLayers;
diff --git a/Engine/Source/Runtime/InputCore/Private/InputCoreTypes.cpp b/Engine/Source/Runtime/InputCore/Private/InputCoreTypes.cpp
--- a/Engine/Source/Runtime/InputCore/Private/InputCoreTypes.cpp
+++ b/Engine/Source/Runtime/InputCore/Private/InputCoreTypes.cpp
@@ -1432,7 +1432,7 @@ void FKey::ConditionalLookupKeyDetails() const
 
 bool FKey::SerializeFromMismatchedTag(struct FPropertyTag const& Tag, FStructuredArchive::FSlot Slot)
 {
-	if (Tag.Type == NAME_ByteProperty && Tag.EnumName == TEXT("EKeys"))
+	if (Tag.GetType().IsEnum(TEXT("EKeys")))
 	{
 		Slot << KeyName;
 		const FString KeyNameString(KeyName.ToString());
diff --git a/Engine/Source/Runtime/LevelSequence/Private/LevelSequenceBindingReference.cpp b/Engine/Source/Runtime/LevelSequence/Private/LevelSequenceBindingReference.cpp
--- a/Engine/Source/Runtime/LevelSequence/Private/LevelSequenceBindingReference.cpp
+++ b/Engine/Source/Runtime/LevelSequence/Private/LevelSequenceBindingReference.cpp
@@ -191,8 +191,7 @@ bool FUpgradedLevelSequenceBindingReferences::SerializeFromMismatchedTag(const F
 {
 	using namespace UE::UniversalObjectLocator;
 
-	if (Tag.Type != NAME_StructProperty
-		|| Tag.StructName != FLevelSequenceBindingReferences::StaticStruct()->GetFName())
+	if (!Tag.GetType().IsStruct(FLevelSequenceBindingReferences::StaticStruct()->GetFName()))
 	{
 		return false;
 	}
diff --git a/Engine/Source/Runtime/MovieScene/Private/Channels/MovieSceneBoolChannel.cpp b/Engine/Source/Runtime/MovieScene/Private/Channels/MovieSceneBoolChannel.cpp
--- a/Engine/Source/Runtime/MovieScene/Private/Channels/MovieSceneBoolChannel.cpp
+++ b/Engine/Source/Runtime/MovieScene/Private/Channels/MovieSceneBoolChannel.cpp
@@ -12,7 +12,7 @@
 bool FMovieSceneBoolChannel::SerializeFromMismatchedTag(const FPropertyTag& Tag, FStructuredArchive::FSlot Slot)
 {
 	static const FName IntegralCurveName("IntegralCurve");
-	if (Tag.Type == NAME_StructProperty && Tag.StructName == IntegralCurveName)
+	if (Tag.GetType().IsStruct(IntegralCurveName))
 	{
 		FIntegralCurve IntegralCurve;
 		FIntegralCurve::StaticStruct()->SerializeItem(Slot, &IntegralCurve, nullptr);
diff --git a/Engine/Source/Runtime/MovieScene/Private/Channels/MovieSceneByteChannel.cpp b/Engine/Source/Runtime/MovieScene/Private/Channels/MovieSceneByteChannel.cpp
--- a/Engine/Source/Runtime/MovieScene/Private/Channels/MovieSceneByteChannel.cpp
+++ b/Engine/Source/Runtime/MovieScene/Private/Channels/MovieSceneByteChannel.cpp
@@ -12,7 +12,7 @@
 bool FMovieSceneByteChannel::SerializeFromMismatchedTag(const FPropertyTag& Tag, FStructuredArchive::FSlot Slot)
 {
 	static const FName IntegralCurveName("IntegralCurve");
-	if (Tag.Type == NAME_StructProperty && Tag.StructName == IntegralCurveName)
+	if (Tag.GetType().IsStruct(IntegralCurveName))
 	{
 		FIntegralCurve IntegralCurve;
 		FIntegralCurve::StaticStruct()->SerializeItem(Slot, &IntegralCurve, nullptr);
diff --git a/Engine/Source/Runtime/MovieScene/Private/Channels/MovieSceneCurveChannelImpl.cpp b/Engine/Source/Runtime/MovieScene/Private/Channels/MovieSceneCurveChannelImpl.cpp
--- a/Engine/Source/Runtime/MovieScene/Private/Channels/MovieSceneCurveChannelImpl.cpp
+++ b/Engine/Source/Runtime/MovieScene/Private/Channels/MovieSceneCurveChannelImpl.cpp
@@ -1556,7 +1556,7 @@ bool TMovieSceneCurveChannelImpl<ChannelType>::SerializeFromRichCurve(ChannelTyp
 
 	check(InChannel);
 
-	if (Tag.Type == NAME_StructProperty && Tag.StructName == RichCurveName)
+	if (Tag.GetType().IsStruct(RichCurveName))
 	{
 		FRichCurve RichCurve;
 		FRichCurve::StaticStruct()->SerializeItem(Slot, &RichCurve, nullptr);
diff --git a/Engine/Source/Runtime/MovieScene/Private/Channels/MovieSceneDoubleChannel.cpp b/Engine/Source/Runtime/MovieScene/Private/Channels/MovieSceneDoubleChannel.cpp
--- a/Engine/Source/Runtime/MovieScene/Private/Channels/MovieSceneDoubleChannel.cpp
+++ b/Engine/Source/Runtime/MovieScene/Private/Channels/MovieSceneDoubleChannel.cpp
@@ -261,7 +261,7 @@ bool FMovieSceneDoubleChannel::SerializeFromMismatchedTag(const FPropertyTag& Ta
 
 	// Load pre-LWC content that was saved with a float channel.
 	static const FName FloatChannelName("MovieSceneFloatChannel");
-	if (Tag.Type == NAME_StructProperty && Tag.StructName == FloatChannelName)
+	if (Tag.GetType().IsStruct(FloatChannelName))
 	{
 		// We have to load the whole structure into a float channel, and then convert it into our data.
 		// It's not ideal but it's the safest way to make it work.
diff --git a/Engine/Source/Runtime/MovieScene/Private/Channels/MovieSceneIntegerChannel.cpp b/Engine/Source/Runtime/MovieScene/Private/Channels/MovieSceneIntegerChannel.cpp
--- a/Engine/Source/Runtime/MovieScene/Private/Channels/MovieSceneIntegerChannel.cpp
+++ b/Engine/Source/Runtime/MovieScene/Private/Channels/MovieSceneIntegerChannel.cpp
@@ -13,7 +13,7 @@
 bool FMovieSceneIntegerChannel::SerializeFromMismatchedTag(const FPropertyTag& Tag, FStructuredArchive::FSlot Slot)
 {
 	static const FName IntegralCurveName("IntegralCurve");
-	if (Tag.Type == NAME_StructProperty && Tag.StructName == IntegralCurveName)
+	if (Tag.GetType().IsStruct(IntegralCurveName))
 	{
 		FIntegralCurve IntegralCurve;
 		FIntegralCurve::StaticStruct()->SerializeItem(Slot, &IntegralCurve, nullptr);
diff --git a/Engine/Source/Runtime/MovieScene/Private/MovieSceneFrameMigration.cpp b/Engine/Source/Runtime/MovieScene/Private/MovieSceneFrameMigration.cpp
--- a/Engine/Source/Runtime/MovieScene/Private/MovieSceneFrameMigration.cpp
+++ b/Engine/Source/Runtime/MovieScene/Private/MovieSceneFrameMigration.cpp
@@ -36,7 +36,7 @@ TRange<FFrameNumber> FMovieSceneFrameRange::FromFloatRange(const TRange<float>&
 bool FMovieSceneFrameRange::SerializeFromMismatchedTag(FPropertyTag const& Tag, FStructuredArchive::FSlot Slot)
 {
 	static const FName NAME_FloatRange("FloatRange");
-	if (Tag.Type == NAME_StructProperty && Tag.StructName == NAME_FloatRange)
+	if (Tag.GetType().IsStruct(NAME_FloatRange))
 	{
 		UScriptStruct* FloatStruct = TBaseStructure<FFloatRange>::Get();
 
diff --git a/Engine/Source/Runtime/MovieScene/Private/MovieSceneSequencePlayer.cpp b/Engine/Source/Runtime/MovieScene/Private/MovieSceneSequencePlayer.cpp
--- a/Engine/Source/Runtime/MovieScene/Private/MovieSceneSequencePlayer.cpp
+++ b/Engine/Source/Runtime/MovieScene/Private/MovieSceneSequencePlayer.cpp
@@ -81,7 +81,7 @@ bool FMovieSceneSequenceLoopCount::SerializeFromMismatchedTag( const FPropertyTa
 
 bool FMovieSceneSequencePlaybackSettings::SerializeFromMismatchedTag( const FPropertyTag& Tag, FStructuredArchive::FSlot Slot )
 {
-	if (Tag.Type == NAME_StructProperty && Tag.StructName == "LevelSequencePlaybackSettings")
+	if (Tag.GetType().IsStruct("LevelSequencePlaybackSettings"))
 	{
 		StaticStruct()->SerializeItem(Slot, this, nullptr);
 		return true;
diff --git a/Engine/Source/Runtime/MovieSceneTracks/Private/Channels/MovieSceneStringChannel.cpp b/Engine/Source/Runtime/MovieSceneTracks/Private/Channels/MovieSceneStringChannel.cpp
--- a/Engine/Source/Runtime/MovieSceneTracks/Private/Channels/MovieSceneStringChannel.cpp
+++ b/Engine/Source/Runtime/MovieSceneTracks/Private/Channels/MovieSceneStringChannel.cpp
@@ -22,7 +22,7 @@ const FString* FMovieSceneStringChannel::Evaluate(FFrameTime InTime) const
 bool FMovieSceneStringChannel::SerializeFromMismatchedTag(const FPropertyTag& Tag, FStructuredArchive::FSlot Slot)
 {
 	static const FName StringCurveName("StringCurve");
-	if (Tag.Type == NAME_StructProperty && Tag.StructName == StringCurveName)
+	if (Tag.GetType().IsStruct(StringCurveName))
 	{
 		FStringCurve StringCurve;
 		FStringCurve::StaticStruct()->SerializeItem(Slot, &StringCurve, nullptr);
diff --git a/Engine/Source/Runtime/SlateCore/Private/Styling/SlateColor.cpp b/Engine/Source/Runtime/SlateCore/Private/Styling/SlateColor.cpp
--- a/Engine/Source/Runtime/SlateCore/Private/Styling/SlateColor.cpp
+++ b/Engine/Source/Runtime/SlateCore/Private/Styling/SlateColor.cpp
@@ -9,24 +9,21 @@
 
 bool FSlateColor::SerializeFromMismatchedTag(const FPropertyTag& Tag, FStructuredArchive::FSlot Slot)
 {
-	if (Tag.Type == NAME_StructProperty)
+	if (Tag.GetType().IsStruct(NAME_Color))
 	{
-		if (Tag.StructName == NAME_Color)
-		{
-			FColor OldColor;
-			Slot << OldColor;
-			*this = FSlateColor(FLinearColor(OldColor));
-
-			return true;
-		}
-		else if(Tag.StructName == NAME_LinearColor)
-		{
-			FLinearColor OldColor;
-			Slot << OldColor;
-			*this = FSlateColor(OldColor);
-
-			return true;
-		}
+		FColor OldColor;
+		Slot << OldColor;
+		*this = FSlateColor(FLinearColor(OldColor));
+
+		return true;
+	}
+	else if (Tag.GetType().IsStruct(NAME_LinearColor))
+	{
+		FLinearColor OldColor;
+		Slot << OldColor;
+		*this = FSlateColor(OldColor);
+
+		return true;
 	}
 
 	return false;
diff --git a/Engine/Source/Runtime/SlateCore/Private/Types/SlateVector2.cpp b/Engine/Source/Runtime/SlateCore/Private/Types/SlateVector2.cpp
--- a/Engine/Source/Runtime/SlateCore/Private/Types/SlateVector2.cpp
+++ b/Engine/Source/Runtime/SlateCore/Private/Types/SlateVector2.cpp
@@ -38,21 +38,21 @@ bool FDeprecateSlateVector2D::Serialize(FStructuredArchive::FSlot Slot)
 
 bool FDeprecateSlateVector2D::SerializeFromMismatchedTag(const struct FPropertyTag& Tag, FStructuredArchive::FSlot Slot)
 {
-	if (Tag.StructName == NAME_Vector2d)
+	if (Tag.GetType().IsStruct(NAME_Vector2d))
 	{
 		FVector2d OldVector;
 		Slot << OldVector;
 		*this = UE::Slate::CastToVector2f(OldVector);
 		return true;
 	}
-	else if (Tag.StructName == NAME_Vector2f)
+	else if (Tag.GetType().IsStruct(NAME_Vector2f))
 	{
 		FVector2f OldVector;
 		Slot << OldVector;
 		*this = OldVector;
 		return true;
 	}
-	else if (Tag.StructName == NAME_Vector2D)
+	else if (Tag.GetType().IsStruct(NAME_Vector2D))
 	{
 		if (Slot.GetUnderlyingArchive().UEVer() < EUnrealEngineObjectUE5Version::LARGE_WORLD_COORDINATES)
 		{
diff --git a/Engine/Source/Runtime/UniversalObjectLocator/Private/UniversalObjectLocator.cpp b/Engine/Source/Runtime/UniversalObjectLocator/Private/UniversalObjectLocator.cpp
--- a/Engine/Source/Runtime/UniversalObjectLocator/Private/UniversalObjectLocator.cpp
+++ b/Engine/Source/Runtime/UniversalObjectLocator/Private/UniversalObjectLocator.cpp
@@ -639,7 +639,7 @@ bool FUniversalObjectLocator::SerializeFromMismatchedTag(const FPropertyTag& Tag
 		Fragments.Emplace(TUniversalObjectLocatorFragment<FDirectPathObjectLocator>(OldProperty.ToSoftObjectPath()));
 		return true;
 	}
-	else if (Tag.Type == NAME_StructProperty && Tag.StructName == NAME_SoftObjectPath)
+	else if (Tag.GetType().IsStruct(NAME_SoftObjectPath))
 	{
 		FSoftObjectPath OldPath;
 		Slot << OldPath;
