Fixed use of a deprecated variable in PropertyBagTest

#rnx

diff --git a/Engine/Source/Runtime/CoreUObject/Tests/PropertyBagTest.cpp b/Engine/Source/Runtime/CoreUObject/Tests/PropertyBagTest.cpp
--- a/Engine/Source/Runtime/CoreUObject/Tests/PropertyBagTest.cpp
+++ b/Engine/Source/Runtime/CoreUObject/Tests/PropertyBagTest.cpp
@@ -286,7 +286,7 @@ namespace PropertyBagTestUtils
 		TArray<uint8> Buffer = EncodeToBuffer(InValue);
 
 		FPropertyTag Tag;
-		Tag.Prop = Prop;
+		Tag.SetProperty(Prop);
 		Tag.Name = FName(TEXT("TagTmp"));
 
 		LoadDataByTag(Dst, Path, Tag, Buffer);
