Fix the fact that not all instanced properties are referencing instanced objects.

diff --git a/Engine/Source/Runtime/CoreUObject/Private/UObject/SavePackage/PackageHarvester.cpp b/Engine/Source/Runtime/CoreUObject/Private/UObject/SavePackage/PackageHarvester.cpp
--- a/Engine/Source/Runtime/CoreUObject/Private/UObject/SavePackage/PackageHarvester.cpp
+++ b/Engine/Source/Runtime/CoreUObject/Private/UObject/SavePackage/PackageHarvester.cpp
@@ -636,7 +636,12 @@ void FPackageHarvester::TryHarvestExportInternal(UObject* InObject)
 		FFindInstancedReferenceSubobjectHelper::GetInstancedSubObjects(InObject, InstancedSubObjects);
 		for (UObject* InstancedSubObject : InstancedSubObjects)
 		{
-			TryHarvestExportInternal(InstancedSubObject);
+			checkf(InstancedSubObject, TEXT("Expecting valid subobject"));
+
+			if (InstancedSubObject->IsInPackage(SaveContext.GetPackage()))
+			{
+				TryHarvestExportInternal(InstancedSubObject);
+			}
 		}
 	}
 }
