Slightly refactored FEnumProperty::ImportText_Internal to make an upcoming change easier.

diff --git a/Engine/Source/Runtime/CoreUObject/Private/UObject/EnumProperty.cpp b/Engine/Source/Runtime/CoreUObject/Private/UObject/EnumProperty.cpp
--- a/Engine/Source/Runtime/CoreUObject/Private/UObject/EnumProperty.cpp
+++ b/Engine/Source/Runtime/CoreUObject/Private/UObject/EnumProperty.cpp
@@ -291,37 +291,37 @@ void FEnumProperty::ExportText_Internal(FString& ValueStr, const void* PropertyV
 	}
 }
 
-const TCHAR* FEnumProperty::ImportText_Internal(const TCHAR* InBuffer, void* ContainerOrPropertyPtr, EPropertyPointerType PropertyPointerType, UObject* Parent, int32 PortFlags, FOutputDevice* ErrorText) const
+const TCHAR* FEnumProperty::ImportText_Internal(const TCHAR* Buffer, void* ContainerOrPropertyPtr, EPropertyPointerType PropertyPointerType, UObject* Parent, int32 PortFlags, FOutputDevice* ErrorText) const
 {
 	check(Enum);
 	check(UnderlyingProp);
 	
 	if (!(PortFlags & PPF_ConsoleVariable))
 	{
+		int64 EnumValue = 0;
 		FString Temp;
-		if (const TCHAR* Buffer = FPropertyHelpers::ReadToken(InBuffer, Temp, true))
+		Buffer = FPropertyHelpers::ReadToken(Buffer, Temp, true);
+		if (Buffer)
 		{
 			int32 EnumIndex = Enum->GetIndexByName(*Temp, EGetByNameFlags::CheckAuthoredName);
 			if (EnumIndex == INDEX_NONE && (Temp.IsNumeric() && !Algo::Find(Temp, TEXT('.'))))
 			{
-				int64 EnumValue = INDEX_NONE;
-				LexFromString(EnumValue, *Temp);
-				EnumIndex = Enum->GetIndexByValue(EnumValue);
+				int64 LexedEnumValue = INDEX_NONE;
+				LexFromString(LexedEnumValue, *Temp);
+				EnumIndex = Enum->GetIndexByValue(LexedEnumValue);
 			}
 			if (EnumIndex != INDEX_NONE)
 			{
-				int64 EnumValue = Enum->GetValueByIndex(EnumIndex);
-				if (PropertyPointerType == EPropertyPointerType::Container && HasSetter())
-				{
-					SetValue_InContainer(ContainerOrPropertyPtr, &EnumValue);
-				}
-				else
-				{
-					UnderlyingProp->SetIntPropertyValue(PointerToValuePtr(ContainerOrPropertyPtr, PropertyPointerType), EnumValue);
-				}
-				return Buffer;
+				EnumValue = Enum->GetValueByIndex(EnumIndex);
 			}
+			else
+			{
+				Buffer = nullptr;
+			}
+		}
 
+		if (!Buffer)
+		{
 			// Enum could not be created from value. This indicates a bad value so
 			// return null so that the caller of ImportText can generate a more meaningful
 			// warning/error
@@ -338,11 +338,21 @@ const TCHAR* FEnumProperty::ImportText_Internal(const TCHAR* InBuffer, void* Con
 				bIsNativeOrLoaded ? TEXT("loaded") : TEXT("not loaded"));
 			return nullptr;
 		}
+
+		if (PropertyPointerType == EPropertyPointerType::Container && HasSetter())
+		{
+			SetValue_InContainer(ContainerOrPropertyPtr, &EnumValue);
+		}
+		else
+		{
+			UnderlyingProp->SetIntPropertyValue(PointerToValuePtr(ContainerOrPropertyPtr, PropertyPointerType), EnumValue);
+		}
+		return Buffer;
 	}
 
 	// UnderlyingProp has a 0 offset so we need to make sure we convert the container pointer to the actual value pointer
-	const TCHAR* Result = UnderlyingProp->ImportText_Internal(InBuffer, PointerToValuePtr(ContainerOrPropertyPtr, PropertyPointerType), EPropertyPointerType::Direct, Parent, PortFlags, ErrorText);
-	return Result;
+	Buffer = UnderlyingProp->ImportText_Internal(Buffer, PointerToValuePtr(ContainerOrPropertyPtr, PropertyPointerType), EPropertyPointerType::Direct, Parent, PortFlags, ErrorText);
+	return Buffer;
 }
 
 FString FEnumProperty::GetCPPMacroType(FString& ExtendedTypeText) const
