New engine version to serialize FPropertyTag more compactly and with complete property type names

diff --git a/Engine/Source/Runtime/Core/Public/UObject/ObjectVersion.h b/Engine/Source/Runtime/Core/Public/UObject/ObjectVersion.h
--- a/Engine/Source/Runtime/Core/Public/UObject/ObjectVersion.h
+++ b/Engine/Source/Runtime/Core/Public/UObject/ObjectVersion.h
@@ -17,20 +17,20 @@ class FCbWriter;
 // This is useful as enum vales cannot be seen by the preprocessor.
 #define PREPROCESSOR_ENUM_PROTECT(a) ((unsigned int)(a))
 
-/** 
+/**
  * Versioning
- * 
+ *
  * We now have two enums for global versioning, EUnrealEngineObjectUE5Version and EUnrealEngineObjectUE4Version
  * which can be updated independently of each other. This allows us to potentially make version changed in the UE4
  * code base in the future and still have packages saved with those changes to be compatible with the UE5 code base.
  * 
  * The struct FPackageFileVersion should be used to contain the version number rather than storing the UE4/UE5 specific
  * version numbers. It provides overloads for testing against both version enum types for convenience.
- * 
- * Note that when reporting versioning problems we often use FPackageFileVersion::ToValue which returns the highest 
+ *
+ * Note that when reporting versioning problems we often use FPackageFileVersion::ToValue which returns the highest
  * version valid in the structure. This allows us to keep reporting code simple and hide the fact that there are two
- * version numbers from end users. However this does mean that if a version mismatch is encountered due to a future 
- * extension of EUnrealEngineObjectUE4Version, then the error message given might look odd. Although given that we are 
+ * version numbers from end users. However this does mean that if a version mismatch is encountered due to a future
+ * extension of EUnrealEngineObjectUE4Version, then the error message given might look odd. Although given that we are
  * unlikely to be changing EUnrealEngineObjectUE4Version in the future the risk is deemed worth it.
  */
 
@@ -48,31 +48,30 @@ enum class EUnrealEngineObjectUE5Version : uint32
 	// Support stripping names that are not referenced from export data
 	NAMES_REFERENCED_FROM_EXPORT_DATA,
 
-	// Added a payload table of contents to the package summary 
+	// Added a payload table of contents to the package summary
 	PAYLOAD_TOC,
 
 	// Added data to identify references from and to optional package
 	OPTIONAL_RESOURCES,
-	
+
 	// Large world coordinates converts a number of core types to double components by default.
-	LARGE_WORLD_COORDINATES,       
+	LARGE_WORLD_COORDINATES,
 
 	// Remove package GUID from FObjectExport
 	REMOVE_OBJECT_EXPORT_PACKAGE_GUID,
 
 	// Add IsInherited to the FObjectExport entry
 	TRACK_OBJECT_EXPORT_IS_INHERITED,
-	
+
 	// Replace FName asset path in FSoftObjectPath with (package name, asset name) pair FTopLevelAssetPath
 	FSOFTOBJECTPATH_REMOVE_ASSET_PATH_FNAMES,
-	
 
 	// Add a soft object path list to the package summary for fast remap
 	ADD_SOFTOBJECTPATH_LIST,
 
 	// Added bulk/data resource table
 	DATA_RESOURCES,
-	
+
 	// Added script property serialization offset to export table entries for saved, versioned packages
 	SCRIPT_SERIALIZATION_OFFSET,
 
@@ -81,6 +80,9 @@ enum class EUnrealEngineObjectUE5Version : uint32
 	// Support for overridable logic in containers
 	PROPERTY_TAG_EXTENSION_AND_OVERRIDABLE_SERIALIZATION,
 
+	// Added property tag complete type name and serialization type
+	PROPERTY_TAG_COMPLETE_TYPE_NAME,
+
 	// -----<new versions can be added before this line>-------------------------------------------------
 	// - this needs to be the last line (see note below)
 	AUTOMATIC_VERSION_PLUS_ONE,
@@ -167,7 +169,7 @@ enum EUnrealEngineObjectUE4Version
 	// Remove LevelBodySetup from ULevel
 	VER_UE4_REMOVE_LEVELBODYSETUP,
 	// Refactor character crouching
-	VER_UE4_REFACTOR_CHARACTER_CROUCH,	
+	VER_UE4_REFACTOR_CHARACTER_CROUCH,
 	// Trimmed down material shader debug information.
 	VER_UE4_SMALLER_DEBUG_MATERIALSHADER_UNIFORM_EXPRESSIONS,
 	// APEX Clothing
@@ -733,10 +735,10 @@ enum EUnrealEngineObjectLicenseeUEVersion
 
 static_assert(EUnrealEngineObjectUE4Version::VER_UE4_AUTOMATIC_VERSION < (int32)EUnrealEngineObjectUE5Version::INITIAL_VERSION, "UE4 versioning must be lower than the start of UE5 versioning");
 
-/** 
+/**
  * This object combines all of our version enums into a single easy to use structure
- * which allows us to update older version numbers independently of the newer version numbers. 
- */ 
+ * which allows us to update older version numbers independently of the newer version numbers.
+ */
 struct FPackageFileVersion
 {
 	FPackageFileVersion() = default;
@@ -776,13 +778,13 @@ struct FPackageFileVersion
 	{
 		return FileVersionUE4 != Version;
 	}
-	
+
 	/** UE4 version comparisons */
 	bool operator <(EUnrealEngineObjectUE4Version Version) const
 	{
 		return FileVersionUE4 < Version;
 	}
-	
+
 	/** UE4 version comparisons */
 	bool operator >=(EUnrealEngineObjectUE4Version Version) const
 	{
@@ -794,7 +796,7 @@ struct FPackageFileVersion
 	{
 		return FileVersionUE5 != (int32)Version;
 	}
-	
+
 	/** UE5 version comparisons  */
 	bool operator <(EUnrealEngineObjectUE5Version Version) const
 	{
@@ -807,9 +809,9 @@ struct FPackageFileVersion
 		return FileVersionUE5 >= (int32)Version;
 	}
 
-	/** 
-	 * Returns true if this object is compatible with the FPackageFileVersion passed in as the parameter. 
-	 * This means that  all version numbers for the current object are equal or greater than the 
+	/**
+	 * Returns true if this object is compatible with the FPackageFileVersion passed in as the parameter.
+	 * This means that  all version numbers for the current object are equal or greater than the
 	 * corresponding version numbers of the other structure.
 	 */
 	bool IsCompatible(const FPackageFileVersion& Other) const
@@ -860,5 +862,3 @@ extern const CORE_API int32					GPackageFileLicenseeUE4Version;		// Licensee Ver
 extern const CORE_API FPackageFileVersion	GPackageFileUEVersion;				// Unreal Version Number.
 extern const CORE_API FPackageFileVersion	GOldestLoadablePackageFileUEVersion;// The oldest Unreal Version number that we can load.
 extern const CORE_API int32					GPackageFileLicenseeUEVersion;		// Licensee Version Number.
-
-
diff --git a/Engine/Source/Runtime/CoreUObject/Private/UObject/EnumProperty.cpp b/Engine/Source/Runtime/CoreUObject/Private/UObject/EnumProperty.cpp
--- a/Engine/Source/Runtime/CoreUObject/Private/UObject/EnumProperty.cpp
+++ b/Engine/Source/Runtime/CoreUObject/Private/UObject/EnumProperty.cpp
@@ -420,8 +420,8 @@ EConvertFromTypeResult FEnumProperty::ConvertFromType(const FPropertyTag& Tag, F
 		uint8 PreviousValue = 0;
 		if (Tag.GetType().GetParameterCount() == 0)
 		{
-			// A nested property loses its enum name in the property tag. Handle this case for backward compatibility reasons.
-			if (GetOwner<FProperty>())
+			// A nested property would lose its enum name on previous versions. Handle this case for backward compatibility reasons.
+			if (GetOwner<FProperty>() && Slot.GetArchiveState().UEVer() < EUnrealEngineObjectUE5Version::PROPERTY_TAG_COMPLETE_TYPE_NAME)
 			{
 				UE::FPropertyTypeNameBuilder TypeBuilder;
 				TypeBuilder.AddName(Tag.Type);
diff --git a/Engine/Source/Runtime/CoreUObject/Private/UObject/PropertyArray.cpp b/Engine/Source/Runtime/CoreUObject/Private/UObject/PropertyArray.cpp
--- a/Engine/Source/Runtime/CoreUObject/Private/UObject/PropertyArray.cpp
+++ b/Engine/Source/Runtime/CoreUObject/Private/UObject/PropertyArray.cpp
@@ -176,7 +176,7 @@ void FArrayProperty::SerializeItem(FStructuredArchive::FSlot Slot, void* Value,
 		return;
 	}
 
-	if (bIsTextFormat && Inner->IsA<FStructProperty>())
+	if (bIsTextFormat && Inner->IsA<FStructProperty>() && UnderlyingArchive.UEVer() < EUnrealEngineObjectUE5Version::PROPERTY_TAG_COMPLETE_TYPE_NAME)
 	{
 		MaybeInnerTag.Emplace(Inner, /*Index*/ 0, (uint8*)Value);
 
@@ -558,7 +558,8 @@ void FArrayProperty::SerializeItem(FStructuredArchive::FSlot Slot, void* Value,
 	FUObjectSerializeContext* Context = FUObjectThreadContext::Get().GetSerializeContext();
 
 	// Serialize a PropertyTag for the inner property of this array, allows us to validate the inner struct to see if it has changed
-	if (UnderlyingArchive.UEVer() >= VER_UE4_INNER_ARRAY_TAG_INFO &&
+	if (UnderlyingArchive.UEVer() < EUnrealEngineObjectUE5Version::PROPERTY_TAG_COMPLETE_TYPE_NAME &&
+		UnderlyingArchive.UEVer() >= VER_UE4_INNER_ARRAY_TAG_INFO &&
 		Inner->IsA<FStructProperty>())
 	{
 		if (!MaybeInnerTag)
@@ -1170,10 +1171,21 @@ EConvertFromTypeResult FArrayProperty::ConvertFromType(const FPropertyTag& Tag,
 
 	FArchive& UnderlyingArchive = Slot.GetUnderlyingArchive();
 
-	const FName InnerTypeName = Tag.GetType().GetParameterName();
-	if (InnerTypeName.IsNone() || InnerTypeName == Inner->GetID())
+	const FPackageFileVersion Version = UnderlyingArchive.UEVer();
+	if (Version >= EUnrealEngineObjectUE5Version::PROPERTY_TAG_COMPLETE_TYPE_NAME)
 	{
-		return EConvertFromTypeResult::UseSerializeItem;
+		if (CanSerializeFromTypeName(Tag.GetType()))
+		{
+			return EConvertFromTypeResult::UseSerializeItem;
+		}
+	}
+	else
+	{
+		const FName InnerTypeName = Tag.GetType().GetParameterName();
+		if (InnerTypeName.IsNone() || InnerTypeName == Inner->GetID())
+		{
+			return EConvertFromTypeResult::UseSerializeItem;
+		}
 	}
 
 	if (Tag.bExperimentalOverridableLogic)
@@ -1198,7 +1210,7 @@ EConvertFromTypeResult FArrayProperty::ConvertFromType(const FPropertyTag& Tag,
 	InnerPropertyTag.SetType(Tag.GetType().GetParameter());
 	InnerPropertyTag.ArrayIndex = 0;
 
-	if (UnderlyingArchive.UEVer() >= VER_UE4_INNER_ARRAY_TAG_INFO && InnerPropertyTag.Type == NAME_StructProperty)
+	if (Version < EUnrealEngineObjectUE5Version::PROPERTY_TAG_COMPLETE_TYPE_NAME && Version >= VER_UE4_INNER_ARRAY_TAG_INFO && InnerPropertyTag.Type == NAME_StructProperty)
 	{
 		UnderlyingArchive << InnerPropertyTag;
 	}
diff --git a/Engine/Source/Runtime/CoreUObject/Private/UObject/PropertyBool.cpp b/Engine/Source/Runtime/CoreUObject/Private/UObject/PropertyBool.cpp
--- a/Engine/Source/Runtime/CoreUObject/Private/UObject/PropertyBool.cpp
+++ b/Engine/Source/Runtime/CoreUObject/Private/UObject/PropertyBool.cpp
@@ -285,8 +285,8 @@ EConvertFromTypeResult FBoolProperty::ConvertFromType(const FPropertyTag& Tag, F
 			return EConvertFromTypeResult::UseSerializeItem;
 		}
 
-		// Disallow a nested byte property because it is impossible to distinguish from a nested enum.
-		if (GetOwner<FProperty>())
+		// Disallow a nested byte property prior to complete type names because it was impossible to distinguish from a nested enum.
+		if (GetOwner<FProperty>() && Slot.GetArchiveState().UEVer() < EUnrealEngineObjectUE5Version::PROPERTY_TAG_COMPLETE_TYPE_NAME)
 		{
 			return EConvertFromTypeResult::UseSerializeItem;
 		}
diff --git a/Engine/Source/Runtime/CoreUObject/Private/UObject/PropertyByte.cpp b/Engine/Source/Runtime/CoreUObject/Private/UObject/PropertyByte.cpp
--- a/Engine/Source/Runtime/CoreUObject/Private/UObject/PropertyByte.cpp
+++ b/Engine/Source/Runtime/CoreUObject/Private/UObject/PropertyByte.cpp
@@ -228,8 +228,8 @@ EConvertFromTypeResult FByteProperty::ConvertFromType(const FPropertyTag& Tag, F
 			uint8 PreviousValue = 0;
 			if (Enum)
 			{
-				// A nested property loses its enum name in the property tag. Handle this case for backward compatibility reasons.
-				if (GetOwner<FProperty>())
+				// A nested property would lose its enum name on previous versions. Handle this case for backward compatibility reasons.
+				if (GetOwner<FProperty>() && Slot.GetArchiveState().UEVer() < EUnrealEngineObjectUE5Version::PROPERTY_TAG_COMPLETE_TYPE_NAME)
 				{
 					return EConvertFromTypeResult::UseSerializeItem;
 				}
diff --git a/Engine/Source/Runtime/CoreUObject/Private/UObject/PropertyMap.cpp b/Engine/Source/Runtime/CoreUObject/Private/UObject/PropertyMap.cpp
--- a/Engine/Source/Runtime/CoreUObject/Private/UObject/PropertyMap.cpp
+++ b/Engine/Source/Runtime/CoreUObject/Private/UObject/PropertyMap.cpp
@@ -1438,11 +1438,27 @@ EConvertFromTypeResult FMapProperty::ConvertFromType(const FPropertyTag& Tag, FS
 	const UE::FPropertyTypeName ValueType = Tag.GetType().GetParameter(1);
 	const FName KeyTypeName = KeyType.GetName();
 	const FName ValueTypeName = ValueType.GetName();
-	const bool bCanSerializeKey = (KeyTypeName == KeyProp->GetID());
-	const bool bCanSerializeValue = (ValueTypeName == ValueProp->GetID());
-	if ((bCanSerializeKey || KeyTypeName.IsNone()) && (bCanSerializeValue || ValueTypeName.IsNone()))
+	bool bCanSerializeKey;
+	bool bCanSerializeValue;
+
+	const FPackageFileVersion Version = UnderlyingArchive.UEVer();
+	if (Version >= EUnrealEngineObjectUE5Version::PROPERTY_TAG_COMPLETE_TYPE_NAME)
 	{
-		return EConvertFromTypeResult::UseSerializeItem;
+		bCanSerializeKey = KeyProp->CanSerializeFromTypeName(KeyType);
+		bCanSerializeValue = ValueProp->CanSerializeFromTypeName(ValueType);
+		if (bCanSerializeKey && bCanSerializeValue)
+		{
+			return EConvertFromTypeResult::UseSerializeItem;
+		}
+	}
+	else
+	{
+		bCanSerializeKey = (KeyTypeName == KeyProp->GetID());
+		bCanSerializeValue = (ValueTypeName == ValueProp->GetID());
+		if ((bCanSerializeKey || KeyTypeName.IsNone()) && (bCanSerializeValue || ValueTypeName.IsNone()))
+		{
+			return EConvertFromTypeResult::UseSerializeItem;
+		}
 	}
 
 	if (Tag.bExperimentalOverridableLogic)
@@ -1456,8 +1472,6 @@ EConvertFromTypeResult FMapProperty::ConvertFromType(const FPropertyTag& Tag, FS
 		{
 			switch (Inner->ConvertFromType(InnerTag, InnerSlot, InnerData, InnerDefaultsStruct, nullptr))
 			{
-				default:
-					checkNoEntry();
 				case EConvertFromTypeResult::Converted:
 				case EConvertFromTypeResult::Serialized:
 					return true;
@@ -1470,6 +1484,9 @@ EConvertFromTypeResult FMapProperty::ConvertFromType(const FPropertyTag& Tag, FS
 					}
 					// Fall through to default SerializeItem
 					break;
+				default:
+					checkNoEntry();
+					return false;
 			}
 		}
 
diff --git a/Engine/Source/Runtime/CoreUObject/Private/UObject/PropertyOptional.cpp b/Engine/Source/Runtime/CoreUObject/Private/UObject/PropertyOptional.cpp
--- a/Engine/Source/Runtime/CoreUObject/Private/UObject/PropertyOptional.cpp
+++ b/Engine/Source/Runtime/CoreUObject/Private/UObject/PropertyOptional.cpp
@@ -201,7 +201,8 @@ void FOptionalProperty::SerializeItem(FStructuredArchive::FSlot Slot, void* Data
 		
 		const void* ValueDefaults = Defaults ? GetValuePointerForReadIfSet(Defaults) : nullptr;
 
-		if (Slot.GetArchiveState().UseUnversionedPropertySerialization())
+		if (Slot.GetArchiveState().UseUnversionedPropertySerialization() ||
+			UnderlyingArchive.UEVer() >= EUnrealEngineObjectUE5Version::PROPERTY_TAG_COMPLETE_TYPE_NAME)
 		{
 			// Simply serialize the inner value if using unversioned property serialization.
 			void* ValueData = bIsLoading
@@ -219,6 +220,11 @@ void FOptionalProperty::SerializeItem(FStructuredArchive::FSlot Slot, void* Data
 			// Serialize the value's tag.
 			ValueSlot << ValueTag;
 
+			if (UE::FPropertyTypeName NewTypeName = ApplyRedirectsToPropertyType(ValueTag.GetType(), ValueProperty); !NewTypeName.IsEmpty())
+			{
+				ValueTag.SetType(NewTypeName);
+			}
+
 			// Deserialize/convert the value.
 			void* ValueData = MarkSetAndGetInitializedValuePointerToReplace(Data);
 
@@ -231,8 +237,11 @@ void FOptionalProperty::SerializeItem(FStructuredArchive::FSlot Slot, void* Data
 					bSuccessfullyDeserialized = true;
 					break;
 				case EConvertFromTypeResult::UseSerializeItem:
-					GetValueProperty()->SerializeItem(ValueSlot, ValueData, ValueDefaults);
-					bSuccessfullyDeserialized = !UnderlyingArchive.IsCriticalError();
+					if (ValueTag.Type == GetValueProperty()->GetID())
+					{
+						ValueTag.SerializeTaggedProperty(ValueSlot, GetValueProperty(), (uint8*)ValueData, (const uint8*)ValueDefaults);
+						bSuccessfullyDeserialized = !UnderlyingArchive.IsCriticalError();
+					}
 					break;
 				case EConvertFromTypeResult::CannotConvert:
 					break;
@@ -450,6 +459,50 @@ int32 FOptionalProperty::GetMinAlignment() const
 
 EConvertFromTypeResult FOptionalProperty::ConvertFromType(const FPropertyTag& Tag, FStructuredArchive::FSlot Slot, uint8* ContainerData, UStruct* DefaultsStruct, const uint8* DefaultsContainer)
 {
+	if (Slot.GetArchiveState().UEVer() >= EUnrealEngineObjectUE5Version::PROPERTY_TAG_COMPLETE_TYPE_NAME)
+	{
+		if (CanSerializeFromTypeName(Tag.GetType()))
+		{
+			return EConvertFromTypeResult::UseSerializeItem;
+		}
+
+		FPropertyTag ValueTag = Tag;
+		TOptional<FStructuredArchive::FSlot> MaybeValueSlot;
+		if (Tag.Type == NAME_OptionalProperty)
+		{
+			FStructuredArchive::FRecord Record = Slot.EnterRecord();
+			MaybeValueSlot = Record.TryEnterField(TEXT("Value"), /*bEnterForSaving*/ false);
+			if (!MaybeValueSlot)
+			{
+				MarkUnset(ContainerPtrToValuePtr<void>(ContainerData));
+				return EConvertFromTypeResult::Converted;
+			}
+			ValueTag.SetType(Tag.GetType().GetParameter());
+		}
+
+		FStructuredArchive::FSlot ValueSlot = MaybeValueSlot.Get(Slot);
+		uint8* ValueData = (uint8*)MarkSetAndGetInitializedValuePointerToReplace(ContainerPtrToValuePtr<void>(ContainerData));
+		switch (GetValueProperty()->ConvertFromType(ValueTag, ValueSlot, ValueData, nullptr, nullptr))
+		{
+		case EConvertFromTypeResult::Converted:
+			return EConvertFromTypeResult::Converted;
+		case EConvertFromTypeResult::Serialized:
+			return EConvertFromTypeResult::Serialized;
+		case EConvertFromTypeResult::CannotConvert:
+			return EConvertFromTypeResult::CannotConvert;
+		case EConvertFromTypeResult::UseSerializeItem:
+			if (ValueTag.Type == GetValueProperty()->GetID())
+			{
+				GetValueProperty()->SerializeItem(ValueSlot, ValueData);
+				return EConvertFromTypeResult::Serialized;
+			}
+			return EConvertFromTypeResult::CannotConvert;
+		default:
+			checkNoEntry();
+			return EConvertFromTypeResult::CannotConvert;
+		}
+	}
+
 	static const FName NAME_OptionProperty("OptionProperty");
 	if (Tag.Type != NAME_OptionProperty)
 	{
diff --git a/Engine/Source/Runtime/CoreUObject/Private/UObject/PropertySet.cpp b/Engine/Source/Runtime/CoreUObject/Private/UObject/PropertySet.cpp
--- a/Engine/Source/Runtime/CoreUObject/Private/UObject/PropertySet.cpp
+++ b/Engine/Source/Runtime/CoreUObject/Private/UObject/PropertySet.cpp
@@ -900,10 +900,21 @@ EConvertFromTypeResult FSetProperty::ConvertFromType(const FPropertyTag& Tag, FS
 		return EConvertFromTypeResult::UseSerializeItem;
 	}
 
-	const FName InnerTypeName = Tag.GetType().GetParameterName();
-	if (InnerTypeName.IsNone() || InnerTypeName == ElementProp->GetID())
+	const FPackageFileVersion Version = UnderlyingArchive.UEVer();
+	if (Version >= EUnrealEngineObjectUE5Version::PROPERTY_TAG_COMPLETE_TYPE_NAME)
 	{
-		return EConvertFromTypeResult::UseSerializeItem;
+		if (CanSerializeFromTypeName(Tag.GetType()))
+		{
+			return EConvertFromTypeResult::UseSerializeItem;
+		}
+	}
+	else
+	{
+		const FName InnerTypeName = Tag.GetType().GetParameterName();
+		if (InnerTypeName.IsNone() || InnerTypeName == ElementProp->GetID())
+		{
+			return EConvertFromTypeResult::UseSerializeItem;
+		}
 	}
 
 	FScriptSetHelper ScriptSetHelper(this, ContainerPtrToValuePtr<void>(Data));
diff --git a/Engine/Source/Runtime/CoreUObject/Private/UObject/PropertyTag.cpp b/Engine/Source/Runtime/CoreUObject/Private/UObject/PropertyTag.cpp
--- a/Engine/Source/Runtime/CoreUObject/Private/UObject/PropertyTag.cpp
+++ b/Engine/Source/Runtime/CoreUObject/Private/UObject/PropertyTag.cpp
@@ -13,6 +13,19 @@
 #include "UObject/Package.h"
 #include "UObject/UnrealType.h"
 
+/** Flags used for serialization of FPropertyTag. DO NOT EDIT THESE VALUES! */
+enum class EPropertyTagFlags : uint8
+{
+	None						= 0x00,
+	HasArrayIndex				= 0x01,
+	HasPropertyGuid				= 0x02,
+	HasPropertyExtensions		= 0x04,
+	HasBinaryOrNativeSerialize	= 0x08,
+	BoolTrue					= 0x10,
+};
+
+ENUM_CLASS_FLAGS(EPropertyTagFlags);
+
 /**
  * Enum flags that indicate that additional data was serialized for that property tag.
  * Registered flags should be serialized in ascending order.
@@ -53,8 +66,6 @@ FPropertyTag::FPropertyTag(FProperty* Property, int32 InIndex, uint8* Value)
 	Property->SaveTypeName(TypeBuilder);
 	SetType(TypeBuilder.Build());
 
-	Property->SaveToTag(*this);
-
 	if (FBoolProperty* Bool = CastField<FBoolProperty>(Property))
 	{
 		BoolVal = Bool->GetPropertyValue(Value);
@@ -164,23 +175,8 @@ static void ParsePathName(FStringView Path, UE::FPropertyTypeNameBuilder& Builde
 	}
 }
 
-// Serializer.
-FArchive& operator<<(FArchive& Ar, FPropertyTag& Tag)
-{
-	FStructuredArchiveFromArchive(Ar).GetSlot() << Tag;
-	return Ar;
-}
-
-void operator<<(FStructuredArchive::FSlot Slot, FPropertyTag& Tag)
+FORCENOINLINE void LoadPropertyTagNoFullType(FStructuredArchive::FSlot Slot, FPropertyTag& Tag)
 {
-	FArchive& UnderlyingArchive = Slot.GetUnderlyingArchive();
-	bool bIsTextFormat = UnderlyingArchive.IsTextFormat();
-
-	const FPackageFileVersion Version = UnderlyingArchive.UEVer();
-
-	check(!UnderlyingArchive.GetArchiveState().UseUnversionedPropertySerialization());
-	checkf(!UnderlyingArchive.IsSaving() || Tag.GetProperty(), TEXT("FPropertyTag must be constructed with a valid property when used for saving data!"));
-
 	const auto AddEnumPath = [](UE::FPropertyTypeNameBuilder& Builder, FName EnumName)
 	{
 		TStringBuilder<256> EnumPath(InPlace, EnumName);
@@ -194,9 +190,14 @@ void operator<<(FStructuredArchive::FSlot Slot, FPropertyTag& Tag)
 		}
 	};
 
+	FArchive& UnderlyingArchive = Slot.GetUnderlyingArchive();
+	const bool bIsTextFormat = UnderlyingArchive.IsTextFormat();
+	const FPackageFileVersion Version = UnderlyingArchive.UEVer();
+
+	check(UnderlyingArchive.IsLoading());
+
 	if (!bIsTextFormat)
 	{
-		// Name.
 		Slot << SA_ATTRIBUTE(TEXT("Name"), Tag.Name);
 		if (Tag.Name.IsNone())
 		{
@@ -206,35 +207,26 @@ void operator<<(FStructuredArchive::FSlot Slot, FPropertyTag& Tag)
 
 	Slot << SA_ATTRIBUTE(TEXT("Type"), Tag.Type);
 
-	if (UnderlyingArchive.IsSaving())
-	{
-		// remember the offset of the Size variable - UStruct::SerializeTaggedProperties will update it after the
-		// property has been serialized.
-		Tag.SizeOffset = UnderlyingArchive.Tell();
-	}
-
 	if (!bIsTextFormat)
 	{
 		Slot << SA_ATTRIBUTE(TEXT("Size"), Tag.Size);
 		Slot << SA_ATTRIBUTE(TEXT("ArrayIndex"), Tag.ArrayIndex);
 	}
 
-	if (Tag.Type.GetNumber() == 0)
+	if (Tag.Type.GetNumber() == NAME_NO_NUMBER_INTERNAL)
 	{
 		// Build Tag.TypeName from the partial type name that was saved in older versions.
 		UE::FPropertyTypeNameBuilder TypeBuilder;
 		TypeBuilder.AddName(Tag.Type);
 
-		FNameEntryId TagType = Tag.Type.GetComparisonIndex();
+		const FNameEntryId TagType = Tag.Type.GetComparisonIndex();
 
 		// only need to serialize this for structs
 		if (TagType == NAME_StructProperty)
 		{
 			TypeBuilder.BeginParameters();
-			PRAGMA_DISABLE_DEPRECATION_WARNINGS
-			FName StructName = Tag.StructName;
-			FGuid StructGuid = Tag.StructGuid;
-			PRAGMA_ENABLE_DEPRECATION_WARNINGS
+			FName StructName;
+			FGuid StructGuid;
 			Slot << SA_ATTRIBUTE(TEXT("StructName"), StructName);
 			TypeBuilder.AddName(StructName);
 			if (Version >= VER_UE4_STRUCT_GUID_IN_PROPERTY_TAG)
@@ -274,9 +266,7 @@ void operator<<(FStructuredArchive::FSlot Slot, FPropertyTag& Tag)
 		// only need to serialize this for bytes/enums
 		else if (TagType == NAME_ByteProperty)
 		{
-			PRAGMA_DISABLE_DEPRECATION_WARNINGS
-			FName EnumName = Tag.EnumName;
-			PRAGMA_ENABLE_DEPRECATION_WARNINGS
+			FName EnumName;
 			if (UnderlyingArchive.IsTextFormat())
 			{
 				Slot << SA_OPTIONAL_ATTRIBUTE(TEXT("EnumName"), EnumName, NAME_None);
@@ -297,9 +287,7 @@ void operator<<(FStructuredArchive::FSlot Slot, FPropertyTag& Tag)
 		}
 		else if (TagType == NAME_EnumProperty)
 		{
-			PRAGMA_DISABLE_DEPRECATION_WARNINGS
-			FName EnumName = Tag.EnumName;
-			PRAGMA_ENABLE_DEPRECATION_WARNINGS
+			FName EnumName;
 			Slot << SA_ATTRIBUTE(TEXT("EnumName"), EnumName);
 			TypeBuilder.BeginParameters();
 			AddEnumPath(TypeBuilder, EnumName);
@@ -312,9 +300,7 @@ void operator<<(FStructuredArchive::FSlot Slot, FPropertyTag& Tag)
 		// need to serialize the InnerType for arrays
 		else if (TagType == NAME_ArrayProperty)
 		{
-			PRAGMA_DISABLE_DEPRECATION_WARNINGS
-			FName InnerType = Tag.InnerType;
-			PRAGMA_ENABLE_DEPRECATION_WARNINGS
+			FName InnerType;
 			if (Version >= VAR_UE4_ARRAY_PROPERTY_INNER_TAGS)
 			{
 				Slot << SA_ATTRIBUTE(TEXT("InnerType"), InnerType);
@@ -329,9 +315,7 @@ void operator<<(FStructuredArchive::FSlot Slot, FPropertyTag& Tag)
 		// need to serialize the InnerType for optionals.
 		else if (TagType == NAME_OptionalProperty)
 		{
-			PRAGMA_DISABLE_DEPRECATION_WARNINGS
-			FName InnerType = Tag.InnerType;
-			PRAGMA_ENABLE_DEPRECATION_WARNINGS
+			FName InnerType;
 			Slot << SA_ATTRIBUTE(TEXT("InnerType"), InnerType);
 			TypeBuilder.BeginParameters();
 			TypeBuilder.AddName(InnerType);
@@ -344,9 +328,7 @@ void operator<<(FStructuredArchive::FSlot Slot, FPropertyTag& Tag)
 		{
 			if (TagType == NAME_SetProperty)
 			{
-				PRAGMA_DISABLE_DEPRECATION_WARNINGS
-				FName InnerType = Tag.InnerType;
-				PRAGMA_ENABLE_DEPRECATION_WARNINGS
+				FName InnerType;
 				Slot << SA_ATTRIBUTE(TEXT("InnerType"), InnerType);
 				TypeBuilder.BeginParameters();
 				TypeBuilder.AddName(InnerType);
@@ -357,10 +339,8 @@ void operator<<(FStructuredArchive::FSlot Slot, FPropertyTag& Tag)
 			}
 			else if (TagType == NAME_MapProperty)
 			{
-				PRAGMA_DISABLE_DEPRECATION_WARNINGS
-				FName InnerType = Tag.InnerType;
-				FName ValueType = Tag.ValueType;
-				PRAGMA_ENABLE_DEPRECATION_WARNINGS
+				FName InnerType;
+				FName ValueType;
 				Slot << SA_ATTRIBUTE(TEXT("InnerType"), InnerType);
 				Slot << SA_ATTRIBUTE(TEXT("ValueType"), ValueType);
 				TypeBuilder.BeginParameters();
@@ -374,10 +354,7 @@ void operator<<(FStructuredArchive::FSlot Slot, FPropertyTag& Tag)
 			}
 		}
 
-		if (UnderlyingArchive.IsLoading())
-		{
-			Tag.TypeName = TypeBuilder.Build();
-		}
+		Tag.TypeName = TypeBuilder.Build();
 	}
 
 	// Property tags to handle renamed blueprint properties effectively.
@@ -403,6 +380,137 @@ void operator<<(FStructuredArchive::FSlot Slot, FPropertyTag& Tag)
 		EPropertyTagExtension PropertyTagExtensions = CalculatePropertyExtensionFlags(UnderlyingArchive, Tag);
 		SerializePropertyExtensions(Slot, PropertyTagExtensions, Tag);
 	}
+
+	Tag.SerializeType = EPropertyTagSerializeType::Unknown;
+}
+
+FORCENOINLINE void SerializePropertyTagAsText(FStructuredArchive::FSlot Slot, FPropertyTag& Tag)
+{
+	FArchive& UnderlyingArchive = Slot.GetUnderlyingArchive();
+
+	Slot << SA_ATTRIBUTE(TEXT("Type"), Tag.TypeName);
+	if (UnderlyingArchive.IsLoading())
+	{
+		Tag.SetType(Tag.TypeName);
+	}
+
+	Slot << SA_OPTIONAL_ATTRIBUTE(TEXT("PropertyGuid"), Tag.PropertyGuid, FGuid());
+	Tag.HasPropertyGuid = Tag.PropertyGuid.IsValid();
+
+	bool bHasBinaryOrNativeSerialize = (Tag.SerializeType == EPropertyTagSerializeType::BinaryOrNative);
+	Slot << SA_OPTIONAL_ATTRIBUTE(TEXT("HasBinaryOrNativeSerialize"), bHasBinaryOrNativeSerialize, false);
+	Tag.SerializeType = bHasBinaryOrNativeSerialize ? EPropertyTagSerializeType::BinaryOrNative : EPropertyTagSerializeType::Property;
+
+	EPropertyTagExtension PropertyTagExtensions = CalculatePropertyExtensionFlags(UnderlyingArchive, Tag);
+	SerializePropertyExtensions(Slot, PropertyTagExtensions, Tag);
+}
+
+// Serializer.
+FArchive& operator<<(FArchive& Ar, FPropertyTag& Tag)
+{
+	FStructuredArchiveFromArchive(Ar).GetSlot() << Tag;
+	return Ar;
+}
+
+void operator<<(FStructuredArchive::FSlot Slot, FPropertyTag& Tag)
+{
+	FArchive& UnderlyingArchive = Slot.GetUnderlyingArchive();
+	const bool bIsTextFormat = UnderlyingArchive.IsTextFormat();
+	const FPackageFileVersion Version = UnderlyingArchive.UEVer();
+
+	check(!UnderlyingArchive.GetArchiveState().UseUnversionedPropertySerialization());
+	checkf(!UnderlyingArchive.IsSaving() || Tag.GetProperty(), TEXT("FPropertyTag must be constructed with a valid property when used for saving data!"));
+
+	if (UNLIKELY(Version < EUnrealEngineObjectUE5Version::PROPERTY_TAG_COMPLETE_TYPE_NAME))
+	{
+		LoadPropertyTagNoFullType(Slot, Tag);
+		return;
+	}
+
+	if (UNLIKELY(bIsTextFormat))
+	{
+		SerializePropertyTagAsText(Slot, Tag);
+		return;
+	}
+
+	Slot << SA_ATTRIBUTE(TEXT("Name"), Tag.Name);
+	if (Tag.Name.IsNone())
+	{
+		return;
+	}
+
+	Slot << SA_ATTRIBUTE(TEXT("Type"), Tag.TypeName);
+	if (UnderlyingArchive.IsLoading())
+	{
+		Tag.SetType(Tag.TypeName);
+	}
+
+	if (UnderlyingArchive.IsSaving())
+	{
+		// Store the serialized offset of the Size field.
+		// UStruct::SerializeTaggedProperties will rewrite it after the property has been serialized.
+		Tag.SizeOffset = UnderlyingArchive.Tell();
+	}
+	Slot << SA_ATTRIBUTE(TEXT("Size"), Tag.Size);
+
+	EPropertyTagFlags PropertyTagFlags = EPropertyTagFlags::None;
+	EPropertyTagExtension PropertyTagExtensions = CalculatePropertyExtensionFlags(UnderlyingArchive, Tag);
+
+	if (UnderlyingArchive.IsSaving())
+	{
+		Tag.SerializeType = Tag.GetProperty()->UseBinaryOrNativeSerialization(UnderlyingArchive)
+			? EPropertyTagSerializeType::BinaryOrNative : EPropertyTagSerializeType::Property;
+
+		if (Tag.ArrayIndex != 0)
+		{
+			PropertyTagFlags |= EPropertyTagFlags::HasArrayIndex;
+		}
+		if (Tag.HasPropertyGuid)
+		{
+			PropertyTagFlags = EPropertyTagFlags::HasPropertyGuid;
+		}
+		if (PropertyTagExtensions != EPropertyTagExtension::NoExtension)
+		{
+			PropertyTagFlags |= EPropertyTagFlags::HasPropertyExtensions;
+		}
+		if (Tag.SerializeType == EPropertyTagSerializeType::BinaryOrNative)
+		{
+			PropertyTagFlags |= EPropertyTagFlags::HasBinaryOrNativeSerialize;
+		}
+		if (Tag.BoolVal && Tag.Type == NAME_BoolProperty)
+		{
+			PropertyTagFlags |= EPropertyTagFlags::BoolTrue;
+		}
+	}
+
+	Slot << SA_ATTRIBUTE(TEXT("Flags"), PropertyTagFlags);
+
+	if (UnderlyingArchive.IsLoading())
+	{
+		Tag.HasPropertyGuid = EnumHasAnyFlags(PropertyTagFlags, EPropertyTagFlags::HasPropertyGuid);
+		Tag.SerializeType = EnumHasAnyFlags(PropertyTagFlags, EPropertyTagFlags::HasBinaryOrNativeSerialize)
+			? EPropertyTagSerializeType::BinaryOrNative : EPropertyTagSerializeType::Property;
+		Tag.BoolVal = EnumHasAnyFlags(PropertyTagFlags, EPropertyTagFlags::BoolTrue);
+	}
+
+	if (EnumHasAnyFlags(PropertyTagFlags, EPropertyTagFlags::HasArrayIndex))
+	{
+		Slot << SA_ATTRIBUTE(TEXT("ArrayIndex"), Tag.ArrayIndex);
+	}
+	else
+	{
+		Tag.ArrayIndex = 0;
+	}
+
+	if (EnumHasAnyFlags(PropertyTagFlags, EPropertyTagFlags::HasPropertyGuid))
+	{
+		Slot << SA_ATTRIBUTE(TEXT("PropertyGuid"), Tag.PropertyGuid);
+	}
+
+	if (EnumHasAnyFlags(PropertyTagFlags, EPropertyTagFlags::HasPropertyExtensions))
+	{
+		SerializePropertyExtensions(Slot, PropertyTagExtensions, Tag);
+	}
 }
 
 // Property serializer.
diff --git a/Engine/Source/Runtime/CoreUObject/Public/UObject/PropertyTag.h b/Engine/Source/Runtime/CoreUObject/Public/UObject/PropertyTag.h
--- a/Engine/Source/Runtime/CoreUObject/Public/UObject/PropertyTag.h
+++ b/Engine/Source/Runtime/CoreUObject/Public/UObject/PropertyTag.h
@@ -16,6 +16,19 @@ enum class EOverriddenPropertyOperation : uint8;
 class FArchive;
 class FProperty;
 
+/**
+ * Used by the tag to describe how the property was serialized.
+ */
+enum class EPropertyTagSerializeType : uint8
+{
+	/** Tag was loaded from an older version or has not yet been saved. */
+	Unknown,
+	/** Serialized with tagged property serialization. */
+	Property,
+	/** Serialized with binary or native serialization. */
+	BinaryOrNative,
+};
+
 /**
  *  A tag describing a class property, to aid in serialization.
  */
@@ -48,6 +61,7 @@ public:
 	FGuid	PropertyGuid;
 	uint8	HasPropertyGuid = 0;
 	uint8	BoolVal = 0;// a boolean property's value (never need to serialize data for bool properties except here)
+	EPropertyTagSerializeType SerializeType = EPropertyTagSerializeType::Unknown;
 	EOverriddenPropertyOperation OverrideOperation; // Overridable serialization state reconstruction 
 	bool	bExperimentalOverridableLogic = false; // Remember if property had CPF_ExperimentalOverridableLogic when saved
 
@@ -82,6 +96,10 @@ public:
 	// Property serializer.
 	void SerializeTaggedProperty(FArchive& Ar, FProperty* Property, uint8* Value, const uint8* Defaults) const;
 	UE_INTERNAL COREUOBJECT_API void SerializeTaggedProperty(FStructuredArchive::FSlot Slot, FProperty* Property, uint8* Value, const uint8* Defaults) const;
+
+private:
+	friend void LoadPropertyTagNoFullType(FStructuredArchive::FSlot Slot, FPropertyTag& Tag);
+	friend void SerializePropertyTagAsText(FStructuredArchive::FSlot Slot, FPropertyTag& Tag);
 };
 
 struct UE_INTERNAL FPropertyTagScope
