Changed FPropertyBag to use complete type names from the property tag

#rnx

diff --git a/Engine/Source/Runtime/CoreUObject/Private/UObject/InstanceDataObjectUtils.cpp b/Engine/Source/Runtime/CoreUObject/Private/UObject/InstanceDataObjectUtils.cpp
--- a/Engine/Source/Runtime/CoreUObject/Private/UObject/InstanceDataObjectUtils.cpp
+++ b/Engine/Source/Runtime/CoreUObject/Private/UObject/InstanceDataObjectUtils.cpp
@@ -54,94 +54,6 @@ namespace UE
 		return bIsEnabled;
 	}
 
-	static void BuildSegmentTypeFromProperty(const FProperty* Property, FPropertyTypeNameBuilder& OutType)
-	{
-#if false // TODO: @jordan.hoffmann when complete type info is finished use this branch instead
-		OutType.AddName(Property->GetID());
-		if (const FObjectProperty* AsObjectProperty = CastField<FObjectProperty>(Property))
-		{
-			OutType.BeginParameters();
-			OutType.AddName(AsObjectProperty->GetID());
-			OutType.EndParameters();
-		}
-		else if (const FEnumProperty* AsEnumProperty = CastField<FEnumProperty>(Property))
-		{
-			OutType.BeginParameters();
-			OutType.AddPath(AsEnumProperty->GetEnum());
-			OutType.EndParameters();
-		}
-		else if (const FArrayProperty* AsArrayProperty = CastField<FArrayProperty>(Property))
-		{
-			OutType.BeginParameters();
-			BuildSegmentTypeFromProperty(AsArrayProperty->Inner, OutType);
-			OutType.EndParameters();
-		}
-		else if (const FSetProperty* AsSetProperty = CastField<FSetProperty>(Property))
-		{
-			OutType.BeginParameters();
-			BuildSegmentTypeFromProperty(AsSetProperty->ElementProp, OutType);
-			OutType.EndParameters();
-		}
-		else if (const FMapProperty* AsMapProperty = CastField<FMapProperty>(Property))
-		{
-			OutType.BeginParameters();
-			BuildSegmentTypeFromProperty(AsMapProperty->KeyProp, OutType);
-			BuildSegmentTypeFromProperty(AsMapProperty->ValueProp, OutType);
-			OutType.EndParameters();
-		}
-#else
-		OutType.AddName(Property->GetID());
-		if (const FStructProperty* AsStructProperty = CastField<FStructProperty>(Property))
-		{
-			OutType.BeginParameters();
-			OutType.AddName(AsStructProperty->Struct->GetFName());
-			OutType.EndParameters();
-		}
-		else if (const FEnumProperty* AsEnumProperty = CastField<FEnumProperty>(Property))
-		{
-			OutType.BeginParameters();
-			OutType.AddName(AsEnumProperty->GetEnum()->GetFName());
-			OutType.EndParameters();
-		}
-		else if (const FArrayProperty* AsArrayProperty = CastField<FArrayProperty>(Property))
-		{
-			OutType.BeginParameters();
-			if (AsArrayProperty->Inner->IsA<FEnumProperty>())
-			{
-				// enum paths currently don't recurse when they're in containers
-				OutType.AddName(AsArrayProperty->Inner->GetID());
-			}
-			else
-			{
-				BuildSegmentTypeFromProperty(AsArrayProperty->Inner, OutType);
-			}
-			OutType.EndParameters();
-		}
-		else if (const FSetProperty* AsSetProperty = CastField<FSetProperty>(Property))
-		{
-			// sets currently don't recurse their element property types
-			OutType.BeginParameters();
-			OutType.AddName(AsSetProperty->ElementProp->GetID());
-			OutType.EndParameters();
-		}
-		else if (const FMapProperty* AsMapProperty = CastField<FMapProperty>(Property))
-		{
-			// maps currently don't recurse their key/value property types
-			OutType.BeginParameters();
-			OutType.AddName(AsMapProperty->KeyProp->GetID());
-			OutType.AddName(AsMapProperty->ValueProp->GetID());
-			OutType.EndParameters();
-		}
-#endif
-	}
-	
-	static FPropertyTypeName GetSegmentTypeFromProperty(const FProperty* Property)
-	{
-		FPropertyTypeNameBuilder TypeBuilder;
-		BuildSegmentTypeFromProperty(Property, TypeBuilder);
-		return TypeBuilder.Build();
-	}
-	
 	static bool ResolvePropertyPathNameHelper(const UStruct* Struct, ResolvePropertyPathNameHelperParams& Params);
 	static bool ResolvePropertyPathNameHelper(ResolvePropertyPathNameHelperParams& Params)
 	{
@@ -227,8 +139,7 @@ namespace UE
 		for (FProperty* Property : TFieldRange<FProperty>(Struct))
 		{
 			// find a property that matches the segment
-			const FPropertyTypeName PropertySegmentType = GetSegmentTypeFromProperty(Property);
-			if (Segment.Name != Property->GetFName() || Segment.Type != PropertySegmentType)
+			if (Segment.Name != Property->GetFName() || !Property->CanSerializeFromTypeName(Segment.Type))
 			{
 				continue;
 			}
@@ -275,10 +186,13 @@ namespace UE
 
 	static FPropertyPathNameSegment CreateSegmentFromProperty(const FProperty* Inner, int32 Index = INDEX_NONE)
 	{
+		FPropertyTypeNameBuilder TypeBuilder;
+		Inner->SaveTypeName(TypeBuilder);
+
 		FPropertyPathNameSegment Result;
 		Result.Index = INDEX_NONE; 
 		Result.Name = Inner->GetFName();
-		Result.Type = GetSegmentTypeFromProperty(Inner);
+		Result.Type = TypeBuilder.Build();
 		return Result;
 	}
 
diff --git a/Engine/Source/Runtime/CoreUObject/Private/UObject/PropertyBag.cpp b/Engine/Source/Runtime/CoreUObject/Private/UObject/PropertyBag.cpp
--- a/Engine/Source/Runtime/CoreUObject/Private/UObject/PropertyBag.cpp
+++ b/Engine/Source/Runtime/CoreUObject/Private/UObject/PropertyBag.cpp
@@ -102,7 +102,7 @@ void FPropertyBag::LoadPropertyByTag(const FPropertyPathName& Path, const FPrope
 
 	// Construct a property from the tag and try to use it to serialize the value.
 	FField* Field = FField::TryConstruct(Tag.Type, {}, Tag.Name, RF_NoFlags);
-	if (FProperty* Property = CastField<FProperty>(Field); Property && Property->LoadFromTag(Value.Tag))
+	if (FProperty* Property = CastField<FProperty>(Field); Property && Property->LoadTypeName(Value.Tag.GetType()))
 	{
 		Property->Link(UnderlyingArchive);
 		Value.bOwnsProperty = true;
