Removed new tag-related functions from FProperty

These have been replaced by LoadTypeName/SaveTypeName/CanSerializeFromTypeName.

#rnx

diff --git a/Engine/Source/Runtime/CoreUObject/Private/UObject/EnumProperty.cpp b/Engine/Source/Runtime/CoreUObject/Private/UObject/EnumProperty.cpp
--- a/Engine/Source/Runtime/CoreUObject/Private/UObject/EnumProperty.cpp
+++ b/Engine/Source/Runtime/CoreUObject/Private/UObject/EnumProperty.cpp
@@ -4,7 +4,6 @@
 
 #include "Algo/Find.h"
 #include "Hash/Blake3.h"
-#include "UObject/Package.h"
 #include "UObject/UnrealTypePrivate.h"
 #include "UObject/UObjectThreadContext.h"
 
@@ -528,50 +527,6 @@ uint64 FEnumProperty::GetMaxNetSerializeBits() const
 	return FMath::Min(DesiredBits, MaxBits);
 }
 
-bool FEnumProperty::LoadFromTag(const FPropertyTag& Tag)
-{
-	if (!Super::LoadFromTag(Tag))
-	{
-		return false;
-	}
-
-	// Update FByteProperty when making changes here.
-	PRAGMA_DISABLE_DEPRECATION_WARNINGS;
-	TStringBuilder<256> EnumName(InPlace, Tag.EnumName);
-	PRAGMA_ENABLE_DEPRECATION_WARNINGS;
-	if (UEnum* LocalEnum = FindFirstObject<UEnum>(*EnumName, EFindFirstObjectOptions::NativeFirst))
-	{
-		Enum = LocalEnum;
-		UE_CLOG(Enum->GetMaxEnumValue() >= 256, LogClass, Warning,
-			TEXT("Enum '%s' does not fit in a byte property loading property '%s'."), *Enum->GetName(), *GetName());
-		AddCppProperty(new FByteProperty(this, GetFName(), RF_NoFlags));
-		return true;
-	}
-
-	return false;
-}
-
-void FEnumProperty::SaveToTag(FPropertyTag& Tag)
-{
-	Super::SaveToTag(Tag);
-
-	if (const UEnum* LocalEnum = Enum)
-	{
-		PRAGMA_DISABLE_DEPRECATION_WARNINGS;
-		// RobM: Ugly hack so that we can avoid content changes in most of the packages
-		// Update FByteProperty when making changes here.
-		if (LocalEnum->GetPackage()->HasAnyPackageFlags(PKG_CompiledIn))
-		{				
-			Tag.EnumName = LocalEnum->GetFName();
-		}
-		else
-		{
-			Tag.EnumName = FName(*LocalEnum->GetPathName());
-		}
-		PRAGMA_ENABLE_DEPRECATION_WARNINGS;
-	}
-}
-
 bool FEnumProperty::LoadTypeName(UE::FPropertyTypeName Type, const FPropertyTag* Tag)
 {
 	if (!Super::LoadTypeName(Type, Tag))
diff --git a/Engine/Source/Runtime/CoreUObject/Private/UObject/Property.cpp b/Engine/Source/Runtime/CoreUObject/Private/UObject/Property.cpp
--- a/Engine/Source/Runtime/CoreUObject/Private/UObject/Property.cpp
+++ b/Engine/Source/Runtime/CoreUObject/Private/UObject/Property.cpp
@@ -2088,25 +2088,6 @@ bool FProperty::UseBinaryOrNativeSerialization(const FArchive& Ar) const
 	return Ar.WantBinaryPropertySerialization();
 }
 
-bool FProperty::LoadFromTag(const FPropertyTag& Tag)
-{
-	checkf(GetID() == Tag.Type, TEXT("Failed to load property '%s' of type '%s' from tag of type '%s'"),
-		*WriteToString<64>(GetFName()), *WriteToString<64>(GetID()), *WriteToString<64>(Tag.Type));
-	return true;
-}
-
-void FProperty::SaveToTag(FPropertyTag& Tag)
-{
-	Tag.SetProperty(this);
-	Tag.Type = GetID();
-	Tag.Name = GetFName();
-}
-
-void FProperty::AssignToTag(FPropertyTag& Tag)
-{
-	Tag.SetProperty(this);
-}
-
 bool FProperty::LoadTypeName(UE::FPropertyTypeName Type, const FPropertyTag* Tag)
 {
 	return ensureMsgf(GetID() == Type.GetName(),
diff --git a/Engine/Source/Runtime/CoreUObject/Private/UObject/PropertyArray.cpp b/Engine/Source/Runtime/CoreUObject/Private/UObject/PropertyArray.cpp
--- a/Engine/Source/Runtime/CoreUObject/Private/UObject/PropertyArray.cpp
+++ b/Engine/Source/Runtime/CoreUObject/Private/UObject/PropertyArray.cpp
@@ -1286,48 +1286,6 @@ bool FArrayProperty::UseBinaryOrNativeSerialization(const FArchive& Ar) const
 	return LocalInner->UseBinaryOrNativeSerialization(Ar);
 }
 
-bool FArrayProperty::LoadFromTag(const FPropertyTag& Tag)
-{
-	if (!Super::LoadFromTag(Tag))
-	{
-		return false;
-	}
-
-	PRAGMA_DISABLE_DEPRECATION_WARNINGS;
-	FField* Field = FField::TryConstruct(Tag.InnerType, {}, Tag.Name, RF_NoFlags);
-	if (FProperty* Property = CastField<FProperty>(Field))
-	{
-		FPropertyTag InnerTag = Tag;
-		InnerTag.Type = Tag.InnerType;
-		InnerTag.InnerType = {};
-		// Skip property types that are missing the name of the inner type.
-		// Structs have their name in a tag in the serialized data, but we cannot
-		// proceed safely unless we know if the struct used native serialization.
-		if (!Property->IsA<FStructProperty>() &&
-			!Property->IsA<FByteProperty>() &&
-			!Property->IsA<FEnumProperty>() &&
-			Property->LoadFromTag(InnerTag))
-		{
-			Inner = Property;
-			return true;
-		}
-	}
-	PRAGMA_ENABLE_DEPRECATION_WARNINGS;
-	delete Field;
-	return false;
-}
-
-void FArrayProperty::SaveToTag(FPropertyTag& Tag)
-{
-	Super::SaveToTag(Tag);
-
-	const FProperty* LocalInner = Inner;
-	check(LocalInner);
-	PRAGMA_DISABLE_DEPRECATION_WARNINGS;
-	Tag.InnerType = LocalInner->GetID();
-	PRAGMA_ENABLE_DEPRECATION_WARNINGS;
-}
-
 bool FArrayProperty::LoadTypeName(UE::FPropertyTypeName Type, const FPropertyTag* Tag)
 {
 	if (!Super::LoadTypeName(Type, Tag))
diff --git a/Engine/Source/Runtime/CoreUObject/Private/UObject/PropertyByte.cpp b/Engine/Source/Runtime/CoreUObject/Private/UObject/PropertyByte.cpp
--- a/Engine/Source/Runtime/CoreUObject/Private/UObject/PropertyByte.cpp
+++ b/Engine/Source/Runtime/CoreUObject/Private/UObject/PropertyByte.cpp
@@ -4,7 +4,6 @@
 
 #include "Algo/Find.h"
 #include "Hash/Blake3.h"
-#include "UObject/Package.h"
 #include "UObject/UnrealTypePrivate.h"
 #include "UObject/UObjectThreadContext.h"
 
@@ -516,52 +515,6 @@ uint64 FByteProperty::GetMaxNetSerializeBits() const
 	return FMath::Min(DesiredBits, MaxBits);
 }
 
-bool FByteProperty::LoadFromTag(const FPropertyTag& Tag)
-{
-	if (!Super::LoadFromTag(Tag))
-	{
-		return false;
-	}
-
-	PRAGMA_DISABLE_DEPRECATION_WARNINGS;
-	if (Tag.EnumName.IsNone())
-	{
-		return true;
-	}
-
-	// Update FEnumProperty when making changes here.
-	TStringBuilder<256> EnumName(InPlace, Tag.EnumName);
-	PRAGMA_ENABLE_DEPRECATION_WARNINGS;
-	if (UEnum* LocalEnum = FindFirstObject<UEnum>(*EnumName, EFindFirstObjectOptions::NativeFirst))
-	{
-		Enum = LocalEnum;
-		return true;
-	}
-
-	return false;
-}
-
-void FByteProperty::SaveToTag(FPropertyTag& Tag)
-{
-	Super::SaveToTag(Tag);
-
-	if (const UEnum* LocalEnum = Enum)
-	{
-		PRAGMA_DISABLE_DEPRECATION_WARNINGS;
-		// RobM: Ugly hack so that we can avoid content changes in most of the packages
-		// Update FEnumProperty when making changes here.
-		if (LocalEnum->GetPackage()->HasAnyPackageFlags(PKG_CompiledIn))
-		{				
-			Tag.EnumName = LocalEnum->GetFName();
-		}
-		else
-		{
-			Tag.EnumName = FName(*LocalEnum->GetPathName());
-		}
-		PRAGMA_ENABLE_DEPRECATION_WARNINGS;
-	}
-}
-
 bool FByteProperty::LoadTypeName(UE::FPropertyTypeName Type, const FPropertyTag* Tag)
 {
 	if (!Super::LoadTypeName(Type, Tag))
diff --git a/Engine/Source/Runtime/CoreUObject/Private/UObject/PropertyMap.cpp b/Engine/Source/Runtime/CoreUObject/Private/UObject/PropertyMap.cpp
--- a/Engine/Source/Runtime/CoreUObject/Private/UObject/PropertyMap.cpp
+++ b/Engine/Source/Runtime/CoreUObject/Private/UObject/PropertyMap.cpp
@@ -1731,30 +1731,6 @@ bool FMapProperty::UseBinaryOrNativeSerialization(const FArchive& Ar) const
 	return LocalKeyProp->UseBinaryOrNativeSerialization(Ar) || LocalValueProp->UseBinaryOrNativeSerialization(Ar);
 }
 
-bool FMapProperty::LoadFromTag(const FPropertyTag& Tag)
-{
-	if (!Super::LoadFromTag(Tag))
-	{
-		return false;
-	}
-
-	return false;
-}
-
-void FMapProperty::SaveToTag(FPropertyTag& Tag)
-{
-	Super::SaveToTag(Tag);
-
-	const FProperty* LocalKeyProp = KeyProp;
-	const FProperty* LocalValueProp = ValueProp;
-	check(LocalKeyProp);
-	check(LocalValueProp);
-	PRAGMA_DISABLE_DEPRECATION_WARNINGS;
-	Tag.InnerType = LocalKeyProp->GetID();
-	Tag.ValueType = LocalValueProp->GetID();
-	PRAGMA_ENABLE_DEPRECATION_WARNINGS;
-}
-
 bool FMapProperty::LoadTypeName(UE::FPropertyTypeName Type, const FPropertyTag* Tag)
 {
 	if (!Super::LoadTypeName(Type, Tag))
diff --git a/Engine/Source/Runtime/CoreUObject/Private/UObject/PropertyMulticastDelegate.cpp b/Engine/Source/Runtime/CoreUObject/Private/UObject/PropertyMulticastDelegate.cpp
--- a/Engine/Source/Runtime/CoreUObject/Private/UObject/PropertyMulticastDelegate.cpp
+++ b/Engine/Source/Runtime/CoreUObject/Private/UObject/PropertyMulticastDelegate.cpp
@@ -689,17 +689,6 @@ void FMulticastSparseDelegateProperty::ClearDelegate(UObject* Parent, void* Prop
 	SparseDelegate.__Internal_Clear(Parent, SparseDelegateFunc->DelegateName);
 }
 
-bool FMulticastSparseDelegateProperty::LoadFromTag(const FPropertyTag& Tag)
-{
-	if (!Super::LoadFromTag(Tag))
-	{
-		return false;
-	}
-
-	// This cannot be used without its SignatureFunction and the tag lacks the information needed to load it.
-	return false;
-}
-
 bool FMulticastSparseDelegateProperty::LoadTypeName(UE::FPropertyTypeName Type, const FPropertyTag* Tag)
 {
 	if (!Super::LoadTypeName(Type, Tag))
diff --git a/Engine/Source/Runtime/CoreUObject/Private/UObject/PropertyOptional.cpp b/Engine/Source/Runtime/CoreUObject/Private/UObject/PropertyOptional.cpp
--- a/Engine/Source/Runtime/CoreUObject/Private/UObject/PropertyOptional.cpp
+++ b/Engine/Source/Runtime/CoreUObject/Private/UObject/PropertyOptional.cpp
@@ -620,27 +620,6 @@ bool FOptionalProperty::UseBinaryOrNativeSerialization(const FArchive& Ar) const
 	return LocalValueProperty->UseBinaryOrNativeSerialization(Ar);
 }
 
-bool FOptionalProperty::LoadFromTag(const FPropertyTag& Tag)
-{
-	if (!Super::LoadFromTag(Tag))
-	{
-		return false;
-	}
-
-	return false;
-}
-
-void FOptionalProperty::SaveToTag(FPropertyTag& Tag)
-{
-	Super::SaveToTag(Tag);
-
-	const FProperty* LocalValueProperty = ValueProperty;
-	check(LocalValueProperty);
-	PRAGMA_DISABLE_DEPRECATION_WARNINGS;
-	Tag.InnerType = LocalValueProperty->GetID();
-	PRAGMA_ENABLE_DEPRECATION_WARNINGS;
-}
-
 bool FOptionalProperty::LoadTypeName(UE::FPropertyTypeName Type, const FPropertyTag* Tag)
 {
 	if (!Super::LoadTypeName(Type, Tag))
diff --git a/Engine/Source/Runtime/CoreUObject/Private/UObject/PropertySet.cpp b/Engine/Source/Runtime/CoreUObject/Private/UObject/PropertySet.cpp
--- a/Engine/Source/Runtime/CoreUObject/Private/UObject/PropertySet.cpp
+++ b/Engine/Source/Runtime/CoreUObject/Private/UObject/PropertySet.cpp
@@ -1102,27 +1102,6 @@ bool FSetProperty::UseBinaryOrNativeSerialization(const FArchive& Ar) const
 	return LocalElementProp->UseBinaryOrNativeSerialization(Ar);
 }
 
-bool FSetProperty::LoadFromTag(const FPropertyTag& Tag)
-{
-	if (!Super::LoadFromTag(Tag))
-	{
-		return false;
-	}
-
-	return false;
-}
-
-void FSetProperty::SaveToTag(FPropertyTag& Tag)
-{
-	Super::SaveToTag(Tag);
-
-	const FProperty* LocalElementProp = ElementProp;
-	check(LocalElementProp);
-	PRAGMA_DISABLE_DEPRECATION_WARNINGS;
-	Tag.InnerType = LocalElementProp->GetID();
-	PRAGMA_ENABLE_DEPRECATION_WARNINGS;
-}
-
 bool FSetProperty::LoadTypeName(UE::FPropertyTypeName Type, const FPropertyTag* Tag)
 {
 	if (!Super::LoadTypeName(Type, Tag))
diff --git a/Engine/Source/Runtime/CoreUObject/Private/UObject/PropertyStruct.cpp b/Engine/Source/Runtime/CoreUObject/Private/UObject/PropertyStruct.cpp
--- a/Engine/Source/Runtime/CoreUObject/Private/UObject/PropertyStruct.cpp
+++ b/Engine/Source/Runtime/CoreUObject/Private/UObject/PropertyStruct.cpp
@@ -454,79 +454,9 @@ void FStructProperty::AppendSchemaHash(FBlake3& Builder, bool bSkipEditorOnly) c
 }
 #endif
 
-static const FName NAME_StructOriginalType(ANSITEXTVIEW("OriginalType"));
-
-static const FString* FindOriginalTypeName(const FStructProperty* Struct)
-{
-	FUObjectSerializeContext* Context = FUObjectThreadContext::Get().GetSerializeContext();
-	if (Context && Context->bImpersonateProperties)
-	{
 #if WITH_EDITORONLY_DATA
-		if (const FString* OriginalType = Struct->FindMetaData(NAME_StructOriginalType))
-		{
-			return OriginalType;
-		}
-		//@note: To support metadata defined on array of struct in UPROPERTY for testing purposes
-		if (FField* OwnerField = Struct->Owner.ToField())
-		{
-			return OwnerField->FindMetaData(NAME_StructOriginalType);
-		}
-#endif
-	}
-	return nullptr;
-}
-
-bool FStructProperty::LoadFromTag(const FPropertyTag& Tag)
-{
-	if (!Super::LoadFromTag(Tag))
-	{
-		return false;
-	}
-
-	PRAGMA_DISABLE_DEPRECATION_WARNINGS;
-	if (!Tag.StructName.IsNone())
-	{
-		TStringBuilder<256> StructName(InPlace, Tag.StructName);
-		if (UScriptStruct* LocalStruct = FindFirstObject<UScriptStruct>(*StructName, EFindFirstObjectOptions::NativeFirst))
-		{
-			Struct = LocalStruct;
-			return true;
-		}
-	}
-	PRAGMA_ENABLE_DEPRECATION_WARNINGS;
-
-	return false;
-}
-
-void FStructProperty::SaveToTag(FPropertyTag& Tag)
-{
-	Super::SaveToTag(Tag);
-
-	const FString* OriginalType = FindOriginalTypeName(this);
-	const UScriptStruct* LocalStruct = Struct;
-	check(LocalStruct);
-	PRAGMA_DISABLE_DEPRECATION_WARNINGS;
-	Tag.StructName = OriginalType ? FName(**OriginalType) : LocalStruct->GetFName();
-	Tag.StructGuid = LocalStruct->GetCustomGuid();
-	PRAGMA_ENABLE_DEPRECATION_WARNINGS;
-}
-
-void FStructProperty::AssignToTag(FPropertyTag& Tag)
-{
-	Super::AssignToTag(Tag);
-
-	const FString* OriginalType = FindOriginalTypeName(this);
-	const UScriptStruct* LocalStruct = Struct;
-	check(LocalStruct);
-	PRAGMA_DISABLE_DEPRECATION_WARNINGS;
-	if (OriginalType && FName(**OriginalType) == Tag.StructName)
-	{
-		Tag.StructName = LocalStruct->GetFName();
-	}
-	PRAGMA_ENABLE_DEPRECATION_WARNINGS;
-}
+static const FName NAME_StructOriginalType(ANSITEXTVIEW("OriginalType"));
 
-#if WITH_EDITORONLY_DATA
 static UE::FPropertyTypeName FindOriginalType(const FStructProperty* Struct)
 {
 	FUObjectSerializeContext* Context = FUObjectThreadContext::Get().GetSerializeContext();
diff --git a/Engine/Source/Runtime/CoreUObject/Public/UObject/EnumProperty.h b/Engine/Source/Runtime/CoreUObject/Public/UObject/EnumProperty.h
--- a/Engine/Source/Runtime/CoreUObject/Public/UObject/EnumProperty.h
+++ b/Engine/Source/Runtime/CoreUObject/Public/UObject/EnumProperty.h
@@ -76,8 +76,6 @@ public:
 #if WITH_EDITORONLY_DATA
 	COREUOBJECT_API virtual void AppendSchemaHash(FBlake3& Builder, bool bSkipEditorOnly) const override;
 #endif
-	COREUOBJECT_API virtual bool LoadFromTag(const FPropertyTag& Tag) override;
-	COREUOBJECT_API virtual void SaveToTag(FPropertyTag& Tag) override;
 	COREUOBJECT_API virtual bool LoadTypeName(UE::FPropertyTypeName Type, const FPropertyTag* Tag = nullptr) override;
 	COREUOBJECT_API virtual void SaveTypeName(UE::FPropertyTypeNameBuilder& Type) const override;
 	COREUOBJECT_API virtual bool CanSerializeFromTypeName(UE::FPropertyTypeName Type) const override;
diff --git a/Engine/Source/Runtime/CoreUObject/Public/UObject/PropertyOptional.h b/Engine/Source/Runtime/CoreUObject/Public/UObject/PropertyOptional.h
--- a/Engine/Source/Runtime/CoreUObject/Public/UObject/PropertyOptional.h
+++ b/Engine/Source/Runtime/CoreUObject/Public/UObject/PropertyOptional.h
@@ -230,8 +230,6 @@ public:
 	virtual uint32 GetValueTypeHashInternal(const void* Src) const override;
 	virtual void AddReferencedObjects(FReferenceCollector& Collector) override;
 	virtual bool UseBinaryOrNativeSerialization(const FArchive& Ar) const override;
-	virtual bool LoadFromTag(const FPropertyTag& Tag) override;
-	virtual void SaveToTag(FPropertyTag& Tag) override;
 	virtual bool LoadTypeName(UE::FPropertyTypeName Type, const FPropertyTag* Tag = nullptr) override;
 	virtual void SaveTypeName(UE::FPropertyTypeNameBuilder& Type) const override;
 	virtual bool CanSerializeFromTypeName(UE::FPropertyTypeName Type) const override;
diff --git a/Engine/Source/Runtime/CoreUObject/Public/UObject/UnrealType.h b/Engine/Source/Runtime/CoreUObject/Public/UObject/UnrealType.h
--- a/Engine/Source/Runtime/CoreUObject/Public/UObject/UnrealType.h
+++ b/Engine/Source/Runtime/CoreUObject/Public/UObject/UnrealType.h
@@ -1042,23 +1042,6 @@ public:
 		return ShouldPort() && GetOwnerClass() != UObject::StaticClass();
 	}
 
-	/**
-	 * Loads types and other details of the property from the tag. Must be a compatible type.
-	 *
-	 * @return true if this property is usable after loading from the tag.
-	 */
-	COREUOBJECT_API virtual bool LoadFromTag(const FPropertyTag& Tag);
-
-	/**
-	 * Saves types and other details of the property to the tag.
-	 */
-	COREUOBJECT_API virtual void SaveToTag(FPropertyTag& Tag);
-
-	/**
-	 * Assign the property to tag and optionally modify other fields
-	 */
-	COREUOBJECT_API virtual void AssignToTag(FPropertyTag& Tag);
-
 	/**
 	 * Restores this property and its owned properties from the type name.
 	 *
@@ -1948,8 +1931,6 @@ public:
 #if WITH_EDITORONLY_DATA
 	virtual void AppendSchemaHash(FBlake3& Builder, bool bSkipEditorOnly) const override;
 #endif
-	virtual bool LoadFromTag(const FPropertyTag& Tag) override;
-	virtual void SaveToTag(FPropertyTag& Tag) override;
 	virtual bool LoadTypeName(UE::FPropertyTypeName Type, const FPropertyTag* Tag = nullptr) override;
 	virtual void SaveTypeName(UE::FPropertyTypeNameBuilder& Type) const override;
 	virtual bool CanSerializeFromTypeName(UE::FPropertyTypeName Type) const override;
@@ -3461,8 +3442,6 @@ public:
 
 	virtual void* GetValueAddressAtIndex_Direct(const FProperty* Inner, void* InValueAddress, int32 Index) const override;
 	virtual bool UseBinaryOrNativeSerialization(const FArchive& Ar) const override;
-	virtual bool LoadFromTag(const FPropertyTag& Tag) override;
-	virtual void SaveToTag(FPropertyTag& Tag) override;
 	virtual bool LoadTypeName(UE::FPropertyTypeName Type, const FPropertyTag* Tag = nullptr) override;
 	virtual void SaveTypeName(UE::FPropertyTypeNameBuilder& Type) const override;
 	virtual bool CanSerializeFromTypeName(UE::FPropertyTypeName Type) const override;
@@ -3595,8 +3574,6 @@ public:
 	virtual EConvertFromTypeResult ConvertFromType(const FPropertyTag& Tag, FStructuredArchive::FSlot Slot, uint8* Data, UStruct* DefaultsStruct, const uint8* Defaults) override;
 	virtual void* GetValueAddressAtIndex_Direct(const FProperty* Inner, void* InValueAddress, int32 LogicalIndex) const override;
 	virtual bool UseBinaryOrNativeSerialization(const FArchive& Ar) const override;
-	virtual bool LoadFromTag(const FPropertyTag& Tag) override;
-	virtual void SaveToTag(FPropertyTag& Tag) override;
 	virtual bool LoadTypeName(UE::FPropertyTypeName Type, const FPropertyTag* Tag = nullptr) override;
 	virtual void SaveTypeName(UE::FPropertyTypeNameBuilder& Type) const override;
 	virtual bool CanSerializeFromTypeName(UE::FPropertyTypeName Type) const override;
@@ -3724,8 +3701,6 @@ public:
 	virtual EConvertFromTypeResult ConvertFromType(const FPropertyTag& Tag, FStructuredArchive::FSlot Slot, uint8* Data, UStruct* DefaultsStruct, const uint8* Defaults) override;
 	virtual void* GetValueAddressAtIndex_Direct(const FProperty* Inner, void* InValueAddress, int32 LogicalIndex) const override;
 	virtual bool UseBinaryOrNativeSerialization(const FArchive& Ar) const override;
-	virtual bool LoadFromTag(const FPropertyTag& Tag) override;
-	virtual void SaveToTag(FPropertyTag& Tag) override;
 	virtual bool LoadTypeName(UE::FPropertyTypeName Type, const FPropertyTag* Tag = nullptr) override;
 	virtual void SaveTypeName(UE::FPropertyTypeNameBuilder& Type) const override;
 	virtual bool CanSerializeFromTypeName(UE::FPropertyTypeName Type) const override;
@@ -5847,9 +5822,6 @@ public:
 	virtual void AppendSchemaHash(FBlake3& Builder, bool bSkipEditorOnly) const override;
 #endif
 	virtual bool UseBinaryOrNativeSerialization(const FArchive& Ar) const override;
-	virtual bool LoadFromTag(const FPropertyTag& Tag) override;
-	virtual void SaveToTag(FPropertyTag& Tag) override;
-	virtual void AssignToTag(FPropertyTag& Tag) override;
 	virtual bool LoadTypeName(UE::FPropertyTypeName Type, const FPropertyTag* Tag = nullptr) override;
 	virtual void SaveTypeName(UE::FPropertyTypeNameBuilder& Type) const override;
 	virtual bool CanSerializeFromTypeName(UE::FPropertyTypeName Type) const override;
@@ -6127,7 +6099,6 @@ public:
 	// FProperty interface
 	virtual void SerializeItem(FStructuredArchive::FSlot Slot, void* Value, void const* Defaults) const override;
 	virtual const TCHAR* ImportText_Internal(const TCHAR* Buffer, void* ContainerOrPropertyPtr, EPropertyPointerType PropertyPointerType, UObject* OwnerObject, int32 PortFlags, FOutputDevice* ErrorText) const override;
-	virtual bool LoadFromTag(const FPropertyTag& Tag) override;
 	virtual bool LoadTypeName(UE::FPropertyTypeName Type, const FPropertyTag* Tag = nullptr) override;
 	// End of FProperty interface
 
