Fixed property bag tests to use the new SetType function on FPropertyTag

#rnx

diff --git a/Engine/Source/Runtime/CoreUObject/Tests/PropertyBagTest.cpp b/Engine/Source/Runtime/CoreUObject/Tests/PropertyBagTest.cpp
--- a/Engine/Source/Runtime/CoreUObject/Tests/PropertyBagTest.cpp
+++ b/Engine/Source/Runtime/CoreUObject/Tests/PropertyBagTest.cpp
@@ -4,21 +4,22 @@
 
 #include "PropertyBagTest.h"
 
-#include "UObject/Class.h"
-#include "UObject/PropertyBag.h"
-#include "UObject/PropertyBagRepository.h"
-#include "UObject/UObjectGlobals.h"
-#include "UObject/CoreRedirects.h"
-#include "UObject/SavePackage.h"
-#include "UObject/UObjectThreadContext.h"
-#include "UObject/UObjectHash.h"
+#include "HAL/FileManager.h"
 #include "Misc/StringBuilder.h"
 #include "Misc/PackageName.h"
-#include "HAL/FileManager.h"
 #include "Serialization/MemoryReader.h"
 #include "Serialization/MemoryWriter.h"
 #include "Tasks/Task.h"
 #include "Tests/TestHarnessAdapter.h"
+#include "UObject/Class.h"
+#include "UObject/CoreRedirects.h"
+#include "UObject/PropertyBag.h"
+#include "UObject/PropertyBagRepository.h"
+#include "UObject/PropertyTypeName.h"
+#include "UObject/SavePackage.h"
+#include "UObject/UObjectGlobals.h"
+#include "UObject/UObjectHash.h"
+#include "UObject/UObjectThreadContext.h"
 
 DEFINE_LOG_CATEGORY_STATIC(LogPropertyBagTests, Log, All);
 
@@ -273,8 +274,11 @@ namespace PropertyBagTestUtils
 	{
 		TArray<uint8> Buffer = EncodeToBuffer(InValue);
 
+		FPropertyTypeNameBuilder TypeBuilder;
+		TypeBuilder.AddName(GetENameForType<T>());
+
 		FPropertyTag Tag;
-		Tag.Type = GetENameForType<T>();
+		Tag.SetType(TypeBuilder.Build());
 		Tag.Name = FName(TEXT("TagTmp"));
 
 		LoadDataByTag(Dst, Path, Tag, Buffer);
