IncrementalCook: Add FCookDependency to record external-file and generic-function dependencies that should invalidate the cook-saved version of a package when they change.
#rnx

diff --git a/Engine/Source/Editor/UnrealEd/Private/CookOnTheFlyServer.cpp b/Engine/Source/Editor/UnrealEd/Private/CookOnTheFlyServer.cpp
--- a/Engine/Source/Editor/UnrealEd/Private/CookOnTheFlyServer.cpp
+++ b/Engine/Source/Editor/UnrealEd/Private/CookOnTheFlyServer.cpp
@@ -6392,7 +6392,7 @@ void FSaveCookedPackageContext::FinishPlatform()
 		if (COTFS.bHybridIterativeEnabled)
 		{
 			UE_SCOPED_HIERARCHICAL_COOKTIMER(TargetDomainDependencies);
-			UE::TargetDomain::CollectAndStoreCookAttachments(Package, TargetPlatform, Info.Attachments);
+			UE::TargetDomain::CollectAndStoreCookAttachments(Package, TargetPlatform, &SavePackageResult, Info.Attachments);
 		}
 		if (bSuccessful)
 		{
@@ -10774,6 +10774,14 @@ namespace UE::Cook
 void FBeginCookConfigSettings::LoadLocal(FBeginCookContext& BeginContext)
 {
 	GConfig->GetBool(TEXT("CookSettings"), TEXT("HybridIterativeEnabled"), bHybridIterativeEnabled, GEditorIni);
+	if (FParse::Param(FCommandLine::Get(), TEXT("CookIncremental")))
+	{
+		bHybridIterativeEnabled = true;
+	}
+	else if (FParse::Param(FCommandLine::Get(), TEXT("NoCookIncremental")))
+	{
+		bHybridIterativeEnabled = false;
+	}
 	// TODO: HybridIterative is not yet implemented for DLC
 	bHybridIterativeEnabled &= !BeginContext.COTFS.IsCookingDLC();
 	// HybridIterative uses TargetDomain storage of dependencies which is only implemented in ZenStore
diff --git a/Engine/Source/Editor/UnrealEd/Private/EditorDomain/EditorDomainUtils.cpp b/Engine/Source/Editor/UnrealEd/Private/EditorDomain/EditorDomainUtils.cpp
--- a/Engine/Source/Editor/UnrealEd/Private/EditorDomain/EditorDomainUtils.cpp
+++ b/Engine/Source/Editor/UnrealEd/Private/EditorDomain/EditorDomainUtils.cpp
@@ -1917,7 +1917,7 @@ bool TrySavePackage(UPackage* Package)
 	if (bStorageResultValid)
 	{
 		TArray<IPackageWriter::FCommitAttachmentInfo, TInlineAllocator<2>> Attachments;
-		UE::TargetDomain::CollectAndStoreCookAttachments(Package, nullptr, Attachments);
+		UE::TargetDomain::CollectAndStoreCookAttachments(Package, nullptr, nullptr, Attachments);
 		if (!Attachments.IsEmpty())
 		{
 			UE::TargetDomain::CommitEditorDomainCookAttachments(Package->GetFName(), Attachments);
diff --git a/Engine/Source/Editor/UnrealEd/Private/TargetDomain/TargetDomainUtils.cpp b/Engine/Source/Editor/UnrealEd/Private/TargetDomain/TargetDomainUtils.cpp
--- a/Engine/Source/Editor/UnrealEd/Private/TargetDomain/TargetDomainUtils.cpp
+++ b/Engine/Source/Editor/UnrealEd/Private/TargetDomain/TargetDomainUtils.cpp
@@ -9,7 +9,9 @@
 #include "Containers/Set.h"
 #include "Containers/UnrealString.h"
 #include "Cooker/CookConfigAccessTracker.h"
+#include "Cooker/CookDependency.h"
 #include "Cooker/PackageBuildDependencyTracker.h"
+#include "CookOnTheSide/CookLog.h"
 #include "DerivedDataBuildDefinition.h"
 #include "DerivedDataBuildKey.h"
 #include "DerivedDataSharedString.h"
@@ -17,6 +19,7 @@
 #include "EditorDomain/EditorDomainUtils.h"
 #include "HAL/PlatformFile.h"
 #include "HAL/PlatformFileManager.h"
+#include "Hash/Blake3.h"
 #include "IO/IoDispatcher.h"
 #include "IO/IoHash.h"
 #include "Misc/App.h"
@@ -80,6 +83,17 @@ private:
 };
 TUniquePtr<FEditorDomainOplog> GEditorDomainOplog;
 
+// Constructor/Destructor defined here in cpp rather than header so we can 
+// avoid needing the definition of FCookDependency in the header; it is needed
+// for construct/destruct of TArray<FCookDependency>.
+FCookDependencies::FCookDependencies() = default;
+FCookDependencies::~FCookDependencies() = default;
+FCookDependencies::FCookDependencies(const FCookDependencies&) = default;
+FCookDependencies::FCookDependencies(FCookDependencies&&) = default;
+FCookDependencies& FCookDependencies::operator=(const FCookDependencies&) = default;
+FCookDependencies& FCookDependencies::operator=(FCookDependencies&&) = default;
+
+
 bool FCookDependencies::IsValid() const
 {
 	return bValid;
@@ -170,6 +184,36 @@ bool FCookDependencies::TryCalculateCurrentKey(FString* OutErrorMessage)
 		}
 	}
 
+	if (!CookDependencies.IsEmpty())
+	{
+		bool bError = false;
+		UE::Cook::FCookDependencyContext Context(&KeyBuilder, [&bError, OutErrorMessage](FString&& ErrorMessage)
+			{
+				if (OutErrorMessage)
+				{
+					if (bError)
+					{
+						*OutErrorMessage += TEXT("\n");
+						*OutErrorMessage += ErrorMessage;
+					}
+					else
+					{
+						*OutErrorMessage = MoveTemp(ErrorMessage);
+					}
+				}
+				bError = true;
+			});
+
+		for (UE::Cook::FCookDependency& CookDependency : CookDependencies)
+		{
+			CookDependency.UpdateHash(Context);
+		}
+		if (bError)
+		{
+			return false;
+		}
+	}
+
 	if (OutErrorMessage) OutErrorMessage->Reset();
 	CurrentKey = KeyBuilder.Finalize();
 	return true;
@@ -194,7 +238,8 @@ void FCookDependencies::Empty()
 	RuntimePackageDependencies.Empty();
 }
 
-FCookDependencies FCookDependencies::Collect(UPackage* Package, const ITargetPlatform* TargetPlatform, FString* OutErrorMessage)
+FCookDependencies FCookDependencies::Collect(UPackage* Package, const ITargetPlatform* TargetPlatform,
+	FSavePackageResultStruct* SaveResult, FString* OutErrorMessage)
 {
 	if (!Package)
 	{
@@ -292,6 +337,12 @@ FCookDependencies FCookDependencies::Collect(UPackage* Package, const ITargetPla
 		}
 	}
 #endif
+	if (SaveResult)
+	{
+		Result.CookDependencies = MoveTemp(SaveResult->CookDependencies);
+		Algo::Sort(Result.CookDependencies);
+	}
+
 	if (!Result.TryCalculateCurrentKey(OutErrorMessage))
 	{
 		return FCookDependencies();
@@ -311,46 +362,57 @@ bool LoadFromCompactBinary(FCbObjectView ObjectView, UE::TargetDomain::FCookDepe
 
 	Dependencies.Reset();
 	int32 Version = -1;
-	for (FCbFieldView FieldView : ObjectView)
+
+	for (FCbFieldViewIterator FieldView(ObjectView.CreateViewIterator()); FieldView; )
 	{
-		FUtf8StringView FieldName = FieldView.GetName();
-		if (FieldName == "Version")
+		const FCbFieldViewIterator Last = FieldView;
+		if (FieldView.GetName().Equals(UTF8TEXTVIEW("Version")))
 		{
 			Version = FieldView.AsInt32();
-			if (FieldView.HasError() || Version != CookDependenciesVersion)
+			if ((FieldView++).HasError() || Version != CookDependenciesVersion)
 			{
 				return false;
 			}
 		}
-		else if (FieldName == "StoredKey")
+		if (FieldView.GetName().Equals(UTF8TEXTVIEW("StoredKey")))
 		{
-			if (!LoadFromCompactBinary(FieldView, Dependencies.StoredKey))
+			if (!LoadFromCompactBinary(FieldView++, Dependencies.StoredKey))
 			{
 				return false;
 			}
 		}
-		else if (FieldName == "PackageDependencies")
+		if (FieldView.GetName().Equals(UTF8TEXTVIEW("PackageDependencies")))
 		{
-			if (!LoadFromCompactBinary(FieldView, Dependencies.PackageDependencies))
+			if (!LoadFromCompactBinary(FieldView++, Dependencies.PackageDependencies))
 			{
 				return false;
 			}
 		}
-		else if (FieldName == "ConfigDependencies")
+		if (FieldView.GetName().Equals(UTF8TEXTVIEW("ConfigDependencies")))
 		{
-			if (!LoadFromCompactBinary(FieldView, Dependencies.ConfigDependencies))
+			if (!LoadFromCompactBinary(FieldView++, Dependencies.ConfigDependencies))
 			{
 				return false;
 			}
 		}
-		else if (FieldName == "RuntimePackageDependencies")
+		if (FieldView.GetName().Equals(UTF8TEXTVIEW("RuntimePackageDependencies")))
 		{
-			if (!LoadFromCompactBinary(FieldView, Dependencies.RuntimePackageDependencies))
+			if (!LoadFromCompactBinary(FieldView++, Dependencies.RuntimePackageDependencies))
 			{
 				return false;
 			}
 		}
-		// else ignore
+		if (FieldView.GetName().Equals(UTF8TEXTVIEW("CookDependencies")))
+		{
+			if (!LoadFromCompactBinary(FieldView++, Dependencies.CookDependencies))
+			{
+				return false;
+			}
+		}
+		if (FieldView == Last)
+		{
+			++FieldView;
+		}
 	}
 	if (Version == -1)
 	{
@@ -379,6 +441,11 @@ FCbWriter& operator<<(FCbWriter& Writer, const UE::TargetDomain::FCookDependenci
 	{
 		Writer << "RuntimePackageDependencies" << CookDependencies.RuntimePackageDependencies;
 	}
+	if (!CookDependencies.CookDependencies.IsEmpty())
+	{
+		Writer << "CookDependencies" << CookDependencies.CookDependencies;
+	}
+
 	Writer.EndObject();
 	return Writer;
 }
@@ -479,11 +546,17 @@ void FCookAttachments::Empty()
 }
 
 bool TryCollectAndStoreCookDependencies(UPackage* Package, const ITargetPlatform* TargetPlatform,
-	IPackageWriter::FCommitAttachmentInfo& OutResult)
+	FSavePackageResultStruct* SaveResult, IPackageWriter::FCommitAttachmentInfo& OutResult)
 {
-	FCookDependencies CookDependencies = FCookDependencies::Collect(Package, TargetPlatform);
+	FString ErrorMessage;
+	FCookDependencies CookDependencies = FCookDependencies::Collect(Package, TargetPlatform, SaveResult, &ErrorMessage);
 	if (!CookDependencies.IsValid())
 	{
+		// CookPackageSplitterTODO: This error occurs for generated packages. Need to register them with EditorDomain.
+#if 0
+		UE_LOG(LogCook, Error, TEXT("Could not collect CookDependencies for package '%s': %s"),
+			*Package->GetName(), *ErrorMessage);
+#endif
 		OutResult.Value = FCbObject();
 		return false;
 	}
diff --git a/Engine/Source/Editor/UnrealEd/Private/TargetDomain/TargetDomainUtils.h b/Engine/Source/Editor/UnrealEd/Private/TargetDomain/TargetDomainUtils.h
--- a/Engine/Source/Editor/UnrealEd/Private/TargetDomain/TargetDomainUtils.h
+++ b/Engine/Source/Editor/UnrealEd/Private/TargetDomain/TargetDomainUtils.h
@@ -17,8 +17,10 @@ class FCbWriter;
 class FString;
 class ITargetPlatform;
 class UPackage;
+struct FSavePackageResultStruct;
 template <typename FuncType> class TUniqueFunction;
 
+namespace UE::Cook { class FCookDependency; }
 namespace UE::TargetDomain { class FCookDependencies; }
 namespace UE::TargetDomain { struct FBuildDefinitionList; }
 namespace UE::TargetDomain { struct FCookAttachments; }
@@ -46,6 +48,12 @@ void CookInitialize();
 class FCookDependencies
 {
 public:
+	FCookDependencies();
+	~FCookDependencies();
+	FCookDependencies(const FCookDependencies&);
+	FCookDependencies(FCookDependencies&&);
+	FCookDependencies& operator=(const FCookDependencies&);
+	FCookDependencies& operator=(FCookDependencies&&);
 	// Build Dependencies
 	const TArray<FName>& GetPackageDependencies() const { return PackageDependencies; }
 	const TArray<FString>& GetConfigDependencies() const { return ConfigDependencies; }
@@ -53,6 +61,9 @@ public:
 	// Runtime Dependencies
 	const TArray<FName>& GetRuntimePackageDependencies() const { return RuntimePackageDependencies; }
 
+	// Cook Dependencies
+	const TArray<UE::Cook::FCookDependency>& GetCookDependencies() const { return CookDependencies; }
+
 	// Data about the Key and Package that are not dependencies
 	/**
 	 * True if the structure has been calculated or fetched and accurately reports dependencies and
@@ -75,7 +86,7 @@ public:
 	 * that have recorded its data during the package's load/other/save operations in the current cook session.
 	 */
 	static FCookDependencies Collect(UPackage* Package, const ITargetPlatform* TargetPlatform,
-		FString* OutErrorMessage = nullptr);
+		FSavePackageResultStruct* SaveResult, FString* OutErrorMessage = nullptr);
 
 	// Fetch function to load the dependencies from a PackageStore is not yet implemented independently for
 	// this structure. Use FCookAttachments instead. 
@@ -84,6 +95,7 @@ private:
 	TArray<FName> PackageDependencies;
 	TArray<FString> ConfigDependencies;
 	TArray<FName> RuntimePackageDependencies;
+	TArray<UE::Cook::FCookDependency> CookDependencies;
 	FName PackageName;
 	FIoHash StoredKey;
 	FIoHash CurrentKey;
@@ -130,16 +142,16 @@ struct FCookAttachments
 };
 
 bool TryCollectAndStoreCookDependencies(UPackage* Package, const ITargetPlatform* TargetPlatform,
-	IPackageWriter::FCommitAttachmentInfo& OutResult);
+	FSavePackageResultStruct* SaveResult, IPackageWriter::FCommitAttachmentInfo& OutResult);
 bool TryCollectAndStoreBuildDefinitionList(UPackage* Package, const ITargetPlatform* TargetPlatform,
 	IPackageWriter::FCommitAttachmentInfo& OutResult);
 
 template <typename ArrayType>
 void CollectAndStoreCookAttachments(UPackage* Package, const ITargetPlatform* TargetPlatform,
-	ArrayType& Output)
+	FSavePackageResultStruct* SaveResult, ArrayType& Output)
 {
 	IPackageWriter::FCommitAttachmentInfo Result;
-	if (TryCollectAndStoreCookDependencies(Package, TargetPlatform, Result))
+	if (TryCollectAndStoreCookDependencies(Package, TargetPlatform, SaveResult, Result))
 	{
 		Output.Add(MoveTemp(Result));
 	}
diff --git a/Engine/Source/Runtime/CoreUObject/Private/Cooker/CookDependency.cpp b/Engine/Source/Runtime/CoreUObject/Private/Cooker/CookDependency.cpp
new file mode 100644
--- /dev/null
+++ b/Engine/Source/Runtime/CoreUObject/Private/Cooker/CookDependency.cpp
@@ -0,0 +1,397 @@
+// Copyright Epic Games, Inc. All Rights Reserved.
+
+#include "Cooker/CookDependency.h"
+
+#if WITH_EDITOR
+#include "Containers/Array.h"
+#include "Containers/Map.h"
+#include "CoreGlobals.h"
+#include "HAL/FileManager.h"
+#include "HAL/PlatformMath.h"
+#include "Hash/Blake3.h"
+#include "Misc/AssertionMacros.h"
+#include "Misc/StringBuilder.h"
+#include "Serialization/Archive.h"
+#include "Serialization/CompactBinary.h"
+#include "Serialization/CompactBinarySerialization.h"
+#include "Serialization/CompactBinaryWriter.h"
+#include "Templates/UniquePtr.h"
+#endif
+
+#if WITH_EDITOR
+
+namespace UE::Cook::Dependency::Private
+{
+
+/**
+ * Return the map from Name to function created by UE_COOK_DEPENDENCY_FUNCTION macros.
+ * Initialized when first called and afterwards immutable
+ */
+const TMap<FName, FCookDependencyFunction>& GetCookDependencyFunctions();
+
+} // namespace UE::Cook::Dependency::Private
+
+namespace UE::Cook
+{
+
+FCookDependency FCookDependency::File(FStringView InFileName)
+{
+	FCookDependency Result(ECookDependency::File);
+	Result.FileName = InFileName;
+	return Result;
+}
+
+FCookDependency FCookDependency::Function(FName InFunctionName, FCbFieldIterator&& InArgs)
+{
+	FCookDependency Result(ECookDependency::Function);
+	Result.FunctionData.Name = InFunctionName;
+	Result.FunctionData.Args = MoveTemp(InArgs);
+	Result.FunctionData.Args.MakeOwned();
+	return Result;
+}
+
+FCookDependency::FCookDependency()
+	: Type(ECookDependency::None)
+{
+}
+
+FCookDependency::FCookDependency(ECookDependency InType)
+	: Type(InType)
+{
+	Construct();
+}
+
+FCookDependency::~FCookDependency()
+{
+	Destruct();
+}
+
+FCookDependency::FCookDependency(const FCookDependency& Other)
+	: Type(ECookDependency::None)
+{
+	*this = Other;
+}
+
+FCookDependency::FCookDependency(FCookDependency&& Other)
+	: Type(ECookDependency::None)
+{
+	*this = MoveTemp(Other);
+}
+
+FCookDependency& FCookDependency::operator=(const FCookDependency& Other)
+{
+	Destruct();
+	Type = Other.Type;
+	Construct();
+	switch (Type)
+	{
+	case ECookDependency::None:
+		break;
+	case ECookDependency::File:
+		FileName = Other.FileName;
+		break;
+	case ECookDependency::Function:
+		FunctionData = Other.FunctionData;
+		break;
+	default:
+		checkNoEntry();
+		break;
+	}
+	return *this;
+}
+
+FCookDependency& FCookDependency::operator=(FCookDependency&& Other)
+{
+	Destruct();
+	Type = Other.Type;
+	Construct();
+	switch (Type)
+	{
+	case ECookDependency::None:
+		break;
+	case ECookDependency::File:
+		FileName = MoveTemp(Other.FileName);
+		break;
+	case ECookDependency::Function:
+		FunctionData = MoveTemp(Other.FunctionData);
+		break;
+	default:
+		checkNoEntry();
+		break;
+	}
+	return *this;
+}
+
+void FCookDependency::UpdateHash(FCookDependencyContext& Context) const
+{
+	switch (GetType())
+	{
+	case ECookDependency::None:
+		// Nothing to add, Nones are never invalidated
+		return;
+	case ECookDependency::File:
+	{
+		const TCHAR* LocalFilename = GetFileName().GetData(); // contract for GetFileName is a null-terminated TCHAR*.
+		LocalFilename = LocalFilename ? LocalFilename : TEXT("");
+		TUniquePtr<FArchive> Ar(IFileManager::Get().CreateFileReader(LocalFilename, FILEREAD_Silent));
+		if (!Ar)
+		{
+			Context.LogError(FString::Printf(
+				TEXT("FCookDependency::File('%s') failed to UpdateHash: could not read file."),
+				LocalFilename));
+			return;
+		}
+		TArray64<uint8> Buffer;
+		Buffer.SetNumUninitialized(64 * 1024);
+		const int64 Size = Ar->TotalSize();
+		int64 Position = 0;
+		// Read in BufferSize chunks
+		while (Position < Size)
+		{
+			const int64 ReadNum = FMath::Min(Size - Position, (int64)Buffer.Num());
+			Ar->Serialize(Buffer.GetData(), static_cast<uint64>(ReadNum));
+			Context.Update(Buffer.GetData(), ReadNum);
+			Position += ReadNum;
+		}
+		break;
+	}
+	case ECookDependency::Function:
+	{
+		const TMap<FName, FCookDependencyFunction>& CookDependencyFunctions =
+			UE::Cook::Dependency::Private::GetCookDependencyFunctions();
+
+		const FCookDependencyFunction* Function = CookDependencyFunctions.Find(GetFunctionName());
+		if (!Function)
+		{
+			Context.LogError(FString::Printf(
+				TEXT("FCookDependency::Function('%s') failed to UpdateHash: Function not found."),
+				*GetFunctionName().ToString()));
+			return;
+		}
+		FCookDependencyContext::FErrorHandlerScope Scope = Context.ErrorHandlerScope([this](FString&& Message)
+			{
+				return FString::Printf(TEXT("FCookDependency::Function('%s') failed to UpdateHash: %s"),
+				*GetFunctionName().ToString(), *Message);
+			});
+		(*Function)(GetFunctionArgs(), Context);
+		return;
+	}
+	default:
+		checkNoEntry();
+		return;
+	}
+}
+
+void FCookDependency::Construct()
+{
+	switch (Type)
+	{
+	case ECookDependency::None:
+		break;
+	case ECookDependency::File:
+		new(&FileName) FString();
+		break;
+	case ECookDependency::Function:
+		new(&FunctionData) FFunctionData();
+		break;
+	default:
+		checkNoEntry();
+		break;
+	}
+}
+
+void FCookDependency::Destruct()
+{
+	switch (Type)
+	{
+	case ECookDependency::None:
+		break;
+	case ECookDependency::File:
+		FileName.~FString();
+		break;
+	case ECookDependency::Function:
+		FunctionData.~FFunctionData();
+		break;
+	default:
+		checkNoEntry();
+		break;
+	}
+}
+
+void FCookDependency::Save(FCbWriter& Writer) const
+{
+	Writer.BeginArray();
+	Writer << static_cast<uint8>(Type);
+	switch (Type)
+	{
+	case ECookDependency::None:
+		break;
+	case ECookDependency::File:
+		Writer << FileName;
+		break;
+	case ECookDependency::Function:
+		Writer << FunctionData.Name;
+		for (FCbFieldViewIterator Iter(FunctionData.Args); Iter; ++Iter)
+		{
+			Writer << Iter;
+		}
+		break;
+	default:
+		checkNoEntry();
+		break;
+	}
+	Writer.EndArray();
+}
+
+bool FCookDependency::Load(FCbFieldView Value)
+{
+	*this = FCookDependency(ECookDependency::None);
+	FCbArrayView ArrayView = Value.AsArrayView();
+	if (ArrayView.Num() < 1)
+	{
+		return false;
+	}
+	FCbFieldViewIterator Field(Value.CreateViewIterator());
+	ECookDependency LocalType = static_cast<ECookDependency>(Field.AsUInt8());
+	if ((Field++).HasError())
+	{
+		return false;
+	}
+
+	switch (LocalType)
+	{
+	case ECookDependency::None:
+		return true;
+	case ECookDependency::File:
+	{
+		FUtf8StringView LocalFilename = Field.AsString();
+		if ((Field++).HasError())
+		{
+			return false;
+		}
+		*this = FCookDependency::File(WriteToString<256>(LocalFilename));
+		return true;
+	}
+	case ECookDependency::Function:
+	{
+		FName LocalFuncName;
+		if (!LoadFromCompactBinary(Field++, LocalFuncName))
+		{
+			return false;
+		}
+		*this = FCookDependency::Function(LocalFuncName, FCbFieldIterator::CloneRange(Field));
+		FunctionData.Args.MakeOwned();
+		return true;
+	}
+	default:
+		return false;
+	}
+	checkNoEntry();
+	return true;
+}
+
+void FCookDependencyContext::Update(const void* Data, uint64 Size)
+{
+	FBlake3& Blake3Hasher = *(reinterpret_cast<FBlake3*>(Hasher));
+	Blake3Hasher.Update(Data, Size);
+}
+
+void FCookDependencyContext::LogError(FString Message)
+{
+	for (TUniqueFunction<FString(FString&&)>& Handler : ErrorHandlers)
+	{
+		Message = Handler(MoveTemp(Message));
+	}
+	OnLogError(MoveTemp(Message));
+}
+
+FCookDependencyContext::FErrorHandlerScope::FErrorHandlerScope(FCookDependencyContext& InContext)
+	: Context(InContext)
+{
+}
+
+FCookDependencyContext::FErrorHandlerScope::~FErrorHandlerScope()
+{
+	check(!Context.ErrorHandlers.IsEmpty());
+	Context.ErrorHandlers.Pop(EAllowShrinking::No);
+}
+
+FCookDependencyContext::FErrorHandlerScope FCookDependencyContext::ErrorHandlerScope(
+	TUniqueFunction<FString(FString&&)>&& ErrorHandler)
+{
+	ErrorHandlers.Add(MoveTemp(ErrorHandler));
+	// We rely on Copy elision to not move-construct this scope, so that it only calls the destructor once.
+	return FErrorHandlerScope(*this);
+}
+
+} // namespace UE::Cook
+
+namespace UE::Cook::Dependency::Private
+{
+
+FCookDependencyFunctionRegistration* List = nullptr;
+bool bGCookDependencyFunctionsInitialized = false;
+TMap<FName, FCookDependencyFunction> GCookDependencyFunctions;
+
+const TMap<FName, FCookDependencyFunction>& GetCookDependencyFunctions()
+{
+	if (!bGCookDependencyFunctionsInitialized)
+	{
+		GCookDependencyFunctions.Empty();
+		FCookDependencyFunctionRegistration* Reg = UE::Cook::Dependency::Private::List;
+		while (Reg)
+		{
+			FCookDependencyFunction& ExistingFunction = GCookDependencyFunctions.FindOrAdd(Reg->GetFName(), nullptr);
+			checkf(ExistingFunction == nullptr || ExistingFunction == Reg->Function,
+				TEXT("UE_COOK_DEPENDENCY_FUNCTION name '%s' is duplicated. UE_COOK_DEPENDENCY_FUNCTION names must be unique."),
+				*Reg->GetFName().ToString());
+			ExistingFunction = Reg->Function;
+			Reg = Reg->Next;
+		}
+		bGCookDependencyFunctionsInitialized = true;
+	}
+	return GCookDependencyFunctions;
+}
+
+void FCookDependencyFunctionRegistration::Construct()
+{
+	using namespace UE::Cook::Dependency::Private;
+
+	Next = UE::Cook::Dependency::Private::List;
+	UE::Cook::Dependency::Private::List = this;
+	bGCookDependencyFunctionsInitialized = false;
+}
+
+FCookDependencyFunctionRegistration::~FCookDependencyFunctionRegistration()
+{
+	using namespace UE::Cook::Dependency::Private;
+
+	// Remove this from the list when destructed, but for better shutdown performance skip this cost during engine
+	// exit, and just leave the list with dangling pointers, since the list can only be read by this destructor or by
+	// GetCookDependencyFunctions, which is not called during shutdown.
+	if (!IsEngineExitRequested())
+	{
+		if (UE::Cook::Dependency::Private::List == this)
+		{
+			UE::Cook::Dependency::Private::List = this->Next;
+		}
+		else
+		{
+			FCookDependencyFunctionRegistration* Current = UE::Cook::Dependency::Private::List;
+			while (Current)
+			{
+				if (Current->Next == this)
+				{
+					Current->Next = this->Next;
+					break;
+				}
+				Current = Current->Next;
+			}
+		}
+		bGCookDependencyFunctionsInitialized = false;
+	}
+}
+
+} // namespace UE::Cook::Dependency::Private
+
+#endif // WITH_EDITOR
\ No newline at end of file
diff --git a/Engine/Source/Runtime/CoreUObject/Private/UObject/Package.cpp b/Engine/Source/Runtime/CoreUObject/Private/UObject/Package.cpp
--- a/Engine/Source/Runtime/CoreUObject/Private/UObject/Package.cpp
+++ b/Engine/Source/Runtime/CoreUObject/Private/UObject/Package.cpp
@@ -3,6 +3,7 @@
 #include "UObject/Package.h"
 
 #include "AssetRegistry/AssetData.h"
+#include "Cooker/CookDependency.h"
 #include "HAL/FileManager.h"
 #include "HAL/PlatformMath.h"
 #include "Misc/AssetRegistryInterface.h"
diff --git a/Engine/Source/Runtime/CoreUObject/Private/UObject/SavePackage/SaveContext.cpp b/Engine/Source/Runtime/CoreUObject/Private/UObject/SavePackage/SaveContext.cpp
--- a/Engine/Source/Runtime/CoreUObject/Private/UObject/SavePackage/SaveContext.cpp
+++ b/Engine/Source/Runtime/CoreUObject/Private/UObject/SavePackage/SaveContext.cpp
@@ -3,6 +3,7 @@
 #include "SaveContext.h"
 
 #include "Algo/Find.h"
+#include "Cooker/CookDependency.h"
 #include "Misc/ConfigCacheIni.h"
 #include "Serialization/PackageWriter.h"
 #include "UObject/UObjectGlobals.h"
@@ -150,6 +151,9 @@ FSavePackageResultStruct FSaveContext::GetFinalResult()
 	}
 	TSet<FName>& SoftPackageReferenceList = GetSoftPackageReferenceList();
 	ResultData.SoftPackageReferences = SoftPackageReferenceList.Array();
+#if WITH_EDITOR
+	ResultData.CookDependencies = MoveTemp(ObjectSaveContext.CookDependencies);
+#endif
 
 	return ResultData;
 }
diff --git a/Engine/Source/Runtime/CoreUObject/Private/UObject/SavePackage/SavePackageUtilities.cpp b/Engine/Source/Runtime/CoreUObject/Private/UObject/SavePackage/SavePackageUtilities.cpp
--- a/Engine/Source/Runtime/CoreUObject/Private/UObject/SavePackage/SavePackageUtilities.cpp
+++ b/Engine/Source/Runtime/CoreUObject/Private/UObject/SavePackage/SavePackageUtilities.cpp
@@ -7,6 +7,7 @@
 #include "AssetRegistry/AssetData.h"
 #include "AssetRegistry/CookTagList.h"
 #include "Blueprint/BlueprintSupport.h"
+#include "Cooker/CookDependency.h"
 #include "CoreMinimal.h"
 #include "HAL/FileManager.h"
 #include "Interfaces/ITargetPlatform.h"
@@ -746,6 +747,12 @@ FAddResaveOnDemandPackage OnAddResaveOnDemandPackage;
 
 } // end namespace UE::SavePackageUtilities
 
+// Constructor/Destructor defined here in cpp rather than header so we can 
+// avoid needing the definition of FCookDependency in the header; it is needed
+// for construct/destruct of TArray<FCookDependency>.
+FObjectSaveContextData::FObjectSaveContextData() = default;
+FObjectSaveContextData::~FObjectSaveContextData() = default;
+
 FObjectSaveContextData::FObjectSaveContextData(UPackage* Package, const ITargetPlatform* InTargetPlatform, const TCHAR* InTargetFilename, uint32 InSaveFlags)
 {
 	Set(Package, InTargetPlatform, InTargetFilename, InSaveFlags);
@@ -756,6 +763,13 @@ FObjectSaveContextData::FObjectSaveContextData(UPackage* Package, const ITargetP
 	Set(Package, InTargetPlatform, TargetPath, InSaveFlags);
 }
 
+#if WITH_EDITOR
+void FObjectPreSaveContext::AddCookDependency(UE::Cook::FCookDependency CookDependency)
+{
+	Data.CookDependencies.Add(MoveTemp(CookDependency));
+}
+#endif
+
 void FObjectSaveContextData::Set(UPackage* Package, const ITargetPlatform* InTargetPlatform, const TCHAR* InTargetFilename, uint32 InSaveFlags)
 {
 	FPackagePath PackagePath(FPackagePath::FromLocalPath(InTargetFilename));
diff --git a/Engine/Source/Runtime/CoreUObject/Public/Cooker/CookDependency.h b/Engine/Source/Runtime/CoreUObject/Public/Cooker/CookDependency.h
new file mode 100644
--- /dev/null
+++ b/Engine/Source/Runtime/CoreUObject/Public/Cooker/CookDependency.h
@@ -0,0 +1,316 @@
+// Copyright Epic Games, Inc. All Rights Reserved.
+
+#pragma once
+
+#if WITH_EDITOR
+#include "Containers/StringView.h"
+#include "Containers/UnrealString.h"
+#include "HAL/PreprocessorHelpers.h"
+#include "Serialization/CompactBinary.h"
+#include "Templates/Function.h"
+#include "UObject/NameTypes.h"
+
+class FCbFieldView;
+class FCbWriter;
+#endif
+
+#if WITH_EDITOR
+
+namespace UE::Cook
+{
+
+/** Context passed into FCookDependencyFunction to provide calling flags and receive their hash output. */
+struct FCookDependencyContext
+{
+public:
+	/** InHasher is void* to mask the implementation details of the hashbuilder. See Update function. */
+	explicit FCookDependencyContext(void* InHasher, TUniqueFunction<void(FString&&)>&& InOnLogError);
+
+	/**
+	 * Update the hashbuilder for the key being constructed (e.g. TargetDomainKey for cooked packages)
+	 * with the given Data of Size bytes.
+	 */
+	COREUOBJECT_API void Update(const void* Data, uint64 Size);
+
+	/**
+	 * Reports failure to compute the hash (e.g. because a file cannot be read).
+	 * When calculating the initial hash during package save, this error will be logged as an error
+	 * and the package will be recooked on the next cook. When calculating the hash during an incremental cook
+	 * the message will be logged at Log level and will cause the package to be recooked.
+	 */
+	COREUOBJECT_API void LogError(FString Message);
+
+	/**
+	 * Private implementation struct used for AddErrorHandlerScope. Should only be used via
+	 * FCookDependencyContext::FErrorHandlerScope& Scope = Context.ErrorHandlerScope(<Function>);
+	 */
+	struct FErrorHandlerScope
+	{
+	public:
+		COREUOBJECT_API ~FErrorHandlerScope();
+	private:
+		COREUOBJECT_API FErrorHandlerScope(FCookDependencyContext& InContext);
+		friend FCookDependencyContext;
+		FCookDependencyContext& Context;
+	};
+	/**
+	 * Add a function that will be removed when the return value goes out of scope, to modify error strings reported
+	 * inside the scope before passing them on to higher scopes or the error consumer.
+	 * e.g. 
+	 * FCookDependencyContext::FErrorHandlerScope Scope = Context.ErrorHandlerScope([](FString&& Inner)
+	 * { return FString::Printf(TEXT("OuterClass for %s: %s"), *Name, *Inner);});
+	 */
+	[[nodiscard]] COREUOBJECT_API FErrorHandlerScope ErrorHandlerScope(
+		TUniqueFunction<FString(FString&&)>&& ErrorHandler);
+
+private:
+	TUniqueFunction<void(FString&&)> OnLogError;
+	TArray<TUniqueFunction<FString(FString&&)>, TInlineAllocator<1>> ErrorHandlers;
+	void* Hasher; // Type is void* to mask the implementation detail
+};
+
+/**
+ * TypeSelector enum for the FCookDependency variable type. Values are serialized into the oplog as integers,
+ * so do not change them without changing oplog version.
+ */
+enum class ECookDependency : uint8
+{
+	None = 0x00,
+	File = 0x01,
+	Function = 0x02,
+};
+
+/**
+ * TargetDomain dependencies that can be reported from the class instances in a package. These dependencies are
+ * stored in the cook oplog and are evaluated during incremental cook. If any of them changes, the package is
+ * invalidated and must be recooked (loaded/saved). These dependencies do not impact whether DDC keys built
+ * from the package need to be recalculated.
+ */
+class FCookDependency
+{
+public:
+	/**
+	 * Create a dependency on the contents of the file. Filename will be normalized. Contents are loaded via
+	 * IFileManager::Get().CreateFileReader and contents are hashed for comparison.
+	 */
+	COREUOBJECT_API static FCookDependency File(FStringView InFileName);
+	/**
+	 * Create a dependency on a call to the specified function with the given arguments. Arguments should be
+	 * created using FCbWriter Writer; ... <append arbitrary number of fields to Writer> ...; Writer.Save().
+	 * The function should read the arguments using the corresponding FCbFieldIteratorView methods and
+	 * LoadFromCompactBinary calls.
+	 * 
+	 * The function must be registered during editor startup for use with FCookDependency via
+	 * UE_COOK_DEPENDENCY_FUNCTION(CppTokenUsedAsName, CppNameOfFunctionToCall).
+	 * The name to pass to FCookDependency::Function can be retrieved via
+	 * UE_COOK_DEPENDENCY_FUNCTION_CALL(CppTokenUsedAsName).
+	 */
+	COREUOBJECT_API static FCookDependency Function(FName InFunctionName, FCbFieldIterator&& InArgs);
+	/** Construct an empty dependency; it will never be invalidated. */
+	COREUOBJECT_API FCookDependency();
+
+	COREUOBJECT_API ~FCookDependency();
+	COREUOBJECT_API FCookDependency(const FCookDependency& Other);
+	COREUOBJECT_API FCookDependency(FCookDependency&& Other);
+	COREUOBJECT_API FCookDependency& operator=(const FCookDependency& Other);
+	COREUOBJECT_API FCookDependency& operator=(FCookDependency&& Other);
+	
+	/** FCookDependency is a vararg type. Return the type of this instance. */
+	ECookDependency GetType() const;
+
+	/** FileName if GetType() == File, else empty. StringView points to null or a null-terminated string. */
+	FStringView GetFileName() const;
+
+	/** FunctionName if GetType() == Function, else NAME_None. */
+	FName GetFunctionName() const;
+	/** FunctionArgs if GetType() == Function, else FCbFieldViewIterator(). */
+	FCbFieldViewIterator GetFunctionArgs() const;
+
+	/**
+	 * Comparison operator for e.g. deterministic ordering of dependencies.
+	 * Uses persistent comparison data and is somewhat expensive.
+	 */
+	bool operator<(const FCookDependency& Other) const;
+
+	/** Calculate the current hash of this CookDependency, and add it into Context. */
+	COREUOBJECT_API void UpdateHash(FCookDependencyContext& Context) const;
+
+private:
+	explicit FCookDependency(ECookDependency InType);
+	void Construct();
+	void Destruct();
+	COREUOBJECT_API void Save(FCbWriter& Writer) const;
+	COREUOBJECT_API bool Load(FCbFieldView Value);
+
+	/** Public hidden friend for operator<< into an FCbWriter. */
+	friend FCbWriter& operator<<(FCbWriter& Writer, const FCookDependency& CookDependencies)
+	{
+		CookDependencies.Save(Writer);
+		return Writer;
+	}
+	/** Public hidden friend for LoadFromCompactBinary. */
+	friend bool LoadFromCompactBinary(FCbFieldView Value, FCookDependency& CookDependencies)
+	{
+		return CookDependencies.Load(Value);
+	}
+
+private:
+	ECookDependency Type;
+	struct FFunctionData
+	{
+		FName Name;
+		FCbFieldIterator Args;
+	};
+	union
+	{
+		FString FileName;
+		FFunctionData FunctionData;
+	};
+};
+
+/**
+ * Type of functions used in FCookDependency to append the hash values of arbitrary data.
+ * 
+ * @param Args Variable-length, variable-typed input data (e.g. names of files, configuration flags)
+ *             that specify which hash data. The function should read this data using FCbFieldViewIterator
+ *             methods and LoadFromCompactBinary calls that correspond to the FCbWriter methods used
+ *             at the callsite of FCookDependency::Function.
+ * @param Context that provides calling flags and receives the hashdata. The function should call
+ *        Context.Update with the data to be added to the target key (e.g. the hash of the contents
+ *        of a filename that was specified in Args).
+ */
+using FCookDependencyFunction = void (*)(FCbFieldViewIterator Args, FCookDependencyContext& Context);
+
+} // namespace UE::Cook
+
+namespace UE::Cook::Dependency::Private
+{
+
+/**
+ * Implementation struct of UE_COOK_DEPENDENCY_FUNCTION. Instances of this class are stored in global or
+ * namespace scope and add themselves to a list during c++ pre-main static initialization. This list
+ * is read later to create a map from FName to c++ function.
+ */
+struct FCookDependencyFunctionRegistration
+{
+	template<int N>
+	FCookDependencyFunctionRegistration(const TCHAR(&InName)[N], FCookDependencyFunction InFunction)
+		: Name(InName), Function(InFunction), Next(nullptr)
+	{
+		static_assert(N > 0, "Name must be provided");
+		check(InName[0] != '\0');
+		Construct();
+	}
+	COREUOBJECT_API ~FCookDependencyFunctionRegistration();
+	COREUOBJECT_API void Construct();
+	FName GetFName();
+
+	FLazyName Name;
+	FCookDependencyFunction Function;
+	FCookDependencyFunctionRegistration* Next;
+};
+
+} // namespace UE::Cook::Dependency::Private
+
+/**
+ * UE_COOK_DEPENDENCY_FUNCTION(<CppToken> Name, void (*)(FCbFieldView Args, FCookDependencyContext& Context))
+ * 
+ * Registers the given function pointer to handle FCookDependency::Function(Name, Args) calls.
+ * @see FCookDependencyFunction. @see FCookDependency::Function.
+ * 
+ * Name should be a bare cpptoken, e.g.
+ * UE_COOK_DEPENDENCY_FUNCTION(MyTypeDependencies, UE::MyTypeDependencies::ImplementationFunction).
+*/
+#define UE_COOK_DEPENDENCY_FUNCTION(Name, Function) \
+	UE::Cook::Dependency::Private::FCookDependencyFunctionRegistration \
+	PREPROCESSOR_JOIN(FCookDependencyFunctionRegistration_,Name)(TEXT(#Name), Function)
+/**
+ * Return the FName to use to call a function that was registered via UE_COOK_DEPENDENCY_FUNCTION(Name, Function).
+ * Name should be the same bare cpptoken that was passed into UE_COOK_DEPENDENCY_FUNCTION.
+ */
+#define UE_COOK_DEPENDENCY_FUNCTION_CALL(Name) \
+	PREPROCESSOR_JOIN(FCookDependencyFunctionRegistration_,Name).GetFName()
+
+#else // WITH_EDITOR
+
+#define UE_COOK_DEPENDENCY_FUNCTION(Name, Function)
+#define UE_COOK_DEPENDENCY_FUNCTION_CALL(Name) NAME_None
+
+#endif // !WITH_EDITOR
+
+#if WITH_EDITOR
+namespace UE::Cook
+{
+
+inline ECookDependency FCookDependency::GetType() const
+{
+	return Type;
+}
+
+inline FStringView FCookDependency::GetFileName() const
+{
+	return Type == ECookDependency::File ? FileName : FStringView();
+}
+
+inline FName FCookDependency::GetFunctionName() const
+{
+	return Type == ECookDependency::Function ? FunctionData.Name : NAME_None;
+}
+
+inline FCbFieldViewIterator FCookDependency::GetFunctionArgs() const
+{
+	return Type == ECookDependency::Function ? FunctionData.Args : FCbFieldViewIterator();
+}
+
+inline bool FCookDependency::operator<(const FCookDependency& Other) const
+{
+	if (static_cast<uint8>(Type) != static_cast<uint8>(Other.Type))
+	{
+		return static_cast<uint8>(Type) < static_cast<uint8>(Other.Type);
+	}
+
+	switch (Type)
+	{
+	case ECookDependency::None:
+		return false;
+	case ECookDependency::File:
+		return FileName.Compare(Other.FileName, ESearchCase::IgnoreCase) < 0;
+	case ECookDependency::Function:
+	{
+		int32 Compare = FunctionData.Name.Compare(Other.FunctionData.Name);
+		if (Compare != 0) return Compare < 0;
+		FMemoryView ViewA;
+		FMemoryView ViewB;
+		bool bHasViewA = FunctionData.Args.TryGetRangeView(ViewA);
+		bool bHasViewB = FunctionData.Args.TryGetRangeView(ViewB);
+		if ((!bHasViewA) | (!bHasViewB))
+		{
+			return bHasViewB; // If both false, return false. If only one, return true only if A is the false.
+		}
+		return ViewA.CompareBytes(ViewB) < 0;
+	}
+	default:
+		checkNoEntry();
+		return false;
+	}
+}
+
+inline FCookDependencyContext::FCookDependencyContext(void* InHasher, TUniqueFunction<void(FString&&)>&& InOnLogError)
+	: OnLogError(MoveTemp(InOnLogError))
+	, Hasher(InHasher)
+{
+}
+
+} // namespace UE::Cook
+
+namespace UE::Cook::Dependency::Private
+{
+
+inline FName FCookDependencyFunctionRegistration::GetFName()
+{
+	return Name.Resolve();
+}
+
+} // namespace UE::Cook::Dependency::Private
+
+#endif
\ No newline at end of file
diff --git a/Engine/Source/Runtime/CoreUObject/Public/UObject/ObjectSaveContext.h b/Engine/Source/Runtime/CoreUObject/Public/UObject/ObjectSaveContext.h
--- a/Engine/Source/Runtime/CoreUObject/Public/UObject/ObjectSaveContext.h
+++ b/Engine/Source/Runtime/CoreUObject/Public/UObject/ObjectSaveContext.h
@@ -10,11 +10,15 @@
 class FPackagePath;
 class ITargetPlatform;
 class UPackage;
+#if WITH_EDITOR
+namespace UE::Cook { class FCookDependency; }
+#endif
 
 /** Data used to provide information about the save parameters during PreSave/PostSave. */
 struct FObjectSaveContextData
 {
-	FObjectSaveContextData() = default;
+	COREUOBJECT_API FObjectSaveContextData();
+	COREUOBJECT_API ~FObjectSaveContextData();
 	/** Standard constructor; calculates derived fields from the given externally-specified fields. */
 	COREUOBJECT_API FObjectSaveContextData(UPackage* Package, const ITargetPlatform* InTargetPlatform, const TCHAR* InTargetFilename, uint32 InSaveFlags);
 	COREUOBJECT_API FObjectSaveContextData(UPackage* Package, const ITargetPlatform* InTargetPlatform, const FPackagePath& TargetPath, uint32 InSaveFlags);
@@ -68,6 +72,11 @@ struct FObjectSaveContextData
 	/** Set to false if the save failed, before calling any PostSaves. */
 	bool bSaveSucceeded = true;
 
+	// Collection variables that are written but not read during the PreSave/PostSave functions
+#if WITH_EDITOR
+	TArray<UE::Cook::FCookDependency> CookDependencies;
+#endif
+
 	// Per-object Output variables; writable from PreSave functions, readable from PostSave functions
 
 	/** List of property overrides per object to apply to during save */
@@ -120,6 +129,14 @@ public:
 	UE::Cook::ECookType GetCookType() const { return Data.CookType; }
 	UE::Cook::ECookingDLC GetCookingDLC() const { return Data.CookingDLC; }
 
+#if WITH_EDITOR
+	/**
+	 * Add the given FCookDependency to the record for the package being cook-saved. Iterative cooks will
+	 * invalidate the package and recook it if the CookDependency changes.
+	 */
+	COREUOBJECT_API void AddCookDependency(UE::Cook::FCookDependency CookDependency);
+#endif
+
 	/**
 	 * Return whether the package is being saved due to a procedural save.
 	 * Any save without the possibility of user-generated edits to the package is a procedural save (Cooking, EditorDomain).
diff --git a/Engine/Source/Runtime/CoreUObject/Public/UObject/Package.h b/Engine/Source/Runtime/CoreUObject/Public/UObject/Package.h
--- a/Engine/Source/Runtime/CoreUObject/Public/UObject/Package.h
+++ b/Engine/Source/Runtime/CoreUObject/Public/UObject/Package.h
@@ -52,6 +52,10 @@ struct FMD5Hash;
 struct FPackageSaveInfo;
 struct FSavePackageArgs;
 
+#if WITH_EDITOR
+namespace UE::Cook { class FCookDependency; }
+#endif
+
 /**
 * Represents the result of saving a package
 */
@@ -109,6 +113,10 @@ struct FSavePackageResultStruct
 	TArray<FName> ImportPackages;
 	TArray<FName> SoftPackageReferences;
 
+#if WITH_EDITOR
+	TArray<UE::Cook::FCookDependency> CookDependencies;
+#endif
+
 	/** Constructors, it will implicitly construct from the result enum */
 	COREUOBJECT_API FSavePackageResultStruct();
 	COREUOBJECT_API FSavePackageResultStruct(ESavePackageResult InResult);
