Fix FOptionalProperty::ConvertFromType handling of converting from an optional boolean property to something else.
#okforversepublic

diff --git a/Engine/Source/Runtime/CoreUObject/Private/UObject/PropertyOptional.cpp b/Engine/Source/Runtime/CoreUObject/Private/UObject/PropertyOptional.cpp
--- a/Engine/Source/Runtime/CoreUObject/Private/UObject/PropertyOptional.cpp
+++ b/Engine/Source/Runtime/CoreUObject/Private/UObject/PropertyOptional.cpp
@@ -478,6 +478,12 @@ EConvertFromTypeResult FOptionalProperty::ConvertFromType(const FPropertyTag& Ta
 				return EConvertFromTypeResult::Converted;
 			}
 			ValueTag.SetType(Tag.GetType().GetParameter());
+
+			// Mimic FPropertyTag::SerializeTaggedProperty's special code path for bools that serializes the value as part of the tag, even though in this case we didn't serialize a separate tag for the inner value.
+			if (ValueTag.Type == NAME_BoolProperty)
+			{
+				*MaybeValueSlot << ValueTag.BoolVal;
+			}
 		}
 
 		FStructuredArchive::FSlot ValueSlot = MaybeValueSlot.Get(Slot);
