Fixed UEnum::GetMaxValue() for flag enums.
Added UEnum::GetValueOrBitfieldFromString() for parsing a bar-separated string (but not bar-separated integers) into a flag-based UEnum.
Added UEnum::GetValueOrBitfieldAsAuthoredNameString() and UEnum::GetValueOrBitfieldAsDisplayNameText() to match UEnum::GetValueOrBitfieldAsString(), for authored names and display names respectively.
Added support for serialising enum flags and converting to/from text.

diff --git a/Engine/Source/Runtime/CoreUObject/Private/UObject/Enum.cpp b/Engine/Source/Runtime/CoreUObject/Private/UObject/Enum.cpp
--- a/Engine/Source/Runtime/CoreUObject/Private/UObject/Enum.cpp
+++ b/Engine/Source/Runtime/CoreUObject/Private/UObject/Enum.cpp
@@ -2,6 +2,7 @@
 
 #include "CoreMinimal.h"
 #include "Misc/ConfigCacheIni.h"
+#include "String/ParseTokens.h"
 #include "UObject/ObjectMacros.h"
 #include "UObject/Class.h"
 #include "UObject/Package.h"
@@ -207,6 +208,39 @@ int64 UEnum::GetValueByName(FName InName, EGetByNameFlags Flags) const
 	return INDEX_NONE;
 }
 
+int64 UEnum::GetValueOrBitfieldFromString(FStringView InString, EGetByNameFlags LookupFlags) const
+{
+	bool  bSuccess = true;
+	int64 Result   = 0;
+
+	UE::String::ParseTokens(
+		InString,
+		TEXTVIEW("|"),
+		[this, &bSuccess, &Result, LookupFlags](FStringView Flag)
+		{
+			if (!bSuccess)
+			{
+				return;
+			}
+
+			FName FlagName(Flag);
+
+			int32 FlagIndex = GetIndexByName(FlagName, LookupFlags);
+			if (FlagIndex != INDEX_NONE)
+			{
+				Result |= GetValueByName(FlagName);
+			}
+			else
+			{
+				bSuccess = false;
+			}
+		},
+		UE::String::EParseTokensOptions::SkipEmpty | UE::String::EParseTokensOptions::Trim
+	);
+
+	return bSuccess ? Result : INDEX_NONE;
+}
+
 int64 UEnum::GetMaxEnumValue() const
 {
 	int32 NamesNum = Names.Num();
@@ -215,17 +249,31 @@ int64 UEnum::GetMaxEnumValue() const
 		return 0;
 	}
 
-	int64 MaxValue = Names[0].Value;
-	for (int32 i = 1; i < NamesNum; ++i)
+	if (EnumHasAnyFlags(EnumFlags, EEnumFlags::Flags))
 	{
-		int64 CurrentValue = Names[i].Value;
-		if (CurrentValue > MaxValue)
+		// The max value of a set of flags is the combination of all the flags
+		int64 MaxFlag = 0;
+		for (const TPair<FName, int64>& NameAndValue : Names)
 		{
-			MaxValue = CurrentValue;
+			MaxFlag |= NameAndValue.Value;
 		}
+
+		return MaxFlag;
 	}
+	else
+	{
+		int64 MaxValue = Names[0].Value;
+		for (int32 i = 1; i < NamesNum; ++i)
+		{
+			int64 CurrentValue = Names[i].Value;
+			if (CurrentValue > MaxValue)
+			{
+				MaxValue = CurrentValue;
+			}
+		}
 
-	return MaxValue;
+		return MaxValue;
+	}
 }
 
 bool UEnum::IsValidEnumValue(int64 InValue) const
@@ -499,36 +547,69 @@ FString UEnum::GetNameStringByValue(int64 Value) const
 	return GetNameStringByIndex(Index);
 }
 
-FString UEnum::GetValueOrBitfieldAsString(int64 InValue) const
+namespace
 {
-	if (!HasAnyEnumFlags(EEnumFlags::Flags) || InValue == 0)
-	{
-		return GetNameStringByValue(InValue);
-	}
-	else
+	template <typename NameGetterType>
+	FString CreateStringFromValueOrBitfield(const UEnum* Enum, int64 InValue, NameGetterType&& NameGetter)
 	{
+		if (!Enum->HasAnyEnumFlags(EEnumFlags::Flags) || InValue == 0)
+		{
+			return NameGetter(InValue);
+		}
+
 		FString BitfieldString;
-		bool WroteFirstFlag = false;
+		bool bWroteFirstFlag = false;
 		while (InValue != 0)
 		{
 			int64 NextValue = 1ll << FMath::CountTrailingZeros64(InValue);
 			InValue = InValue & ~NextValue;
-			if (WroteFirstFlag)
-			{
-				// We don't just want to use the NameValuePair.Key because we want to strip enum class prefixes
-				BitfieldString.Appendf(TEXT(" | %s"), *GetNameStringByValue(NextValue));
-			}
-			else
-			{
-				// We don't just want to use the NameValuePair.Key because we want to strip enum class prefixes
-				BitfieldString.Appendf(TEXT("%s"), *GetNameStringByValue(NextValue));
-				WroteFirstFlag = true;
-			}
+
+			// We don't just want to use the NameValuePair.Key because we want to strip enum class prefixes
+			BitfieldString.Appendf(TEXT("%s%s"), bWroteFirstFlag ? TEXT(" | ") : TEXT(""), *NameGetter(NextValue));
+			bWroteFirstFlag = true;
 		}
 		return BitfieldString;
 	}
 }
 
+FString UEnum::GetValueOrBitfieldAsString(int64 InValue) const
+{
+	return CreateStringFromValueOrBitfield(
+		this,
+		InValue,
+		[this](int64 Value)
+		{
+			return GetNameStringByValue(Value);
+		}
+	);
+}
+
+FString UEnum::GetValueOrBitfieldAsAuthoredNameString(int64 InValue) const
+{
+	return CreateStringFromValueOrBitfield(
+		this,
+		InValue,
+		[this](int64 Value)
+		{
+			return GetAuthoredNameStringByValue(Value);
+		}
+	);
+}
+
+FText UEnum::GetValueOrBitfieldAsDisplayNameText(int64 InValue) const
+{
+	return FText::FromString(
+		CreateStringFromValueOrBitfield(
+			this,
+			InValue,
+			[this](int64 Value) -> const FString&
+			{
+				return GetDisplayNameTextByValue(Value).ToString();
+			}
+		)
+	);
+}
+
 bool UEnum::FindNameStringByValue(FString& Out, int64 InValue) const
 {
 	int32 Index = GetIndexByValue(InValue);
diff --git a/Engine/Source/Runtime/CoreUObject/Private/UObject/EnumProperty.cpp b/Engine/Source/Runtime/CoreUObject/Private/UObject/EnumProperty.cpp
--- a/Engine/Source/Runtime/CoreUObject/Private/UObject/EnumProperty.cpp
+++ b/Engine/Source/Runtime/CoreUObject/Private/UObject/EnumProperty.cpp
@@ -4,6 +4,7 @@
 
 #include "Algo/Find.h"
 #include "Hash/Blake3.h"
+#include "UObject/PropertyHelper.h"
 #include "UObject/UnrealTypePrivate.h"
 #include "UObject/UObjectThreadContext.h"
 
@@ -134,16 +135,26 @@ void FEnumProperty::SerializeItem(FStructuredArchive::FSlot Slot, void* Value, v
 				UnderlyingArchive.Preload(Enum);
 			}
 
-			// There's no guarantee EnumValueName is still present in Enum, in which case Value will be set to the enum's max value.
-			// On save, it will then be serialized as NAME_None.
-			const int32 EnumIndex = Enum->GetIndexByName(EnumValueName, EGetByNameFlags::ErrorIfNotFound);
-			if (EnumIndex == INDEX_NONE)
+			if (Enum->HasAnyEnumFlags(EEnumFlags::Flags))
 			{
-				NewEnumValue = Enum->GetMaxEnumValue();
+				if (EnumValueName != NAME_None)
+				{
+					NewEnumValue = Enum->GetValueOrBitfieldFromString(*EnumValueName.ToString());
+				}
 			}
 			else
 			{
-				NewEnumValue = Enum->GetValueByIndex(EnumIndex);
+				// There's no guarantee EnumValueName is still present in Enum, in which case Value will be set to the enum's max value.
+				// On save, it will then be serialized as NAME_None.
+				const int32 EnumIndex = Enum->GetIndexByName(EnumValueName, EGetByNameFlags::ErrorIfNotFound);
+				if (EnumIndex == INDEX_NONE)
+				{
+					NewEnumValue = Enum->GetMaxEnumValue();
+				}
+				else
+				{
+					NewEnumValue = Enum->GetValueByIndex(EnumIndex);
+				}
 			}
 		}
 
@@ -157,9 +168,19 @@ void FEnumProperty::SerializeItem(FStructuredArchive::FSlot Slot, void* Value, v
 		{
 			const int64 IntValue = UnderlyingProp->GetSignedIntPropertyValue(Value);
 
-			if (Enum->IsValidEnumValue(IntValue))
+			if (Enum->HasAnyEnumFlags(EEnumFlags::Flags))
 			{
-				EnumValueName = Enum->GetNameByValue(IntValue);
+				if (IntValue != 0)
+				{
+					EnumValueName = *Enum->GetValueOrBitfieldAsString(IntValue);
+				}
+			}
+			else
+			{
+				if (Enum->IsValidEnumValue(IntValue))
+				{
+					EnumValueName = Enum->GetNameByValue(IntValue);
+				}
 			}
 		}
 
@@ -269,7 +290,7 @@ void FEnumProperty::ExportText_Internal(FString& ValueStr, const void* PropertyV
 
 	// if the value is the max value (the autogenerated *_MAX value), export as "INVALID", unless we're exporting text for copy/paste (for copy/paste,
 	// the property text value must actually match an entry in the enum's names array)
-	if (!Enum->IsValidEnumValue(Value) || (!(PortFlags & PPF_Copy) && Value == Enum->GetMaxEnumValue()))
+	if (!Enum->HasAnyEnumFlags(EEnumFlags::Flags) && (!Enum->IsValidEnumValue(Value) || (!(PortFlags & PPF_Copy) && Value == Enum->GetMaxEnumValue())))
 	{
 		ValueStr += TEXT("(INVALID)");
 		return;
@@ -278,15 +299,15 @@ void FEnumProperty::ExportText_Internal(FString& ValueStr, const void* PropertyV
 	// We do not want to export the enum text for non-display uses, localization text is very dynamic and would cause issues on import
 	if (PortFlags & PPF_PropertyWindow)
 	{
-		ValueStr += Enum->GetDisplayNameTextByValue(Value).ToString();
+		ValueStr += Enum->GetValueOrBitfieldAsDisplayNameText(Value).ToString();
 	}
 	else if (PortFlags & PPF_ExternalEditor)
 	{
-		ValueStr += Enum->GetAuthoredNameStringByValue(Value);
+		ValueStr += Enum->GetValueOrBitfieldAsAuthoredNameString(Value);
 	}
 	else
 	{
-		ValueStr += Enum->GetNameStringByValue(Value);
+		ValueStr += Enum->GetValueOrBitfieldAsString(Value);
 	}
 }
 
@@ -297,25 +318,42 @@ const TCHAR* FEnumProperty::ImportText_Internal(const TCHAR* Buffer, void* Conta
 	
 	if (!(PortFlags & PPF_ConsoleVariable))
 	{
-		int64 EnumValue = 0;
+		bool    bIsEnumOfFlags = Enum->HasAnyEnumFlags(EEnumFlags::Flags);
+		int64   EnumValue      = 0;
 		FString Temp;
-		Buffer = FPropertyHelpers::ReadToken(Buffer, Temp, true);
-		if (Buffer)
+		for (;;)
 		{
-			int32 EnumIndex = Enum->GetIndexByName(*Temp, EGetByNameFlags::CheckAuthoredName);
-			if (EnumIndex == INDEX_NONE && (Temp.IsNumeric() && !Algo::Find(Temp, TEXT('.'))))
-			{
-				int64 LexedEnumValue = INDEX_NONE;
-				LexFromString(LexedEnumValue, *Temp);
-				EnumIndex = Enum->GetIndexByValue(LexedEnumValue);
-			}
-			if (EnumIndex != INDEX_NONE)
-			{
-				EnumValue = Enum->GetValueByIndex(EnumIndex);
-			}
-			else
+			Buffer = FPropertyHelpers::ReadToken(Buffer, Temp, true);
+			if (Buffer)
 			{
-				Buffer = nullptr;
+				int32 EnumIndex = Enum->GetIndexByName(*Temp, EGetByNameFlags::CheckAuthoredName);
+				if (!bIsEnumOfFlags && EnumIndex == INDEX_NONE && (Temp.IsNumeric() && !Algo::Find(Temp, TEXT('.'))))
+				{
+					int64 LexedEnumValue = INDEX_NONE;
+					LexFromString(LexedEnumValue, *Temp);
+					EnumIndex = Enum->GetIndexByValue(LexedEnumValue);
+				}
+				if (EnumIndex == INDEX_NONE)
+				{
+					Buffer = nullptr;
+					break;
+				}
+
+				EnumValue |= Enum->GetValueByIndex(EnumIndex);
+
+				if (!bIsEnumOfFlags)
+				{
+					break;
+				}
+
+				Temp.Reset();
+				SkipWhitespace(Buffer);
+				if (*Buffer != TEXT('|'))
+				{
+					break;
+				}
+				++Buffer;
+				SkipWhitespace(Buffer);
 			}
 		}
 
diff --git a/Engine/Source/Runtime/CoreUObject/Public/UObject/Class.h b/Engine/Source/Runtime/CoreUObject/Public/UObject/Class.h
--- a/Engine/Source/Runtime/CoreUObject/Public/UObject/Class.h
+++ b/Engine/Source/Runtime/CoreUObject/Public/UObject/Class.h
@@ -2059,6 +2059,9 @@ public:
 	/** Gets enum value by name, returns INDEX_NONE and optionally errors when name is not found. This is faster than ByNameString if the FName is exact, but will fall back if needed */
 	COREUOBJECT_API int64 GetValueByName(FName InName, EGetByNameFlags Flags = EGetByNameFlags::None) const;
 
+	/** Gets enum value by name, returns INDEX_NONE and optionally errors when name is not found. This is faster than ByNameString if the FName is exact, but will fall back if needed */
+	COREUOBJECT_API int64 GetValueOrBitfieldFromString(FStringView InString, EGetByNameFlags LookupFlags = EGetByNameFlags::None) const;
+
 	/** Returns the short name at the enum index, returns empty string if invalid */
 	COREUOBJECT_API FString GetNameStringByIndex(int32 InIndex) const;
 
@@ -2071,6 +2074,12 @@ public:
 	/** If the enumeration is declared as UENUM(Flags), returns a string of the form A | B | C representing set bits A, B, and C. If it is not a bitfield, the result is the same as calling GetNameStringByValue*/
 	COREUOBJECT_API FString GetValueOrBitfieldAsString(int64 InValue) const;
 
+	/** If the enumeration is declared as UENUM(Flags), returns a string of the form A | B | C representing set bits A, B, and C. If it is not a bitfield, the result is the same as calling GetAuthoredNameStringByValue*/
+	COREUOBJECT_API FString GetValueOrBitfieldAsAuthoredNameString(int64 InValue) const;
+
+	/** If the enumeration is declared as UENUM(Flags), returns a string of the form A | B | C representing set bits A, B, and C. If it is not a bitfield, the result is the same as calling GetDisplayNameTextByValue*/
+	COREUOBJECT_API FText GetValueOrBitfieldAsDisplayNameText(int64 InValue) const;
+
 	/** Looks for a name with a given value and returns true and writes the name to Out if one was found */
 	COREUOBJECT_API bool FindNameStringByValue(FString& Out, int64 InValue) const;
 
