Fix property bag serialization for failed conversions and correctly disable IDO support when serializing an IDO

Contains a few improvements to logging property names and types when a type mismatch occurs.

#rnx

diff --git a/Engine/Source/Runtime/CoreUObject/Private/UObject/Class.cpp b/Engine/Source/Runtime/CoreUObject/Private/UObject/Class.cpp
--- a/Engine/Source/Runtime/CoreUObject/Private/UObject/Class.cpp
+++ b/Engine/Source/Runtime/CoreUObject/Private/UObject/Class.cpp
@@ -1622,17 +1622,13 @@ void UStruct::SerializeVersionedTaggedProperties(FStructuredArchive::FSlot Slot,
 							}
 						}
 
+						bool bTryLoadIntoPropertyBag = false;
+
 						switch (Property->ConvertFromType(Tag, ValueSlot, Data, DefaultsStruct, Defaults))
 						{
 							case EConvertFromTypeResult::Converted:
 								bAdvanceProperty = true;
-								if (FPropertyBag* PropertyBag = TryFindPropertyBag())
-								{
-									UnderlyingArchive.Seek(StartOfProperty);
-									FStructuredArchive::FSlot CopySlot = PropertyRecord.EnterField(TEXT("Value"));
-									Tag.SetProperty(nullptr);
-									PropertyBag->LoadPropertyByTag(SerializeContext->SerializedPropertyPath, Tag, CopySlot, Property->ContainerPtrToValuePtrForDefaults<uint8>(DefaultsStruct, Defaults, Tag.ArrayIndex));
-								}
+								bTryLoadIntoPropertyBag = true;
 								break;
 
 							case EConvertFromTypeResult::Serialized:
@@ -1646,12 +1642,7 @@ void UStruct::SerializeVersionedTaggedProperties(FStructuredArchive::FSlot Slot,
 										*WriteToString<32>(Tag.Name), *WriteToString<32>(GetFName()),
 										*WriteToString<32>(Tag.Type), *WriteToString<32>(PropID),
 										*UnderlyingArchive.GetArchiveName());
-									if (FPropertyBag* PropertyBag = TryFindPropertyBag())
-									{
-										Tag.SetProperty(nullptr);
-										PropertyBag->LoadPropertyByTag(SerializeContext->SerializedPropertyPath, Tag, ValueSlot, Property->ContainerPtrToValuePtrForDefaults<uint8>(DefaultsStruct, Defaults, Tag.ArrayIndex));
-										bAdvanceProperty = !UnderlyingArchive.IsCriticalError();
-									}
+									bTryLoadIntoPropertyBag = true;
 								}
 								else
 								{
@@ -1665,16 +1656,23 @@ void UStruct::SerializeVersionedTaggedProperties(FStructuredArchive::FSlot Slot,
 								break;
 
 							case EConvertFromTypeResult::CannotConvert:
-								if (FPropertyBag* PropertyBag = TryFindPropertyBag())
-								{
-									Tag.SetProperty(nullptr);
-									PropertyBag->LoadPropertyByTag(SerializeContext->SerializedPropertyPath, Tag, ValueSlot, Property->ContainerPtrToValuePtrForDefaults<uint8>(DefaultsStruct, Defaults, Tag.ArrayIndex));
-									bAdvanceProperty = !UnderlyingArchive.IsCriticalError();
-								}
+								bTryLoadIntoPropertyBag = true;
 								break;
 
 							default:
-								check(false);
+								checkNoEntry();
+								break;
+						}
+
+						if (bTryLoadIntoPropertyBag)
+						{
+							if (FPropertyBag* PropertyBag = TryFindPropertyBag())
+							{
+								Tag.SetProperty(nullptr);
+								UnderlyingArchive.Seek(StartOfProperty);
+								FStructuredArchive::FSlot ValueSlotCopy = PropertyRecord.EnterField(TEXT("Value"));
+								PropertyBag->LoadPropertyByTag(SerializeContext->SerializedPropertyPath, Tag, ValueSlotCopy, Property->ContainerPtrToValuePtrForDefaults<uint8>(DefaultsStruct, Defaults, Tag.ArrayIndex));
+							}
 						}
 					}
 				}
@@ -1694,16 +1692,16 @@ void UStruct::SerializeVersionedTaggedProperties(FStructuredArchive::FSlot Slot,
 
 				int64 Loaded = UnderlyingArchive.Tell() - StartOfProperty;
 
-				if (!bAdvanceProperty)
-				{
-					UnderlyingArchive.Seek(StartOfProperty + Tag.Size);
-				}
-				else
+				if (bAdvanceProperty)
 				{
 					checkf(Tag.Size == Loaded,
 						TEXT("Size mismatch in %s of %s of type %s. Loaded %" INT64_FMT " bytes but expected %d. Package: %s"),
 						*Tag.Name.ToString(), *GetName(), *WriteToString<64>(Tag.GetType()), Loaded, Tag.Size, *UnderlyingArchive.GetArchiveName());
 				}
+				else if (Tag.Size != Loaded)
+				{
+					UnderlyingArchive.Seek(StartOfProperty + Tag.Size);
+				}
 			}
 		}
 	}
diff --git a/Engine/Source/Runtime/CoreUObject/Private/UObject/EnumProperty.cpp b/Engine/Source/Runtime/CoreUObject/Private/UObject/EnumProperty.cpp
--- a/Engine/Source/Runtime/CoreUObject/Private/UObject/EnumProperty.cpp
+++ b/Engine/Source/Runtime/CoreUObject/Private/UObject/EnumProperty.cpp
@@ -468,6 +468,7 @@ EConvertFromTypeResult FEnumProperty::ConvertFromType(const FPropertyTag& Tag, F
 
 				FPropertyTag InnerPropertyTag;
 				InnerPropertyTag.SetType(TypeBuilder.Build());
+				InnerPropertyTag.Name = Tag.Name;
 				InnerPropertyTag.ArrayIndex = 0;
 
 				PreviousValue = (uint8)FNumericProperty::ReadEnumAsInt64(Slot, DefaultsStruct, InnerPropertyTag);
diff --git a/Engine/Source/Runtime/CoreUObject/Private/UObject/LinkerLoad.cpp b/Engine/Source/Runtime/CoreUObject/Private/UObject/LinkerLoad.cpp
--- a/Engine/Source/Runtime/CoreUObject/Private/UObject/LinkerLoad.cpp
+++ b/Engine/Source/Runtime/CoreUObject/Private/UObject/LinkerLoad.cpp
@@ -4648,19 +4648,13 @@ void FLinkerLoad::Preload( UObject* Object )
 						ESoftObjectPathSerializeType::AlwaysSerialize);
 #endif
 
-					// Enable IDO, if needed
+					// Toggle support for IDOs
 					FUObjectSerializeContext* LoadContext = FUObjectThreadContext::Get().GetSerializeContext();
-					TOptional<TGuardValue<bool>> ScopedTrackSerializedPropertyPath;
-					TOptional<TGuardValue<bool>> ScopedSerializeUnknownProperty;
-
-					// Do not enable IDO when impersonation is enabled as we are probably deserializing an IDO already at that point
-					bool bIDOEnabled = UE::IsInstanceDataObjectSupportEnabled(Object) && !LoadContext->bImpersonateProperties;
-					if (bIDOEnabled )
-					{
-						// this will had property path tracking and create a property bag to hold data not matching the current class schema
-						ScopedTrackSerializedPropertyPath.Emplace(LoadContext->bTrackSerializedPropertyPath, true);
-						ScopedSerializeUnknownProperty.Emplace(LoadContext->bSerializeUnknownProperty, true);
-					}
+					// Do not enable IDOs when impersonation is enabled as it implies we are deserializing an IDO
+					const bool bIDOEnabled = UE::IsInstanceDataObjectSupportEnabled(Object) && !LoadContext->bImpersonateProperties;
+					// This will add property path tracking and create a property bag to hold any property that does not match the current class schema
+					TGuardValue<bool> ScopedTrackSerializedPropertyPath(LoadContext->bTrackSerializedPropertyPath, bIDOEnabled);
+					TGuardValue<bool> ScopedSerializeUnknownProperty(LoadContext->bSerializeUnknownProperty, bIDOEnabled);
 
 					if (Object->HasAnyFlags(RF_ClassDefaultObject))
 					{
diff --git a/Engine/Source/Runtime/CoreUObject/Private/UObject/Property.cpp b/Engine/Source/Runtime/CoreUObject/Private/UObject/Property.cpp
--- a/Engine/Source/Runtime/CoreUObject/Private/UObject/Property.cpp
+++ b/Engine/Source/Runtime/CoreUObject/Private/UObject/Property.cpp
@@ -2091,7 +2091,7 @@ bool FProperty::UseBinaryOrNativeSerialization(const FArchive& Ar) const
 bool FProperty::LoadTypeName(UE::FPropertyTypeName Type, const FPropertyTag* Tag)
 {
 	return ensureMsgf(GetID() == Type.GetName(),
-		TEXT("Failed to load property '%s' of type '%s' from tag of type '%s'"),
+		TEXT("Failed to load property '%s' of type '%s' from type name '%s'"),
 		*WriteToString<64>(GetFName()), *WriteToString<64>(GetID()), *WriteToString<64>(Type.GetName()));
 }
 
diff --git a/Engine/Source/Runtime/CoreUObject/Private/UObject/PropertyArray.cpp b/Engine/Source/Runtime/CoreUObject/Private/UObject/PropertyArray.cpp
--- a/Engine/Source/Runtime/CoreUObject/Private/UObject/PropertyArray.cpp
+++ b/Engine/Source/Runtime/CoreUObject/Private/UObject/PropertyArray.cpp
@@ -583,15 +583,17 @@ void FArrayProperty::SerializeItem(FStructuredArchive::FSlot Slot, void* Value,
 				FStructProperty* StructProperty = CastFieldChecked<FStructProperty>(Inner);
 
 				// Attempt mismatched tag serialization if available
-				const FName StructName = InnerTag.GetType().GetParameterName();
+				const FName StructName = InnerTag.GetType().GetParameterName(0);
 				if ((StructProperty->Struct->StructFlags & STRUCT_SerializeFromMismatchedTag) && (InnerTag.Type != NAME_StructProperty || StructName != StructProperty->Struct->GetFName()))
 				{
 					SerializeFromMismatchedTag = InnerTag;
 				}
 				else
 				{
-					UE_LOG(LogClass, Warning, TEXT("Array Property %s contains a struct type mismatch (tag %s != prop %s) in package: %s. If that struct got renamed, add an entry to ActiveStructRedirects."),
-						*WriteToString<32>(InnerTag.Name), *WriteToString<32>(StructName), *WriteToString<32>(StructProperty->Struct->GetFName()), *UnderlyingArchive.GetArchiveName());
+					UE::FPropertyTypeNameBuilder Builder;
+					StructProperty->SaveTypeName(Builder);
+					UE_LOG(LogClass, Warning, TEXT("Array Property %s has a struct type mismatch (tag %s != prop %s) in package: %s. If that struct got renamed, add an entry to ActiveStructRedirects."),
+						*WriteToString<32>(InnerTag.Name), *WriteToString<64>(InnerTag.GetType().GetParameter(0)), *WriteToString<64>(Builder.Build().GetParameter(0)), *UnderlyingArchive.GetArchiveName());
 
 #if WITH_EDITOR
 					// Ensure the structure is initialized
@@ -1208,6 +1210,7 @@ EConvertFromTypeResult FArrayProperty::ConvertFromType(const FPropertyTag& Tag,
 
 	FPropertyTag InnerPropertyTag;
 	InnerPropertyTag.SetType(Tag.GetType().GetParameter());
+	InnerPropertyTag.Name = Tag.Name;
 	InnerPropertyTag.ArrayIndex = 0;
 
 	if (Version < EUnrealEngineObjectUE5Version::PROPERTY_TAG_COMPLETE_TYPE_NAME && Version >= VER_UE4_INNER_ARRAY_TAG_INFO && InnerPropertyTag.Type == NAME_StructProperty)
diff --git a/Engine/Source/Runtime/CoreUObject/Private/UObject/PropertyMap.cpp b/Engine/Source/Runtime/CoreUObject/Private/UObject/PropertyMap.cpp
--- a/Engine/Source/Runtime/CoreUObject/Private/UObject/PropertyMap.cpp
+++ b/Engine/Source/Runtime/CoreUObject/Private/UObject/PropertyMap.cpp
@@ -1509,10 +1509,12 @@ EConvertFromTypeResult FMapProperty::ConvertFromType(const FPropertyTag& Tag, FS
 
 	FPropertyTag KeyPropertyTag;
 	KeyPropertyTag.SetType(KeyType);
+	KeyPropertyTag.Name = Tag.Name;
 	KeyPropertyTag.ArrayIndex = 0;
 
 	FPropertyTag ValuePropertyTag;
 	ValuePropertyTag.SetType(ValueType);
+	ValuePropertyTag.Name = Tag.Name;
 	ValuePropertyTag.ArrayIndex = 0;
 
 	bool bConversionSucceeded = true;
diff --git a/Engine/Source/Runtime/CoreUObject/Private/UObject/PropertySet.cpp b/Engine/Source/Runtime/CoreUObject/Private/UObject/PropertySet.cpp
--- a/Engine/Source/Runtime/CoreUObject/Private/UObject/PropertySet.cpp
+++ b/Engine/Source/Runtime/CoreUObject/Private/UObject/PropertySet.cpp
@@ -930,6 +930,7 @@ EConvertFromTypeResult FSetProperty::ConvertFromType(const FPropertyTag& Tag, FS
 
 	FPropertyTag InnerPropertyTag;
 	InnerPropertyTag.SetType(Tag.GetType().GetParameter());
+	InnerPropertyTag.Name = Tag.Name;
 	InnerPropertyTag.ArrayIndex = 0;
 
 	bool bConversionSucceeded = true;
diff --git a/Engine/Source/Runtime/CoreUObject/Private/UObject/PropertyStruct.cpp b/Engine/Source/Runtime/CoreUObject/Private/UObject/PropertyStruct.cpp
--- a/Engine/Source/Runtime/CoreUObject/Private/UObject/PropertyStruct.cpp
+++ b/Engine/Source/Runtime/CoreUObject/Private/UObject/PropertyStruct.cpp
@@ -426,7 +426,7 @@ EConvertFromTypeResult FStructProperty::ConvertFromType(const FPropertyTag& Tag,
 					UE::FPropertyTypeNameBuilder Builder;
 					SaveTypeName(Builder);
 					UE_LOG(LogClass, Warning, TEXT("SerializeFromMismatchedTag failed: Type mismatch in %s - Previous (%s) Current(%s) in package: %s"),
-						*WriteToString<32>(Tag.Name), *WriteToString<32>(Tag.GetType()), *WriteToString<32>(Builder.Build()), *UnderlyingArchive.GetArchiveName());
+						*WriteToString<32>(Tag.Name), *WriteToString<64>(Tag.GetType()), *WriteToString<64>(Builder.Build()), *UnderlyingArchive.GetArchiveName());
 					return EConvertFromTypeResult::CannotConvert;
 				}
 			}
@@ -434,8 +434,10 @@ EConvertFromTypeResult FStructProperty::ConvertFromType(const FPropertyTag& Tag,
 
 		if (Tag.Type == NAME_StructProperty && !bCanSerialize && (UnderlyingArchive.UEVer() >= VER_UE4_STRUCT_GUID_IN_PROPERTY_TAG || !UseBinaryOrNativeSerialization(UnderlyingArchive)))
 		{
+			UE::FPropertyTypeNameBuilder Builder;
+			SaveTypeName(Builder);
 			UE_LOG(LogClass, Warning, TEXT("Struct Property %s has a struct type mismatch (tag %s != prop %s) in package: %s. If that struct got renamed, add an entry to ActiveStructRedirects."),
-				*WriteToString<32>(Tag.Name), *WriteToString<32>(Tag.GetType().GetParameterName(0)), *WriteToString<32>(Struct->GetFName()), *UnderlyingArchive.GetArchiveName());
+				*WriteToString<32>(Tag.Name), *WriteToString<64>(Tag.GetType().GetParameter(0)), *WriteToString<64>(Builder.Build().GetParameter(0)), *UnderlyingArchive.GetArchiveName());
 			return EConvertFromTypeResult::CannotConvert;
 		}
 	}
