Removed the default parameter index from FPropertyTypeName

#rnx

diff --git a/Engine/Source/Runtime/CoreUObject/Private/UObject/PropertyArray.cpp b/Engine/Source/Runtime/CoreUObject/Private/UObject/PropertyArray.cpp
--- a/Engine/Source/Runtime/CoreUObject/Private/UObject/PropertyArray.cpp
+++ b/Engine/Source/Runtime/CoreUObject/Private/UObject/PropertyArray.cpp
@@ -627,7 +627,7 @@ void FArrayProperty::SerializeItem(FStructuredArchive::FSlot Slot, void* Value,
 				UE::FPropertyPathNameSegment Segment = Context->SerializedPropertyPath.GetSegment(SegmentIndex);
 				if (Segment.Type.GetName() == NAME_ArrayProperty)
 				{
-					const UE::FPropertyTypeName InnerTypeName = Segment.Type.GetParameter();
+					const UE::FPropertyTypeName InnerTypeName = Segment.Type.GetParameter(0);
 					if (InnerTypeName.GetName() == NAME_StructProperty && InnerTypeName.GetParameterCount() == 0)
 					{
 						UE::FPropertyTypeNameBuilder NewTypeBuilder;
@@ -1183,7 +1183,7 @@ EConvertFromTypeResult FArrayProperty::ConvertFromType(const FPropertyTag& Tag,
 	}
 	else
 	{
-		const FName InnerTypeName = Tag.GetType().GetParameterName();
+		const FName InnerTypeName = Tag.GetType().GetParameterName(0);
 		if (InnerTypeName.IsNone() || InnerTypeName == Inner->GetID())
 		{
 			return EConvertFromTypeResult::UseSerializeItem;
@@ -1209,7 +1209,7 @@ EConvertFromTypeResult FArrayProperty::ConvertFromType(const FPropertyTag& Tag,
 	ScriptArrayHelper.EmptyAndAddValues(ElementCount);
 
 	FPropertyTag InnerPropertyTag;
-	InnerPropertyTag.SetType(Tag.GetType().GetParameter());
+	InnerPropertyTag.SetType(Tag.GetType().GetParameter(0));
 	InnerPropertyTag.Name = Tag.Name;
 	InnerPropertyTag.ArrayIndex = 0;
 
@@ -1296,7 +1296,7 @@ bool FArrayProperty::LoadTypeName(UE::FPropertyTypeName Type, const FPropertyTag
 		return false;
 	}
 
-	const UE::FPropertyTypeName InnerType = Type.GetParameter();
+	const UE::FPropertyTypeName InnerType = Type.GetParameter(0);
 	FField* Field = FField::TryConstruct(InnerType.GetName(), this, GetFName(), RF_NoFlags);
 	if (FProperty* Property = CastField<FProperty>(Field); Property && Property->LoadTypeName(InnerType, Tag))
 	{
@@ -1327,5 +1327,5 @@ bool FArrayProperty::CanSerializeFromTypeName(UE::FPropertyTypeName Type) const
 
 	const FProperty* LocalInner = Inner;
 	check(LocalInner);
-	return LocalInner->CanSerializeFromTypeName(Type.GetParameter());
+	return LocalInner->CanSerializeFromTypeName(Type.GetParameter(0));
 }
diff --git a/Engine/Source/Runtime/CoreUObject/Private/UObject/PropertyByte.cpp b/Engine/Source/Runtime/CoreUObject/Private/UObject/PropertyByte.cpp
--- a/Engine/Source/Runtime/CoreUObject/Private/UObject/PropertyByte.cpp
+++ b/Engine/Source/Runtime/CoreUObject/Private/UObject/PropertyByte.cpp
@@ -522,7 +522,7 @@ bool FByteProperty::LoadTypeName(UE::FPropertyTypeName Type, const FPropertyTag*
 		return false;
 	}
 
-	const FName EnumName = Type.GetParameterName();
+	const FName EnumName = Type.GetParameterName(0);
 	if (EnumName.IsNone())
 	{
 		return true;
@@ -556,7 +556,7 @@ bool FByteProperty::CanSerializeFromTypeName(UE::FPropertyTypeName Type) const
 		return false;
 	}
 
-	const FName EnumName = Type.GetParameterName();
+	const FName EnumName = Type.GetParameterName(0);
 	if (const UEnum* LocalEnum = Enum)
 	{
 		return EnumName == LocalEnum->GetFName();
diff --git a/Engine/Source/Runtime/CoreUObject/Private/UObject/PropertyNumeric.cpp b/Engine/Source/Runtime/CoreUObject/Private/UObject/PropertyNumeric.cpp
--- a/Engine/Source/Runtime/CoreUObject/Private/UObject/PropertyNumeric.cpp
+++ b/Engine/Source/Runtime/CoreUObject/Private/UObject/PropertyNumeric.cpp
@@ -25,7 +25,7 @@ int64 FNumericProperty::ReadEnumAsInt64(FStructuredArchive::FSlot Slot, UStruct*
 	FName EnumName;
 	Slot << EnumName;
 
-	FName EnumTypeName = Tag.GetType().GetParameterName();
+	FName EnumTypeName = Tag.GetType().GetParameterName(0);
 
 	UEnum* Enum = FindUField<UEnum>(dynamic_cast<UClass*>(DefaultsStruct) ? static_cast<UClass*>(DefaultsStruct) : DefaultsStruct->GetTypedOuter<UClass>(), EnumTypeName);
 	if (!Enum)
diff --git a/Engine/Source/Runtime/CoreUObject/Private/UObject/PropertyOptional.cpp b/Engine/Source/Runtime/CoreUObject/Private/UObject/PropertyOptional.cpp
--- a/Engine/Source/Runtime/CoreUObject/Private/UObject/PropertyOptional.cpp
+++ b/Engine/Source/Runtime/CoreUObject/Private/UObject/PropertyOptional.cpp
@@ -477,9 +477,9 @@ EConvertFromTypeResult FOptionalProperty::ConvertFromType(const FPropertyTag& Ta
 				MarkUnset(ContainerPtrToValuePtr<void>(ContainerData));
 				return EConvertFromTypeResult::Converted;
 			}
-			ValueTag.SetType(Tag.GetType().GetParameter());
+			ValueTag.SetType(Tag.GetType().GetParameter(0));
 
-			// Mimic FPropertyTag::SerializeTaggedProperty's special code path for bools that serializes the value as part of the tag, even though in this case we didn't serialize a separate tag for the inner value.
+			// Mimic FPropertyTag::SerializeTaggedProperty storing bool values in the tag because implementations of ConvertFromType expect it there.
 			if (ValueTag.Type == NAME_BoolProperty)
 			{
 				*MaybeValueSlot << ValueTag.BoolVal;
@@ -633,7 +633,7 @@ bool FOptionalProperty::LoadTypeName(UE::FPropertyTypeName Type, const FProperty
 		return false;
 	}
 
-	const UE::FPropertyTypeName ValueType = Type.GetParameter();
+	const UE::FPropertyTypeName ValueType = Type.GetParameter(0);
 	FField* Field = FField::TryConstruct(ValueType.GetName(), this, GetFName(), RF_NoFlags);
 	if (FProperty* Property = CastField<FProperty>(Field); Property && Property->LoadTypeName(ValueType, Tag))
 	{
@@ -664,5 +664,5 @@ bool FOptionalProperty::CanSerializeFromTypeName(UE::FPropertyTypeName Type) con
 
 	const FProperty* LocalValueProperty = ValueProperty;
 	check(LocalValueProperty);
-	return LocalValueProperty->CanSerializeFromTypeName(Type.GetParameter());
+	return LocalValueProperty->CanSerializeFromTypeName(Type.GetParameter(0));
 }
diff --git a/Engine/Source/Runtime/CoreUObject/Private/UObject/PropertySet.cpp b/Engine/Source/Runtime/CoreUObject/Private/UObject/PropertySet.cpp
--- a/Engine/Source/Runtime/CoreUObject/Private/UObject/PropertySet.cpp
+++ b/Engine/Source/Runtime/CoreUObject/Private/UObject/PropertySet.cpp
@@ -910,7 +910,7 @@ EConvertFromTypeResult FSetProperty::ConvertFromType(const FPropertyTag& Tag, FS
 	}
 	else
 	{
-		const FName InnerTypeName = Tag.GetType().GetParameterName();
+		const FName InnerTypeName = Tag.GetType().GetParameterName(0);
 		if (InnerTypeName.IsNone() || InnerTypeName == ElementProp->GetID())
 		{
 			return EConvertFromTypeResult::UseSerializeItem;
@@ -929,7 +929,7 @@ EConvertFromTypeResult FSetProperty::ConvertFromType(const FPropertyTag& Tag, FS
 	};
 
 	FPropertyTag InnerPropertyTag;
-	InnerPropertyTag.SetType(Tag.GetType().GetParameter());
+	InnerPropertyTag.SetType(Tag.GetType().GetParameter(0));
 	InnerPropertyTag.Name = Tag.Name;
 	InnerPropertyTag.ArrayIndex = 0;
 
@@ -1110,7 +1110,7 @@ bool FSetProperty::LoadTypeName(UE::FPropertyTypeName Type, const FPropertyTag*
 		return false;
 	}
 
-	const UE::FPropertyTypeName ElementType = Type.GetParameter();
+	const UE::FPropertyTypeName ElementType = Type.GetParameter(0);
 	FField* Field = FField::TryConstruct(ElementType.GetName(), this, GetFName(), RF_NoFlags);
 	if (FProperty* Property = CastField<FProperty>(Field); Property && Property->LoadTypeName(ElementType, Tag))
 	{
@@ -1141,5 +1141,5 @@ bool FSetProperty::CanSerializeFromTypeName(UE::FPropertyTypeName Type) const
 
 	const FProperty* LocalElementProp = ElementProp;
 	check(LocalElementProp);
-	return LocalElementProp->CanSerializeFromTypeName(Type.GetParameter());
+	return LocalElementProp->CanSerializeFromTypeName(Type.GetParameter(0));
 }
diff --git a/Engine/Source/Runtime/CoreUObject/Private/UObject/PropertySoftObjectPtr.cpp b/Engine/Source/Runtime/CoreUObject/Private/UObject/PropertySoftObjectPtr.cpp
--- a/Engine/Source/Runtime/CoreUObject/Private/UObject/PropertySoftObjectPtr.cpp
+++ b/Engine/Source/Runtime/CoreUObject/Private/UObject/PropertySoftObjectPtr.cpp
@@ -241,7 +241,7 @@ EConvertFromTypeResult FSoftObjectProperty::ConvertFromType(const FPropertyTag&
 	}
 	else if (Tag.Type == NAME_StructProperty)
 	{
-		const FName StructName = Tag.GetType().GetParameterName();
+		const FName StructName = Tag.GetType().GetParameterName(0);
 		if (StructName == NAME_SoftObjectPath || StructName == NAME_SoftClassPath || StructName == NAME_StringAssetReference || StructName == NAME_StringClassReference)
 		{
 			// This property used to be a FSoftObjectPath but is now a TSoftObjectPtr<Foo>
diff --git a/Engine/Source/Runtime/CoreUObject/Private/UObject/PropertyTag.cpp b/Engine/Source/Runtime/CoreUObject/Private/UObject/PropertyTag.cpp
--- a/Engine/Source/Runtime/CoreUObject/Private/UObject/PropertyTag.cpp
+++ b/Engine/Source/Runtime/CoreUObject/Private/UObject/PropertyTag.cpp
@@ -105,12 +105,12 @@ void FPropertyTag::SetType(UE::FPropertyTypeName InFullType)
 			break;
 		case NAME_ByteProperty:
 		case NAME_EnumProperty:
-			EnumName = TypeName.GetParameterName();
+			EnumName = TypeName.GetParameterName(0);
 			break;
 		case NAME_ArrayProperty:
 		case NAME_OptionalProperty:
 		case NAME_SetProperty:
-			InnerType = TypeName.GetParameterName();
+			InnerType = TypeName.GetParameterName(0);
 			break;
 		case NAME_MapProperty:
 			InnerType = TypeName.GetParameterName(0);
diff --git a/Engine/Source/Runtime/CoreUObject/Public/UObject/Class.h b/Engine/Source/Runtime/CoreUObject/Public/UObject/Class.h
--- a/Engine/Source/Runtime/CoreUObject/Public/UObject/Class.h
+++ b/Engine/Source/Runtime/CoreUObject/Public/UObject/Class.h
@@ -1460,7 +1460,7 @@ public:
 					FName StructName;
 					if (Tag.GetType().GetName() == NAME_StructProperty && Tag.GetType().GetParameterCount() >= 1)
 					{
-						StructName = Tag.GetType().GetParameterName();
+						StructName = Tag.GetType().GetParameterName(0);
 					}
 					return ((CPPSTRUCT*)Data)->SerializeFromMismatchedTag(StructName, Ar);
 				}
@@ -1485,7 +1485,7 @@ public:
 					FName StructName;
 					if (Tag.GetType().GetName() == NAME_StructProperty && Tag.GetType().GetParameterCount() >= 1)
 					{
-						StructName = Tag.GetType().GetParameterName();
+						StructName = Tag.GetType().GetParameterName(0);
 					}
 					return ((CPPSTRUCT*)Data)->SerializeFromMismatchedTag(StructName, Slot);
 				}
diff --git a/Engine/Source/Runtime/CoreUObject/Public/UObject/PropertyTypeName.h b/Engine/Source/Runtime/CoreUObject/Public/UObject/PropertyTypeName.h
--- a/Engine/Source/Runtime/CoreUObject/Public/UObject/PropertyTypeName.h
+++ b/Engine/Source/Runtime/CoreUObject/Public/UObject/PropertyTypeName.h
@@ -74,7 +74,7 @@ public:
 	 * - GetParameter(0) -> StructProperty(KeyStruct)
 	 * - GetParameter(1) -> EnumProperty(ByteEnum,ByteProperty)
 	 */
-	UE_API FPropertyTypeName GetParameter(int32 ParamIndex = 0) const;
+	UE_API FPropertyTypeName GetParameter(int32 ParamIndex) const;
 
 	/**
 	 * Returns the indexed parameter type name under the root of this property type name.
@@ -85,7 +85,7 @@ public:
 	 * - GetParameterName(0) -> StructProperty
 	 * - GetParameterName(1) -> EnumProperty
 	 */
-	inline FName GetParameterName(int32 ParamIndex = 0) const
+	inline FName GetParameterName(int32 ParamIndex) const
 	{
 		return GetParameter(ParamIndex).GetName();
 	}
