Removed the non-persistent serialization path for FPropertyTypeName

There are non-persistent archives being used for persistent data. This type is too widely-used to reasonably expect every incorrect archive to be found and fixed.

#rnx

diff --git a/Engine/Source/Runtime/CoreUObject/Private/UObject/PropertyTypeName.cpp b/Engine/Source/Runtime/CoreUObject/Private/UObject/PropertyTypeName.cpp
--- a/Engine/Source/Runtime/CoreUObject/Private/UObject/PropertyTypeName.cpp
+++ b/Engine/Source/Runtime/CoreUObject/Private/UObject/PropertyTypeName.cpp
@@ -307,11 +307,7 @@ bool operator<(const FPropertyTypeName& Lhs, const FPropertyTypeName& Rhs)
 
 FArchive& operator<<(FArchive& Ar, FPropertyTypeName& TypeName)
 {
-	if (!Ar.IsPersistent())
-	{
-		Ar << TypeName.Index;
-	}
-	else if (Ar.IsLoading())
+	if (Ar.IsLoading())
 	{
 		TArray<FPropertyTypeNameNode, TInlineAllocator<16>> Nodes;
 		int32 Remaining = 1;
@@ -340,11 +336,7 @@ FArchive& operator<<(FArchive& Ar, FPropertyTypeName& TypeName)
 void operator<<(FStructuredArchiveSlot Slot, FPropertyTypeName& TypeName)
 {
 	const FArchiveState& State = Slot.GetArchiveState();
-	if (!State.IsPersistent())
-	{
-		Slot << TypeName.Index;
-	}
-	else if (State.IsLoading())
+	if (State.IsLoading())
 	{
 		if (State.IsTextFormat())
 		{
