Changed  FPropertyTypeName to serialize when neither saving nor loading

#rnx

diff --git a/Engine/Source/Runtime/CoreUObject/Private/UObject/PropertyTypeName.cpp b/Engine/Source/Runtime/CoreUObject/Private/UObject/PropertyTypeName.cpp
--- a/Engine/Source/Runtime/CoreUObject/Private/UObject/PropertyTypeName.cpp
+++ b/Engine/Source/Runtime/CoreUObject/Private/UObject/PropertyTypeName.cpp
@@ -320,7 +320,7 @@ FArchive& operator<<(FArchive& Ar, FPropertyTypeName& TypeName)
 		while (Remaining > 0);
 		TypeName.Index = GPropertyTypeNameTable.FindOrAddByName(Nodes.GetData());
 	}
-	else if (Ar.IsSaving())
+	else
 	{
 		const FPropertyTypeNameNode* Node = GPropertyTypeNameTable.ResolveByIndex(TypeName.Index);
 		for (int32 Remaining = 1; Remaining > 0; --Remaining, ++Node)
@@ -368,7 +368,7 @@ void operator<<(FStructuredArchiveSlot Slot, FPropertyTypeName& TypeName)
 			TypeName.Index = GPropertyTypeNameTable.FindOrAddByName(Nodes.GetData());
 		}
 	}
-	else if (State.IsSaving())
+	else
 	{
 		if (State.IsTextFormat())
 		{
