Early out of FByteProperty::ExportText_Internal if it's to be treated numerically, i.e. non-enum.

diff --git a/Engine/Source/Runtime/CoreUObject/Private/UObject/PropertyByte.cpp b/Engine/Source/Runtime/CoreUObject/Private/UObject/PropertyByte.cpp
--- a/Engine/Source/Runtime/CoreUObject/Private/UObject/PropertyByte.cpp
+++ b/Engine/Source/Runtime/CoreUObject/Private/UObject/PropertyByte.cpp
@@ -358,47 +358,46 @@ void FByteProperty::AppendSchemaHash(FBlake3& Builder, bool bSkipEditorOnly) con
 
 void FByteProperty::ExportText_Internal( FString& ValueStr, const void* PropertyValueOrContainer, EPropertyPointerType PropertyPointerType, const void* DefaultValue, UObject* Parent, int32 PortFlags, UObject* ExportRootScope ) const
 {
-	if( Enum && (PortFlags & PPF_ConsoleVariable) == 0 )
+	if (!Enum || (PortFlags & PPF_ConsoleVariable))
 	{
-		int64 ActualValue = 0;
-		if (PropertyPointerType == EPropertyPointerType::Container && HasGetter())
-		{
-			uint8 Value = 0;
-			GetValue_InContainer(PropertyValueOrContainer, &Value);
-			ActualValue = Value;
-		}
-		else
+		Super::ExportText_Internal(ValueStr, PropertyValueOrContainer, PropertyPointerType, DefaultValue, Parent, PortFlags, ExportRootScope);
+		return;
+	}
+
+	int64 ActualValue = 0;
+	if (PropertyPointerType == EPropertyPointerType::Container && HasGetter())
+	{
+		uint8 Value = 0;
+		GetValue_InContainer(PropertyValueOrContainer, &Value);
+		ActualValue = Value;
+	}
+	else
+	{
+		ActualValue = *(const uint8*)PointerToValuePtr(PropertyValueOrContainer, PropertyPointerType);
+	}
+	// if the value is the max value (the autogenerated *_MAX value), export as "INVALID", unless we're exporting text for copy/paste (for copy/paste,
+	// the property text value must actually match an entry in the enum's names array)
+	bool bIsValid = Enum->IsValidEnumValue(ActualValue);
+	bool bIsMax = ActualValue == Enum->GetMaxEnumValue();
+	if (bIsValid && (!bIsMax || (PortFlags & PPF_Copy)))
+	{
+		// We do not want to export the enum text for non-display uses, localization text is very dynamic and would cause issues on import
+		if (PortFlags & PPF_PropertyWindow)
 		{
-			ActualValue = *(const uint8*)PointerToValuePtr(PropertyValueOrContainer, PropertyPointerType);
+			ValueStr += Enum->GetDisplayNameTextByValue(ActualValue).ToString();
 		}
-		// if the value is the max value (the autogenerated *_MAX value), export as "INVALID", unless we're exporting text for copy/paste (for copy/paste,
-		// the property text value must actually match an entry in the enum's names array)
-		bool bIsValid = Enum->IsValidEnumValue(ActualValue);
-		bool bIsMax = ActualValue == Enum->GetMaxEnumValue();
-		if (bIsValid && (!bIsMax || (PortFlags & PPF_Copy)))
+		else if (PortFlags & PPF_ExternalEditor)
 		{
-			// We do not want to export the enum text for non-display uses, localization text is very dynamic and would cause issues on import
-			if (PortFlags & PPF_PropertyWindow)
-			{
-				ValueStr += Enum->GetDisplayNameTextByValue(ActualValue).ToString();
-			}
-			else if (PortFlags & PPF_ExternalEditor)
-			{
-				ValueStr += Enum->GetAuthoredNameStringByValue(ActualValue);
-			}
-			else
-			{
-				ValueStr += Enum->GetNameStringByValue(ActualValue);
-			}
+			ValueStr += Enum->GetAuthoredNameStringByValue(ActualValue);
 		}
 		else
 		{
-			ValueStr += TEXT("(INVALID)");
+			ValueStr += Enum->GetNameStringByValue(ActualValue);
 		}
 	}
 	else
 	{
-		Super::ExportText_Internal(ValueStr, PropertyValueOrContainer, PropertyPointerType, DefaultValue, Parent, PortFlags, ExportRootScope);
+		ValueStr += TEXT("(INVALID)");
 	}
 }
 const TCHAR* FByteProperty::ImportText_Internal( const TCHAR* InBuffer, void* ContainerOrPropertyPtr, EPropertyPointerType PropertyPointerType, UObject* Parent, int32 PortFlags, FOutputDevice* ErrorText ) const
