Early out on exporting invalid enum values in FByteProperty::ExportText_Internal.

diff --git a/Engine/Source/Runtime/CoreUObject/Private/UObject/PropertyByte.cpp b/Engine/Source/Runtime/CoreUObject/Private/UObject/PropertyByte.cpp
--- a/Engine/Source/Runtime/CoreUObject/Private/UObject/PropertyByte.cpp
+++ b/Engine/Source/Runtime/CoreUObject/Private/UObject/PropertyByte.cpp
@@ -375,29 +375,27 @@ void FByteProperty::ExportText_Internal( FString& ValueStr, const void* Property
 	{
 		ActualValue = *(const uint8*)PointerToValuePtr(PropertyValueOrContainer, PropertyPointerType);
 	}
+
 	// if the value is the max value (the autogenerated *_MAX value), export as "INVALID", unless we're exporting text for copy/paste (for copy/paste,
 	// the property text value must actually match an entry in the enum's names array)
-	bool bIsValid = Enum->IsValidEnumValue(ActualValue);
-	bool bIsMax = ActualValue == Enum->GetMaxEnumValue();
-	if (bIsValid && (!bIsMax || (PortFlags & PPF_Copy)))
+	if (!Enum->IsValidEnumValue(ActualValue) || (!(PortFlags & PPF_Copy) && ActualValue == Enum->GetMaxEnumValue()))
 	{
-		// We do not want to export the enum text for non-display uses, localization text is very dynamic and would cause issues on import
-		if (PortFlags & PPF_PropertyWindow)
-		{
-			ValueStr += Enum->GetDisplayNameTextByValue(ActualValue).ToString();
-		}
-		else if (PortFlags & PPF_ExternalEditor)
-		{
-			ValueStr += Enum->GetAuthoredNameStringByValue(ActualValue);
-		}
-		else
-		{
-			ValueStr += Enum->GetNameStringByValue(ActualValue);
-		}
+		ValueStr += TEXT("(INVALID)");
+		return;
+	}
+
+	// We do not want to export the enum text for non-display uses, localization text is very dynamic and would cause issues on import
+	if (PortFlags & PPF_PropertyWindow)
+	{
+		ValueStr += Enum->GetDisplayNameTextByValue(ActualValue).ToString();
+	}
+	else if (PortFlags & PPF_ExternalEditor)
+	{
+		ValueStr += Enum->GetAuthoredNameStringByValue(ActualValue);
 	}
 	else
 	{
-		ValueStr += TEXT("(INVALID)");
+		ValueStr += Enum->GetNameStringByValue(ActualValue);
 	}
 }
 const TCHAR* FByteProperty::ImportText_Internal( const TCHAR* InBuffer, void* ContainerOrPropertyPtr, EPropertyPointerType PropertyPointerType, UObject* Parent, int32 PortFlags, FOutputDevice* ErrorText ) const
