Fixes for the reuse of LocalUnderlyingProp in FEnumProperty::ExportText_Internal.

diff --git a/Engine/Source/Runtime/CoreUObject/Private/UObject/EnumProperty.cpp b/Engine/Source/Runtime/CoreUObject/Private/UObject/EnumProperty.cpp
--- a/Engine/Source/Runtime/CoreUObject/Private/UObject/EnumProperty.cpp
+++ b/Engine/Source/Runtime/CoreUObject/Private/UObject/EnumProperty.cpp
@@ -264,11 +264,11 @@ void FEnumProperty::ExportText_Internal(FString& ValueStr, const void* PropertyV
 		return;
 	}
 
-	check(UnderlyingProp);
+	FNumericProperty* LocalUnderlyingProp = UnderlyingProp;
+	check(LocalUnderlyingProp);
 
 	int64 LocalValue = 0;
 	void* PropertyValue = nullptr;
-	FNumericProperty* LocalUnderlyingProp = UnderlyingProp;
 
 	if (PropertyPointerType == EPropertyPointerType::Container && HasGetter())
 	{
@@ -282,7 +282,7 @@ void FEnumProperty::ExportText_Internal(FString& ValueStr, const void* PropertyV
 
 	if (PortFlags & PPF_ConsoleVariable)
 	{
-		UnderlyingProp->ExportText_Internal(ValueStr, PropertyValue, EPropertyPointerType::Direct, DefaultValue, Parent, PortFlags, ExportRootScope);
+		LocalUnderlyingProp->ExportText_Internal(ValueStr, PropertyValue, EPropertyPointerType::Direct, DefaultValue, Parent, PortFlags, ExportRootScope);
 		return;
 	}
 
