Refactored FEnumProperty's and FByteProperty's ExportText_Internal to bring them closer to each other, to make them easier to unify.

diff --git a/Engine/Source/Runtime/CoreUObject/Private/UObject/EnumProperty.cpp b/Engine/Source/Runtime/CoreUObject/Private/UObject/EnumProperty.cpp
--- a/Engine/Source/Runtime/CoreUObject/Private/UObject/EnumProperty.cpp
+++ b/Engine/Source/Runtime/CoreUObject/Private/UObject/EnumProperty.cpp
@@ -267,26 +267,29 @@ void FEnumProperty::ExportText_Internal(FString& ValueStr, const void* PropertyV
 	FNumericProperty* LocalUnderlyingProp = UnderlyingProp;
 	check(LocalUnderlyingProp);
 
-	int64 LocalValue = 0;
-	void* PropertyValue = nullptr;
+	alignas(int64) uint8 GetterResult[sizeof(int64)];
+	void* ValuePtr = nullptr;
 
 	if (PropertyPointerType == EPropertyPointerType::Container && HasGetter())
 	{
-		PropertyValue = &LocalValue;
-		GetValue_InContainer(PropertyValueOrContainer, PropertyValue);
+		// Put the bytes returned by the getter into GetterResult (which is big and aligned enough to hold any numeric value)
+		// - these bytes will be re-read by a call to LocalUnderlyingProp below.
+		ValuePtr = &GetterResult;
+		GetValue_InContainer(PropertyValueOrContainer, ValuePtr);
 	}
 	else
 	{
-		PropertyValue = PointerToValuePtr(PropertyValueOrContainer, PropertyPointerType);
+		// Otherwise read directly from the property
+		ValuePtr = PointerToValuePtr(PropertyValueOrContainer, PropertyPointerType);
 	}
 
 	if (PortFlags & PPF_ConsoleVariable)
 	{
-		LocalUnderlyingProp->ExportText_Internal(ValueStr, PropertyValue, EPropertyPointerType::Direct, DefaultValue, Parent, PortFlags, ExportRootScope);
+		LocalUnderlyingProp->ExportText_Internal(ValueStr, ValuePtr, EPropertyPointerType::Direct, DefaultValue, Parent, PortFlags, ExportRootScope);
 		return;
 	}
 
-	int64 Value = LocalUnderlyingProp->GetSignedIntPropertyValue(PropertyValue);
+	int64 Value = LocalUnderlyingProp->GetSignedIntPropertyValue(ValuePtr);
 
 	// if the value is the max value (the autogenerated *_MAX value), export as "INVALID", unless we're exporting text for copy/paste (for copy/paste,
 	// the property text value must actually match an entry in the enum's names array)
diff --git a/Engine/Source/Runtime/CoreUObject/Private/UObject/PropertyByte.cpp b/Engine/Source/Runtime/CoreUObject/Private/UObject/PropertyByte.cpp
--- a/Engine/Source/Runtime/CoreUObject/Private/UObject/PropertyByte.cpp
+++ b/Engine/Source/Runtime/CoreUObject/Private/UObject/PropertyByte.cpp
@@ -364,21 +364,27 @@ void FByteProperty::ExportText_Internal( FString& ValueStr, const void* Property
 		return;
 	}
 
-	int64 ActualValue = 0;
+	uint8 GetterResult;
+	void* ValuePtr = nullptr;
+
 	if (PropertyPointerType == EPropertyPointerType::Container && HasGetter())
 	{
-		uint8 Value = 0;
-		GetValue_InContainer(PropertyValueOrContainer, &Value);
-		ActualValue = Value;
+		// Put the byte returned by the getter into GetterResult
+		// - this will be re-read by a call to GetSignedIntPropertyValue below.
+		ValuePtr = &GetterResult;
+		GetValue_InContainer(PropertyValueOrContainer, (uint8*)ValuePtr);
 	}
 	else
 	{
-		ActualValue = *(const uint8*)PointerToValuePtr(PropertyValueOrContainer, PropertyPointerType);
+		// Otherwise read directly from the property
+		ValuePtr = PointerToValuePtr(PropertyValueOrContainer, PropertyPointerType);
 	}
 
+	int64 Value = GetSignedIntPropertyValue(ValuePtr);
+
 	// if the value is the max value (the autogenerated *_MAX value), export as "INVALID", unless we're exporting text for copy/paste (for copy/paste,
 	// the property text value must actually match an entry in the enum's names array)
-	if (!Enum->IsValidEnumValue(ActualValue) || (!(PortFlags & PPF_Copy) && ActualValue == Enum->GetMaxEnumValue()))
+	if (!Enum->IsValidEnumValue(Value) || (!(PortFlags & PPF_Copy) && Value == Enum->GetMaxEnumValue()))
 	{
 		ValueStr += TEXT("(INVALID)");
 		return;
@@ -387,15 +393,15 @@ void FByteProperty::ExportText_Internal( FString& ValueStr, const void* Property
 	// We do not want to export the enum text for non-display uses, localization text is very dynamic and would cause issues on import
 	if (PortFlags & PPF_PropertyWindow)
 	{
-		ValueStr += Enum->GetDisplayNameTextByValue(ActualValue).ToString();
+		ValueStr += Enum->GetDisplayNameTextByValue(Value).ToString();
 	}
 	else if (PortFlags & PPF_ExternalEditor)
 	{
-		ValueStr += Enum->GetAuthoredNameStringByValue(ActualValue);
+		ValueStr += Enum->GetAuthoredNameStringByValue(Value);
 	}
 	else
 	{
-		ValueStr += Enum->GetNameStringByValue(ActualValue);
+		ValueStr += Enum->GetNameStringByValue(Value);
 	}
 }
 const TCHAR* FByteProperty::ImportText_Internal( const TCHAR* InBuffer, void* ContainerOrPropertyPtr, EPropertyPointerType PropertyPointerType, UObject* Parent, int32 PortFlags, FOutputDevice* ErrorText ) const
