Renamed old-style namespace.

diff --git a/Engine/Source/Runtime/CoreUObject/Private/UObject/EnumProperty.cpp b/Engine/Source/Runtime/CoreUObject/Private/UObject/EnumProperty.cpp
--- a/Engine/Source/Runtime/CoreUObject/Private/UObject/EnumProperty.cpp
+++ b/Engine/Source/Runtime/CoreUObject/Private/UObject/EnumProperty.cpp
@@ -8,7 +8,7 @@
 #include "UObject/UnrealTypePrivate.h"
 #include "UObject/UObjectThreadContext.h"
 
-namespace UEEnumProperty_Private
+namespace UE::CoreUObject::Private
 {
 	template <typename OldIntType>
 	void ConvertIntValueToEnumProperty(OldIntType OldValue, FEnumProperty* EnumProp, FNumericProperty* UnderlyingProp, UEnum* Enum, void* Obj)
@@ -492,28 +492,28 @@ EConvertFromTypeResult FEnumProperty::ConvertFromType(const FPropertyTag& Tag, F
 		return EConvertFromTypeResult::Converted;
 	}
 	case NAME_Int8Property:
-		UEEnumProperty_Private::ConvertIntToEnumProperty<int8>(Slot, this, UnderlyingProp, Enum, ContainerPtrToValuePtr<void>(Data, Tag.ArrayIndex));
+		UE::CoreUObject::Private::ConvertIntToEnumProperty<int8>(Slot, this, UnderlyingProp, Enum, ContainerPtrToValuePtr<void>(Data, Tag.ArrayIndex));
 		return EConvertFromTypeResult::Converted;
 	case NAME_Int16Property:
-		UEEnumProperty_Private::ConvertIntToEnumProperty<int16>(Slot, this, UnderlyingProp, Enum, ContainerPtrToValuePtr<void>(Data, Tag.ArrayIndex));
+		UE::CoreUObject::Private::ConvertIntToEnumProperty<int16>(Slot, this, UnderlyingProp, Enum, ContainerPtrToValuePtr<void>(Data, Tag.ArrayIndex));
 		return EConvertFromTypeResult::Converted;
 	case NAME_IntProperty:
-		UEEnumProperty_Private::ConvertIntToEnumProperty<int32>(Slot, this, UnderlyingProp, Enum, ContainerPtrToValuePtr<void>(Data, Tag.ArrayIndex));
+		UE::CoreUObject::Private::ConvertIntToEnumProperty<int32>(Slot, this, UnderlyingProp, Enum, ContainerPtrToValuePtr<void>(Data, Tag.ArrayIndex));
 		return EConvertFromTypeResult::Converted;
 	case NAME_Int64Property:
-		UEEnumProperty_Private::ConvertIntToEnumProperty<int64>(Slot, this, UnderlyingProp, Enum, ContainerPtrToValuePtr<void>(Data, Tag.ArrayIndex));
+		UE::CoreUObject::Private::ConvertIntToEnumProperty<int64>(Slot, this, UnderlyingProp, Enum, ContainerPtrToValuePtr<void>(Data, Tag.ArrayIndex));
 		return EConvertFromTypeResult::Converted;
 	case NAME_UInt16Property:
-		UEEnumProperty_Private::ConvertIntToEnumProperty<uint16>(Slot, this, UnderlyingProp, Enum, ContainerPtrToValuePtr<void>(Data, Tag.ArrayIndex));
+		UE::CoreUObject::Private::ConvertIntToEnumProperty<uint16>(Slot, this, UnderlyingProp, Enum, ContainerPtrToValuePtr<void>(Data, Tag.ArrayIndex));
 		return EConvertFromTypeResult::Converted;
 	case NAME_UInt32Property:
-		UEEnumProperty_Private::ConvertIntToEnumProperty<uint32>(Slot, this, UnderlyingProp, Enum, ContainerPtrToValuePtr<void>(Data, Tag.ArrayIndex));
+		UE::CoreUObject::Private::ConvertIntToEnumProperty<uint32>(Slot, this, UnderlyingProp, Enum, ContainerPtrToValuePtr<void>(Data, Tag.ArrayIndex));
 		return EConvertFromTypeResult::Converted;
 	case NAME_UInt64Property:
-		UEEnumProperty_Private::ConvertIntToEnumProperty<uint64>(Slot, this, UnderlyingProp, Enum, ContainerPtrToValuePtr<void>(Data, Tag.ArrayIndex));
+		UE::CoreUObject::Private::ConvertIntToEnumProperty<uint64>(Slot, this, UnderlyingProp, Enum, ContainerPtrToValuePtr<void>(Data, Tag.ArrayIndex));
 		return EConvertFromTypeResult::Converted;
 	case NAME_BoolProperty:
-		UEEnumProperty_Private::ConvertIntValueToEnumProperty<uint8>(Tag.BoolVal, this, UnderlyingProp, Enum, ContainerPtrToValuePtr<void>(Data, Tag.ArrayIndex));
+		UE::CoreUObject::Private::ConvertIntValueToEnumProperty<uint8>(Tag.BoolVal, this, UnderlyingProp, Enum, ContainerPtrToValuePtr<void>(Data, Tag.ArrayIndex));
 		return EConvertFromTypeResult::Converted;
 	}
 }
diff --git a/Engine/Source/Runtime/CoreUObject/Private/UObject/PropertyPrivate.cpp b/Engine/Source/Runtime/CoreUObject/Private/UObject/PropertyPrivate.cpp
--- a/Engine/Source/Runtime/CoreUObject/Private/UObject/PropertyPrivate.cpp
+++ b/Engine/Source/Runtime/CoreUObject/Private/UObject/PropertyPrivate.cpp
@@ -162,7 +162,7 @@ void UEnumProperty::AddReferencedObjects(UObject* InThis, FReferenceCollector& C
 	Super::AddReferencedObjects(InThis, Collector);
 }
 
-namespace UEEnumProperty_Private
+namespace UE::CoreUObject::Private
 {
 	struct FEnumPropertyFriend
 	{
