Another attempt at controlling when IDO creation and property path tracking are enabled

#rnx

diff --git a/Engine/Source/Runtime/CoreUObject/Private/UObject/LinkerLoad.cpp b/Engine/Source/Runtime/CoreUObject/Private/UObject/LinkerLoad.cpp
--- a/Engine/Source/Runtime/CoreUObject/Private/UObject/LinkerLoad.cpp
+++ b/Engine/Source/Runtime/CoreUObject/Private/UObject/LinkerLoad.cpp
@@ -4650,11 +4650,13 @@ void FLinkerLoad::Preload( UObject* Object )
 
 					// Toggle support for IDOs
 					FUObjectSerializeContext* LoadContext = FUObjectThreadContext::Get().GetSerializeContext();
-					// Do not enable IDOs when impersonation is enabled as it implies we are deserializing an IDO
-					const bool bIDOEnabled = UE::IsInstanceDataObjectSupportEnabled(Object) && !LoadContext->bImpersonateProperties;
-					// This will add property path tracking and create a property bag to hold any property that does not match the current class schema
-					TGuardValue<bool> ScopedTrackSerializedPropertyPath(LoadContext->bTrackSerializedPropertyPath, bIDOEnabled);
-					TGuardValue<bool> ScopedSerializeUnknownProperty(LoadContext->bSerializeUnknownProperty, bIDOEnabled);
+					// Enable property path tracking when IDO support is enabled. Both property bags and IDO creation require the paths.
+					const bool bHasIDOSupport = UE::IsInstanceDataObjectSupportEnabled(Object);
+					TGuardValue<bool> ScopedTrackSerializedPropertyPath(LoadContext->bTrackSerializedPropertyPath, bHasIDOSupport);
+					// Enable creation of a property bag to hold any property that does not match the current class schema,
+					// except when impersonation is enabled because that implies we are deserializing an IDO.
+					const bool bCreateIDO = bHasIDOSupport && !LoadContext->bImpersonateProperties;
+					TGuardValue<bool> ScopedSerializeUnknownProperty(LoadContext->bSerializeUnknownProperty, bCreateIDO);
 
 					if (Object->HasAnyFlags(RF_ClassDefaultObject))
 					{
@@ -4779,7 +4781,7 @@ void FLinkerLoad::Preload( UObject* Object )
 					}
 
 					// Object has been deserialized, if IDO is enabled, generate it
-					if (bIDOEnabled)
+					if (bCreateIDO)
 					{
 						UE::FPropertyBagRepository::Get().CreateInstanceDataObject(Object);
 					}
diff --git a/Engine/Source/Runtime/CoreUObject/Private/UObject/PropertyBagRepository.cpp b/Engine/Source/Runtime/CoreUObject/Private/UObject/PropertyBagRepository.cpp
--- a/Engine/Source/Runtime/CoreUObject/Private/UObject/PropertyBagRepository.cpp
+++ b/Engine/Source/Runtime/CoreUObject/Private/UObject/PropertyBagRepository.cpp
@@ -297,7 +297,6 @@ void FPropertyBagRepository::CreateInstanceDataObjectUnsafe(const UObjectBase* O
 	// setup load context to mark properties the that were set by serialization
 	FUObjectSerializeContext* LoadContext = FUObjectThreadContext::Get().GetSerializeContext();
 	TGuardValue<bool> ScopedImpersonateProperties(LoadContext->bImpersonateProperties, true);
-	TGuardValue<bool> ScopedTrackSerializedPropertyPath(LoadContext->bTrackSerializedPropertyPath, true);
 	
 	UObject* OwnerAsObject = (UObject*)Owner;
 	if (FLinkerLoad* Linker = OwnerAsObject->GetLinker())
