Better support of arrays in arrays of subobjects.
Added unittest to test this case

diff --git a/Engine/Source/Runtime/CoreUObject/Private/UObject/OverriddenPropertySet.cpp b/Engine/Source/Runtime/CoreUObject/Private/UObject/OverriddenPropertySet.cpp
--- a/Engine/Source/Runtime/CoreUObject/Private/UObject/OverriddenPropertySet.cpp
+++ b/Engine/Source/Runtime/CoreUObject/Private/UObject/OverriddenPropertySet.cpp
@@ -666,9 +666,9 @@ void FOverriddenPropertySet::NotifyPropertyChange(FOverriddenPropertyNode* Paren
 						}
 					}
 				}
-				else if (Notification == EPropertyNotificationType::PostEdit)
+				// Can only forward to subobject if we have a valid index
+				else if (ArrayHelper.IsValidIndex(ArrayIndex))
 				{
-					checkf(ArrayHelper.IsValidIndex(ArrayIndex), TEXT("Any sub operation is expected to have a valid index"));
 					if (UObject* SubObject = InnerObjectProperty->GetObjectPropertyValue(ArrayHelper.GetElementPtr(ArrayIndex)))
 					{
 						// This should not be needed in the property grid, as it should already been called on the subobject itself.
@@ -915,10 +915,9 @@ void FOverriddenPropertySet::NotifyPropertyChange(FOverriddenPropertyNode* Paren
 				}
 			}
 		}
-		else if (Notification == EPropertyNotificationType::PostEdit)
+		// Can only forward to subobject if we have a valid index
+		else if (MapHelper.IsValidIndex(InternalMapIndex))
 		{
-			checkf(MapHelper.IsValidIndex(InternalMapIndex), TEXT("Any sub operation is expected to not have a valid index"));
-
 			// @todo support instanced object as a key in maps
 			//if (UObject* SubObject = KeyInstancedObjectProperty ? KeyInstancedObjectProperty->GetObjectPropertyValue(MapHelper.GetValuePtr(InternalMapIndex)) : nullptr)
 			//{
@@ -1026,7 +1025,7 @@ bool FOverriddenPropertySet::ClearOverriddenProperty(const FPropertyChangedEvent
 void FOverriddenPropertySet::OverrideProperty(const FPropertyChangedEvent& PropertyEvent, const FEditPropertyChain::TDoubleLinkedListNode* PropertyNode, const void* Data)
 {
 	FOverriddenPropertyNode& RootPropertyNode = OverriddenPropertyNodes.FindOrAddByHash(GetTypeHash(RootNodeID), RootNodeID);
-	NotifyPropertyChange(&RootPropertyNode, EPropertyNotificationType::PreEdit, FPropertyChangedEvent(nullptr), PropertyNode, Data);
+	NotifyPropertyChange(&RootPropertyNode, EPropertyNotificationType::PreEdit, PropertyEvent, PropertyNode, Data);
 	NotifyPropertyChange(&RootPropertyNode, EPropertyNotificationType::PostEdit, PropertyEvent, PropertyNode, Data);
 }
 
