Fix nodes not clearing up their modified state.

diff --git a/Engine/Source/Runtime/CoreUObject/Private/UObject/OverriddenPropertySet.cpp b/Engine/Source/Runtime/CoreUObject/Private/UObject/OverriddenPropertySet.cpp
--- a/Engine/Source/Runtime/CoreUObject/Private/UObject/OverriddenPropertySet.cpp
+++ b/Engine/Source/Runtime/CoreUObject/Private/UObject/OverriddenPropertySet.cpp
@@ -284,7 +284,9 @@ bool FOverriddenPropertySet::ClearOverriddenProperty(FOverriddenPropertyNode& Pa
 	const FEditPropertyChain::TDoubleLinkedListNode* PropertyIterator = PropertyNode;
 	FOverriddenPropertyNode* OverriddenPropertyNode = &ParentPropertyNode;
 	int32 ArrayIndex = INDEX_NONE;
-	while (PropertyIterator && OverriddenPropertyNode && OverriddenPropertyNode->Operation != EOverriddenPropertyOperation::Replace)
+	TArray<FOverriddenPropertyNode*> PropertyNodePath;
+	PropertyNodePath.Push(OverriddenPropertyNode);
+	while (PropertyIterator && (!OverriddenPropertyNode || OverriddenPropertyNode->Operation != EOverriddenPropertyOperation::Replace))
 	{
 		ArrayIndex = INDEX_NONE;
 
@@ -299,6 +301,7 @@ bool FOverriddenPropertySet::ClearOverriddenProperty(FOverriddenPropertyNode& Pa
 			{
 				CurrentOverriddenPropertyNode = OverriddenPropertyNodes.FindByHash(GetTypeHash(*CurrentPropKey), *CurrentPropKey);
 				checkf(CurrentOverriddenPropertyNode, TEXT("Expecting a node"));
+				PropertyNodePath.Push(CurrentOverriddenPropertyNode);
 			}
 		}
 
@@ -445,8 +448,27 @@ bool FOverriddenPropertySet::ClearOverriddenProperty(FOverriddenPropertyNode& Pa
 		PropertyIterator = PropertyIterator->GetNextNode();
 	}
 
+	auto CleanupClearedNodes = [this, &PropertyNodePath]()
+	{
+		// Need to cleanup up the chain of property nodes if they endup empty
+		while (FOverriddenPropertyNode* CurrentPropertyNode = !PropertyNodePath.IsEmpty() ? PropertyNodePath.Pop() : nullptr)
+		{
+			if (CurrentPropertyNode->SubPropertyNodeKeys.Num() > 1)
+			{
+				break;
+			}
+
+			RemoveOverriddenSubProperties(*CurrentPropertyNode);
+		}
+	};
+
 	if (PropertyIterator != nullptr || OverriddenPropertyNode == nullptr)
 	{
+		if (bClearedOverrides)
+		{
+			CleanupClearedNodes();
+		}
+
 		return bClearedOverrides;
 	}
 
@@ -455,8 +477,7 @@ bool FOverriddenPropertySet::ClearOverriddenProperty(FOverriddenPropertyNode& Pa
 		return false;
 	}
 
-	RemoveOverriddenSubProperties(*OverriddenPropertyNode);
-	OverriddenPropertyNode->Operation = EOverriddenPropertyOperation::None;
+	CleanupClearedNodes();
 	return true;
 }
 
@@ -469,12 +490,12 @@ void FOverriddenPropertySet::NotifyPropertyChange(FOverriddenPropertyNode* Paren
 	{
 		if (ParentPropertyNode && Notification == EPropertyNotificationType::PostEdit)
 		{
-			// Replacing this entire property
-			ParentPropertyNode->Operation = EOverriddenPropertyOperation::Replace;
-
 			// Sub-property overrides are not needed from now on, so clear them
 			RemoveOverriddenSubProperties(*ParentPropertyNode);
 
+			// Replacing this entire property
+			ParentPropertyNode->Operation = EOverriddenPropertyOperation::Replace;
+
 			// If we are overriding the root node, need to propagate the overrides to all instanced sub object
 			const FOverriddenPropertyNode* RootNode = OverriddenPropertyNodes.FindByHash(GetTypeHash(RootNodeID), RootNodeID);
 			checkf(RootNode, TEXT("Expecting to always have a "));
@@ -1001,6 +1022,7 @@ void FOverriddenPropertySet::RemoveOverriddenSubProperties(FOverriddenPropertyNo
 		RemoveOverriddenSubProperties(*RemovedPropertyNode);
 		OverriddenPropertyNodes.RemoveByHash(GetTypeHash(Pair.Value), Pair.Value);
 	}
+	PropertyNode.Operation = EOverriddenPropertyOperation::None;
 	PropertyNode.SubPropertyNodeKeys.Empty();
 }
 
