Fixed typo.

diff --git a/Engine/Source/Runtime/CoreUObject/Private/UObject/Package.cpp b/Engine/Source/Runtime/CoreUObject/Private/UObject/Package.cpp
--- a/Engine/Source/Runtime/CoreUObject/Private/UObject/Package.cpp
+++ b/Engine/Source/Runtime/CoreUObject/Private/UObject/Package.cpp
@@ -33,7 +33,7 @@ UPackage::FOnPackageSavedWithContext UPackage::PackageSavedWithContextEvent;
  *  Use Package->IsDirty() to get the updated dirty state of the package */
 UPackage::FOnPackageDirtyStateChanged UPackage::PackageDirtyStateChangedEvent;
 /** 
- * Delegate to notify subscribers when a package is marked as dirty via UObjectBaseUtilty::MarkPackageDirty 
+ * Delegate to notify subscribers when a package is marked as dirty via UObjectBaseUtility::MarkPackageDirty 
  * Note: Unlike FOnPackageDirtyStateChanged, this is always called, even when the package is already dirty
  * Use bWasDirty to check the previous dirty state of the package
  * Use Package->IsDirty() to get the updated dirty state of the package
diff --git a/Engine/Source/Runtime/CoreUObject/Public/UObject/Package.h b/Engine/Source/Runtime/CoreUObject/Public/UObject/Package.h
--- a/Engine/Source/Runtime/CoreUObject/Public/UObject/Package.h
+++ b/Engine/Source/Runtime/CoreUObject/Public/UObject/Package.h
@@ -173,7 +173,7 @@ public:
 	DECLARE_MULTICAST_DELEGATE_TwoParams(FOnPackageSaved, const FString&, UObject*);					
 	/** delegate type for package saved events ( Params: const FString& PackageFileName, UObject* Outer, FObjectPostSaveContext ObjectSaveContext ) */
 	DECLARE_MULTICAST_DELEGATE_ThreeParams(FOnPackageSavedWithContext, const FString&, UPackage*, FObjectPostSaveContext);
-	/** delegate type for when a package is marked as dirty via UObjectBaseUtilty::MarkPackageDirty ( Params: UPackage* ModifiedPackage, bool bWasDirty ) */
+	/** delegate type for when a package is marked as dirty via UObjectBaseUtility::MarkPackageDirty ( Params: UPackage* ModifiedPackage, bool bWasDirty ) */
 	DECLARE_MULTICAST_DELEGATE_TwoParams(FOnPackageMarkedDirty, class UPackage*, bool);
 	/** delegate type for when a package is about to be saved */
 	DECLARE_MULTICAST_DELEGATE_OneParam(FPreSavePackage, class UPackage*);
@@ -194,7 +194,7 @@ public:
 	*  Use Package->IsDirty() to get the updated dirty state of the package */
 	COREUOBJECT_API static FOnPackageDirtyStateChanged PackageDirtyStateChangedEvent;
 	/** 
-	* Delegate to notify subscribers when a package is marked as dirty via UObjectBaseUtilty::MarkPackageDirty 
+	* Delegate to notify subscribers when a package is marked as dirty via UObjectBaseUtility::MarkPackageDirty 
 	* Note: Unlike FOnPackageDirtyStateChanged, this is always called, even when the package is already dirty
 	* Use bWasDirty to check the previous dirty state of the package
 	* Use Package->IsDirty() to get the updated dirty state of the package
