Fixed static asserts in Cast which prevent casts of UObjectBase and UObjectBaseUtility.

diff --git a/Engine/Source/Runtime/CoreUObject/Public/Templates/Casts.h b/Engine/Source/Runtime/CoreUObject/Public/Templates/Casts.h
--- a/Engine/Source/Runtime/CoreUObject/Public/Templates/Casts.h
+++ b/Engine/Source/Runtime/CoreUObject/Public/Templates/Casts.h
@@ -123,7 +123,7 @@ FORCEINLINE To* Cast(From* Src)
 		}
 		else
 		{
-			static_assert(std::is_base_of_v<UObject, From>, "Attempting to use Cast<> on a type that is not a UObject or an Interface");
+			static_assert(std::is_base_of_v<UObjectBase, From>, "Attempting to use Cast<> on a type that is not a UObject or an Interface");
 
 			if constexpr (UE_USE_CAST_FLAGS && UE::CoreUObject::Private::TCastFlags_V<To> != CASTCLASS_None)
 			{
@@ -269,7 +269,7 @@ FORCEINLINE T* ExactCast( UObject* Src )
 		}
 		else
 		{
-			static_assert(std::is_base_of_v<UObject, From>, "Attempting to use Cast<> on a type that is not a UObject or an Interface");
+			static_assert(std::is_base_of_v<UObjectBase, From>, "Attempting to use Cast<> on a type that is not a UObject or an Interface");
 
 			if constexpr (TIsIInterface<To>::Value)
 			{
