Added struct hashing and custom struct guid support to InstanceDataObject structs

#rnx

diff --git a/Engine/Source/Runtime/CoreUObject/Private/UObject/InstanceDataObjectUtils.cpp b/Engine/Source/Runtime/CoreUObject/Private/UObject/InstanceDataObjectUtils.cpp
--- a/Engine/Source/Runtime/CoreUObject/Private/UObject/InstanceDataObjectUtils.cpp
+++ b/Engine/Source/Runtime/CoreUObject/Private/UObject/InstanceDataObjectUtils.cpp
@@ -1,18 +1,109 @@
 // Copyright Epic Games, Inc. All Rights Reserved.
 
 #include "UObject/InstanceDataObjectUtils.h"
-#include "UObject/PropertyBag.h"
-#include "UObject/UnrealType.h"
+
+#include "HAL/IConsoleManager.h"
+#include "UObject/Class.h"
 #include "UObject/EnumProperty.h"
 #include "UObject/Field.h"
-#include "HAL/IConsoleManager.h"
+#include "UObject/PropertyBag.h"
+#include "UObject/UnrealType.h"
+
+static const FName NAME_ValuesSetBySerialization(ANSITEXTVIEW("_ValuesSetBySerialization"));
+
+/** Type used for InstanceDataObject structs to provide support for hashing and custom guids. */
+class UInstanceDataObjectStruct final : public UScriptStruct
+{
+public:
+	DECLARE_CASTED_CLASS_INTRINSIC(UInstanceDataObjectStruct, UScriptStruct, CLASS_Transient, TEXT("/Script/CoreUObject"), CASTCLASS_UScriptStruct)
+
+	uint32 GetStructTypeHash(const void* Src) const final;
+	FGuid GetCustomGuid() const final { return Guid; }
+
+	FGuid Guid;
+};
+
+IMPLEMENT_CORE_INTRINSIC_CLASS(UInstanceDataObjectStruct, UScriptStruct,
+{
+});
+
+uint32 UInstanceDataObjectStruct::GetStructTypeHash(const void* Src) const
+{
+	class FBoolHash
+	{
+	public:
+		inline void Hash(bool bValue)
+		{
+			BoolValues = (BoolValues << 1) | (bValue ? 1 : 0);
+			if ((++BoolCount & 63) == 0)
+			{
+				Flush();
+			}
+		}
+
+		inline uint32 CalculateHash()
+		{
+			if (BoolCount & 63)
+			{
+				Flush();
+			}
+			return BoolHash;
+		}
+
+	private:
+		inline void Flush()
+		{
+			BoolHash = HashCombineFast(BoolHash, GetTypeHash(BoolValues));
+			BoolValues = 0;
+		}
+
+		uint32 BoolHash = 0;
+		uint32 BoolCount = 0;
+		uint64 BoolValues = 0;
+	};
+
+	FBoolHash BoolHash;
+	uint32 ValueHash = 0;
+	for (TFieldIterator<const FProperty> It(this); It; ++It)
+	{
+		if (It->GetFName() == NAME_ValuesSetBySerialization)
+		{
+			continue;
+		}
+		if (const FBoolProperty* BoolProperty = CastField<const FBoolProperty>(*It))
+		{
+			for (int32 I = 0; I < It->ArrayDim; ++I)
+			{
+				BoolHash.Hash(BoolProperty->GetPropertyValue_InContainer(Src, I));
+			}
+		}
+		else if (ensure(It->HasAllPropertyFlags(CPF_HasGetValueTypeHash)))
+		{
+			for (int32 I = 0; I < It->ArrayDim; ++I)
+			{
+				uint32 Hash = It->GetValueTypeHash(It->ContainerPtrToValuePtr<void>(Src, I));
+				ValueHash = HashCombineFast(ValueHash, Hash);
+			}
+		}
+		else
+		{
+			ValueHash = HashCombineFast(ValueHash, It->ArrayDim);
+		}
+	}
+
+	if (const uint32 Hash = BoolHash.CalculateHash())
+	{
+		ValueHash = HashCombineFast(ValueHash, Hash);
+	}
+
+	return ValueHash;
+}
 
 namespace UE
 {
 	// typedef to help make it clearer when a pathName has indices and when the indices are wildcarded away
 	using FWildcardPropertyPathName = FPropertyPathName;
-	
-	static FName Name_ValuesSetBySerialization(TEXT("_ValuesSetBySerialization"));
+
 	static const FName NAME_StructOriginalTypeMetadata(TEXT("OriginalType"));
 	static const FName NAME_PresentAsTypeMetadata(TEXT("PresentAsType"));
 	static const FName NAME_IsLooseMetadata(TEXT("IsLoose"));
@@ -197,7 +288,7 @@ namespace UE
 	}
 
 	// recursively re-instances all structs contained by this property to include loose properties
-	static void ConvertToInstanceDataObjectProperty(FProperty* Property, UObject* Outer,
+	static void ConvertToInstanceDataObjectProperty(FProperty* Property, FPropertyTypeName PropertyType, UObject* Outer,
 		const TMap<FWildcardPropertyPathName, TMap<FName, const FProperty*>>& LooseProperties, FWildcardPropertyPathName& Path)
 	{
 		if (FStructProperty* AsStructProperty = CastField<FStructProperty>(Property))
@@ -227,7 +318,12 @@ namespace UE
 					OriginalName = WriteToString<256>(OriginalNameBuilder.Build()).ToView();
 				}
 #endif
-				AsStructProperty->Struct = CreateInstanceDataObjectStructRec<UScriptStruct>(AsStructProperty->Struct, Outer, LooseProperties, Path);
+				UInstanceDataObjectStruct* Struct = CreateInstanceDataObjectStructRec<UInstanceDataObjectStruct>(AsStructProperty->Struct, Outer, LooseProperties, Path);
+				if (const FName StructGuidName = PropertyType.GetParameterName(1); !StructGuidName.IsNone())
+				{
+					FGuid::Parse(StructGuidName.ToString(), Struct->Guid);
+				}
+				AsStructProperty->Struct = Struct;
 #if WITH_EDITORONLY_DATA
 				AsStructProperty->SetMetaData(NAME_StructOriginalTypeMetadata, *OriginalName);
 				AsStructProperty->SetMetaData(NAME_PresentAsTypeMetadata, *OriginalName);
@@ -237,21 +333,21 @@ namespace UE
 		}
 		else if (const FArrayProperty* AsArrayProperty = CastField<FArrayProperty>(Property))
 		{
-			ConvertToInstanceDataObjectProperty(AsArrayProperty->Inner, Outer, LooseProperties, Path);
+			ConvertToInstanceDataObjectProperty(AsArrayProperty->Inner, PropertyType.GetParameter(0), Outer, LooseProperties, Path);
 		}
 		else if (const FSetProperty* AsSetProperty = CastField<FSetProperty>(Property))
 		{
-			ConvertToInstanceDataObjectProperty(AsSetProperty->ElementProp, Outer, LooseProperties, Path);
+			ConvertToInstanceDataObjectProperty(AsSetProperty->ElementProp, PropertyType.GetParameter(0), Outer, LooseProperties, Path);
 		}
 		else if (const FMapProperty* AsMapProperty = CastField<FMapProperty>(Property))
 		{
 			// todo: This will likely need revisiting once devin has maps working
 			Path.Push(CreateSegmentFromProperty(AsMapProperty->KeyProp));
-			ConvertToInstanceDataObjectProperty(AsMapProperty->KeyProp, Outer, LooseProperties, Path);
+			ConvertToInstanceDataObjectProperty(AsMapProperty->KeyProp, PropertyType.GetParameter(0), Outer, LooseProperties, Path);
 			Path.Pop();
 			
 			Path.Push(CreateSegmentFromProperty(AsMapProperty->ValueProp));
-			ConvertToInstanceDataObjectProperty(AsMapProperty->ValueProp, Outer, LooseProperties, Path);
+			ConvertToInstanceDataObjectProperty(AsMapProperty->ValueProp, PropertyType.GetParameter(1), Outer, LooseProperties, Path);
 			Path.Pop();
 		}
 	}
@@ -264,7 +360,7 @@ namespace UE
 #if WITH_EDITORONLY_DATA
 		FField::CopyMetaData(TemplateProperty, InstanceDataObjectProperty);
 #endif
-		ConvertToInstanceDataObjectProperty(InstanceDataObjectProperty, Outer, LooseProperties, Path);
+		ConvertToInstanceDataObjectProperty(InstanceDataObjectProperty, Path.GetSegment(Path.GetSegmentCount() - 1).Type, Outer, LooseProperties, Path);
 		return InstanceDataObjectProperty;
 	}
 
@@ -437,7 +533,7 @@ namespace UE
 
 		// add a hidden set property used to record whether this struct's properties were set serialization.
 		{
-			FSetProperty* ValuesSetBySerializationProperty = CastFieldChecked<FSetProperty>(FSetProperty::Construct(Result, Name_ValuesSetBySerialization, RF_Transient | RF_MarkAsNative));
+			FSetProperty* ValuesSetBySerializationProperty = CastFieldChecked<FSetProperty>(FSetProperty::Construct(Result, NAME_ValuesSetBySerialization, RF_Transient | RF_MarkAsNative));
 			static FName Name_PropertyName(TEXT("PropertyName"));
 			ValuesSetBySerializationProperty->ElementProp = CastFieldChecked<FProperty>(FInt64Property::Construct(ValuesSetBySerializationProperty, Name_PropertyName, RF_Transient));
 			ValuesSetBySerializationProperty->SetPropertyFlags(CPF_Transient | CPF_EditorOnly | CPF_NativeAccessSpecifierPrivate);
@@ -504,7 +600,7 @@ namespace UE
 
 	static void MarkPropertySetBySerialization(const UStruct* Struct, const void* StructData, const void* PropertyDataPtr)
 	{
-		if (const FSetProperty* ValuesSetByPropertyBagProperty = CastField<FSetProperty>(Struct->FindPropertyByName(Name_ValuesSetBySerialization)))
+		if (const FSetProperty* ValuesSetByPropertyBagProperty = CastField<FSetProperty>(Struct->FindPropertyByName(NAME_ValuesSetBySerialization)))
 		{
 			FScriptSetHelper ValuesSetByPropertyBag(ValuesSetByPropertyBagProperty, ValuesSetByPropertyBagProperty->ContainerPtrToValuePtr<void>(StructData));
 			const int64 ValueOffset = static_cast<const uint8*>(PropertyDataPtr) - static_cast<const uint8*>(StructData);
@@ -599,7 +695,7 @@ namespace UE
 		{
 			ArrayIndex = 0;
 		}
-		if (const FSetProperty* ValuesSetByPropertyBagProperty = CastField<FSetProperty>(Struct->FindPropertyByName(Name_ValuesSetBySerialization)))
+		if (const FSetProperty* ValuesSetByPropertyBagProperty = CastField<FSetProperty>(Struct->FindPropertyByName(NAME_ValuesSetBySerialization)))
 		{
 			const uint8* PropertyDataPtr;
 			if (ArrayIndex == INDEX_NONE || Property->IsA<FArrayProperty>() || Property->IsA<FMapProperty>() || Property->IsA<FSetProperty>())
