Added UEnum::IsValidEnumValueOrBitfield() for checking if an enum value allows a combination of flags.
Removed the MAX entry from a flags-based UEnum, because it doesn't make any sense and typically isn't a valid flag nor combination of flags.

diff --git a/Engine/Source/Runtime/CoreUObject/Private/UObject/Enum.cpp b/Engine/Source/Runtime/CoreUObject/Private/UObject/Enum.cpp
--- a/Engine/Source/Runtime/CoreUObject/Private/UObject/Enum.cpp
+++ b/Engine/Source/Runtime/CoreUObject/Private/UObject/Enum.cpp
@@ -251,11 +251,15 @@ int64 UEnum::GetMaxEnumValue() const
 
 	if (EnumHasAnyFlags(EnumFlags, EEnumFlags::Flags))
 	{
-		// The max value of a set of flags is the combination of all the flags
+		// The max value of a set of flags is the combination of all the power-of-two flags
 		int64 MaxFlag = 0;
 		for (const TPair<FName, int64>& NameAndValue : Names)
 		{
-			MaxFlag |= NameAndValue.Value;
+			int64 EnumeratorValue = NameAndValue.Value;
+			if (FMath::IsPowerOfTwo(EnumeratorValue))
+			{
+				MaxFlag |= EnumeratorValue;
+			}
 		}
 
 		return MaxFlag;
@@ -291,6 +295,41 @@ bool UEnum::IsValidEnumValue(int64 InValue) const
 	return false;
 }
 
+bool UEnum::IsValidEnumValueOrBitfield(int64 InValue) const
+{
+	if (!EnumHasAnyFlags(EnumFlags, EEnumFlags::Flags))
+	{
+		return IsValidEnumValue(InValue);
+	}
+
+	// Remove the known flags from the value
+	int32 NamesNum  = Names.Num();
+	int32 NameIndex = 0;
+	for (;;)
+	{
+		// If all the flags have been removed then it must be a valid value
+		if (InValue == 0)
+		{
+			return true;
+		}
+
+		// If all the flags in the enum have been tested and there are still left in the value, it wasn't a valid value
+		if (NameIndex == NamesNum)
+		{
+			return false;
+		}
+
+		// Remove flag from value, but only if it's a power of two
+		int64 EnumeratorValue = Names[NameIndex].Value;
+		if (FMath::IsPowerOfTwo(EnumeratorValue))
+		{
+			InValue &= ~EnumeratorValue;
+		}
+
+		++NameIndex;
+	}
+}
+
 bool UEnum::IsValidEnumName(FName InName) const
 {
 	int32 NamesNum = Names.Num();
@@ -853,7 +892,7 @@ bool UEnum::SetEnums(TArray<TPair<FName, int64>>& InNames, UEnum::ECppForm InCpp
 
 	if (bAddMaxKeyIfMissing)
 	{
-		if (!ContainsExistingMax())
+		if (!ContainsExistingMax() && !EnumHasAnyFlags(EnumFlags, EEnumFlags::Flags))
 		{
 			FName MaxEnumItem = *GenerateFullEnumName(*(GenerateEnumPrefix() + TEXT("_MAX")));
 			if (LookupEnumName(GetOutermost()->GetFName(), MaxEnumItem) != INDEX_NONE)
diff --git a/Engine/Source/Runtime/CoreUObject/Public/UObject/Class.h b/Engine/Source/Runtime/CoreUObject/Public/UObject/Class.h
--- a/Engine/Source/Runtime/CoreUObject/Public/UObject/Class.h
+++ b/Engine/Source/Runtime/CoreUObject/Public/UObject/Class.h
@@ -2125,6 +2125,9 @@ public:
 	/** Checks if enum has entry with given value. Includes autogenerated _MAX entry. */
 	COREUOBJECT_API bool IsValidEnumValue(int64 InValue) const;
 
+	/** Checks if enum has entry with given value, or combination of flag values. If the enum is not a bitfield, this is equivalent to IsValidEnumValue. */
+	COREUOBJECT_API bool IsValidEnumValueOrBitfield(int64 InValue) const;
+
 	/** Checks if enum has entry with given name. Includes autogenerated _MAX entry. */
 	COREUOBJECT_API bool IsValidEnumName(FName InName) const;
 
