Though of an other case I was not covering when clearing overridden nodes
Rearrange how I saved path to the cleared nodes, I'm now keeping NodeID instead of ptrs into the container elements, those ptr are no good after starting to remove elements.
Added unittest to cover the discovered problem.

diff --git a/Engine/Source/Runtime/CoreUObject/Private/UObject/OverriddenPropertySet.cpp b/Engine/Source/Runtime/CoreUObject/Private/UObject/OverriddenPropertySet.cpp
--- a/Engine/Source/Runtime/CoreUObject/Private/UObject/OverriddenPropertySet.cpp
+++ b/Engine/Source/Runtime/CoreUObject/Private/UObject/OverriddenPropertySet.cpp
@@ -284,8 +284,8 @@ bool FOverriddenPropertySet::ClearOverriddenProperty(FOverriddenPropertyNode& Pa
 	const FEditPropertyChain::TDoubleLinkedListNode* PropertyIterator = PropertyNode;
 	FOverriddenPropertyNode* OverriddenPropertyNode = &ParentPropertyNode;
 	int32 ArrayIndex = INDEX_NONE;
-	TArray<FOverriddenPropertyNode*> PropertyNodePath;
-	PropertyNodePath.Push(OverriddenPropertyNode);
+	TArray<FOverriddenPropertyNodeID> PropertyNodePath;
+	PropertyNodePath.Push(OverriddenPropertyNode->NodeID);
 	while (PropertyIterator && (!OverriddenPropertyNode || OverriddenPropertyNode->Operation != EOverriddenPropertyOperation::Replace))
 	{
 		ArrayIndex = INDEX_NONE;
@@ -301,7 +301,7 @@ bool FOverriddenPropertySet::ClearOverriddenProperty(FOverriddenPropertyNode& Pa
 			{
 				CurrentOverriddenPropertyNode = OverriddenPropertyNodes.FindByHash(GetTypeHash(*CurrentPropKey), *CurrentPropKey);
 				checkf(CurrentOverriddenPropertyNode, TEXT("Expecting a node"));
-				PropertyNodePath.Push(CurrentOverriddenPropertyNode);
+				PropertyNodePath.Push(CurrentOverriddenPropertyNode->NodeID);
 			}
 		}
 
@@ -451,14 +451,26 @@ bool FOverriddenPropertySet::ClearOverriddenProperty(FOverriddenPropertyNode& Pa
 	auto CleanupClearedNodes = [this, &PropertyNodePath]()
 	{
 		// Need to cleanup up the chain of property nodes if they endup empty
-		while (FOverriddenPropertyNode* CurrentPropertyNode = !PropertyNodePath.IsEmpty() ? PropertyNodePath.Pop() : nullptr)
+		FOverriddenPropertyNodeID ChildPropertyNodeID;
+		while (FOverriddenPropertyNode* CurrentPropertyNode = !PropertyNodePath.IsEmpty() ? OverriddenPropertyNodes.FindByHash(GetTypeHash(PropertyNodePath.Top()), PropertyNodePath.Top()) : nullptr)
 		{
+			PropertyNodePath.Pop();
 			if (CurrentPropertyNode->SubPropertyNodeKeys.Num() > 1)
 			{
+				// Now need to remove the child from this node
+				if (ChildPropertyNodeID.IsValid())
+				{
+					const FOverriddenPropertyNodeID* NodeToRemove = CurrentPropertyNode->SubPropertyNodeKeys.FindKey(ChildPropertyNodeID);
+					checkf(NodeToRemove, TEXT("Expecting a node"));
+					CurrentPropertyNode->SubPropertyNodeKeys.RemoveByHash(GetTypeHash(*NodeToRemove),*NodeToRemove);
+
+					verifyf(OverriddenPropertyNodes.RemoveByHash(GetTypeHash(ChildPropertyNodeID), ChildPropertyNodeID), TEXT("Expecting the node to be removed"));
+				}
 				break;
 			}
 
 			RemoveOverriddenSubProperties(*CurrentPropertyNode);
+			ChildPropertyNodeID = CurrentPropertyNode->NodeID;
 		}
 	};
 
@@ -1020,7 +1032,7 @@ void FOverriddenPropertySet::RemoveOverriddenSubProperties(FOverriddenPropertyNo
 		FOverriddenPropertyNode* RemovedPropertyNode = OverriddenPropertyNodes.FindByHash(GetTypeHash(Pair.Value), Pair.Value);
 		checkf(RemovedPropertyNode, TEXT("Expecting a node"));
 		RemoveOverriddenSubProperties(*RemovedPropertyNode);
-		OverriddenPropertyNodes.RemoveByHash(GetTypeHash(Pair.Value), Pair.Value);
+		verifyf(OverriddenPropertyNodes.RemoveByHash(GetTypeHash(Pair.Value), Pair.Value), TEXT("Expecting the node to be removed"));
 	}
 	PropertyNode.Operation = EOverriddenPropertyOperation::None;
 	PropertyNode.SubPropertyNodeKeys.Empty();
