Added scope types for managing SerializedPropertyPathName

#rnx

diff --git a/Engine/Source/Runtime/CoreUObject/Private/UObject/Class.cpp b/Engine/Source/Runtime/CoreUObject/Private/UObject/Class.cpp
--- a/Engine/Source/Runtime/CoreUObject/Private/UObject/Class.cpp
+++ b/Engine/Source/Runtime/CoreUObject/Private/UObject/Class.cpp
@@ -1574,11 +1574,13 @@ void UStruct::SerializeVersionedTaggedProperties(FStructuredArchive::FSlot Slot,
 					}
 				}
 
+				TOptional<UE::FSerializedPropertyPathScope> SerializedPropertyPath;
 				if (SerializeContext && SerializeContext->bTrackSerializedPropertyPath)
 				{
 					const FName Name = Property ? Property->GetFName() : Tag.Name;
 					const int32 Index = Tag.ArrayIndex > 0 || (Property && Property->ArrayDim > 1) ? Tag.ArrayIndex : INDEX_NONE;
-					SerializeContext->SerializedPropertyPath.Push({Name, Tag.GetType(), Index});
+					const UE::FPropertyPathNameSegment Segment{Name, Tag.GetType(), Index};
+					SerializedPropertyPath.Emplace(SerializeContext, Segment, UE::ESerializedPropertyPathNotify::Yes);
 				}
 
 				if (Property)
@@ -1683,12 +1685,8 @@ void UStruct::SerializeVersionedTaggedProperties(FStructuredArchive::FSlot Slot,
 					PropertyBag->LoadPropertyByTag(SerializeContext->SerializedPropertyPath, Tag, ValueSlot);
 				}
 
-				if (SerializeContext && SerializeContext->bTrackSerializedPropertyPath)
-				{
-					// broadcast that a property was serialized
-					SerializeContext->OnTaggedPropertySerialize.Broadcast(*SerializeContext);
-					SerializeContext->SerializedPropertyPath.Pop();
-				}
+				// Broadcast that a property was serialized if tracking the serialized property path.
+				SerializedPropertyPath.Reset();
 
 				int64 Loaded = UnderlyingArchive.Tell() - StartOfProperty;
 
diff --git a/Engine/Source/Runtime/CoreUObject/Private/UObject/PropertyArray.cpp b/Engine/Source/Runtime/CoreUObject/Private/UObject/PropertyArray.cpp
--- a/Engine/Source/Runtime/CoreUObject/Private/UObject/PropertyArray.cpp
+++ b/Engine/Source/Runtime/CoreUObject/Private/UObject/PropertyArray.cpp
@@ -680,13 +680,7 @@ void FArrayProperty::SerializeItem(FStructuredArchive::FSlot Slot, void* Value,
 
 				// Serialize the item at this array index
 				i = PropertyNode->ArrayIndex;
-				if (Context && Context->bTrackSerializedPropertyPath)
-				{
-					Context->SerializedPropertyPath.SetIndex(i);
-
-					// broadcast that a property will be serialized
-					Context->OnTaggedPropertySerialize.Broadcast(*Context);
-				}
+				UE::FSerializedPropertyPathIndexScope SerializedPropertyPathIndex(Context, i, UE::ESerializedPropertyPathNotify::Yes);
 				SerializeContainerItem(Array.EnterElement(), ArrayHelper.GetRawPtr(i));
 				PropertyNode = PropertyNode->PropertyListNext;
 
@@ -711,13 +705,7 @@ void FArrayProperty::SerializeItem(FStructuredArchive::FSlot Slot, void* Value,
 			NAME_UArraySerializeCount.SetNumber(i);
 			FArchive::FScopeAddDebugData P(UnderlyingArchive, NAME_UArraySerializeCount);
 #endif
-			if (Context && Context->bTrackSerializedPropertyPath)
-			{
-				Context->SerializedPropertyPath.SetIndex(i);
-
-				// broadcast that a property will be serialized
-				Context->OnTaggedPropertySerialize.Broadcast(*Context);
-			}
+			UE::FSerializedPropertyPathIndexScope SerializedPropertyPathIndex(Context, i, UE::ESerializedPropertyPathNotify::Yes);
 			SerializeContainerItem(Array.EnterElement(), ArrayHelper.GetRawPtr(i++));
 		}
 
@@ -745,11 +733,6 @@ void FArrayProperty::SerializeItem(FStructuredArchive::FSlot Slot, void* Value,
 			UnderlyingArchive.Seek(DataOffset);
 		}
 	}
-
-	if (Context && Context->bTrackSerializedPropertyPath)
-	{
-		Context->SerializedPropertyPath.SetIndex(INDEX_NONE);
-	}
 }
 
 bool FArrayProperty::NetSerializeItem(FArchive& Ar, UPackageMap* Map, void* Data, TArray<uint8> * MetaData) const
@@ -1223,13 +1206,19 @@ EConvertFromTypeResult FArrayProperty::ConvertFromType(const FPropertyTag& Tag,
 		return EConvertFromTypeResult::Converted;
 	}
 
+	FUObjectSerializeContext* Context = FUObjectThreadContext::Get().GetSerializeContext();
 	FStructuredArchive::FStream ValueStream = Slot.EnterStream();
 
-	EConvertFromTypeResult ConvertResult = Inner->ConvertFromType(InnerPropertyTag, ValueStream.EnterElement(), ScriptArrayHelper.GetRawPtr(0), DefaultsStruct, nullptr);
+	EConvertFromTypeResult ConvertResult;
+	{
+		UE::FSerializedPropertyPathIndexScope SerializedPropertyPathIndex(Context, 0, UE::ESerializedPropertyPathNotify::Yes);
+		ConvertResult = Inner->ConvertFromType(InnerPropertyTag, ValueStream.EnterElement(), ScriptArrayHelper.GetRawPtr(0), DefaultsStruct, nullptr);
+	}
 	if (ConvertResult == EConvertFromTypeResult::Converted || ConvertResult == EConvertFromTypeResult::Serialized)
 	{
 		for (int32 ElementIndex = 1; ElementIndex < ElementCount; ++ElementIndex)
 		{
+			UE::FSerializedPropertyPathIndexScope SerializedPropertyPathIndex(Context, ElementIndex, UE::ESerializedPropertyPathNotify::Yes);
 			ConvertResult = Inner->ConvertFromType(InnerPropertyTag, ValueStream.EnterElement(), ScriptArrayHelper.GetRawPtr(ElementIndex), DefaultsStruct, nullptr);
 			check(ConvertResult == EConvertFromTypeResult::Converted || ConvertResult == EConvertFromTypeResult::Serialized);
 		}
diff --git a/Engine/Source/Runtime/CoreUObject/Private/UObject/PropertyHelper.h b/Engine/Source/Runtime/CoreUObject/Private/UObject/PropertyHelper.h
--- a/Engine/Source/Runtime/CoreUObject/Private/UObject/PropertyHelper.h
+++ b/Engine/Source/Runtime/CoreUObject/Private/UObject/PropertyHelper.h
@@ -3,12 +3,14 @@
 
 #include "CoreTypes.h"
 #include "Logging/LogMacros.h"
+#include "UObject/PropertyPathName.h"
+#include "UObject/PropertyTypeName.h"
 #include "UObject/ScriptDelegateFwd.h"
+#include "UObject/UObjectThreadContext.h"
 #include "UObject/WeakObjectPtrFwd.h"
 
 class UFunction;
 class UObject;
-namespace UE { class FPropertyTypeName; }
 
 DECLARE_LOG_CATEGORY_EXTERN(LogProperty, Log, All);
 
@@ -89,4 +91,85 @@ namespace UE
  */
 FPropertyTypeName ApplyRedirectsToPropertyType(FPropertyTypeName OldTypeName, const FProperty* Property = nullptr);
 
+/** Whether to notify that a property has been serialized when terminating the property path scopes below. */
+enum class ESerializedPropertyPathNotify
+{
+	No,
+	Yes,
+};
+
+/**
+ * Pushes a segment to SerializedPropertyPath for the lifetime of this object if path tracking is active.
+ */
+class FSerializedPropertyPathScope
+{
+public:
+	[[nodiscard]] FSerializedPropertyPathScope(FUObjectSerializeContext* InContext, const FPropertyPathNameSegment& InSegment, ESerializedPropertyPathNotify InNotify)
+	{
+		if (InContext && InContext->bTrackSerializedPropertyPath)
+		{
+			Context = InContext;
+			Context->SerializedPropertyPath.Push(InSegment);
+			Notify = InNotify;
+		}
+	}
+
+	~FSerializedPropertyPathScope()
+	{
+		if (Context)
+		{
+			if (Notify == ESerializedPropertyPathNotify::Yes)
+			{
+				Context->OnTaggedPropertySerialize.Broadcast(*Context);
+			}
+			Context->SerializedPropertyPath.Pop();
+		}
+	}
+
+	FSerializedPropertyPathScope(const FSerializedPropertyPathScope&) = delete;
+	FSerializedPropertyPathScope& operator=(const FSerializedPropertyPathScope&) = delete;
+
+private:
+	FUObjectSerializeContext* Context = nullptr;
+	ESerializedPropertyPathNotify Notify = ESerializedPropertyPathNotify::No;
+};
+
+/**
+ * Sets the index of the last segment of SerializedPropertyPath for the lifetime of this object if path tracking is active.
+ *
+ * Resets the index to INDEX_NONE when destructed.
+ */
+class FSerializedPropertyPathIndexScope
+{
+public:
+	[[nodiscard]] FSerializedPropertyPathIndexScope(FUObjectSerializeContext* InContext, int32 InIndex, ESerializedPropertyPathNotify InNotify)
+	{
+		if (InContext && InContext->bTrackSerializedPropertyPath)
+		{
+			Context = InContext;
+			Context->SerializedPropertyPath.SetIndex(InIndex);
+			Notify = InNotify;
+		}
+	}
+
+	~FSerializedPropertyPathIndexScope()
+	{
+		if (Context)
+		{
+			if (Notify == ESerializedPropertyPathNotify::Yes)
+			{
+				Context->OnTaggedPropertySerialize.Broadcast(*Context);
+			}
+			Context->SerializedPropertyPath.SetIndex(INDEX_NONE);
+		}
+	}
+
+	FSerializedPropertyPathIndexScope(const FSerializedPropertyPathIndexScope&) = delete;
+	FSerializedPropertyPathIndexScope& operator=(const FSerializedPropertyPathIndexScope&) = delete;
+
+private:
+	FUObjectSerializeContext* Context = nullptr;
+	ESerializedPropertyPathNotify Notify = ESerializedPropertyPathNotify::No;
+};
+
 } // UE
diff --git a/Engine/Source/Runtime/CoreUObject/Private/UObject/PropertyMap.cpp b/Engine/Source/Runtime/CoreUObject/Private/UObject/PropertyMap.cpp
--- a/Engine/Source/Runtime/CoreUObject/Private/UObject/PropertyMap.cpp
+++ b/Engine/Source/Runtime/CoreUObject/Private/UObject/PropertyMap.cpp
@@ -288,6 +288,8 @@ void FMapProperty::SerializeItem(FStructuredArchive::FSlot Slot, void* Value, co
 	checkSlow(KeyProp);
 	checkSlow(ValueProp);
 
+	FUObjectSerializeContext* Context = FUObjectThreadContext::Get().GetSerializeContext();
+
 	FScriptMapHelper MapHelper(this, Value);
 
 	// *** Experimental *** Special serialization path for map with overridable serialization logic
@@ -335,7 +337,7 @@ void FMapProperty::SerializeItem(FStructuredArchive::FSlot Slot, void* Value, co
 						FStructuredArchive::FRecord EntryRecord = ReplacedArray.EnterElement().EnterRecord();
 						int32 Index = MapHelper.AddDefaultValue_Invalid_NeedsRehash();
 						uint8* PairPtr = MapHelper.GetPairPtr(Index);
-
+						UE::FSerializedPropertyPathIndexScope SerializedPropertyPathIndex(Context, Index, UE::ESerializedPropertyPathNotify::Yes);
 						{
 							FSerializedPropertyScope SerializedProperty(UnderlyingArchive, KeyProp, this);
 							KeyProp->SerializeItem(EntryRecord.EnterField(TEXT("Key")), PairPtr);
@@ -662,8 +664,6 @@ void FMapProperty::SerializeItem(FStructuredArchive::FSlot Slot, void* Value, co
 
 	if (UnderlyingArchive.IsLoading())
 	{
-		FUObjectSerializeContext* Context = FUObjectThreadContext::Get().GetSerializeContext();
-
 		// Delete any explicitly-removed elements
 		int32 NumKeysToRemove = 0;
 		FStructuredArchive::FArray KeysToRemoveArray = Record.EnterArray(TEXT("KeysToRemove"), NumKeysToRemove);
@@ -714,15 +714,7 @@ void FMapProperty::SerializeItem(FStructuredArchive::FSlot Slot, void* Value, co
 			{
 				FStructuredArchive::FRecord EntryRecord = EntriesArray.EnterElement().EnterRecord();
 				int32 Index = MapHelper.AddDefaultValue_Invalid_NeedsRehash();
-				if (Context && Context->bTrackSerializedPropertyPath)
-				{
-					Context->SerializedPropertyPath.SetIndex(Index);
-					
-					// broadcast that a property will be serialized
-					Context->OnTaggedPropertySerialize.Broadcast(*Context);
-				}
-
-				// TODO: Need a way to indicate that we are serializing a key and value into the property bag. Push Key/Value names?
+				UE::FSerializedPropertyPathIndexScope SerializedPropertyPathIndex(Context, Index, UE::ESerializedPropertyPathNotify::Yes);
 				{
 					FSerializedPropertyScope SerializedProperty(UnderlyingArchive, KeyProp, this);
 					KeyProp->SerializeItem(EntryRecord.EnterField(TEXT("Key")), MapHelper.GetKeyPtr(Index));
diff --git a/Engine/Source/Runtime/CoreUObject/Private/UObject/PropertySet.cpp b/Engine/Source/Runtime/CoreUObject/Private/UObject/PropertySet.cpp
--- a/Engine/Source/Runtime/CoreUObject/Private/UObject/PropertySet.cpp
+++ b/Engine/Source/Runtime/CoreUObject/Private/UObject/PropertySet.cpp
@@ -290,13 +290,7 @@ void FSetProperty::SerializeItem(FStructuredArchive::FSlot Slot, void* Value, co
 			for (; Num; --Num)
 			{
 				int32 Index = SetHelper.AddDefaultValue_Invalid_NeedsRehash();
-				if (Context && Context->bTrackSerializedPropertyPath)
-				{
-					Context->SerializedPropertyPath.SetIndex(Index);
-					
-					// broadcast that a property will be serialized
-					Context->OnTaggedPropertySerialize.Broadcast(*Context);
-				}
+				UE::FSerializedPropertyPathIndexScope SerializedPropertyPathIndex(Context, Index, UE::ESerializedPropertyPathNotify::Yes);
 				ElementProp->SerializeItem(ElementsArray.EnterElement(), SetHelper.GetElementPtrWithoutCheck(Index));
 			}
 		}
@@ -376,11 +370,6 @@ void FSetProperty::SerializeItem(FStructuredArchive::FSlot Slot, void* Value, co
 			}
 		}
 
-		if (Context && Context->bTrackSerializedPropertyPath)
-		{
-			Context->SerializedPropertyPath.SetIndex(INDEX_NONE);
-		}
-
 		SetHelper.Rehash();
 	}
 	else
