Updated the style of InstanceDataObjectUtils.h to be more consistent with CoreUObject

#rnx

diff --git a/Engine/Source/Runtime/CoreUObject/Private/UObject/InstanceDataObjectUtils.h b/Engine/Source/Runtime/CoreUObject/Private/UObject/InstanceDataObjectUtils.h
--- a/Engine/Source/Runtime/CoreUObject/Private/UObject/InstanceDataObjectUtils.h
+++ b/Engine/Source/Runtime/CoreUObject/Private/UObject/InstanceDataObjectUtils.h
@@ -1,29 +1,36 @@
 // Copyright Epic Games, Inc. All Rights Reserved.
 
 #pragma once
+
 #include "HAL/Platform.h"
 
+class FProperty;
 class UClass;
 class UObject;
 class UStruct;
-class FProperty;
+
+namespace UE { class FPropertyBag; }
+namespace UE { class FPropertyPathName; }
 
 namespace UE
 {
-	class FPropertyBag;
-	class FPropertyPathName;
 
-	// Query if InstanceDataObject are enabled for a specific object, passing in nullptr return if the system itself is enabled
-	bool IsInstanceDataObjectSupportEnabled(UObject* InObject = nullptr);
+/**
+ * Query if InstanceDataObject support is enabled for a specific object.
+ *
+ * Pass nullptr to query if the system is enabled.
+ */
+bool IsInstanceDataObjectSupportEnabled(UObject* Object = nullptr);
 
-	// generate a UClass that unions the properties of PropertyBag and TemplateStruct
-	UClass* CreateInstanceDataObjectClass(const FPropertyBag* PropertyBag, UClass* OwnerClass, UObject* Outer);
+/** Generate a UClass that contains the union of the properties of PropertyBag and OwnerClass. */
+UClass* CreateInstanceDataObjectClass(const FPropertyBag* PropertyBag, UClass* OwnerClass, UObject* Outer);
 
-	// notify that a property in an object was set when the object was deserialized
-	void MarkPropertySetBySerialization(UObject* Object, const FPropertyPathName& Path);
+/** Mark a property within the object as having been set during deserialization. */
+void MarkPropertySetBySerialization(UObject* Object, const FPropertyPathName& Path);
 	
-	// query whether a property in an object was set when the object was deserialized
-	bool WasPropertySetBySerialization(UObject* Object, const FPropertyPathName& Path);
-	// query whether a property in Struct was set when the struct was deserialized
-	bool WasPropertySetBySerialization(const UStruct* Struct, const void* StructData, const FProperty* Property, int32 ArrayIndex = 0);
-}
+/** Query whether a property within the object was set when the object was deserialized. */
+bool WasPropertySetBySerialization(UObject* Object, const FPropertyPathName& Path);
+/** Query whether a property in the struct was set when the struct was deserialized. */
+bool WasPropertySetBySerialization(const UStruct* Struct, const void* StructData, const FProperty* Property, int32 ArrayIndex = 0);
+
+} // UE
