Fixed some logic in FByteProperty::ImportText_Internal() - failure to parse a token for an enumerator value is now reported like an unknown token, rather than fall through to other parsing stages which will also fail.

diff --git a/Engine/Source/Runtime/CoreUObject/Private/UObject/PropertyByte.cpp b/Engine/Source/Runtime/CoreUObject/Private/UObject/PropertyByte.cpp
--- a/Engine/Source/Runtime/CoreUObject/Private/UObject/PropertyByte.cpp
+++ b/Engine/Source/Runtime/CoreUObject/Private/UObject/PropertyByte.cpp
@@ -431,75 +431,75 @@ const TCHAR* FByteProperty::ImportText_Internal( const TCHAR* InBuffer, void* Co
 				}
 				return Buffer;
 			}
+		}
 
-			// Enum could not be created from value. This indicates a bad value so
-			// return null so that the caller of ImportText can generate a more meaningful
-			// warning/error
-			UObject* SerializedObject = nullptr;
-			if (FUObjectSerializeContext* LoadContext = FUObjectThreadContext::Get().GetSerializeContext())
-			{
-				SerializedObject = LoadContext->SerializedObject;
-			}
-			const bool bIsNativeOrLoaded = (!Enum->HasAnyFlags(RF_WasLoaded) || Enum->HasAnyFlags(RF_LoadCompleted));
-			ErrorText->Logf(ELogVerbosity::Warning, TEXT("FBP: In asset '%s', there is an enum property of type '%s' with an invalid value of '%s' - %s"), 
-				*GetPathNameSafe(SerializedObject ? SerializedObject : FUObjectThreadContext::Get().ConstructedObject), 
-				*Enum->GetName(), 
-				*Temp,
-				bIsNativeOrLoaded ? TEXT("loaded") : TEXT("not loaded"));
-			return nullptr;
+		// Enum could not be created from value. This indicates a bad value so
+		// return null so that the caller of ImportText can generate a more meaningful
+		// warning/error
+		UObject* SerializedObject = nullptr;
+		if (FUObjectSerializeContext* LoadContext = FUObjectThreadContext::Get().GetSerializeContext())
+		{
+			SerializedObject = LoadContext->SerializedObject;
 		}
+		const bool bIsNativeOrLoaded = (!Enum->HasAnyFlags(RF_WasLoaded) || Enum->HasAnyFlags(RF_LoadCompleted));
+		ErrorText->Logf(ELogVerbosity::Warning, TEXT("FBP: In asset '%s', there is an enum property of type '%s' with an invalid value of '%s' - %s"), 
+			*GetPathNameSafe(SerializedObject ? SerializedObject : FUObjectThreadContext::Get().ConstructedObject), 
+			*Enum->GetName(), 
+			*Temp,
+			bIsNativeOrLoaded ? TEXT("loaded") : TEXT("not loaded"));
+		return nullptr;
 	}
-	
+
 	// Interpret "True" and "False" as 1 and 0. This is mostly for importing a property that was exported as a bool and is imported as a non-enum byte.
 	// Also allow for ConsoleVariable-backed enums to attempt to convert True/False to 1/0 in case a bool cvar has been converted to an enum. 
 	// Enum properties backed by an integer CVar are stored as number values, so this code will only do anything when reading an old .ini file with True/False values
 	// We log a warning so users can fix up their .ini files to use integer values that map to the enum
-	if (!Enum || (PortFlags & PPF_ConsoleVariable))
+	FString Temp;
+	const TCHAR* Buffer = FPropertyHelpers::ReadToken(InBuffer, Temp);
+	if (!Buffer)
 	{
-		FString Temp;
-		if (const TCHAR* Buffer = FPropertyHelpers::ReadToken(InBuffer, Temp))
-		{
-			const FCoreTexts& CoreTexts = FCoreTexts::Get();
+		return nullptr;
+	}
 
-			if (Temp == TEXT("True") || Temp == *(CoreTexts.True.ToString()))
-			{
-				uint64 TrueValue = 1ull;
-				if (PropertyPointerType == EPropertyPointerType::Container && HasSetter())
-				{
-					SetValue_InContainer(ContainerOrPropertyPtr, static_cast<uint8>(TrueValue));
-				}
-				else
-				{
-					SetIntPropertyValue(PointerToValuePtr(ContainerOrPropertyPtr, PropertyPointerType), TrueValue);
-				}
+	const FCoreTexts& CoreTexts = FCoreTexts::Get();
 
-				if (Enum)
-				{
-					UE_LOG(LogClass, Warning, TEXT("ConsoleVariable-Backed Enum Property of type '%s' was set from a string. Please update the cvar in your ini files."), *Enum->GetPathName(), *Enum->GetName());
-				}
+	if (Temp == TEXT("True") || Temp == *(CoreTexts.True.ToString()))
+	{
+		uint64 TrueValue = 1ull;
+		if (PropertyPointerType == EPropertyPointerType::Container && HasSetter())
+		{
+			SetValue_InContainer(ContainerOrPropertyPtr, static_cast<uint8>(TrueValue));
+		}
+		else
+		{
+			SetIntPropertyValue(PointerToValuePtr(ContainerOrPropertyPtr, PropertyPointerType), TrueValue);
+		}
 
-				return Buffer;
-			}
-			else if (Temp == TEXT("False") || Temp == *(CoreTexts.False.ToString()))
-			{
-				uint64 FalseValue = 0ull;
-				if (PropertyPointerType == EPropertyPointerType::Container && HasSetter())
-				{
-					SetValue_InContainer(ContainerOrPropertyPtr, static_cast<uint8>(FalseValue));
-				}
-				else
-				{
-					SetIntPropertyValue(PointerToValuePtr(ContainerOrPropertyPtr, PropertyPointerType), FalseValue);
-				}
+		if (Enum)
+		{
+			UE_LOG(LogClass, Warning, TEXT("ConsoleVariable-Backed Enum Property of type '%s' was set from a string. Please update the cvar in your ini files."), *Enum->GetPathName(), *Enum->GetName());
+		}
 
-				if (Enum)
-				{
-					UE_LOG(LogClass, Warning, TEXT("ConsoleVariable-Backed Enum Property of type '%s' was set from a string. Please update the cvar in your ini files."), *Enum->GetPathName(), *Enum->GetName());
-				}
+		return Buffer;
+	}
+	else if (Temp == TEXT("False") || Temp == *(CoreTexts.False.ToString()))
+	{
+		uint64 FalseValue = 0ull;
+		if (PropertyPointerType == EPropertyPointerType::Container && HasSetter())
+		{
+			SetValue_InContainer(ContainerOrPropertyPtr, static_cast<uint8>(FalseValue));
+		}
+		else
+		{
+			SetIntPropertyValue(PointerToValuePtr(ContainerOrPropertyPtr, PropertyPointerType), FalseValue);
+		}
 
-				return Buffer;
-			}
+		if (Enum)
+		{
+			UE_LOG(LogClass, Warning, TEXT("ConsoleVariable-Backed Enum Property of type '%s' was set from a string. Please update the cvar in your ini files."), *Enum->GetPathName(), *Enum->GetName());
 		}
+
+		return Buffer;
 	}
 
 	return Super::ImportText_Internal( InBuffer, ContainerOrPropertyPtr, PropertyPointerType, Parent, PortFlags, ErrorText );
