Extended the property paths for maps to disambiguate the key and value

Refactored InstanceDataObject path resolution because the two-stage recursive approach was not working for maps.

#rnx

diff --git a/Engine/Source/Runtime/CoreUObject/Private/UObject/InstanceDataObjectUtils.cpp b/Engine/Source/Runtime/CoreUObject/Private/UObject/InstanceDataObjectUtils.cpp
--- a/Engine/Source/Runtime/CoreUObject/Private/UObject/InstanceDataObjectUtils.cpp
+++ b/Engine/Source/Runtime/CoreUObject/Private/UObject/InstanceDataObjectUtils.cpp
@@ -104,11 +104,13 @@ namespace UE
 	// typedef to help make it clearer when a pathName has indices and when the indices are wildcarded away
 	using FWildcardPropertyPathName = FPropertyPathName;
 
-	static const FName NAME_StructOriginalTypeMetadata(TEXT("OriginalType"));
-	static const FName NAME_PresentAsTypeMetadata(TEXT("PresentAsType"));
-	static const FName NAME_IsLooseMetadata(TEXT("IsLoose"));
-	static const FName NAME_VerseClass("VerseClass");
-	
+	static const FName NAME_StructOriginalTypeMetadata(ANSITEXTVIEW("OriginalType"));
+	static const FName NAME_PresentAsTypeMetadata(ANSITEXTVIEW("PresentAsType"));
+	static const FName NAME_IsLooseMetadata(ANSITEXTVIEW("IsLoose"));
+	static const FName NAME_VerseClass(ANSITEXTVIEW("VerseClass"));
+	static const FName NAME_IDOMapKey(ANSITEXTVIEW("Key"));
+	static const FName NAME_IDOMapValue(ANSITEXTVIEW("Value"));
+
 	struct ResolvePropertyPathNameHelperParams
 	{
 		void* Data = nullptr;
@@ -145,126 +147,144 @@ namespace UE
 		return bIsEnabled;
 	}
 
-	static bool ResolvePropertyPathNameHelper(const UStruct* Struct, ResolvePropertyPathNameHelperParams& Params);
-	static bool ResolvePropertyPathNameHelper(ResolvePropertyPathNameHelperParams& Params)
+	static FProperty* FindPropertyByType(UStruct* Struct, FName PropertyName, FPropertyTypeName PropertyType)
 	{
-		if (Params.CurPathIndex == Params.Path.GetSegmentCount() || Params.CurPathIndex == Params.EndPathIndex)
+		for (FProperty* Property : TFieldRange<FProperty>(Struct))
 		{
-			return true;
+			if (Property->GetFName() == PropertyName && Property->CanSerializeFromTypeName(PropertyType))
+			{
+				return Property;
+			}
 		}
-		FPropertyPathNameSegment Segment = Params.Path.GetSegment(Params.CurPathIndex);
-		if (const FStructProperty* AsStructProperty = CastField<FStructProperty>(Params.ResultProperty))
+		return nullptr;
+	}
+
+	static bool ResolvePropertyPathName(
+		UObject* Object, const FPropertyPathName& Path,
+		void*& OutData, FProperty*& OutProperty,
+		void*& OutOwnerData, FStructProperty*& OutOwnerProperty)
+	{
+		OutData = Object;
+		OutProperty = nullptr;
+		OutOwnerData = nullptr;
+		OutOwnerProperty = nullptr;
+
+		if (!Object)
 		{
-			return ResolvePropertyPathNameHelper(AsStructProperty->Struct, Params);
+			return false;
 		}
-		if (const FArrayProperty* AsArrayProperty = CastField<FArrayProperty>(Params.ResultProperty))
+
+		const int32 SegmentCount = Path.GetSegmentCount();
+		for (int32 SegmentIndex = 0; SegmentIndex < SegmentCount;)
 		{
-			FScriptArrayHelper Array(AsArrayProperty, Params.Data);
-			if (!Array.IsValidIndex(Segment.Index))
+			// Assign the owner of the property in this segment.
+			OutOwnerProperty = CastField<FStructProperty>(OutProperty);
+			OutOwnerData = OutOwnerProperty ? OutData : nullptr;
+
+			// Find the struct to search within.
+			UStruct* Struct = nullptr;
+			if (SegmentIndex == 0)
 			{
-				if (Params.bAddIfNeeded)
-				{
-					Array.Resize(Segment.Index + 1);
-				}
-				else
-				{
-					return false;
-				}
+				Struct = Object->GetClass();
 			}
-			++Params.CurPathIndex;
-			Params.ResultProperty = AsArrayProperty->Inner;
-			Params.Data = Array.GetElementPtr(Segment.Index);
-			return ResolvePropertyPathNameHelper(Params);
-		}
-		if (const FSetProperty* AsSetProperty = CastField<FSetProperty>(Params.ResultProperty))
-		{
-			FScriptSetHelper Set(AsSetProperty, Params.Data);
-			if (!Set.IsValidIndex(Segment.Index))
+			else if (OutOwnerProperty)
 			{
-				if (Params.bAddIfNeeded)
-				{
-					Segment.Index = Set.AddDefaultValue_Invalid_NeedsRehash();
-				}
-				else
-				{
-					return false;
-				}
+				Struct = OutOwnerProperty->Struct;
 			}
-			++Params.CurPathIndex;
-			Params.ResultProperty = AsSetProperty->ElementProp;
-			Params.Data = Set.GetElementPtr(Segment.Index);
-			return ResolvePropertyPathNameHelper(Params);
-		}
-		if (const FMapProperty* AsMapProperty = CastField<FMapProperty>(Params.ResultProperty))
-		{
-			FScriptMapHelper Map(AsMapProperty, Params.Data);
-			if (!Map.IsValidIndex(Segment.Index))
+			else
 			{
-				if (Params.bAddIfNeeded)
-				{
-					Segment.Index = Map.AddDefaultValue_Invalid_NeedsRehash();
-				}
-				else
-				{
-					return false;
-				}
+				return false;
 			}
-			++Params.CurPathIndex;
-			Params.ResultProperty = AsMapProperty->KeyProp;
-			Params.Data = Map.GetKeyPtr(Segment.Index);
-			if (ResolvePropertyPathNameHelper(Params))
+
+			// Find the named property within the struct/class.
+			const FPropertyPathNameSegment Segment = Path.GetSegment(SegmentIndex++);
+			FProperty* Property = FindPropertyByType(Struct, Segment.Name, Segment.Type);
+			if (!Property)
 			{
-				return true;
+				return false;
 			}
-			Params.ResultProperty = AsMapProperty->ValueProp;
-			Params.Data = Map.GetValuePtr(Segment.Index);
-			return ResolvePropertyPathNameHelper(Params);
-		}
-		check(Params.CurPathIndex == Params.Path.GetSegmentCount() - 1)
-		return true;
-	}
-	
-	static bool ResolvePropertyPathNameHelper(const UStruct* Struct, ResolvePropertyPathNameHelperParams& Params)
-    {
-		const FPropertyPathNameSegment Segment = Params.Path.GetSegment(Params.CurPathIndex);
-		for (FProperty* Property : TFieldRange<FProperty>(Struct))
-		{
-			// find a property that matches the segment
-			if (Segment.Name != Property->GetFName() || !Property->CanSerializeFromTypeName(Segment.Type))
+
+			// Find the address of the property value.
+			const bool bIsIndexedProperty = Property->IsA<FArrayProperty>() || Property->IsA<FSetProperty>() || Property->IsA<FMapProperty>();
+			const int32 ArrayIndex = !bIsIndexedProperty && Segment.Index >= 0 ? Segment.Index : 0;
+			if (ArrayIndex >= Property->ArrayDim)
 			{
-				continue;
+				return false;
 			}
-			const bool bStaticArrayIndex = !Property->IsA<FArrayProperty>() && !Property->IsA<FSetProperty>() && !Property->IsA<FMapProperty>();
-			
-			if (bStaticArrayIndex && Segment.Index != INDEX_NONE)
+			OutProperty = Property;
+			OutData = Property->ContainerPtrToValuePtr<void>(OutData, ArrayIndex);
+
+			if (bIsIndexedProperty)
 			{
-				if (Segment.Index < Property->ArrayDim)
+				if (Segment.Index == INDEX_NONE)
 				{
-					++Params.CurPathIndex;
-					Params.ResultProperty = Property;
-					Params.Data = Property->ContainerPtrToValuePtr<void>(Params.Data, Segment.Index);
-					return ResolvePropertyPathNameHelper(Params);
+					// A segment may only resolve directly to an indexed property if it is the last segment.
+					if (SegmentIndex < SegmentCount)
+					{
+						return false;
+					}
 				}
-				else
+				else if (FArrayProperty* ArrayProperty = CastField<FArrayProperty>(Property))
+				{
+					FScriptArrayHelper Array(ArrayProperty, OutData);
+					if (!Array.IsValidIndex(Segment.Index))
+					{
+						return false;
+					}
+					OutProperty = ArrayProperty->Inner;
+					OutData = Array.GetElementPtr(Segment.Index);
+				}
+				else if (FSetProperty* SetProperty = CastField<FSetProperty>(Property))
 				{
-					// out of static array bounds
-					return false;
+					FScriptSetHelper Set(SetProperty, OutData);
+					if (!Set.IsValidIndex(Segment.Index))
+					{
+						return false;
+					}
+					OutProperty = SetProperty->ElementProp;
+					OutData = Set.GetElementPtr(Segment.Index);
 				}
-			}
-			else
-			{
-				if (Segment.Index == INDEX_NONE)
+				else if (FMapProperty* MapProperty = CastField<FMapProperty>(Property))
 				{
-					++Params.CurPathIndex;
+					FScriptMapHelper Map(MapProperty, OutData);
+					if (!Map.IsValidIndex(Segment.Index))
+					{
+						return false;
+					}
+					// Resolve to the pair if Key or Value are not in the path.
+					if (SegmentIndex == SegmentCount)
+					{
+						// Clear the property because there is no property for the pair, only the key and value.
+						OutProperty = nullptr;
+						OutData = Map.GetPairPtr(Segment.Index);
+						// Clear the owner because this logically resolves to the element, which is owned by the map.
+						OutOwnerProperty = nullptr;
+						OutOwnerData = nullptr;
+					}
+					else
+					{
+						const FPropertyPathNameSegment MapSegment = Path.GetSegment(SegmentIndex++);
+						if (MapSegment.Name == NAME_IDOMapKey)
+						{
+							OutProperty = MapProperty->KeyProp;
+							OutData = Map.GetKeyPtr(Segment.Index);
+						}
+						else if (MapSegment.Name == NAME_IDOMapValue)
+						{
+							OutProperty = MapProperty->ValueProp;
+							OutData = Map.GetValuePtr(Segment.Index);
+						}
+						else
+						{
+							return false;
+						}
+					}
 				}
-				Params.ResultProperty = Property;
-				Params.Data = Property->ContainerPtrToValuePtr<void>(Params.Data);
-				return ResolvePropertyPathNameHelper(Params);
 			}
 		}
-		
-		return false; // segment not found in struct
-    }
+
+		return true;
+	}
 
 	static UStruct* CreateInstanceDataObjectStructRec(const UClass* StructClass, UStruct* TemplateStruct,
 		UObject* Outer, const TMap<FWildcardPropertyPathName, TMap<FName, const FProperty*>>& LooseProperties, FWildcardPropertyPathName& Path);
@@ -341,12 +361,11 @@ namespace UE
 		}
 		else if (const FMapProperty* AsMapProperty = CastField<FMapProperty>(Property))
 		{
-			// todo: This will likely need revisiting once devin has maps working
-			Path.Push(CreateSegmentFromProperty(AsMapProperty->KeyProp));
+			Path.Push({NAME_IDOMapKey});
 			ConvertToInstanceDataObjectProperty(AsMapProperty->KeyProp, PropertyType.GetParameter(0), Outer, LooseProperties, Path);
 			Path.Pop();
 			
-			Path.Push(CreateSegmentFromProperty(AsMapProperty->ValueProp));
+			Path.Push({NAME_IDOMapValue});
 			ConvertToInstanceDataObjectProperty(AsMapProperty->ValueProp, PropertyType.GetParameter(1), Outer, LooseProperties, Path);
 			Path.Pop();
 		}
@@ -437,18 +456,18 @@ namespace UE
 #endif
 		Property->SetPropertyFlags(CPF_Edit | CPF_EditConst);
 		if (const FArrayProperty* AsArrayProperty = CastField<FArrayProperty>(Property))
-        {
+		{
 			MarkPropertyAsLoose(AsArrayProperty->Inner);
-        }
-        else if (const FSetProperty* AsSetProperty =CastField<FSetProperty>(Property))
-        {
+		}
+		else if (const FSetProperty* AsSetProperty = CastField<FSetProperty>(Property))
+{
 			MarkPropertyAsLoose(AsSetProperty->ElementProp);
-        }
-        else if (const FMapProperty* AsMapProperty =CastField<FMapProperty>(Property))
-        {
+		}
+		else if (const FMapProperty* AsMapProperty = CastField<FMapProperty>(Property))
+		{
 			MarkPropertyAsLoose(AsMapProperty->KeyProp);
 			MarkPropertyAsLoose(AsMapProperty->ValueProp);
-        }
+		}
 	}
 
 	// constructs an InstanceDataObject struct by merging the properties in 
@@ -483,7 +502,7 @@ namespace UE
 				if (MatchesBagProperty(TemplateProperty))
 				{
 					// this property was determined to be loose despite it being in the template.
-					// this likely occured due to an entire struct instance being loose and that instance becoming a template
+					// this likely occurred due to an entire struct instance being loose and that instance becoming a template
 					continue;
 				}
 				Path.Push(CreateSegmentFromProperty(TemplateProperty));
@@ -614,78 +633,47 @@ namespace UE
 	
 	void MarkPropertySetBySerialization(UObject* Object, const FPropertyPathName& Path)
 	{
-		const FPropertyPathNameSegment Segment = Path.GetSegment(Path.GetSegmentCount() - 1);
+		void* ResolvedData = nullptr;
+		void* ResolvedOwnerData = nullptr;
+		FProperty* ResolvedProperty = nullptr;
+		FStructProperty* ResolvedOwnerProperty = nullptr;
+		if (ensureMsgf(ResolvePropertyPathName(Object, Path, ResolvedData, ResolvedProperty, ResolvedOwnerData, ResolvedOwnerProperty),
+			TEXT("Failed to resolve property path name %s"), *WriteToString<256>(Path)))
+		{
+			// only mark properties set if they're in structs/classes
+			if (ResolvedOwnerProperty)
+			{
+				MarkPropertySetBySerialization(ResolvedOwnerProperty->Struct, ResolvedOwnerData, ResolvedData);
+			}
+		}
+	}
 
-		// partially resolve path to find the parent struct and data pointer
-		const UStruct* ParentStruct = Object->GetClass();
-		const void* ParentData = Object;
-		
-		ResolvePropertyPathNameHelperParams Params {
-            .Data = Object,
-            .Path = Path,
-            .CurPathIndex = 0,
-			.EndPathIndex = Path.GetSegmentCount() - 1,
-            .bAddIfNeeded = false
-        };
-		if (Params.CurPathIndex != Params.EndPathIndex)
-		{
-			if (!(ensure(ResolvePropertyPathNameHelper(Object->GetClass(), Params))))
-            {
-                return;
-            }
-            if (const FStructProperty* ParentAsStructProperty = CastField<FStructProperty>(Params.ResultProperty))
-            {
-                ParentData = Params.Data;
-                ParentStruct = ParentAsStructProperty->Struct;
-            }
-            else
-            {
-            	// only mark properties set if they're in structs/classes
-	            return;
-            }
-		}
-		// resolve the last segment of the path
-		Params.EndPathIndex += 1;
-		if (ensure(ResolvePropertyPathNameHelper(ParentStruct, Params)))
-        {
-			MarkPropertySetBySerialization(ParentStruct, ParentData, Params.Data);
-        }
+	static bool WasPropertySetBySerialization(const UStruct* Struct, const void* StructData, const void* PropertyDataPtr)
+	{
+		if (const FSetProperty* ValuesSetByPropertyBagProperty = CastField<FSetProperty>(Struct->FindPropertyByName(NAME_ValuesSetBySerialization)))
+		{
+			const FScriptSetHelper ValuesSetByPropertyBag(ValuesSetByPropertyBagProperty, ValuesSetByPropertyBagProperty->ContainerPtrToValuePtr<void>(StructData));
+			const int64 ValueOffset = static_cast<const uint8*>(PropertyDataPtr) - static_cast<const uint8*>(StructData);
+			return ValuesSetByPropertyBag.FindElementIndex(&ValueOffset) != INDEX_NONE;
+		}
+		return false;
 	}
 
 	bool WasPropertySetBySerialization(UObject* Object, const FPropertyPathName& Path)
 	{
-		const FPropertyPathNameSegment Segment = Path.GetSegment(Path.GetSegmentCount() - 1);
-
-		// partially resolve path to find the parent struct and data pointer
-		const UStruct* ParentStruct = Object->GetClass();
-		const void* ParentData = Object;
-		
-		ResolvePropertyPathNameHelperParams Params {
-			.Data = Object,
-			.Path = Path,
-			.CurPathIndex = 0,
-			.EndPathIndex = Path.GetSegmentCount() - 1,
-			.bAddIfNeeded = false
-		};
-		if (Params.CurPathIndex != Params.EndPathIndex)
-		{
-			if (!(ensure(ResolvePropertyPathNameHelper(Object->GetClass(), Params))))
+		void* ResolvedData = nullptr;
+		void* ResolvedOwnerData = nullptr;
+		FProperty* ResolvedProperty = nullptr;
+		FStructProperty* ResolvedOwnerProperty = nullptr;
+		if (ensureMsgf(ResolvePropertyPathName(Object, Path, ResolvedData, ResolvedProperty, ResolvedOwnerData, ResolvedOwnerProperty),
+			TEXT("Failed to resolve property path name %s"), *WriteToString<256>(Path)))
+		{
+			// only properties in structs/classes have been marked
+			if (ResolvedOwnerProperty)
 			{
-				return false;
-			}
-			if (const FStructProperty* ParentAsStructProperty = CastField<FStructProperty>(Params.ResultProperty))
-			{
-				ParentData = Params.Data;
-				ParentStruct = ParentAsStructProperty->Struct;
+				return WasPropertySetBySerialization(ResolvedOwnerProperty->Struct, ResolvedOwnerData, ResolvedData);
 			}
 		}
-
-		// resolve the last segment of the path
-		Params.EndPathIndex += 1;
-		if (ensure(ResolvePropertyPathNameHelper(ParentStruct, Params)))
-		{
-			return WasPropertySetBySerialization(ParentStruct, ParentData, Params.ResultProperty, Segment.Index);
-		}
 		return false;
 	}
 	
@@ -713,4 +701,4 @@ namespace UE
 		}
 		return false;
 	}
-} // UE
\ No newline at end of file
+} // UE
diff --git a/Engine/Source/Runtime/CoreUObject/Private/UObject/PropertyMap.cpp b/Engine/Source/Runtime/CoreUObject/Private/UObject/PropertyMap.cpp
--- a/Engine/Source/Runtime/CoreUObject/Private/UObject/PropertyMap.cpp
+++ b/Engine/Source/Runtime/CoreUObject/Private/UObject/PropertyMap.cpp
@@ -144,6 +144,9 @@ namespace UEMapProperty_Private
 	}
 }
 
+static const FName NAME_MapKey(ANSITEXTVIEW("Key"));
+static const FName NAME_MapValue(ANSITEXTVIEW("Value"));
+
 IMPLEMENT_FIELD(FMapProperty)
 
 FMapProperty::FMapProperty(FFieldVariant InOwner, const FName& InName, EObjectFlags InObjectFlags, EMapPropertyFlags InMapFlags)
@@ -339,10 +342,12 @@ void FMapProperty::SerializeItem(FStructuredArchive::FSlot Slot, void* Value, co
 						uint8* PairPtr = MapHelper.GetPairPtr(Index);
 						UE::FSerializedPropertyPathIndexScope SerializedPropertyPathIndex(Context, Index, UE::ESerializedPropertyPathNotify::Yes);
 						{
+							UE::FSerializedPropertyPathScope SerializedPropertyPath(Context, {NAME_MapKey}, UE::ESerializedPropertyPathNotify::No);
 							FSerializedPropertyScope SerializedProperty(UnderlyingArchive, KeyProp, this);
 							KeyProp->SerializeItem(EntryRecord.EnterField(TEXT("Key")), PairPtr);
 						}
 						{
+							UE::FSerializedPropertyPathScope SerializedPropertyPath(Context, {NAME_MapValue}, UE::ESerializedPropertyPathNotify::No);
 							FSerializedPropertyScope SerializedProperty(UnderlyingArchive, ValueProp, this);
 							ValueProp->SerializeItem(EntryRecord.EnterField(TEXT("Value")), PairPtr + MapLayout.ValueOffset);
 						}
@@ -378,6 +383,7 @@ void FMapProperty::SerializeItem(FStructuredArchive::FSlot Slot, void* Value, co
 						for (int32 i = 0; i < NumRemoved; ++i)
 						{
 							{
+								UE::FSerializedPropertyPathScope SerializedPropertyPath(Context, {NAME_MapKey}, UE::ESerializedPropertyPathNotify::No);
 								FSerializedPropertyScope SerializedProperty(UnderlyingArchive, KeyProp, this);
 								KeyProp->SerializeItem(RemovedArray.EnterElement().EnterRecord().EnterField(TEXT("Key")), TempKeyValueStorage);
 							}
@@ -410,6 +416,7 @@ void FMapProperty::SerializeItem(FStructuredArchive::FSlot Slot, void* Value, co
 
 							// Read key into temporary storage
 							{
+								UE::FSerializedPropertyPathScope SerializedPropertyPath(Context, {NAME_MapKey}, UE::ESerializedPropertyPathNotify::No);
 								FSerializedPropertyScope SerializedProperty(UnderlyingArchive, KeyProp, this);
 								KeyProp->SerializeItem(EntryRecord.EnterField(TEXT("Key")), TempKeyValueStorage);
 							}
@@ -419,6 +426,7 @@ void FMapProperty::SerializeItem(FStructuredArchive::FSlot Slot, void* Value, co
 
 							// Deserialize value into hash map-owned memory
 							{
+								UE::FSerializedPropertyPathScope SerializedPropertyPath(Context, {NAME_MapValue}, UE::ESerializedPropertyPathNotify::No);
 								FSerializedPropertyScope SerializedProperty(UnderlyingArchive, ValueProp, this);
 								ValueProp->SerializeItem(EntryRecord.EnterField(TEXT("Value")), ValuePtr);
 							}
@@ -442,6 +450,7 @@ void FMapProperty::SerializeItem(FStructuredArchive::FSlot Slot, void* Value, co
 
 							// Read key into temporary storage
 							{
+								UE::FSerializedPropertyPathScope SerializedPropertyPath(Context, {NAME_MapKey}, UE::ESerializedPropertyPathNotify::No);
 								FSerializedPropertyScope SerializedProperty(UnderlyingArchive, KeyProp, this);
 								KeyProp->SerializeItem(EntryRecord.EnterField(TEXT("Key")), TempKeyValueStorage);
 							}
@@ -450,6 +459,7 @@ void FMapProperty::SerializeItem(FStructuredArchive::FSlot Slot, void* Value, co
 
 							// Deserialize value into hash map-owned memory
 							{
+								UE::FSerializedPropertyPathScope SerializedPropertyPath(Context, {NAME_MapValue}, UE::ESerializedPropertyPathNotify::No);
 								FSerializedPropertyScope SerializedProperty(UnderlyingArchive, ValueProp, this);
 								ValueProp->SerializeItem(EntryRecord.EnterField(TEXT("Value")), ValuePtr);
 							}
@@ -600,10 +610,12 @@ void FMapProperty::SerializeItem(FStructuredArchive::FSlot Slot, void* Value, co
 						FStructuredArchive::FRecord EntryRecord = ReplacedArray.EnterElement().EnterRecord();
 						uint8* PairPtr = MapHelper.GetPairPtr(It.GetInternalIndex());
 						{
+							UE::FSerializedPropertyPathScope SerializedPropertyPath(Context, {NAME_MapKey}, UE::ESerializedPropertyPathNotify::No);
 							FSerializedPropertyScope SerializedProperty(UnderlyingArchive, KeyProp, this);
 							KeyProp->SerializeItem(EntryRecord.EnterField(TEXT("Key")), PairPtr);
 						}
 						{
+							UE::FSerializedPropertyPathScope SerializedPropertyPath(Context, {NAME_MapValue}, UE::ESerializedPropertyPathNotify::No);
 							FSerializedPropertyScope SerializedProperty(UnderlyingArchive, ValueProp, this);
 							ValueProp->SerializeItem(EntryRecord.EnterField(TEXT("Value")), PairPtr + MapLayout.ValueOffset);
 						}
@@ -619,6 +631,7 @@ void FMapProperty::SerializeItem(FStructuredArchive::FSlot Slot, void* Value, co
 					for (int32 i : RemovedIndices)
 					{
 						FStructuredArchive::FRecord EntryRecord = RemovedArray.EnterElement().EnterRecord();
+						UE::FSerializedPropertyPathScope SerializedPropertyPath(Context, {NAME_MapKey}, UE::ESerializedPropertyPathNotify::No);
 						FSerializedPropertyScope SerializedProperty(UnderlyingArchive, KeyProp, this);
 						KeyProp->SerializeItem(EntryRecord.EnterField(TEXT("Key")), DefaultsMapHelper.GetKeyPtr(i));
 					}
@@ -630,10 +643,12 @@ void FMapProperty::SerializeItem(FStructuredArchive::FSlot Slot, void* Value, co
 						FStructuredArchive::FRecord EntryRecord = ModifiedArray.EnterElement().EnterRecord();
 						uint8* PairPtr = MapHelper.GetPairPtr(i);
 						{
+							UE::FSerializedPropertyPathScope SerializedPropertyPath(Context, {NAME_MapKey}, UE::ESerializedPropertyPathNotify::No);
 							FSerializedPropertyScope SerializedProperty(UnderlyingArchive, KeyProp, this);
 							KeyProp->SerializeItem(EntryRecord.EnterField(TEXT("Key")), PairPtr);
 						}
 						{
+							UE::FSerializedPropertyPathScope SerializedPropertyPath(Context, {NAME_MapValue}, UE::ESerializedPropertyPathNotify::No);
 							FSerializedPropertyScope SerializedProperty(UnderlyingArchive, ValueProp, this);
 							ValueProp->SerializeItem(EntryRecord.EnterField(TEXT("Value")), PairPtr + MapLayout.ValueOffset);
 						}
@@ -647,10 +662,12 @@ void FMapProperty::SerializeItem(FStructuredArchive::FSlot Slot, void* Value, co
 						FStructuredArchive::FRecord EntryRecord = AddedArray.EnterElement().EnterRecord();
 						uint8* PairPtr = MapHelper.GetPairPtr(i);
 						{
+							UE::FSerializedPropertyPathScope SerializedPropertyPath(Context, {NAME_MapKey}, UE::ESerializedPropertyPathNotify::No);
 							FSerializedPropertyScope SerializedProperty(UnderlyingArchive, KeyProp, this);
 							KeyProp->SerializeItem(EntryRecord.EnterField(TEXT("Key")), PairPtr);
 						}
 						{
+							UE::FSerializedPropertyPathScope SerializedPropertyPath(Context, {NAME_MapValue}, UE::ESerializedPropertyPathNotify::No);
 							FSerializedPropertyScope SerializedProperty(UnderlyingArchive, ValueProp, this);
 							ValueProp->SerializeItem(EntryRecord.EnterField(TEXT("Value")), PairPtr + MapLayout.ValueOffset);
 						}
@@ -716,10 +733,12 @@ void FMapProperty::SerializeItem(FStructuredArchive::FSlot Slot, void* Value, co
 				int32 Index = MapHelper.AddDefaultValue_Invalid_NeedsRehash();
 				UE::FSerializedPropertyPathIndexScope SerializedPropertyPathIndex(Context, Index, UE::ESerializedPropertyPathNotify::Yes);
 				{
+					UE::FSerializedPropertyPathScope SerializedPropertyPath(Context, {NAME_MapKey}, UE::ESerializedPropertyPathNotify::No);
 					FSerializedPropertyScope SerializedProperty(UnderlyingArchive, KeyProp, this);
 					KeyProp->SerializeItem(EntryRecord.EnterField(TEXT("Key")), MapHelper.GetKeyPtr(Index));
 				}
 				{
+					UE::FSerializedPropertyPathScope SerializedPropertyPath(Context, {NAME_MapValue}, UE::ESerializedPropertyPathNotify::No);
 					FSerializedPropertyScope SerializedProperty(UnderlyingArchive, ValueProp, this);
 					ValueProp->SerializeItem(EntryRecord.EnterField(TEXT("Value")), MapHelper.GetValuePtr(Index));
 				}
@@ -750,6 +769,8 @@ void FMapProperty::SerializeItem(FStructuredArchive::FSlot Slot, void* Value, co
 				TempKeyValueStorage = (uint8*)FMemory::Malloc(MapLayout.SetLayout.Size);
 				KeyProp->InitializeValue(TempKeyValueStorage);
 
+				
+				UE::FSerializedPropertyPathScope SerializedPropertyPath(Context, {NAME_MapKey}, UE::ESerializedPropertyPathNotify::No);
 				FSerializedPropertyScope SerializedProperty(UnderlyingArchive, KeyProp, this);
 				for (; NumKeysToRemove; --NumKeysToRemove)
 				{
@@ -788,10 +809,9 @@ void FMapProperty::SerializeItem(FStructuredArchive::FSlot Slot, void* Value, co
 
 				FStructuredArchive::FRecord EntryRecord = EntriesArray.EnterElement().EnterRecord();
 
-				// TODO: Need a way to indicate that we are serializing a key and value into the property bag. Push Key/Value names?
-
 				// Read key into temporary storage
 				{
+					UE::FSerializedPropertyPathScope SerializedPropertyPath(Context, {NAME_MapKey}, UE::ESerializedPropertyPathNotify::No);
 					FSerializedPropertyScope SerializedProperty(UnderlyingArchive, KeyProp, this);
 					KeyProp->SerializeItem(EntryRecord.EnterField(TEXT("Key")), TempKeyValueStorage);
 				}
@@ -800,6 +820,7 @@ void FMapProperty::SerializeItem(FStructuredArchive::FSlot Slot, void* Value, co
 
 				// Deserialize value into hash map-owned memory
 				{
+					UE::FSerializedPropertyPathScope SerializedPropertyPath(Context, {NAME_MapValue}, UE::ESerializedPropertyPathNotify::No);
 					FSerializedPropertyScope SerializedProperty(UnderlyingArchive, ValueProp, this);
 					ValueProp->SerializeItem(EntryRecord.EnterField(TEXT("Value")), ValuePtr);
 				}
@@ -864,10 +885,12 @@ void FMapProperty::SerializeItem(FStructuredArchive::FSlot Slot, void* Value, co
 				FStructuredArchive::FRecord EntryRecord = EntriesArray.EnterElement().EnterRecord();
 
 				{
+					UE::FSerializedPropertyPathScope SerializedPropertyPath(Context, {NAME_MapKey}, UE::ESerializedPropertyPathNotify::No);
 					FSerializedPropertyScope SerializedProperty(UnderlyingArchive, KeyProp, this);
 					KeyProp->SerializeItem(EntryRecord.EnterField(TEXT("Key")), ValuePairPtr);
 				}
 				{
+					UE::FSerializedPropertyPathScope SerializedPropertyPath(Context, {NAME_MapValue}, UE::ESerializedPropertyPathNotify::No);
 					FSerializedPropertyScope SerializedProperty(UnderlyingArchive, ValueProp, this);
 					ValueProp->SerializeItem(EntryRecord.EnterField(TEXT("Value")), ValuePairPtr + MapLayout.ValueOffset);
 				}
@@ -884,10 +907,12 @@ void FMapProperty::SerializeItem(FStructuredArchive::FSlot Slot, void* Value, co
 
 				uint8* ValuePairPtr = MapHelper.GetPairPtr(Iterator);
 				{
+					UE::FSerializedPropertyPathScope SerializedPropertyPath(Context, {NAME_MapKey}, UE::ESerializedPropertyPathNotify::No);
 					FSerializedPropertyScope SerializedProperty(UnderlyingArchive, KeyProp, this);
 					KeyProp->SerializeItem(EntryRecord.EnterField(TEXT("Key")), ValuePairPtr);
 				}
 				{
+					UE::FSerializedPropertyPathScope SerializedPropertyPath(Context, {NAME_MapValue}, UE::ESerializedPropertyPathNotify::No);
 					FSerializedPropertyScope SerializedProperty(UnderlyingArchive, ValueProp, this);
 					ValueProp->SerializeItem(EntryRecord.EnterField(TEXT("Value")), ValuePairPtr + MapLayout.ValueOffset);
 				}
@@ -1458,8 +1483,10 @@ EConvertFromTypeResult FMapProperty::ConvertFromType(const FPropertyTag& Tag, FS
 		return EConvertFromTypeResult::CannotConvert;
 	}
 
-	const auto SerializeOrConvert = [](bool bCanSerialize, FProperty* Inner, const FPropertyTag& InnerTag, FStructuredArchive::FSlot InnerSlot, uint8* InnerData, UStruct* InnerDefaultsStruct) -> bool
+	const auto SerializeOrConvert = [Context = FUObjectThreadContext::Get().GetSerializeContext()](bool bCanSerialize, FProperty* Inner, const FPropertyTag& InnerTag, FName InnerName, FStructuredArchive::FSlot InnerSlot, uint8* InnerData, UStruct* InnerDefaultsStruct) -> bool
 	{
+		UE::FSerializedPropertyPathScope SerializedPropertyPath(Context, {InnerName}, UE::ESerializedPropertyPathNotify::No);
+
 		if (!bCanSerialize)
 		{
 			switch (Inner->ConvertFromType(InnerTag, InnerSlot, InnerData, InnerDefaultsStruct, nullptr))
@@ -1524,7 +1551,7 @@ EConvertFromTypeResult FMapProperty::ConvertFromType(const FPropertyTag& Tag, FS
 		TempKeyValueStorage = (uint8*)FMemory::Malloc(MapLayout.SetLayout.Size);
 		KeyProp->InitializeValue(TempKeyValueStorage);
 
-		if (SerializeOrConvert(bCanSerializeKey, KeyProp, KeyPropertyTag, KeysToRemoveArray.EnterElement(), TempKeyValueStorage, DefaultsStruct))
+		if (SerializeOrConvert(bCanSerializeKey, KeyProp, KeyPropertyTag, NAME_MapKey, KeysToRemoveArray.EnterElement(), TempKeyValueStorage, DefaultsStruct))
 		{
 			// If the key is in the map, remove it
 			int32 Found = MapHelper.FindMapIndexWithKey(TempKeyValueStorage);
@@ -1536,7 +1563,7 @@ EConvertFromTypeResult FMapProperty::ConvertFromType(const FPropertyTag& Tag, FS
 			// things are going fine, remove the rest of the keys:
 			for (int32 I = 1; I < NumKeysToRemove; ++I)
 			{
-				verify(SerializeOrConvert(bCanSerializeKey, KeyProp, KeyPropertyTag, KeysToRemoveArray.EnterElement(), TempKeyValueStorage, DefaultsStruct));
+				verify(SerializeOrConvert(bCanSerializeKey, KeyProp, KeyPropertyTag, NAME_MapKey, KeysToRemoveArray.EnterElement(), TempKeyValueStorage, DefaultsStruct));
 				Found = MapHelper.FindMapIndexWithKey(TempKeyValueStorage);
 				if (Found != INDEX_NONE)
 				{
@@ -1565,7 +1592,7 @@ EConvertFromTypeResult FMapProperty::ConvertFromType(const FPropertyTag& Tag, FS
 
 			FStructuredArchive::FRecord FirstPropertyRecord = EntriesArray.EnterElement().EnterRecord();
 
-			if (SerializeOrConvert(bCanSerializeKey, KeyProp, KeyPropertyTag, FirstPropertyRecord.EnterField(TEXT("Key")), TempKeyValueStorage, DefaultsStruct))
+			if (SerializeOrConvert(bCanSerializeKey, KeyProp, KeyPropertyTag, NAME_MapKey, FirstPropertyRecord.EnterField(TEXT("Key")), TempKeyValueStorage, DefaultsStruct))
 			{
 				// Add a new default value if the key doesn't currently exist in the map
 				bool bKeyAlreadyPresent = true;
@@ -1581,14 +1608,14 @@ EConvertFromTypeResult FMapProperty::ConvertFromType(const FPropertyTag& Tag, FS
 				KeyProp->CopyCompleteValue_InContainer(NextPairPtr, TempKeyValueStorage);
 
 				// Deserialize value
-				if (SerializeOrConvert(bCanSerializeValue, ValueProp, ValuePropertyTag, FirstPropertyRecord.EnterField(TEXT("Value")), NextPairPtr, DefaultsStruct))
+				if (SerializeOrConvert(bCanSerializeValue, ValueProp, ValuePropertyTag, NAME_MapValue, FirstPropertyRecord.EnterField(TEXT("Value")), NextPairPtr, DefaultsStruct))
 				{
 					// first entry went fine, convert the rest:
 					for (int32 I = 1; I < Num; ++I)
 					{
 						FStructuredArchive::FRecord PropertyRecord = EntriesArray.EnterElement().EnterRecord();
 
-						verify(SerializeOrConvert(bCanSerializeKey, KeyProp, KeyPropertyTag, PropertyRecord.EnterField(TEXT("Key")), TempKeyValueStorage, DefaultsStruct));
+						verify(SerializeOrConvert(bCanSerializeKey, KeyProp, KeyPropertyTag, NAME_MapKey, PropertyRecord.EnterField(TEXT("Key")), TempKeyValueStorage, DefaultsStruct));
 						NextPairIndex = MapHelper.FindMapIndexWithKey(TempKeyValueStorage);
 						if (NextPairIndex == INDEX_NONE)
 						{
@@ -1598,7 +1625,7 @@ EConvertFromTypeResult FMapProperty::ConvertFromType(const FPropertyTag& Tag, FS
 						NextPairPtr = MapHelper.GetPairPtrWithoutCheck(NextPairIndex);
 						// This copy is unnecessary when the key was already in the map:
 						KeyProp->CopyCompleteValue_InContainer(NextPairPtr, TempKeyValueStorage);
-						verify(SerializeOrConvert(bCanSerializeValue, ValueProp, ValuePropertyTag, PropertyRecord.EnterField(TEXT("Value")), NextPairPtr, DefaultsStruct));
+						verify(SerializeOrConvert(bCanSerializeValue, ValueProp, ValuePropertyTag, NAME_MapValue, PropertyRecord.EnterField(TEXT("Value")), NextPairPtr, DefaultsStruct));
 					}
 				}
 				else
