Fixed documentation for FPropertyTypeName

#rnx

diff --git a/Engine/Source/Runtime/CoreUObject/Public/UObject/PropertyTypeName.h b/Engine/Source/Runtime/CoreUObject/Public/UObject/PropertyTypeName.h
--- a/Engine/Source/Runtime/CoreUObject/Public/UObject/PropertyTypeName.h
+++ b/Engine/Source/Runtime/CoreUObject/Public/UObject/PropertyTypeName.h
@@ -33,9 +33,9 @@ struct FPropertyTypeNameNode
  *
  * Examples:
  * - int32 -> IntProperty
- * - TArray(int32) -> ArrayProperty(IntProperty)
- * - TArray(FStructType) -> ArrayProperty(StructProperty(StructType(/Script/Module),e21f566f-7153-433a-959d-bfb3abed17e2))
- * - TMap(FKeyStruct, EByteEnum) -> MapProperty(StructProperty(KeyStruct(/Script/Module)),EnumProperty(ByteEnum(/Script/Module),ByteProperty))
+ * - TArray<int32> -> ArrayProperty(IntProperty)
+ * - TArray<FStructType> -> ArrayProperty(StructProperty(StructType(/Script/Module),e21f566f-7153-433a-959d-bfb3abed17e2))
+ * - TMap<FKeyStruct, EByteEnum> -> MapProperty(StructProperty(KeyStruct(/Script/Module)),EnumProperty(ByteEnum(/Script/Module),ByteProperty))
  *
  * Incomplete property types created from sources with incomplete type information. Consumers must support this.
  * - ArrayProperty(StructProperty)
