Changes the trace marker used for denoting GameThread async flushes to now clarify if a flush of all async loads is being performed or if the game thread is flushing only a subset of all loads. The name "Flush All Async Loads GT" was chosen as it closely mirrors the old name ""Flush Async Loading GT" but this specific name has a colour in insights that pops more in traces. This new marker makes it easier to detect and fix bad behaviour since, except for a few special cases, we should never wait for all loads and instead should be specifying a subset.

diff --git a/Engine/Source/Runtime/CoreUObject/Private/Serialization/AsyncLoading2.cpp b/Engine/Source/Runtime/CoreUObject/Private/Serialization/AsyncLoading2.cpp
--- a/Engine/Source/Runtime/CoreUObject/Private/Serialization/AsyncLoading2.cpp
+++ b/Engine/Source/Runtime/CoreUObject/Private/Serialization/AsyncLoading2.cpp
@@ -9437,7 +9437,9 @@ void FAsyncLoadingThread2::FlushLoading(TConstArrayView<int32> RequestIDs)
 			return;
 		}
 
-		SCOPE_CYCLE_COUNTER(STAT_FAsyncPackage_FlushAsyncLoadingGameThread);
+		const bool bIsFullFlush = RequestIDs.IsEmpty();
+		CONDITIONAL_SCOPE_CYCLE_COUNTER(STAT_FAsyncPackage_FlushAsyncLoadingGameThread, !bIsFullFlush);
+		CONDITIONAL_SCOPE_CYCLE_COUNTER(STAT_FAsyncPackage_FlushAllAsyncLoadingGameThread, bIsFullFlush);
 
 		// if the sync count is 0, then this flush is not triggered from a sync load, broadcast the delegate in that case
 		FUObjectThreadContext& ThreadContext = FUObjectThreadContext::Get();
diff --git a/Engine/Source/Runtime/CoreUObject/Public/Serialization/AsyncPackageLoader.h b/Engine/Source/Runtime/CoreUObject/Public/Serialization/AsyncPackageLoader.h
--- a/Engine/Source/Runtime/CoreUObject/Public/Serialization/AsyncPackageLoader.h
+++ b/Engine/Source/Runtime/CoreUObject/Public/Serialization/AsyncPackageLoader.h
@@ -29,6 +29,7 @@ DECLARE_STATS_GROUP(TEXT("Async Load Game Thread"), STATGROUP_AsyncLoadGameThrea
 DECLARE_CYCLE_STAT(TEXT("PostLoadObjects GT"), STAT_FAsyncPackage_PostLoadObjectsGameThread, STATGROUP_AsyncLoadGameThread);
 DECLARE_CYCLE_STAT(TEXT("TickAsyncLoading GT"), STAT_FAsyncPackage_TickAsyncLoadingGameThread, STATGROUP_AsyncLoadGameThread);
 DECLARE_CYCLE_STAT(TEXT("Flush Async Loading GT"), STAT_FAsyncPackage_FlushAsyncLoadingGameThread, STATGROUP_AsyncLoadGameThread);
+DECLARE_CYCLE_STAT(TEXT("Flush All Async Loads GT"), STAT_FAsyncPackage_FlushAllAsyncLoadingGameThread, STATGROUP_AsyncLoadGameThread);
 DECLARE_CYCLE_STAT_WITH_FLAGS(TEXT("PostLoadInstances GT"), STAT_FAsyncPackage_PostLoadInstancesGameThread, STATGROUP_AsyncLoadGameThread, EStatFlags::Verbose);
 DECLARE_CYCLE_STAT_WITH_FLAGS(TEXT("CreateClusters GT"), STAT_FAsyncPackage_CreateClustersGameThread, STATGROUP_AsyncLoadGameThread, EStatFlags::Verbose);
 
