Fixed return-by-value function call being returned by reference from a lambda.

diff --git a/Engine/Source/Runtime/CoreUObject/Private/UObject/Enum.cpp b/Engine/Source/Runtime/CoreUObject/Private/UObject/Enum.cpp
--- a/Engine/Source/Runtime/CoreUObject/Private/UObject/Enum.cpp
+++ b/Engine/Source/Runtime/CoreUObject/Private/UObject/Enum.cpp
@@ -641,7 +641,7 @@ FText UEnum::GetValueOrBitfieldAsDisplayNameText(int64 InValue) const
 		CreateStringFromValueOrBitfield(
 			this,
 			InValue,
-			[this](int64 Value) -> const FString&
+			[this](int64 Value)
 			{
 				return GetDisplayNameTextByValue(Value).ToString();
 			}
