Added optional property support to InstanceDataObjects

#rnx

diff --git a/Engine/Source/Runtime/CoreUObject/Private/UObject/InstanceDataObjectUtils.cpp b/Engine/Source/Runtime/CoreUObject/Private/UObject/InstanceDataObjectUtils.cpp
--- a/Engine/Source/Runtime/CoreUObject/Private/UObject/InstanceDataObjectUtils.cpp
+++ b/Engine/Source/Runtime/CoreUObject/Private/UObject/InstanceDataObjectUtils.cpp
@@ -7,6 +7,7 @@
 #include "UObject/EnumProperty.h"
 #include "UObject/Field.h"
 #include "UObject/PropertyBag.h"
+#include "UObject/PropertyOptional.h"
 #include "UObject/UnrealType.h"
 
 static const FName NAME_ValuesSetBySerialization(ANSITEXTVIEW("_ValuesSetBySerialization"));
@@ -281,6 +282,21 @@ namespace UE
 					}
 				}
 			}
+			else if (FOptionalProperty* OptionalProperty = CastField<FOptionalProperty>(Property))
+			{
+				FOptionalPropertyLayout Optional(OptionalProperty->GetValueProperty());
+				// Resolve to the value if it is set, otherwise resolve to the unset optional.
+				if (void* Data = Optional.GetValuePointerForReplaceIfSet(OutData))
+				{
+					OutProperty = Optional.GetValueProperty();
+					OutData = Data;
+				}
+				// A segment may only resolve directly to an unset optional if it is the last segment.
+				else if (SegmentIndex < SegmentCount)
+				{
+					return false;
+				}
+			}
 		}
 
 		return true;
@@ -369,6 +385,10 @@ namespace UE
 			ConvertToInstanceDataObjectProperty(AsMapProperty->ValueProp, PropertyType.GetParameter(1), Outer, LooseProperties, Path);
 			Path.Pop();
 		}
+		else if (const FOptionalProperty* AsOptionalProperty = CastField<FOptionalProperty>(Property))
+		{
+			ConvertToInstanceDataObjectProperty(AsOptionalProperty->GetValueProperty(), PropertyType.GetParameter(0), Outer, LooseProperties, Path);
+		}
 	}
 	
 	// copy template property then convert it into an InstanceDataObject property by adding loose properties
@@ -423,6 +443,10 @@ namespace UE
 			AddWildcardedProperties(OutProperties, ParentPath, AsMapProperty->KeyProp);
 			AddWildcardedProperties(OutProperties, ParentPath, AsMapProperty->ValueProp);
 		}
+		else if (const FOptionalProperty* AsOptionalProperty = CastField<FOptionalProperty>(Property))
+		{
+			AddWildcardedProperties(OutProperties, ParentPath, AsOptionalProperty->GetValueProperty());
+		}
 		ParentPath.Pop();
 	}
 
@@ -468,6 +492,10 @@ namespace UE
 			MarkPropertyAsLoose(AsMapProperty->KeyProp);
 			MarkPropertyAsLoose(AsMapProperty->ValueProp);
 		}
+		else if (const FOptionalProperty* AsOptionalProperty = CastField<FOptionalProperty>(Property))
+		{
+			MarkPropertyAsLoose(AsOptionalProperty->GetValueProperty());
+		}
 	}
 
 	// constructs an InstanceDataObject struct by merging the properties in 
