Fix the fact that a moved entity would loose all it overrides.
Added unittest to prevent this problem from coming back

diff --git a/Engine/Source/Runtime/CoreUObject/Private/UObject/OverriddenPropertySet.cpp b/Engine/Source/Runtime/CoreUObject/Private/UObject/OverriddenPropertySet.cpp
--- a/Engine/Source/Runtime/CoreUObject/Private/UObject/OverriddenPropertySet.cpp
+++ b/Engine/Source/Runtime/CoreUObject/Private/UObject/OverriddenPropertySet.cpp
@@ -626,8 +626,6 @@ void FOverriddenPropertySet::NotifyPropertyChange(FOverriddenPropertyNode* Paren
 									RemovedSubObjectNode.Operation = EOverriddenPropertyOperation::Remove;
 								}
 							}
-
-							OverridableManager.ClearInstancedSubObjectOverrides(*Owner, *RemovedSubObject);
 						}
 					};
 
@@ -868,16 +866,6 @@ void FOverriddenPropertySet::NotifyPropertyChange(FOverriddenPropertyNode* Paren
 						RemovedKeyNode.Operation = EOverriddenPropertyOperation::Remove;
 					}
 				}
-
-				// @todo support instanced object as a key in maps
-				//if (UObject* RemovedKeySubObject = KeyInstancedObjectProperty ? KeyInstancedObjectProperty->GetObjectPropertyValue(PreEditMapHelper.GetKeyPtr(InternalPreEditMapIndex)) : nullptr)
-				//{
-				//	OverridableManager.ClearInstancedSubObjectOverrides(*Owner, *RemovedKeySubObject);
-				//}
-				if (UObject* RemovedValueSubObject = ValueInstancedObjectProperty ? ValueInstancedObjectProperty->GetObjectPropertyValue(PreEditMapHelper.GetValuePtr(InternalPreEditMapIndex)) : nullptr)
-				{
-					OverridableManager.ClearInstancedSubObjectOverrides(*Owner, *RemovedValueSubObject);
-				}
 			};
 
 			// Only maps flagged overridable logic can be handled here
