[Feature] Added ArMergeOverrides flag that keeps (and adds to) existing overrides rather than clearing them at the beginning of serialization
[Bugfix] Overrides are now read/set even for properties that no longer exist now. This way we can use them for loose properties on IDOs

diff --git a/Engine/Source/Runtime/Core/Public/Serialization/Archive.h b/Engine/Source/Runtime/Core/Public/Serialization/Archive.h
--- a/Engine/Source/Runtime/Core/Public/Serialization/Archive.h
+++ b/Engine/Source/Runtime/Core/Public/Serialization/Archive.h
@@ -860,6 +860,9 @@ public:
 	/** Set TRUE to use the custom property list attribute for serialization. */
 	uint8 ArUseCustomPropertyList : 1;
 
+	/** Set to false if OverriddenProperties should be cleared at the beginning of SerializeVersionedTaggedProperties. */
+	uint8 ArMergeOverrides : 1;
+
 	/** Whether we are currently serializing defaults. > 0 means yes, <= 0 means no. */
 	int32 ArSerializingDefaults;
 
@@ -2077,6 +2080,7 @@ public:
 	using FArchiveState::ArSerializingDefaults;
 	using FArchiveState::ArPortFlags;
 	using FArchiveState::ArMaxSerializeSize;
+	using FArchiveState::ArMergeOverrides;
 
 public:
 	using FArchiveState::SetIsLoading;
diff --git a/Engine/Source/Runtime/CoreUObject/Private/UObject/Class.cpp b/Engine/Source/Runtime/CoreUObject/Private/UObject/Class.cpp
--- a/Engine/Source/Runtime/CoreUObject/Private/UObject/Class.cpp
+++ b/Engine/Source/Runtime/CoreUObject/Private/UObject/Class.cpp
@@ -1423,7 +1423,16 @@ void UStruct::SerializeVersionedTaggedProperties(FStructuredArchive::FSlot Slot,
 				// Overridden values are saved independently in transaction, so do no need to restore them here.
 				if (!UnderlyingArchive.IsTransacting())
 				{
-					ControlContext.OverriddenProperties = &FOverridableManager::Get().SetOverriddenProperties(*(UObject*)Data, Operation);
+					if (UnderlyingArchive.ArMergeOverrides)
+					{
+						ControlContext.OverriddenProperties = FOverridableManager::Get().GetOverriddenProperties(*(UObject*)Data);
+					}
+					
+					if (!ControlContext.OverriddenProperties)
+					{
+						ControlContext.OverriddenProperties = &FOverridableManager::Get().SetOverriddenProperties(*(UObject*)Data, Operation);
+					}
+					
 					ControlContext.OverriddenProperties->bNeedsSubobjectTemplateInstantiation = true;
 				}
 			}
@@ -1683,6 +1692,23 @@ void UStruct::SerializeVersionedTaggedProperties(FStructuredArchive::FSlot Slot,
 					// TODO: Might we find defaults in a property bag for Defaults?
 					FStructuredArchive::FSlot ValueSlot = PropertyRecord.EnterField(TEXT("Value"));
 					PropertyBag->LoadPropertyByTag(SerializeContext->SerializedPropertyPath, Tag, ValueSlot);
+					
+					// we still want to set overrides for loose properties because they will be used by the IDO.
+					if (!UnderlyingArchive.IsTransacting())
+					{
+						if (FOverriddenPropertySet* OverriddenProperties = FOverridableSerializationLogic::GetOverriddenProperties())
+						{
+							// No need to restore none operations
+							if (Tag.OverrideOperation != EOverriddenPropertyOperation::None)
+							{
+								// SetOverriddenPropertyOperation only needs to know the property name so we'll just construct a dummy property
+								FProperty* PropertyToOverride = CastField<FProperty>(FField::TryConstruct(Tag.Type, {}, Tag.Name, RF_NoFlags));
+								check(PropertyToOverride);
+								OverriddenProperties->SetOverriddenPropertyOperation(Tag.OverrideOperation, UnderlyingArchive.GetSerializedPropertyChain(), PropertyToOverride);
+								delete PropertyToOverride;
+							}
+						}
+					}
 				}
 
 				// Broadcast that a property was serialized if tracking the serialized property path.
