Added bitfield text import and export support to FByteProperty with an enum.  This will also fix a bug where a getter wouldn't be used.

diff --git a/Engine/Source/Runtime/CoreUObject/Private/UObject/EnumProperty.cpp b/Engine/Source/Runtime/CoreUObject/Private/UObject/EnumProperty.cpp
--- a/Engine/Source/Runtime/CoreUObject/Private/UObject/EnumProperty.cpp
+++ b/Engine/Source/Runtime/CoreUObject/Private/UObject/EnumProperty.cpp
@@ -42,6 +42,128 @@ namespace UE::CoreUObject::Private
 
 		ConvertIntValueToEnumProperty(OldValue, EnumProp, UnderlyingProp, Enum, Obj);
 	}
+
+	const TCHAR* ImportEnumFromBuffer(UEnum* Enum, const FProperty* PropertyToSet, const FNumericProperty* UnderlyingProp, const TCHAR* PropertyClassName, const TCHAR* Buffer, void* ContainerOrPropertyPtr, EPropertyPointerType PropertyPointerType, FOutputDevice* ErrorText)
+	{
+		bool    bIsEnumOfFlags = Enum->HasAnyEnumFlags(EEnumFlags::Flags);
+		int64   EnumValue      = 0;
+		FString Temp;
+		for (;;)
+		{
+			Buffer = FPropertyHelpers::ReadToken(Buffer, Temp, true);
+			if (Buffer)
+			{
+				int32 EnumIndex = Enum->GetIndexByName(*Temp, EGetByNameFlags::CheckAuthoredName);
+				if (!bIsEnumOfFlags && EnumIndex == INDEX_NONE && (Temp.IsNumeric() && !Algo::Find(Temp, TEXT('.'))))
+				{
+					int64 LexedEnumValue = INDEX_NONE;
+					LexFromString(LexedEnumValue, *Temp);
+					EnumIndex = Enum->GetIndexByValue(LexedEnumValue);
+				}
+				if (EnumIndex == INDEX_NONE)
+				{
+					Buffer = nullptr;
+					break;
+				}
+
+				EnumValue |= Enum->GetValueByIndex(EnumIndex);
+
+				if (!bIsEnumOfFlags)
+				{
+					break;
+				}
+
+				Temp.Reset();
+				SkipWhitespace(Buffer);
+				if (*Buffer != TEXT('|'))
+				{
+					break;
+				}
+				++Buffer;
+				SkipWhitespace(Buffer);
+			}
+		}
+
+		if (!Buffer)
+		{
+			// Enum could not be created from value. This indicates a bad value so
+			// return null so that the caller of ImportText can generate a more meaningful
+			// warning/error
+			UObject* SerializedObject = nullptr;
+			if (FUObjectSerializeContext* LoadContext = FUObjectThreadContext::Get().GetSerializeContext())
+			{
+				SerializedObject = LoadContext->SerializedObject;
+			}
+			const bool bIsNativeOrLoaded = (!Enum->HasAnyFlags(RF_WasLoaded) || Enum->HasAnyFlags(RF_LoadCompleted));
+			ErrorText->Logf(ELogVerbosity::Warning, TEXT("%s: In asset '%s', there is an enum property of type '%s' with an invalid value of '%s' - %s"), 
+				PropertyClassName,
+				*GetPathNameSafe(SerializedObject ? SerializedObject : FUObjectThreadContext::Get().ConstructedObject), 
+				*Enum->GetName(), 
+				*Temp,
+				bIsNativeOrLoaded ? TEXT("loaded") : TEXT("not loaded"));
+			return nullptr;
+		}
+
+		if (PropertyPointerType == EPropertyPointerType::Container && PropertyToSet->HasSetter())
+		{
+			PropertyToSet->SetValue_InContainer(ContainerOrPropertyPtr, &EnumValue);
+		}
+		else
+		{
+			UnderlyingProp->SetIntPropertyValue(PropertyToSet->PointerToValuePtr(ContainerOrPropertyPtr, PropertyPointerType), EnumValue);
+		}
+		return Buffer;
+	}
+
+	void ExportEnumToBuffer(const UEnum* Enum, const FProperty* Prop, const FNumericProperty* NumericProp, FString& ValueStr, const void* PropertyValueOrContainer, EPropertyPointerType PropertyPointerType, const void* DefaultValue, UObject* Parent, int32 PortFlags, UObject* ExportRootScope)
+	{
+		// This should be big enough to import all numeric enums
+		alignas(int64) uint8 GetterResult[sizeof(int64)];
+		void* ValuePtr = nullptr;
+
+		if (PropertyPointerType == EPropertyPointerType::Container && Prop->HasGetter())
+		{
+			// Put the bytes returned by the getter into GetterResult (which is big and aligned enough to hold any numeric value)
+			// - these bytes will be re-read by a call to NumericProp below.
+			ValuePtr = &GetterResult;
+			Prop->GetValue_InContainer(PropertyValueOrContainer, ValuePtr);
+		}
+		else
+		{
+			// Otherwise read directly from the property
+			ValuePtr = Prop->PointerToValuePtr(PropertyValueOrContainer, PropertyPointerType);
+		}
+
+		if (PortFlags & PPF_ConsoleVariable)
+		{
+			NumericProp->ExportText_Internal(ValueStr, ValuePtr, EPropertyPointerType::Direct, DefaultValue, Parent, PortFlags, ExportRootScope);
+			return;
+		}
+
+		int64 Value = NumericProp->GetSignedIntPropertyValue(ValuePtr);
+
+		// if the value is the max value (the autogenerated *_MAX value), export as "INVALID", unless we're exporting text for copy/paste (for copy/paste,
+		// the property text value must actually match an entry in the enum's names array)
+		if (!Enum->HasAnyEnumFlags(EEnumFlags::Flags) && (!Enum->IsValidEnumValue(Value) || (!(PortFlags & PPF_Copy) && Value == Enum->GetMaxEnumValue())))
+		{
+			ValueStr += TEXT("(INVALID)");
+			return;
+		}
+
+		// We do not want to export the enum text for non-display uses, localization text is very dynamic and would cause issues on import
+		if (PortFlags & PPF_PropertyWindow)
+		{
+			ValueStr += Enum->GetValueOrBitfieldAsDisplayNameText(Value).ToString();
+		}
+		else if (PortFlags & PPF_ExternalEditor)
+		{
+			ValueStr += Enum->GetValueOrBitfieldAsAuthoredNameString(Value);
+		}
+		else
+		{
+			ValueStr += Enum->GetValueOrBitfieldAsString(Value);
+		}
+	}
 }
 
 IMPLEMENT_FIELD(FEnumProperty)
@@ -267,51 +389,7 @@ void FEnumProperty::ExportText_Internal(FString& ValueStr, const void* PropertyV
 	FNumericProperty* LocalUnderlyingProp = UnderlyingProp;
 	check(LocalUnderlyingProp);
 
-	alignas(int64) uint8 GetterResult[sizeof(int64)];
-	void* ValuePtr = nullptr;
-
-	if (PropertyPointerType == EPropertyPointerType::Container && HasGetter())
-	{
-		// Put the bytes returned by the getter into GetterResult (which is big and aligned enough to hold any numeric value)
-		// - these bytes will be re-read by a call to LocalUnderlyingProp below.
-		ValuePtr = &GetterResult;
-		GetValue_InContainer(PropertyValueOrContainer, ValuePtr);
-	}
-	else
-	{
-		// Otherwise read directly from the property
-		ValuePtr = PointerToValuePtr(PropertyValueOrContainer, PropertyPointerType);
-	}
-
-	if (PortFlags & PPF_ConsoleVariable)
-	{
-		LocalUnderlyingProp->ExportText_Internal(ValueStr, ValuePtr, EPropertyPointerType::Direct, DefaultValue, Parent, PortFlags, ExportRootScope);
-		return;
-	}
-
-	int64 Value = LocalUnderlyingProp->GetSignedIntPropertyValue(ValuePtr);
-
-	// if the value is the max value (the autogenerated *_MAX value), export as "INVALID", unless we're exporting text for copy/paste (for copy/paste,
-	// the property text value must actually match an entry in the enum's names array)
-	if (!Enum->HasAnyEnumFlags(EEnumFlags::Flags) && (!Enum->IsValidEnumValue(Value) || (!(PortFlags & PPF_Copy) && Value == Enum->GetMaxEnumValue())))
-	{
-		ValueStr += TEXT("(INVALID)");
-		return;
-	}
-
-	// We do not want to export the enum text for non-display uses, localization text is very dynamic and would cause issues on import
-	if (PortFlags & PPF_PropertyWindow)
-	{
-		ValueStr += Enum->GetValueOrBitfieldAsDisplayNameText(Value).ToString();
-	}
-	else if (PortFlags & PPF_ExternalEditor)
-	{
-		ValueStr += Enum->GetValueOrBitfieldAsAuthoredNameString(Value);
-	}
-	else
-	{
-		ValueStr += Enum->GetValueOrBitfieldAsString(Value);
-	}
+	UE::CoreUObject::Private::ExportEnumToBuffer(Enum, this, UnderlyingProp, ValueStr, PropertyValueOrContainer, PropertyPointerType, DefaultValue, Parent, PortFlags, ExportRootScope);
 }
 
 const TCHAR* FEnumProperty::ImportText_Internal(const TCHAR* Buffer, void* ContainerOrPropertyPtr, EPropertyPointerType PropertyPointerType, UObject* Parent, int32 PortFlags, FOutputDevice* ErrorText) const
@@ -321,73 +399,7 @@ const TCHAR* FEnumProperty::ImportText_Internal(const TCHAR* Buffer, void* Conta
 	
 	if (!(PortFlags & PPF_ConsoleVariable))
 	{
-		bool    bIsEnumOfFlags = Enum->HasAnyEnumFlags(EEnumFlags::Flags);
-		int64   EnumValue      = 0;
-		FString Temp;
-		for (;;)
-		{
-			Buffer = FPropertyHelpers::ReadToken(Buffer, Temp, true);
-			if (Buffer)
-			{
-				int32 EnumIndex = Enum->GetIndexByName(*Temp, EGetByNameFlags::CheckAuthoredName);
-				if (!bIsEnumOfFlags && EnumIndex == INDEX_NONE && (Temp.IsNumeric() && !Algo::Find(Temp, TEXT('.'))))
-				{
-					int64 LexedEnumValue = INDEX_NONE;
-					LexFromString(LexedEnumValue, *Temp);
-					EnumIndex = Enum->GetIndexByValue(LexedEnumValue);
-				}
-				if (EnumIndex == INDEX_NONE)
-				{
-					Buffer = nullptr;
-					break;
-				}
-
-				EnumValue |= Enum->GetValueByIndex(EnumIndex);
-
-				if (!bIsEnumOfFlags)
-				{
-					break;
-				}
-
-				Temp.Reset();
-				SkipWhitespace(Buffer);
-				if (*Buffer != TEXT('|'))
-				{
-					break;
-				}
-				++Buffer;
-				SkipWhitespace(Buffer);
-			}
-		}
-
-		if (!Buffer)
-		{
-			// Enum could not be created from value. This indicates a bad value so
-			// return null so that the caller of ImportText can generate a more meaningful
-			// warning/error
-			UObject* SerializedObject = nullptr;
-			if (FUObjectSerializeContext* LoadContext = FUObjectThreadContext::Get().GetSerializeContext())
-			{
-				SerializedObject = LoadContext->SerializedObject;
-			}
-			const bool bIsNativeOrLoaded = (!Enum->HasAnyFlags(RF_WasLoaded) || Enum->HasAnyFlags(RF_LoadCompleted));
-			ErrorText->Logf(ELogVerbosity::Warning, TEXT("FEP: In asset '%s', there is an enum property of type '%s' with an invalid value of '%s' - %s"), 
-				*GetPathNameSafe(SerializedObject ? SerializedObject : FUObjectThreadContext::Get().ConstructedObject), 
-				*Enum->GetName(), 
-				*Temp,
-				bIsNativeOrLoaded ? TEXT("loaded") : TEXT("not loaded"));
-			return nullptr;
-		}
-
-		if (PropertyPointerType == EPropertyPointerType::Container && HasSetter())
-		{
-			SetValue_InContainer(ContainerOrPropertyPtr, &EnumValue);
-		}
-		else
-		{
-			UnderlyingProp->SetIntPropertyValue(PointerToValuePtr(ContainerOrPropertyPtr, PropertyPointerType), EnumValue);
-		}
-		return Buffer;
+		return UE::CoreUObject::Private::ImportEnumFromBuffer(Enum, this, UnderlyingProp, TEXT("FEnumProperty"), Buffer, ContainerOrPropertyPtr, PropertyPointerType, ErrorText);
 	}
 
 	// UnderlyingProp has a 0 offset so we need to make sure we convert the container pointer to the actual value pointer
diff --git a/Engine/Source/Runtime/CoreUObject/Private/UObject/PropertyByte.cpp b/Engine/Source/Runtime/CoreUObject/Private/UObject/PropertyByte.cpp
--- a/Engine/Source/Runtime/CoreUObject/Private/UObject/PropertyByte.cpp
+++ b/Engine/Source/Runtime/CoreUObject/Private/UObject/PropertyByte.cpp
@@ -364,143 +364,66 @@ void FByteProperty::ExportText_Internal( FString& ValueStr, const void* Property
 		return;
 	}
 
-	uint8 GetterResult;
-	void* ValuePtr = nullptr;
-
-	if (PropertyPointerType == EPropertyPointerType::Container && HasGetter())
-	{
-		// Put the byte returned by the getter into GetterResult
-		// - this will be re-read by a call to GetSignedIntPropertyValue below.
-		ValuePtr = &GetterResult;
-		GetValue_InContainer(PropertyValueOrContainer, (uint8*)ValuePtr);
-	}
-	else
-	{
-		// Otherwise read directly from the property
-		ValuePtr = PointerToValuePtr(PropertyValueOrContainer, PropertyPointerType);
-	}
-
-	int64 Value = GetSignedIntPropertyValue(ValuePtr);
-
-	// if the value is the max value (the autogenerated *_MAX value), export as "INVALID", unless we're exporting text for copy/paste (for copy/paste,
-	// the property text value must actually match an entry in the enum's names array)
-	if (!Enum->IsValidEnumValue(Value) || (!(PortFlags & PPF_Copy) && Value == Enum->GetMaxEnumValue()))
-	{
-		ValueStr += TEXT("(INVALID)");
-		return;
-	}
-
-	// We do not want to export the enum text for non-display uses, localization text is very dynamic and would cause issues on import
-	if (PortFlags & PPF_PropertyWindow)
-	{
-		ValueStr += Enum->GetDisplayNameTextByValue(Value).ToString();
-	}
-	else if (PortFlags & PPF_ExternalEditor)
-	{
-		ValueStr += Enum->GetAuthoredNameStringByValue(Value);
-	}
-	else
-	{
-		ValueStr += Enum->GetNameStringByValue(Value);
-	}
+	UE::CoreUObject::Private::ExportEnumToBuffer(Enum, this, this, ValueStr, PropertyValueOrContainer, PropertyPointerType, DefaultValue, Parent, PortFlags, ExportRootScope);
 }
 const TCHAR* FByteProperty::ImportText_Internal( const TCHAR* InBuffer, void* ContainerOrPropertyPtr, EPropertyPointerType PropertyPointerType, UObject* Parent, int32 PortFlags, FOutputDevice* ErrorText ) const
 {
 	if( Enum && (PortFlags & PPF_ConsoleVariable) == 0 )
 	{
-		FString Temp;
-		if (const TCHAR* Buffer = FPropertyHelpers::ReadToken(InBuffer, Temp, true))
-		{
-			int32 EnumIndex = Enum->GetIndexByName(*Temp, EGetByNameFlags::CheckAuthoredName);
-			if (EnumIndex == INDEX_NONE && (Temp.IsNumeric() && !Algo::Find(Temp, TEXT('.'))))
-			{
-				int64 EnumValue = INDEX_NONE;
-				LexFromString(EnumValue, *Temp);
-				EnumIndex = Enum->GetIndexByValue(EnumValue);
-			}
-			if (EnumIndex != INDEX_NONE)
-			{
-				uint8 EnumValue = IntCastChecked<uint8>(Enum->GetValueByIndex(EnumIndex));
-				if (PropertyPointerType == EPropertyPointerType::Container && HasSetter())
-				{
-					SetValue_InContainer(ContainerOrPropertyPtr, EnumValue);
-				}
-				else
-				{
-					*(uint8*)PointerToValuePtr(ContainerOrPropertyPtr, PropertyPointerType) = EnumValue;
-				}
-				return Buffer;
-			}
-		}
-
-		// Enum could not be created from value. This indicates a bad value so
-		// return null so that the caller of ImportText can generate a more meaningful
-		// warning/error
-		UObject* SerializedObject = nullptr;
-		if (FUObjectSerializeContext* LoadContext = FUObjectThreadContext::Get().GetSerializeContext())
-		{
-			SerializedObject = LoadContext->SerializedObject;
-		}
-		const bool bIsNativeOrLoaded = (!Enum->HasAnyFlags(RF_WasLoaded) || Enum->HasAnyFlags(RF_LoadCompleted));
-		ErrorText->Logf(ELogVerbosity::Warning, TEXT("FBP: In asset '%s', there is an enum property of type '%s' with an invalid value of '%s' - %s"), 
-			*GetPathNameSafe(SerializedObject ? SerializedObject : FUObjectThreadContext::Get().ConstructedObject), 
-			*Enum->GetName(), 
-			*Temp,
-			bIsNativeOrLoaded ? TEXT("loaded") : TEXT("not loaded"));
-		return nullptr;
+		return UE::CoreUObject::Private::ImportEnumFromBuffer(Enum, this, this, TEXT("FByteProperty"), InBuffer, ContainerOrPropertyPtr, PropertyPointerType, ErrorText);
 	}
-
+	
 	// Interpret "True" and "False" as 1 and 0. This is mostly for importing a property that was exported as a bool and is imported as a non-enum byte.
 	// Also allow for ConsoleVariable-backed enums to attempt to convert True/False to 1/0 in case a bool cvar has been converted to an enum. 
 	// Enum properties backed by an integer CVar are stored as number values, so this code will only do anything when reading an old .ini file with True/False values
 	// We log a warning so users can fix up their .ini files to use integer values that map to the enum
-	FString Temp;
+		FString Temp;
 	const TCHAR* Buffer = FPropertyHelpers::ReadToken(InBuffer, Temp);
 	if (!Buffer)
-	{
+		{
 		return nullptr;
 	}
 
-	const FCoreTexts& CoreTexts = FCoreTexts::Get();
+			const FCoreTexts& CoreTexts = FCoreTexts::Get();
 
-	if (Temp == TEXT("True") || Temp == *(CoreTexts.True.ToString()))
-	{
-		uint64 TrueValue = 1ull;
-		if (PropertyPointerType == EPropertyPointerType::Container && HasSetter())
-		{
-			SetValue_InContainer(ContainerOrPropertyPtr, static_cast<uint8>(TrueValue));
-		}
-		else
-		{
-			SetIntPropertyValue(PointerToValuePtr(ContainerOrPropertyPtr, PropertyPointerType), TrueValue);
-		}
+			if (Temp == TEXT("True") || Temp == *(CoreTexts.True.ToString()))
+			{
+				uint64 TrueValue = 1ull;
+				if (PropertyPointerType == EPropertyPointerType::Container && HasSetter())
+				{
+					SetValue_InContainer(ContainerOrPropertyPtr, static_cast<uint8>(TrueValue));
+				}
+				else
+				{
+					SetIntPropertyValue(PointerToValuePtr(ContainerOrPropertyPtr, PropertyPointerType), TrueValue);
+				}
 
-		if (Enum)
-		{
-			UE_LOG(LogClass, Warning, TEXT("ConsoleVariable-Backed Enum Property of type '%s' was set from a string. Please update the cvar in your ini files."), *Enum->GetPathName(), *Enum->GetName());
-		}
+				if (Enum)
+				{
+					UE_LOG(LogClass, Warning, TEXT("ConsoleVariable-Backed Enum Property of type '%s' was set from a string. Please update the cvar in your ini files."), *Enum->GetPathName(), *Enum->GetName());
+				}
 
-		return Buffer;
-	}
-	else if (Temp == TEXT("False") || Temp == *(CoreTexts.False.ToString()))
-	{
-		uint64 FalseValue = 0ull;
-		if (PropertyPointerType == EPropertyPointerType::Container && HasSetter())
-		{
-			SetValue_InContainer(ContainerOrPropertyPtr, static_cast<uint8>(FalseValue));
-		}
-		else
-		{
-			SetIntPropertyValue(PointerToValuePtr(ContainerOrPropertyPtr, PropertyPointerType), FalseValue);
-		}
+				return Buffer;
+			}
+			else if (Temp == TEXT("False") || Temp == *(CoreTexts.False.ToString()))
+			{
+				uint64 FalseValue = 0ull;
+				if (PropertyPointerType == EPropertyPointerType::Container && HasSetter())
+				{
+					SetValue_InContainer(ContainerOrPropertyPtr, static_cast<uint8>(FalseValue));
+				}
+				else
+				{
+					SetIntPropertyValue(PointerToValuePtr(ContainerOrPropertyPtr, PropertyPointerType), FalseValue);
+				}
 
-		if (Enum)
-		{
-			UE_LOG(LogClass, Warning, TEXT("ConsoleVariable-Backed Enum Property of type '%s' was set from a string. Please update the cvar in your ini files."), *Enum->GetPathName(), *Enum->GetName());
-		}
+				if (Enum)
+				{
+					UE_LOG(LogClass, Warning, TEXT("ConsoleVariable-Backed Enum Property of type '%s' was set from a string. Please update the cvar in your ini files."), *Enum->GetPathName(), *Enum->GetName());
+				}
 
-		return Buffer;
-	}
+				return Buffer;
+			}
 
 	return Super::ImportText_Internal( InBuffer, ContainerOrPropertyPtr, PropertyPointerType, Parent, PortFlags, ErrorText );
 }
diff --git a/Engine/Source/Runtime/CoreUObject/Public/UObject/UnrealType.h b/Engine/Source/Runtime/CoreUObject/Public/UObject/UnrealType.h
--- a/Engine/Source/Runtime/CoreUObject/Public/UObject/UnrealType.h
+++ b/Engine/Source/Runtime/CoreUObject/Public/UObject/UnrealType.h
@@ -73,6 +73,8 @@ class UPropertyWrapper;
 enum ELifetimeCondition : int;
 struct CGetTypeHashable;
 struct FUObjectSerializeContext;
+class FProperty;
+class FNumericProperty;
 template <typename FuncType> class TFunctionRef;
 namespace UE { class FPropertyTypeName; }
 namespace UE { class FPropertyTypeNameBuilder; }
@@ -156,6 +158,15 @@ enum class EPropertyPointerType
 	Container = 1, /** Property access through its owner container */
 };
 
+namespace UE::CoreUObject::Private
+{
+	// Defined in EnumProperty.cpp and used by both FEnumProperty and FByteProperty.
+	// They don't have an API macro because they're not intended to be called outside of CoreUObject.
+	const TCHAR* ImportEnumFromBuffer(UEnum* Enum, const FProperty* PropertyToSet, const FNumericProperty* UnderlyingProp, const TCHAR* PropertyClassName, const TCHAR* Buffer, void* ContainerOrPropertyPtr, EPropertyPointerType PropertyPointerType, FOutputDevice* ErrorText);
+
+	void ExportEnumToBuffer(const UEnum* Enum, const FProperty* Prop, const FNumericProperty* NumericProp, FString& ValueStr, const void* PropertyValueOrContainer, EPropertyPointerType PropertyPointerType, const void* DefaultValue, UObject* Parent, int32 PortFlags, UObject* ExportRootScope);
+}
+
 //
 // An UnrealScript variable.
 //
@@ -647,6 +658,9 @@ private:
 
 protected:
 
+	friend const TCHAR* UE::CoreUObject::Private::ImportEnumFromBuffer(UEnum* Enum, const FProperty* PropertyToSet, const FNumericProperty* UnderlyingProp, const TCHAR* PropertyClassName, const TCHAR* Buffer, void* ContainerOrPropertyPtr, EPropertyPointerType PropertyPointerType, FOutputDevice* ErrorText);
+	friend void UE::CoreUObject::Private::ExportEnumToBuffer(const UEnum* Enum, const FProperty* Prop, const FNumericProperty* NumericProp, FString& ValueStr, const void* PropertyValueOrContainer, EPropertyPointerType PropertyPointerType, const void* DefaultValue, UObject* Parent, int32 PortFlags, UObject* ExportRootScope);
+
 	FORCEINLINE void* PointerToValuePtr(void const* ContainerOrPropertyPtr, EPropertyPointerType PropertyPointerType, int32 ArrayIndex = 0) const
 	{
 		if (PropertyPointerType == EPropertyPointerType::Container)
