Restore MAX to flag-based UEnums, as it breaks some current editor assumptions.

diff --git a/Engine/Source/Runtime/CoreUObject/Private/UObject/Enum.cpp b/Engine/Source/Runtime/CoreUObject/Private/UObject/Enum.cpp
--- a/Engine/Source/Runtime/CoreUObject/Private/UObject/Enum.cpp
+++ b/Engine/Source/Runtime/CoreUObject/Private/UObject/Enum.cpp
@@ -302,8 +302,9 @@ bool UEnum::IsValidEnumValueOrBitfield(int64 InValue) const
 		return IsValidEnumValue(InValue);
 	}
 
-	// Remove the known flags from the value
-	int32 NamesNum  = Names.Num();
+	// Remove the known flags from the value - ignore the last enumerator because it's MAX.
+	// TODO: Update this when MAX is no longer generated for flag-based enums.
+	int32 NamesNum  = Names.Num() - 1;
 	int32 NameIndex = 0;
 	for (;;)
 	{
@@ -892,7 +893,9 @@ bool UEnum::SetEnums(TArray<TPair<FName, int64>>& InNames, UEnum::ECppForm InCpp
 
 	if (bAddMaxKeyIfMissing)
 	{
-		if (!ContainsExistingMax() && !EnumHasAnyFlags(EnumFlags, EEnumFlags::Flags))
+		// TODO: A MAX value for enum flags doesn't make sense, but keep generating it until we're comfortable that
+		// it's no longer needed.
+		if (!ContainsExistingMax() /*&& !EnumHasAnyFlags(EnumFlags, EEnumFlags::Flags)*/)
 		{
 			FName MaxEnumItem = *GenerateFullEnumName(*(GenerateEnumPrefix() + TEXT("_MAX")));
 			if (LookupEnumName(GetOutermost()->GetFName(), MaxEnumItem) != INDEX_NONE)
