Fixed property conversion to and from enums with bitfields.

diff --git a/Engine/Source/Runtime/CoreUObject/Private/UObject/EnumProperty.cpp b/Engine/Source/Runtime/CoreUObject/Private/UObject/EnumProperty.cpp
--- a/Engine/Source/Runtime/CoreUObject/Private/UObject/EnumProperty.cpp
+++ b/Engine/Source/Runtime/CoreUObject/Private/UObject/EnumProperty.cpp
@@ -16,8 +16,10 @@ namespace UE::CoreUObject::Private
 		using LargeIntType = std::conditional_t<TIsSigned<OldIntType>::Value, int64, uint64>;
 
 		LargeIntType NewValue = OldValue;
-		if (!UnderlyingProp->CanHoldValue(NewValue) || !Enum->IsValidEnumValue(NewValue))
+		if (!UnderlyingProp->CanHoldValue(NewValue) || !Enum->IsValidEnumValueOrBitfield(NewValue))
 		{
+			NewValue = Enum->HasAnyEnumFlags(EEnumFlags::Flags) ? 0 : Enum->GetMaxEnumValue();
+
 			UE_LOG(
 				LogClass,
 				Warning,
@@ -25,10 +27,8 @@ namespace UE::CoreUObject::Private
 				*LexToString(OldValue),
 				*Enum->GetName(),
 				*EnumProp->GetName(),
-				*Enum->GetNameByValue(Enum->GetMaxEnumValue()).ToString()
+				*Enum->GetNameByValue(NewValue).ToString()
 			);
-
-			NewValue = Enum->GetMaxEnumValue();
 		}
 
 		UnderlyingProp->SetIntPropertyValue(Obj, NewValue);
diff --git a/Engine/Source/Runtime/CoreUObject/Private/UObject/PropertyNumeric.cpp b/Engine/Source/Runtime/CoreUObject/Private/UObject/PropertyNumeric.cpp
--- a/Engine/Source/Runtime/CoreUObject/Private/UObject/PropertyNumeric.cpp
+++ b/Engine/Source/Runtime/CoreUObject/Private/UObject/PropertyNumeric.cpp
@@ -43,9 +43,11 @@ int64 FNumericProperty::ReadEnumAsInt64(FStructuredArchive::FSlot Slot, UStruct*
 	Slot.GetUnderlyingArchive().Preload(Enum);
 
 	// This handles redirects internally
-	int64 Result = Enum->GetValueByName(EnumName);
-	if (!Enum->IsValidEnumValue(Result))
+	int64 Result = Enum->GetValueOrBitfieldFromString(EnumName.ToString());
+	if (!Enum->IsValidEnumValueOrBitfield(Result))
 	{
+		int64 ReplacementValue = Enum->HasAnyEnumFlags(EEnumFlags::Flags) ? 0 : Enum->GetMaxEnumValue();
+
 		UE_LOG(
 			LogClass,
 			Warning,
@@ -53,10 +55,10 @@ int64 FNumericProperty::ReadEnumAsInt64(FStructuredArchive::FSlot Slot, UStruct*
 			*EnumName.ToString(),
 			*Enum->GetName(),
 			*Tag.Name.ToString(),
-			*Enum->GetNameByValue(Enum->GetMaxEnumValue()).ToString()
+			*Enum->GetNameByValue(ReplacementValue).ToString()
 			);
 
-		return Enum->GetMaxEnumValue();
+		return ReplacementValue;
 	}
 
 	return Result;
