Handle impersonated structs when finding property type redirects

#rnx

diff --git a/Engine/Source/Runtime/CoreUObject/Private/UObject/PropertyHelper.cpp b/Engine/Source/Runtime/CoreUObject/Private/UObject/PropertyHelper.cpp
--- a/Engine/Source/Runtime/CoreUObject/Private/UObject/PropertyHelper.cpp
+++ b/Engine/Source/Runtime/CoreUObject/Private/UObject/PropertyHelper.cpp
@@ -292,6 +292,33 @@ const TCHAR* DelegatePropertyTools::ImportDelegateFromText( FScriptDelegate& Del
 namespace UE
 {
 
+#if WITH_EDITORONLY_DATA
+FPropertyTypeName FindOriginalType(const FStructProperty* Struct)
+{
+	FUObjectSerializeContext* Context = FUObjectThreadContext::Get().GetSerializeContext();
+	if (Context && Context->bImpersonateProperties)
+	{
+		const FString* OriginalType = Struct->FindMetaData(NAME_OriginalType);
+		if (!OriginalType)
+		{
+			//@note: To support metadata defined on array of struct in UPROPERTY for testing purposes
+			if (FField* OwnerField = Struct->Owner.ToField())
+			{
+				OriginalType = OwnerField->FindMetaData(NAME_OriginalType);
+			}
+		}
+		if (OriginalType)
+		{
+			if (UE::FPropertyTypeNameBuilder Type; Type.TryParse(*OriginalType))
+			{
+				return Type.Build();
+			}
+		}
+	}
+	return {};
+}
+#endif // WITH_EDITORONLY_DATA
+
 static FCoreRedirectObjectName BuildCoreRedirectObjectName(FPropertyTypeName TypeName)
 {
 	FName OuterName;
@@ -371,7 +398,7 @@ static bool FindRedirectForProperty(FPropertyTypeName OldType, FPropertyTypeName
 
 	// If the type in the tag matches the field then skip looking for redirects.
 	// This is necessary to handle redirects where the old name continues to be used.
-	FCoreRedirectObjectName OldNameRedirect = BuildCoreRedirectObjectName(OldType);
+	const FCoreRedirectObjectName OldNameRedirect = BuildCoreRedirectObjectName(OldType);
 	if (Field && Field->GetFName() == OldNameRedirect.ObjectName)
 	{
 		return false;
@@ -387,12 +414,32 @@ static bool FindRedirectForProperty(FPropertyTypeName OldType, FPropertyTypeName
 
 	// If a partial match does not match the field then repeat the lookup without allowing partial matches.
 	// This is necessary to handle redirects of the original type of a property that changed type.
-	if (Field && Field->GetFName() != NewNameRedirect.ObjectName)
+	if (Field)
 	{
-		NewNameRedirect = FCoreRedirects::GetRedirectedName(FTraits::GetFlags(), OldNameRedirect, ECoreRedirectMatchFlags::None);
-		if (OldNameRedirect == NewNameRedirect)
+		FName FieldName = Field->GetFName();
+
+	#if WITH_EDITORONLY_DATA
+		// Compare against the original type of an impersonated struct.
+		if constexpr (std::is_same_v<PropertyType, FStructProperty>)
 		{
-			return false;
+			if (const FPropertyTypeName OriginalType = FindOriginalType(Property); !OriginalType.IsEmpty())
+			{
+				FieldName = OriginalType.GetName();
+				if (FieldName == OldNameRedirect.ObjectName)
+				{
+					return false;
+				}
+			}
+		}
+	#endif // WITH_EDITORONLY_DATA
+
+		if (FieldName != NewNameRedirect.ObjectName)
+		{
+			NewNameRedirect = FCoreRedirects::GetRedirectedName(FTraits::GetFlags(), OldNameRedirect, ECoreRedirectMatchFlags::None);
+			if (OldNameRedirect == NewNameRedirect)
+			{
+				return false;
+			}
 		}
 	}
 
diff --git a/Engine/Source/Runtime/CoreUObject/Private/UObject/PropertyHelper.h b/Engine/Source/Runtime/CoreUObject/Private/UObject/PropertyHelper.h
--- a/Engine/Source/Runtime/CoreUObject/Private/UObject/PropertyHelper.h
+++ b/Engine/Source/Runtime/CoreUObject/Private/UObject/PropertyHelper.h
@@ -3,12 +3,14 @@
 
 #include "CoreTypes.h"
 #include "Logging/LogMacros.h"
+#include "UObject/NameTypes.h"
 #include "UObject/PropertyPathName.h"
 #include "UObject/PropertyTypeName.h"
 #include "UObject/ScriptDelegateFwd.h"
 #include "UObject/UObjectThreadContext.h"
 #include "UObject/WeakObjectPtrFwd.h"
 
+class FStructProperty;
 class UFunction;
 class UObject;
 
@@ -79,6 +81,11 @@ namespace DelegatePropertyTools
 namespace UE
 {
 
+#if WITH_EDITORONLY_DATA
+inline static const FName NAME_OriginalType(ANSITEXTVIEW("OriginalType"));
+FPropertyTypeName FindOriginalType(const FStructProperty* Struct);
+#endif // WITH_EDITORONLY_DATA
+
 /**
  * Applies core redirects to type names and paths within the property type.
  *
diff --git a/Engine/Source/Runtime/CoreUObject/Private/UObject/PropertyStruct.cpp b/Engine/Source/Runtime/CoreUObject/Private/UObject/PropertyStruct.cpp
--- a/Engine/Source/Runtime/CoreUObject/Private/UObject/PropertyStruct.cpp
+++ b/Engine/Source/Runtime/CoreUObject/Private/UObject/PropertyStruct.cpp
@@ -456,35 +456,6 @@ void FStructProperty::AppendSchemaHash(FBlake3& Builder, bool bSkipEditorOnly) c
 }
 #endif
 
-#if WITH_EDITORONLY_DATA
-static const FName NAME_StructOriginalType(ANSITEXTVIEW("OriginalType"));
-
-static UE::FPropertyTypeName FindOriginalType(const FStructProperty* Struct)
-{
-	FUObjectSerializeContext* Context = FUObjectThreadContext::Get().GetSerializeContext();
-	if (Context && Context->bImpersonateProperties)
-	{
-		const FString* OriginalType = Struct->FindMetaData(NAME_StructOriginalType);
-		if (!OriginalType)
-		{
-			//@note: To support metadata defined on array of struct in UPROPERTY for testing purposes
-			if (FField* OwnerField = Struct->Owner.ToField())
-			{
-				OriginalType = OwnerField->FindMetaData(NAME_StructOriginalType);
-			}
-		}
-		if (OriginalType)
-		{
-			if (UE::FPropertyTypeNameBuilder Type; Type.TryParse(*OriginalType))
-			{
-				return Type.Build();
-			}
-		}
-	}
-	return {};
-}
-#endif // WITH_EDITORONLY_DATA
-
 bool FStructProperty::LoadTypeName(UE::FPropertyTypeName Type, const FPropertyTag* Tag)
 {
 	if (!Super::LoadTypeName(Type, Tag))
@@ -515,7 +486,7 @@ void FStructProperty::SaveTypeName(UE::FPropertyTypeNameBuilder& Type) const
 
 	Type.BeginParameters();
 #if WITH_EDITORONLY_DATA
-	if (const UE::FPropertyTypeName OriginalType = FindOriginalType(this); !OriginalType.IsEmpty())
+	if (const UE::FPropertyTypeName OriginalType = UE::FindOriginalType(this); !OriginalType.IsEmpty())
 	{
 		Type.AddType(OriginalType);
 	}
@@ -554,7 +525,7 @@ bool FStructProperty::CanSerializeFromTypeName(UE::FPropertyTypeName Type) const
 	}
 
 #if WITH_EDITORONLY_DATA
-	if (const UE::FPropertyTypeName OriginalType = FindOriginalType(this); !OriginalType.IsEmpty())
+	if (const UE::FPropertyTypeName OriginalType = UE::FindOriginalType(this); !OriginalType.IsEmpty())
 	{
 		return StructName == OriginalType.GetName();
 	}
