Deprecates and removes MatchSubstring CoreRedirects from ini files in favour of `MatchWildcard=true, OldName="...SubString..."` which is equivalent but flexible to prefix and suffix matching as well. This change is a follow up to core redirect work in UE-204063

All existing substring matches have been kept as substring matches. Until prefix matching is optimized, it's _faster_ to perform substring matching (Boot->PIE 4.8s vs 7.1s spent in RedirectNameAndValues)

diff --git a/Engine/Plugins/Animation/ControlRig/Config/BaseControlRig.ini b/Engine/Plugins/Animation/ControlRig/Config/BaseControlRig.ini
--- a/Engine/Plugins/Animation/ControlRig/Config/BaseControlRig.ini
+++ b/Engine/Plugins/Animation/ControlRig/Config/BaseControlRig.ini
@@ -50,8 +50,8 @@
 +PropertyRedirects=(OldName="/Script/ControlRig.UControlRigSettings.SetupEventBorderColor",NewName="ConstructionEventBorderColor")
 
 ; moving core units to RigVM runtime
-;+StructRedirects=(OldName="/Script/ControlRig.RigUnit_",NewName="/Script/RigVM.RigVMFunction_", MatchSubstring=true)
-;+StructRedirects=(OldName="/Script/ControlRig.RigDispatch_",NewName="/Script/RigVM.RigVMDispatch_", MatchSubstring=true)
+;+StructRedirects=(OldName="/Script/ControlRig.RigUnit_...",NewName="/Script/RigVM.RigVMFunction_", MatchWildcard=true)
+;+StructRedirects=(OldName="/Script/ControlRig.RigDispatch_...",NewName="/Script/RigVM.RigVMDispatch_", MatchWildcard=true)
 +StructRedirects=(OldName="/Script/ControlRig.RigUnit_NameBase",NewName="/Script/RigVM.RigVMFunction_NameBase")
 +StructRedirects=(OldName="/Script/ControlRig.RigUnit_NameConcat",NewName="/Script/RigVM.RigVMFunction_NameConcat")
 +StructRedirects=(OldName="/Script/ControlRig.RigUnit_NameTruncate",NewName="/Script/RigVM.RigVMFunction_NameTruncate")
diff --git a/Engine/Plugins/Animation/DeformerGraph/Config/BaseDeformerGraph.ini b/Engine/Plugins/Animation/DeformerGraph/Config/BaseDeformerGraph.ini
--- a/Engine/Plugins/Animation/DeformerGraph/Config/BaseDeformerGraph.ini
+++ b/Engine/Plugins/Animation/DeformerGraph/Config/BaseDeformerGraph.ini
@@ -2,9 +2,9 @@
 ; Uncomment these if you have old Deformer Graph assets that won't load properly.
 ;+StructRedirects=(OldName="OptimusType_ShaderText",NewName="/Script/OptimusCore.OptimusShaderText")
 
-;+ClassRedirects=(OldName="/Script/OptimusDeveloper.",NewName="/Script/OptimusCore.", MatchSubstring=true)
-;+StructRedirects=(OldName="/Script/OptimusDeveloper.",NewName="/Script/OptimusCore.", MatchSubstring=true)
-;+EnumRedirects=(OldName="/Script/OptimusDeveloper.",NewName="/Script/OptimusCore.", MatchSubstring=true)
+;+ClassRedirects=(OldName="/Script/OptimusDeveloper....",NewName="/Script/OptimusCore.", MatchWildcard=true)
+;+StructRedirects=(OldName="/Script/OptimusDeveloper....",NewName="/Script/OptimusCore.", MatchWildcard=true)
+;+EnumRedirects=(OldName="/Script/OptimusDeveloper....",NewName="/Script/OptimusCore.", MatchWildcard=true)
 
 +ClassRedirects=(OldName="/Script/OptimusCore.ClothDataInterface",NewName="/Script/OptimusCore.OptimusClothDataInterface")
 +ClassRedirects=(OldName="/Script/OptimusCore.ClothDataProvider",NewName="/Script/OptimusCore.OptimusClothDataProvider")
diff --git a/Engine/Plugins/Enterprise/LidarPointCloud/Config/BaseLidarPointCloud.ini b/Engine/Plugins/Enterprise/LidarPointCloud/Config/BaseLidarPointCloud.ini
--- a/Engine/Plugins/Enterprise/LidarPointCloud/Config/BaseLidarPointCloud.ini
+++ b/Engine/Plugins/Enterprise/LidarPointCloud/Config/BaseLidarPointCloud.ini
@@ -26,4 +26,4 @@
 +PackageRedirects=(OldName="/PointCloudPlugin/Materials/MF_EDF_SampleTexel_CustomPass.MF_EDF_SampleTexel_CustomPass",NewName="/LidarPointCloud/Materials/MF_EDF_SampleTexel_CustomPass.MF_EDF_SampleTexel_CustomPass")
 
 ; Package
-+PackageRedirects=(MatchSubstring=true,OldName="/PointCloudRuntime",NewName="/LidarPointCloudRuntime")
\ No newline at end of file
++PackageRedirects=(MatchWildcard=true,OldName=".../PointCloudRuntime...",NewName="/LidarPointCloudRuntime")
\ No newline at end of file
diff --git a/Engine/Plugins/Enterprise/MDLImporter/Config/BaseMDLImporter.ini b/Engine/Plugins/Enterprise/MDLImporter/Config/BaseMDLImporter.ini
--- a/Engine/Plugins/Enterprise/MDLImporter/Config/BaseMDLImporter.ini
+++ b/Engine/Plugins/Enterprise/MDLImporter/Config/BaseMDLImporter.ini
@@ -1,3 +1,3 @@
 [CoreRedirects]
 
-+PackageRedirects=(OldName="/DatasmithContent/Materials/MDL/", NewName="/MDLImporter/Materials/MDL/", MatchSubstring=true)
++PackageRedirects=(OldName="/DatasmithContent/Materials/MDL/...", NewName="/MDLImporter/Materials/MDL/", MatchWildcard=true)
diff --git a/Engine/Plugins/FX/NiagaraFluids/Config/BaseNiagaraFluids.ini b/Engine/Plugins/FX/NiagaraFluids/Config/BaseNiagaraFluids.ini
--- a/Engine/Plugins/FX/NiagaraFluids/Config/BaseNiagaraFluids.ini
+++ b/Engine/Plugins/FX/NiagaraFluids/Config/BaseNiagaraFluids.ini
@@ -1,2 +1,2 @@
 [CoreRedirects]
-+PackageRedirects=(OldName="/NiagaraSimulationStages/",NewName="/NiagaraFluids/",MatchSubstring=true)
++PackageRedirects=(OldName="/NiagaraSimulationStages/...",NewName="/NiagaraFluids/",MatchWildcard=true)
diff --git a/Engine/Plugins/Runtime/MixedRealityCaptureFramework/Config/BaseMixedRealityCaptureFramework.ini b/Engine/Plugins/Runtime/MixedRealityCaptureFramework/Config/BaseMixedRealityCaptureFramework.ini
--- a/Engine/Plugins/Runtime/MixedRealityCaptureFramework/Config/BaseMixedRealityCaptureFramework.ini
+++ b/Engine/Plugins/Runtime/MixedRealityCaptureFramework/Config/BaseMixedRealityCaptureFramework.ini
@@ -5,7 +5,7 @@
 ;
 
 ; Plugin/Module redirects
-+PackageRedirects=(OldName="/MixedRealityFramework/", NewName="/MixedRealityCaptureFramework/", MatchSubstring=true)
++PackageRedirects=(OldName="/MixedRealityFramework/...", NewName="/MixedRealityCaptureFramework/", MatchWildcard=true)
 +PackageRedirects=(OldName="/Script/MixedRealityFramework", NewName="/Script/MixedRealityCaptureFramework")
 +PackageRedirects=(OldName="/Script/MixedRealityCalibration", NewName="/Script/MixedRealityCaptureCalibration")
 
diff --git a/Engine/Plugins/VirtualProduction/VirtualCamera/Config/BaseVirtualCamera.ini b/Engine/Plugins/VirtualProduction/VirtualCamera/Config/BaseVirtualCamera.ini
--- a/Engine/Plugins/VirtualProduction/VirtualCamera/Config/BaseVirtualCamera.ini
+++ b/Engine/Plugins/VirtualProduction/VirtualCamera/Config/BaseVirtualCamera.ini
@@ -1,5 +1,5 @@
 [CoreRedirects]
-+PackageRedirects=(OldName="/VirtualCamera/OperatorPanel", NewName="/VirtualCamera/Deprecated/OperatorPanel", MatchSubstring=true)
-+PackageRedirects=(OldName="/VirtualCamera/TestMaps", NewName="/VirtualCamera/Deprecated/TestMaps", MatchSubstring=true)
-+PackageRedirects=(OldName="/VirtualCamera/V2", NewName="/VirtualCamera/Deprecated/V2", MatchSubstring=true)
-+PackageRedirects=(OldName="/VirtualCamera/VCamCore", NewName="/VirtualCamera/Deprecated/VCamCore", MatchSubstring=true)
\ No newline at end of file
++PackageRedirects=(OldName="/VirtualCamera/OperatorPanel...", NewName="/VirtualCamera/Deprecated/OperatorPanel", MatchWildcard=true)
++PackageRedirects=(OldName="/VirtualCamera/TestMaps...", NewName="/VirtualCamera/Deprecated/TestMaps", MatchWildcard=true)
++PackageRedirects=(OldName="/VirtualCamera/V2...", NewName="/VirtualCamera/Deprecated/V2", MatchWildcard=true)
++PackageRedirects=(OldName="/VirtualCamera/VCamCore...", NewName="/VirtualCamera/Deprecated/VCamCore", MatchWildcard=true)
\ No newline at end of file
diff --git a/Engine/Source/Runtime/CoreUObject/Private/UObject/CoreRedirects.cpp b/Engine/Source/Runtime/CoreUObject/Private/UObject/CoreRedirects.cpp
--- a/Engine/Source/Runtime/CoreUObject/Private/UObject/CoreRedirects.cpp
+++ b/Engine/Source/Runtime/CoreUObject/Private/UObject/CoreRedirects.cpp
@@ -1579,13 +1579,14 @@ bool FCoreRedirects::ReadRedirectsFromIni(const FString& IniName)
 
 					if (bMatchWildcard || bMatchSubstring)
 					{
-						// Enable once we deprecate MatchSubstring across the engine
-						//UE_CLOG(bMatchSubstring, LogCoreRedirects, Warning, TEXT("ReadRedirectsFromIni(%s) 'MatchSubstring=' is deprecated. "
-						//	"Please prefer `MatchWildcard=' instead for redirect %s!"), *IniName, *ValueString);
+						UE_CLOG(bMatchSubstring, LogCoreRedirects, Warning, TEXT("ReadRedirectsFromIni(%s) 'MatchSubstring=' is deprecated. "
+							"Please prefer `MatchWildcard=' instead for redirect %s. "
+							"For more information refer to the documentation in Engine/Config/BaseEngine.ini."), *IniName, *ValueString);
 
 						constexpr FStringView WildcardMarker(TEXTVIEW("..."));
-						const bool bMatchPrefix = OldName.EndsWith(WildcardMarker, ESearchCase::CaseSensitive);
-						const bool bMatchSuffix = OldName.StartsWith(WildcardMarker, ESearchCase::CaseSensitive);
+						bool bMatchPrefix = OldName.EndsWith(WildcardMarker, ESearchCase::CaseSensitive);
+						bool bMatchSuffix = OldName.StartsWith(WildcardMarker, ESearchCase::CaseSensitive);
+
 						bMatchSubstring = bMatchSubstring || (bMatchPrefix && bMatchSuffix);
 
 						// Count how many '...' there are to ensure the OldName is not malformed
@@ -1618,6 +1619,12 @@ bool FCoreRedirects::ReadRedirectsFromIni(const FString& IniName)
 							OldName.RightChopInline(WildcardMarker.Len(), EAllowShrinking::No);
 						}
 
+						// UE-209583 - Convert prefix and suffix matches to substring matches since prefix/suffix is currently unoptimized
+						if (bMatchPrefix || bMatchSuffix)
+						{
+							bMatchSubstring = true;
+						}
+
 						NewFlags |= bMatchSubstring ? ECoreRedirectFlags::Option_MatchSubstring : ECoreRedirectFlags::None;
 					}
 
