Fix static analysis warning where bools are used in arithmetic. It is indeed intended to treat boolean results as integer values, so the warning is resolved with casting

diff --git a/Engine/Source/Runtime/CoreUObject/Private/UObject/CoreRedirects.cpp b/Engine/Source/Runtime/CoreUObject/Private/UObject/CoreRedirects.cpp
--- a/Engine/Source/Runtime/CoreUObject/Private/UObject/CoreRedirects.cpp
+++ b/Engine/Source/Runtime/CoreUObject/Private/UObject/CoreRedirects.cpp
@@ -1599,8 +1599,8 @@ bool FCoreRedirects::ReadRedirectsFromIni(const FString& IniName)
 							WildcardCount++;
 						}
 
-						if ((!bMatchPrefix && !bMatchSuffix && !bMatchSubstring) // No wildcards found and not handling MatchSubstring (no '...')
-							|| (WildcardCount > (bMatchPrefix + bMatchSuffix)))	 // Ensure we don't have more wildcards than necessary
+						if ((!bMatchPrefix && !bMatchSuffix && !bMatchSubstring)		  // No wildcards found and not handling MatchSubstring (no '...')
+							|| (WildcardCount > ((int)bMatchPrefix + (int)bMatchSuffix))) // Ensure we don't have more wildcards than necessary
 						{
 							UE_LOG(LogCoreRedirects, Error, TEXT("ReadRedirectsFromIni(%s) failed to parse OldName for wildcard redirect %s! "
 								"OldName must be of the form 'PrefixName...', '...SubstringName...' or '...SuffixName'. For more information refer to the documentation in Engine/Config/BaseEngine.ini."), *IniName, *ValueString);
