Partially Rollbacking cl 32194239 as it creates overridable serialization problems.

diff --git a/Engine/Source/Runtime/CoreUObject/Private/UObject/Class.cpp b/Engine/Source/Runtime/CoreUObject/Private/UObject/Class.cpp
--- a/Engine/Source/Runtime/CoreUObject/Private/UObject/Class.cpp
+++ b/Engine/Source/Runtime/CoreUObject/Private/UObject/Class.cpp
@@ -1423,16 +1423,7 @@ void UStruct::SerializeVersionedTaggedProperties(FStructuredArchive::FSlot Slot,
 				// Overridden values are saved independently in transaction, so do no need to restore them here.
 				if (!UnderlyingArchive.IsTransacting())
 				{
-					if (UnderlyingArchive.ArMergeOverrides)
-					{
-						ControlContext.OverriddenProperties = FOverridableManager::Get().GetOverriddenProperties(*(UObject*)Data);
-					}
-					
-					if (!ControlContext.OverriddenProperties)
-					{
-						ControlContext.OverriddenProperties = &FOverridableManager::Get().SetOverriddenProperties(*(UObject*)Data, Operation);
-					}
-					
+					ControlContext.OverriddenProperties = &FOverridableManager::Get().SetOverriddenProperties(*(UObject*)Data, Operation);
 					ControlContext.OverriddenProperties->bNeedsSubobjectTemplateInstantiation = true;
 				}
 			}
