[Bugfix] Fix for undefined behavior that was backed out in CL32251477

diff --git a/Engine/Source/Runtime/CoreUObject/Private/UObject/Class.cpp b/Engine/Source/Runtime/CoreUObject/Private/UObject/Class.cpp
--- a/Engine/Source/Runtime/CoreUObject/Private/UObject/Class.cpp
+++ b/Engine/Source/Runtime/CoreUObject/Private/UObject/Class.cpp
@@ -1423,7 +1423,15 @@ void UStruct::SerializeVersionedTaggedProperties(FStructuredArchive::FSlot Slot,
 				// Overridden values are saved independently in transaction, so do no need to restore them here.
 				if (!UnderlyingArchive.IsTransacting())
 				{
-					ControlContext.OverriddenProperties = &FOverridableManager::Get().SetOverriddenProperties(*(UObject*)Data, Operation);
+					ControlContext.OverriddenProperties = nullptr;
+					if (UnderlyingArchive.ArMergeOverrides)
+					{
+						ControlContext.OverriddenProperties = FOverridableManager::Get().GetOverriddenProperties(*(UObject*)Data);
+					}
+					if (!ControlContext.OverriddenProperties)
+					{
+						ControlContext.OverriddenProperties = &FOverridableManager::Get().SetOverriddenProperties(*(UObject*)Data, Operation);
+					}
 					ControlContext.OverriddenProperties->bNeedsSubobjectTemplateInstantiation = true;
 				}
 			}
