[CoreUObject]
* Removed 4 million allocations when starting up some games

diff --git a/Engine/Source/Runtime/CoreUObject/Private/UObject/UObjectGlobals.cpp b/Engine/Source/Runtime/CoreUObject/Private/UObject/UObjectGlobals.cpp
--- a/Engine/Source/Runtime/CoreUObject/Private/UObject/UObjectGlobals.cpp
+++ b/Engine/Source/Runtime/CoreUObject/Private/UObject/UObjectGlobals.cpp
@@ -479,6 +479,11 @@ UObject* StaticFindObjectFastSafe(UClass* ObjectClass, UObject* ObjectPackage, F
 	return FoundObject;
 }
 
+// TODO: Should these appear in public API?
+void ConstructorHelpers_StripObjectClass2(FStringBuilderBase& PathName, bool bAssertOnBadPath = false);
+bool ResolveName2(UObject*& InPackage, FStringBuilderBase& InOutName, bool Create, bool Throw, uint32 LoadFlags = LOAD_None, const FLinkerInstancingContext* InstancingContext = nullptr);
+
+
 #if WITH_EDITOR
 static UObject* LoadObjectWhenImportingT3D(UClass* ObjectClass, const TCHAR* OrigInName)
 {
@@ -524,25 +529,24 @@ UObject* StaticFindObject( UClass* ObjectClass, UObject* InObjectPackage, const
 	}
 #endif	//#if !WITH_EDITOR
 
-	FName ObjectName;
-
+	TStringBuilder<512> InName;
+	InName = OrigInName;
+	
 	// Don't resolve the name if we're searching in any package
 	if (!bAnyPackage)
 	{
-		FString InName = OrigInName;
-		if (!ResolveName(ObjectPackage, InName, false, false))
+		if (!ResolveName2(ObjectPackage, InName, false, false))
 		{
 			return nullptr;
 		}
-		ObjectName = FName(*InName, FNAME_Add);
 	}
 	else
 	{
-		FString InName = OrigInName;
-		ConstructorHelpers::StripObjectClass(InName);
-
-		ObjectName = FName(*InName, FNAME_Add);
+		ConstructorHelpers_StripObjectClass2(InName);
 	}
+	
+	FName ObjectName(InName.ToView(), FNAME_Add);
+	
 	PRAGMA_DISABLE_DEPRECATION_WARNINGS
 	return StaticFindObjectFast(ObjectClass, ObjectPackage, ObjectName, bExactClass, bAnyPackage);
 	PRAGMA_ENABLE_DEPRECATION_WARNINGS
@@ -677,9 +681,10 @@ bool StaticFindAllObjects(TArray<UObject*>& OutFoundObjects, UClass* ObjectClass
 #endif	//#if !WITH_EDITOR
 
 	// Don't resolve the name since we're searching in any package
-	FString InName = OrigInName;
-	ConstructorHelpers::StripObjectClass(InName);
-	FName ObjectName(*InName, FNAME_Add);
+	TStringBuilder<512> InName;
+	InName = OrigInName;
+	ConstructorHelpers_StripObjectClass2(InName);
+	FName ObjectName(InName.ToView(), FNAME_Add);
 
 	return StaticFindAllObjectsFast(OutFoundObjects, ObjectClass, ObjectName, ExactClass);
 }
@@ -705,9 +710,10 @@ UObject* StaticFindFirstObject(UClass* Class, const TCHAR* Name, EFindFirstObjec
 	}
 	else
 	{
-		FString InName = Name;
-		ConstructorHelpers::StripObjectClass(InName);
-		ObjectName = FName(*InName, FNAME_Add);
+		TStringBuilder<512> InName;
+		InName = Name;
+		ConstructorHelpers_StripObjectClass2(InName);
+		ObjectName = FName(InName.ToView(), FNAME_Add);
 	}
 
 	if (AmbiguousMessageVerbosity == ELogVerbosity::NoLogging && !(Options & (EFindFirstObjectOptions::NativeFirst | EFindFirstObjectOptions::EnsureIfAmbiguous)))
@@ -949,19 +955,19 @@ bool SafeLoadError( UObject* Outer, uint32 LoadFlags, const TCHAR* ErrorMessage)
 
 UPackage* FindPackage( UObject* InOuter, const TCHAR* PackageName )
 {
-	FString InName;
+	TStringBuilder<512> InName;
 	if( PackageName )
 	{
 		InName = PackageName;
 	}
 	else
 	{
-		InName = MakeUniqueObjectName( InOuter, UPackage::StaticClass() ).ToString();
+		MakeUniqueObjectName( InOuter, UPackage::StaticClass() ).ToString(InName);
 	}
-	ResolveName( InOuter, InName, true, false );
+	ResolveName2( InOuter, InName, true, false );
 
 	UPackage* Result = NULL;
-	if ( InName != TEXT("None") )
+	if ( InName.ToView() != TEXT("None") )
 	{
 		Result = FindObject<UPackage>( InOuter, *InName );
 	}
@@ -1027,32 +1033,32 @@ void RemoveMountPointDefaultPackageFlags(const TArrayView<FString> InMountPoints
 
 UPackage* CreatePackage(const TCHAR* PackageName )
 {
-	FString InName;
+	TStringBuilder<512> InName;
 
 	if( PackageName )
 	{
 		InName = PackageName;
-	}
 
-	if (InName.Contains(TEXT("//"), ESearchCase::CaseSensitive))
-	{
-		UE_LOG(LogUObjectGlobals, Fatal, TEXT("Attempted to create a package with name containing double slashes. PackageName: %s"), PackageName);
-	}
+		if( InName.ToView().Contains(TEXT("//")))
+		{
+			UE_LOG(LogUObjectGlobals, Fatal, TEXT("Attempted to create a package with name containing double slashes. PackageName: %s"), PackageName);
+		}
 
-	if( InName.EndsWith( TEXT( "." ), ESearchCase::CaseSensitive ) )
-	{
-		FString InName2 = InName.Left( InName.Len() - 1 );
-		UE_LOG(LogUObjectGlobals, Log,  TEXT( "Invalid Package Name entered - '%s' renamed to '%s'" ), *InName, *InName2 );
-		InName = InName2;
+		if( InName.ToView().EndsWith( TEXT( "." ) ) )
+		{
+			FStringView InName2 = InName.ToView().Left( InName.Len() - 1 );
+			UE_LOG(LogUObjectGlobals, Log,  TEXT( "Invalid Package Name entered - '%s' renamed to '%s'" ), *InName, InName2.GetData() );
+			InName = InName2;
+		}
 	}
 
 	if(InName.Len() == 0)
 	{
-		InName = MakeUniqueObjectName( nullptr, UPackage::StaticClass() ).ToString();
+		MakeUniqueObjectName( nullptr, UPackage::StaticClass() ).ToString(InName);
 	}
 
 	UObject* Outer = nullptr;
-	ResolveName(Outer, InName, true, false );
+	ResolveName2(Outer, InName, true, false );
 
 
 	UPackage* Result = NULL;
@@ -1061,7 +1067,7 @@ UPackage* CreatePackage(const TCHAR* PackageName )
 		UE_LOG(LogUObjectGlobals, Fatal, TEXT("%s"), TEXT("Attempted to create a package with an empty package name.") );
 	}
 
-	if ( InName != TEXT("None") )
+	if ( InName.ToView() != TEXT("None") )
 	{
 		Result = FindObject<UPackage>( nullptr, *InName );
 		if( Result == NULL )
@@ -1158,10 +1164,10 @@ const FString* GetIniFilenameFromObjectsReference(const FString& Name)
 //
 // Resolve a package and name.
 //
-bool ResolveName(UObject*& InPackage, FString& InOutName, bool Create, bool Throw, uint32 LoadFlags /*= LOAD_None*/, const FLinkerInstancingContext* InstancingContext)
+bool ResolveName2(UObject*& InPackage, FStringBuilderBase& InOutName, bool Create, bool Throw, uint32 LoadFlags, const FLinkerInstancingContext* InstancingContext)
 {
 	// Strip off the object class.
-	ConstructorHelpers::StripObjectClass( InOutName );
+	ConstructorHelpers_StripObjectClass2( InOutName );
 
 	// if you're attempting to find an object in any package using a dotted name that isn't fully
 	// qualified (such as ObjectName.SubobjectName - notice no package name there), you normally call
@@ -1256,10 +1262,19 @@ bool ResolveName(UObject*& InPackage, FString& InOutName, bool Create, bool Thro
 
 			check(InPackage);
 		}
-		InOutName.RemoveAt(0, DotIndex + 1, EAllowShrinking::No);
+		InOutName.RemoveAt(0, DotIndex + 1);
 	}
 }
 
+bool ResolveName(UObject*& InPackage, FString& InOutName, bool Create, bool Throw, uint32 LoadFlags /*= LOAD_None*/, const FLinkerInstancingContext* InstancingContext)
+{
+	TStringBuilder<512> Builder;
+	Builder.Append(InOutName);
+	bool Result = ResolveName2(InPackage, Builder, Create, Throw, LoadFlags, InstancingContext);
+	InOutName = Builder;
+	return Result;
+}
+
 bool ParseObject( const TCHAR* Stream, const TCHAR* Match, UClass* Class, UObject*& DestRes, UObject* InParent, EParseObjectLoadingPolicy LoadingPolicy, bool* bInvalidObject )
 {
 	if (!GAllowParseObjectLoading)
@@ -1347,12 +1362,13 @@ UObject* StaticLoadObjectInternal(UClass* ObjectClass, UObject* InOuter, const T
 	check(InName);
 
 	FScopedLoadingState ScopedLoadingState(InName);
-	FString StrName = InName;
+	TStringBuilder<512> StrName;
+	StrName = InName;
 	UObject* Result = nullptr;
 	const bool bContainsObjectName = !!FCString::Strstr(InName, TEXT("."));
 
 	// break up the name into packages, returning the innermost name and its outer
-	ResolveName(InOuter, StrName, true, true, LoadFlags & (LOAD_EditorOnly | LOAD_NoVerify | LOAD_Quiet | LOAD_NoWarn | LOAD_DeferDependencyLoads), InstancingContext);
+	ResolveName2(InOuter, StrName, true, true, LoadFlags & (LOAD_EditorOnly | LOAD_NoVerify | LOAD_Quiet | LOAD_NoWarn | LOAD_DeferDependencyLoads), InstancingContext);
 	if (InOuter)
 	{
 		// If we have a full UObject name then attempt to find the object in memory first,
@@ -1430,8 +1446,9 @@ UObject* StaticLoadObject(UClass* ObjectClass, UObject* InOuter, const TCHAR* In
 	UObject* Result = StaticLoadObjectInternal(ObjectClass, InOuter, InName, Filename, LoadFlags, Sandbox, bAllowObjectReconciliation, InstancingContext);
 	if (!Result)
 	{
-		FString ObjectName = InName;
-		ResolveName(InOuter, ObjectName, true, true, LoadFlags & LOAD_EditorOnly, InstancingContext);
+		TStringBuilder<512> ObjectName;
+		ObjectName = InName;
+		ResolveName2(InOuter, ObjectName, true, true, LoadFlags & LOAD_EditorOnly, InstancingContext);
 
 		if (InOuter == nullptr || FLinkerLoad::IsKnownMissingPackage(FName(*InOuter->GetPathName())) == false)
 		{
@@ -1439,7 +1456,7 @@ UObject* StaticLoadObject(UClass* ObjectClass, UObject* InOuter, const TCHAR* In
 			FFormatNamedArguments Arguments;
 			Arguments.Add(TEXT("ClassName"), ObjectClass ? FText::FromString(ObjectClass->GetName()) : NSLOCTEXT("Core", "None", "None"));
 			Arguments.Add(TEXT("OuterName"), InOuter ? FText::FromString(InOuter->GetPathName()) : NSLOCTEXT("Core", "None", "None"));
-			Arguments.Add(TEXT("ObjectName"), FText::FromString(ObjectName));
+			Arguments.Add(TEXT("ObjectName"), FText::FromStringView(ObjectName.ToView()));
 			const FString Error = FText::Format(NSLOCTEXT("Core", "ObjectNotFound", "Failed to find object '{ClassName} {OuterName}.{ObjectName}'"), Arguments).ToString();
 			SafeLoadError(InOuter, LoadFlags, *Error);
 
@@ -4637,6 +4654,28 @@ void ConstructorHelpers::StripObjectClass( FString& PathName, bool bAssertOnBadP
 	}
 }
 
+void ConstructorHelpers_StripObjectClass2(FStringBuilderBase& PathName, bool bAssertOnBadPath /*= false */ )
+{
+	int32 NameStartIndex = INDEX_NONE;
+	PathName.ToView().FindChar( TCHAR('\''), NameStartIndex );
+	if( NameStartIndex != INDEX_NONE )
+	{
+		int32 NameEndIndex = INDEX_NONE;
+		PathName.ToView().FindLastChar( TCHAR('\''), NameEndIndex );
+		if(NameEndIndex > NameStartIndex)
+		{
+			TStringBuilder<256> Temp;
+			Temp.Append(PathName.GetData() + NameStartIndex + 1, NameEndIndex - NameStartIndex - 1);
+			PathName.Reset();
+			PathName.Append(Temp);
+		}
+		else
+		{
+			UE_CLOG( bAssertOnBadPath, LogUObjectGlobals, Fatal, TEXT("Bad path name: %s, missing \' or an incorrect format"), *PathName );
+		}
+	}
+}
+
 //////////////////////////////////////////////////////////////////////////
 
 FReferenceCollectorArchive::FReferenceCollectorArchive(const UObject* InSerializingObject, FReferenceCollector& InCollector)
diff --git a/Engine/Source/Runtime/Engine/Private/AssetManager.cpp b/Engine/Source/Runtime/Engine/Private/AssetManager.cpp
--- a/Engine/Source/Runtime/Engine/Private/AssetManager.cpp
+++ b/Engine/Source/Runtime/Engine/Private/AssetManager.cpp
@@ -2,6 +2,7 @@
 
 #include "Engine/AssetManager.h"
 
+#include "Algo/Compare.h"
 #include "Algo/Unique.h"
 #include "AssetRegistry/ARFilter.h"
 #include "AssetRegistry/AssetBundleData.h"
@@ -1719,12 +1720,11 @@ TSharedPtr<FStreamableHandle> UAssetManager::ChangeBundleStateForPrimaryAssets(c
 			bool bLoadIfNeeded = false;
 			
 			// Use pending state if valid
-			TArray<FName> CurrentBundleState = NameData->PendingState.IsValid() ? NameData->PendingState.BundleNames : NameData->CurrentState.BundleNames;
-			TArray<FName> NewBundleState;
+			TArray<FName, TInlineAllocator<32>> NewBundleState;
 
 			if (!bRemoveAllBundles)
 			{
-				NewBundleState = CurrentBundleState;
+				NewBundleState = NameData->PendingState.IsValid() ? NameData->PendingState.BundleNames : NameData->CurrentState.BundleNames;
 
 				for (const FName& RemoveBundle : RemoveBundles)
 				{
@@ -1742,7 +1742,7 @@ TSharedPtr<FStreamableHandle> UAssetManager::ChangeBundleStateForPrimaryAssets(c
 			// If the pending state is valid, check if it is different
 			if (NameData->PendingState.IsValid())
 			{
-				if (NameData->PendingState.BundleNames == NewBundleState)
+				if (Algo::Compare(NameData->PendingState.BundleNames, NewBundleState))
 				{
 					// This will wait on any existing handles to finish
 					ExistingHandles.Add(NameData->PendingState.Handle);
@@ -1752,7 +1752,7 @@ TSharedPtr<FStreamableHandle> UAssetManager::ChangeBundleStateForPrimaryAssets(c
 				// Clear pending state
 				NameData->PendingState.Reset(true);
 			}
-			else if (NameData->CurrentState.IsValid() && NameData->CurrentState.BundleNames == NewBundleState)
+			else if (NameData->CurrentState.IsValid() && Algo::Compare(NameData->CurrentState.BundleNames, NewBundleState))
 			{
 				// If no pending, compare with current
 				continue;
@@ -2935,7 +2935,9 @@ bool UAssetManager::GetAssetDataForPath(const FSoftObjectPath& ObjectPath, FAsse
 static bool EndsWithBlueprint(const FTopLevelAssetPath& Name)
 {
 	// Numbered names can't end with Blueprint
-	return Name.GetAssetName().ToString().EndsWith(TEXT("Blueprint"));
+	TStringBuilder<512> AssetName;
+	Name.GetAssetName().ToString(AssetName);
+	return AssetName.ToView().EndsWith(TEXT("Blueprint"));
 }
 
 static bool ContainsSubobjectDelimiter(FName Name)
