[GarbageCollect]
 - Move PreCollect broadcast outside of the lock because it's too easy for user code to cause deadlocks by calling into other part of the engine.

diff --git a/Engine/Source/Runtime/CoreUObject/Private/UObject/GarbageCollection.cpp b/Engine/Source/Runtime/CoreUObject/Private/UObject/GarbageCollection.cpp
--- a/Engine/Source/Runtime/CoreUObject/Private/UObject/GarbageCollection.cpp
+++ b/Engine/Source/Runtime/CoreUObject/Private/UObject/GarbageCollection.cpp
@@ -5417,24 +5417,34 @@ void PreCollectGarbageImpl(EObjectFlags KeepFlags)
 	// Reset GC skip counter
 	GNumAttemptsSinceLastGC = 0;
 
-	// Flush streaming before GC if requested
-	if (!GIsIncrementalReachabilityPending && GFlushStreamingOnGC && IsAsyncLoading())
+	if (!GIsIncrementalReachabilityPending)
 	{
-		UE_LOG(LogGarbage, Log, TEXT("CollectGarbageInternal() is flushing async loading"));
+		// Do not hold the lock during flush async loading and calling into user code through a delegate
+		// because this could cause deadlocks easily.
 		ReleaseGCLock();
-		FlushAsyncLoading();
-		AcquireGCLock();
 
-		GGCStats.bFlushedAsyncLoading = true;
-	}
+		// Flush streaming before GC if requested
+		if (GFlushStreamingOnGC && IsAsyncLoading())
+		{
+			UE_LOG(LogGarbage, Log, TEXT("CollectGarbageInternal() is flushing async loading"));
+			FlushAsyncLoading();
+			GGCStats.bFlushedAsyncLoading = true;
+		}
+		{
+			TRACE_CPUPROFILER_EVENT_SCOPE(BroadcastPreGarbageCollect);
+			FCoreUObjectDelegates::GetPreGarbageCollectDelegate().Broadcast();
+		}
+		// User code in PreGarbageCollect could have triggered new loading... need to flush that too.
+		if (GFlushStreamingOnGC && IsAsyncLoading())
+		{
+			UE_LOG(LogGarbage, Log, TEXT("CollectGarbageInternal() is flushing async loading because PreGarbageCollect triggered new loads"));
+			FlushAsyncLoading();
+			GGCStats.bFlushedAsyncLoading = true;
+		}
 
-	// Route callbacks so we can ensure that we are e.g. not in the middle of loading something by flushing
-	// the async loading, etc...
-	if (!GIsIncrementalReachabilityPending)
-	{
-		TRACE_CPUPROFILER_EVENT_SCOPE(BroadcastPreGarbageCollect);
-		FCoreUObjectDelegates::GetPreGarbageCollectDelegate().Broadcast();
+		AcquireGCLock();
 	}
+
 	GLastGCFrame = GFrameCounter;
 
 	{
