Verse proxy vars: fix nested class var archetype initialization glitch

#tests Sandbox/SolarisTestbed/.../Attributes/CustomClassVarAccessors.versetest

diff --git a/Engine/Source/Runtime/CoreUObject/Public/UObject/UObjectGlobals.h b/Engine/Source/Runtime/CoreUObject/Public/UObject/UObjectGlobals.h
--- a/Engine/Source/Runtime/CoreUObject/Public/UObject/UObjectGlobals.h
+++ b/Engine/Source/Runtime/CoreUObject/Public/UObject/UObjectGlobals.h
@@ -1236,6 +1236,11 @@ public:
 		return Obj;
 	}
 
+	FORCEINLINE struct FObjectInstancingGraph* GetInstancingGraph()
+	{
+		return InstanceGraph;
+	}
+
 	/**
 	* Return the class of the object that is being constructed
 	**/
