Fixing single entry set properties not being saved correctly in some cases.
Making UObject::SaveConfig fall back to the default behavior if a set property has one or fewer elements to match the load behavior

#tests saving and loading 0, 1, and 2 entry struct set properties

diff --git a/Engine/Source/Runtime/CoreUObject/Private/UObject/Obj.cpp b/Engine/Source/Runtime/CoreUObject/Private/UObject/Obj.cpp
--- a/Engine/Source/Runtime/CoreUObject/Private/UObject/Obj.cpp
+++ b/Engine/Source/Runtime/CoreUObject/Private/UObject/Obj.cpp
@@ -3308,6 +3308,8 @@ void UObject::SaveConfig( uint64 Flags, const TCHAR* InFilename, FConfigCacheIni
 			const bool bShouldCheckIfIdenticalBeforeAdding = !GetClass()->HasAnyClassFlags(CLASS_ConfigDoNotCheckDefaults) && !bPerObject && bIsPropertyInherited;
 			UObject* SuperClassDefaultObject = GetClass()->GetSuperClass()->GetDefaultObject();
 
+			bool bProcessedProperty = false;
+
 			FArrayProperty* Array   = CastField<FArrayProperty>( Property );
 			FSetProperty* SetProperty = CastField<FSetProperty>(Property);
 			if (Array || SetProperty)
@@ -3339,25 +3341,29 @@ void UObject::SaveConfig( uint64 Flags, const TCHAR* InFilename, FConfigCacheIni
 							const FString EmptyKey = FString::Printf(TEXT("!%s"), *Key);
 							Config->AddToSection(*Section, *EmptyKey, TEXT("__ClearArray__"), PropFileName);
 						}
+						bProcessedProperty = true;
 					}
 					else if (SetProperty)
 					{
 						FScriptSetHelper_InContainer SetHelper(SetProperty, this);
-						for (FScriptSetHelper::FIterator It = SetHelper.CreateIterator(); It; ++It)
-						{
-							FString	Buffer;
-							SetProperty->ElementProp->ExportTextItem_Direct(Buffer, SetHelper.GetElementPtr(It), SetHelper.GetElementPtr(It), this, PortFlags);
-							Config->AddToSection(*Section, *CompleteKey, *Buffer, PropFileName);
-						}
-						if (SetHelper.Num() == 0 && bIsADefaultIniWrite)
+
+						// If we have one or fewer elements, fall back to the old export method so we match what we do
+						// on import
+						if (SetHelper.Num() > 1)
 						{
-							const FString EmptyKey = FString::Printf(TEXT("!%s"), *Key);
-							Config->AddToSection(*Section, *EmptyKey, TEXT("__ClearSet__"), PropFileName);
+							for (FScriptSetHelper::FIterator It = SetHelper.CreateIterator(); It; ++It)
+							{
+								FString	Buffer;
+								SetProperty->ElementProp->ExportTextItem_Direct(Buffer, SetHelper.GetElementPtr(It), SetHelper.GetElementPtr(It), this, PortFlags);
+								Config->AddToSection(*Section, *CompleteKey, *Buffer, PropFileName);
+							}
+							bProcessedProperty = true;
 						}
 					}
 				}
 			}
-			else
+			
+			if (!bProcessedProperty)
 			{
 				TCHAR TempKey[MAX_SPRINTF] = {};
 				for( int32 Index=0; Index<Property->ArrayDim; Index++ )
