Always enable IDO for placeholder-typed objects.

Since all properties are loose, they require an IDO to be serialized back out at save time.

diff --git a/Engine/Source/Runtime/CoreUObject/Private/UObject/InstanceDataObjectUtils.cpp b/Engine/Source/Runtime/CoreUObject/Private/UObject/InstanceDataObjectUtils.cpp
--- a/Engine/Source/Runtime/CoreUObject/Private/UObject/InstanceDataObjectUtils.cpp
+++ b/Engine/Source/Runtime/CoreUObject/Private/UObject/InstanceDataObjectUtils.cpp
@@ -7,6 +7,7 @@
 #include "UObject/EnumProperty.h"
 #include "UObject/Field.h"
 #include "UObject/PropertyBag.h"
+#include "UObject/PropertyBagRepository.h"
 #include "UObject/PropertyOptional.h"
 #include "UObject/UnrealType.h"
 
@@ -135,6 +136,12 @@ namespace UE
 		bool bIsEnabled = bEnableIDOSupport;
 		if (bIsEnabled && InObject && !InObject->IsInPackage(GetTransientPackage()))
 		{
+			// Property bag placeholder objects are always enabled for IDO support
+			if (UE::FPropertyBagRepository::IsPropertyBagPlaceholderObject(InObject))
+			{
+				return true;
+			}
+
 			//@todo FH: change to check trait when available or use config object
 			const UClass* ObjClass = InObject->GetClass();
 			while (ObjClass && ObjClass->GetClass()->GetFName() != NAME_VerseClass)
