[ZenLoader]
 - Make ZenLoader the default loader

diff --git a/Engine/Source/Runtime/CoreUObject/Private/Serialization/AsyncPackageLoader.cpp b/Engine/Source/Runtime/CoreUObject/Private/Serialization/AsyncPackageLoader.cpp
--- a/Engine/Source/Runtime/CoreUObject/Private/Serialization/AsyncPackageLoader.cpp
+++ b/Engine/Source/Runtime/CoreUObject/Private/Serialization/AsyncPackageLoader.cpp
@@ -188,7 +188,7 @@ void InitAsyncThread()
 	LLM_SCOPE(ELLMTag::AsyncLoading);
 	if (FIoDispatcher::IsInitialized())
 	{
-		bool bSettingsEnabled = false;
+		bool bSettingsEnabled = true;
 		bool bCommandLineEnabled = false;
 		bool bCommandLineDisabled = false;
 		bool bHasUseIoStoreParamInEditor = false;
