Salvage SCD archetype data when taxonomy evolves

diff --git a/Engine/Source/Runtime/CoreUObject/Private/UObject/Class.cpp b/Engine/Source/Runtime/CoreUObject/Private/UObject/Class.cpp
--- a/Engine/Source/Runtime/CoreUObject/Private/UObject/Class.cpp
+++ b/Engine/Source/Runtime/CoreUObject/Private/UObject/Class.cpp
@@ -5968,9 +5968,27 @@ void* UClass::CreateSparseClassData()
 		const void* SparseArchetypeData = GetArchetypeForSparseClassData();
 		UScriptStruct* SparseClassDataArchetypeStruct = GetSparseClassDataArchetypeStruct();
 
-		if (SparseArchetypeData && SparseClassDataStruct->IsChildOf(SparseClassDataArchetypeStruct))
+		if (SparseArchetypeData)
 		{
-			SparseClassDataArchetypeStruct->CopyScriptStruct(SparseClassData, SparseArchetypeData);
+			// our saved SparseClassDataStruct typically will be a child of 
+			// its parent's SCD, but sometimes (e.g. because a parent introduced
+			// a new SCDStruct) it won't be, and will actually be a child of 
+			// *our* SCDStruct:
+			if (SparseClassDataStruct->IsChildOf(SparseClassDataArchetypeStruct))
+			{
+				SparseClassDataArchetypeStruct->CopyScriptStruct(SparseClassData, SparseArchetypeData);
+			}
+			else if (SparseClassDataArchetypeStruct->IsChildOf(SparseClassDataStruct))
+			{
+				SparseClassDataStruct->CopyScriptStruct(SparseClassData, SparseArchetypeData);
+			}
+			else 
+			{
+				UE_LOG(LogClass, Warning, TEXT("SparseClassData %s for class %s archetype is of unrelated type %s"),
+					*SparseClassDataStruct->GetPathName(),
+					*GetPathName(),
+					*SparseClassDataArchetypeStruct->GetPathName());
+			}
 		}
 
 		FCoreUObjectDelegates::OnPostInitSparseClassData.Broadcast(this, SparseClassDataStruct, SparseClassData);
