Fixed unreachable code warning

#rnx

diff --git a/Engine/Source/Runtime/CoreUObject/Private/Cooker/CookDependency.cpp b/Engine/Source/Runtime/CoreUObject/Private/Cooker/CookDependency.cpp
--- a/Engine/Source/Runtime/CoreUObject/Private/Cooker/CookDependency.cpp
+++ b/Engine/Source/Runtime/CoreUObject/Private/Cooker/CookDependency.cpp
@@ -284,10 +284,10 @@ bool FCookDependency::Load(FCbFieldView Value)
 		return true;
 	}
 	default:
-		return false;
+		break;
 	}
 	checkNoEntry();
-	return true;
+	return false;
 }
 
 void FCookDependencyContext::Update(const void* Data, uint64 Size)
