[GarbageCollect]
 - Move PreGarbageCollect broadcast outside of the GC lock because it's too easy for user code to cause deadlocks by calling into other part of the engine.
 - Introduce a new event when GC has started and the GC lock is held
 - Use the new event for zenloader as it provides thread-safety guarantees with async loading thread

diff --git a/Engine/Source/Runtime/CoreUObject/Private/Serialization/AsyncLoading2.cpp b/Engine/Source/Runtime/CoreUObject/Private/Serialization/AsyncLoading2.cpp
--- a/Engine/Source/Runtime/CoreUObject/Private/Serialization/AsyncLoading2.cpp
+++ b/Engine/Source/Runtime/CoreUObject/Private/Serialization/AsyncLoading2.cpp
@@ -3692,7 +3692,7 @@ public:
 	}
 
 private:
-	void OnPreGarbageCollect();
+	void OnGarbageCollectStarted();
 
 	void CollectUnreachableObjects(TArrayView<FUObjectItem*> UnreachableObjectItems, FUnreachableObjects& OutUnreachableObjects);
 
@@ -8163,7 +8163,7 @@ FAsyncLoadingThread2::FAsyncLoadingThread2(FIoDispatcher& InIoDispatcher, IAsync
 	ThreadResumedEvent = FPlatformProcess::GetSynchEventFromPool();
 	AsyncLoadingTickCounter = 0;
 
-	FCoreUObjectDelegates::GetPreGarbageCollectDelegate().AddRaw(this, &FAsyncLoadingThread2::OnPreGarbageCollect);
+	FCoreUObjectDelegates::GetGarbageCollectStartedDelegate().AddRaw(this, &FAsyncLoadingThread2::OnGarbageCollectStarted);
 
 	FAsyncLoadingThreadState2::TlsSlot = FPlatformTLS::AllocTlsSlot();
 	GameThreadState = MakeUnique<FAsyncLoadingThreadState2>(GraphAllocator, IoDispatcher);
@@ -8637,9 +8637,9 @@ void FAsyncLoadingThread2::CollectUnreachableObjects(
 	}
 }
 
-void FAsyncLoadingThread2::OnPreGarbageCollect()
+void FAsyncLoadingThread2::OnGarbageCollectStarted()
 {
-	TRACE_CPUPROFILER_EVENT_SCOPE(FAsyncLoadingThread2::OnPreGarbageCollect);
+	TRACE_CPUPROFILER_EVENT_SCOPE(FAsyncLoadingThread2::OnGarbageCollectStarted);
 	// Flush the delete queue so that we don't prevent packages from being garbage collected if we're done with them
 	ProcessDeferredDeletePackagesQueue();
 }
diff --git a/Engine/Source/Runtime/CoreUObject/Private/UObject/GarbageCollection.cpp b/Engine/Source/Runtime/CoreUObject/Private/UObject/GarbageCollection.cpp
--- a/Engine/Source/Runtime/CoreUObject/Private/UObject/GarbageCollection.cpp
+++ b/Engine/Source/Runtime/CoreUObject/Private/UObject/GarbageCollection.cpp
@@ -5417,29 +5417,45 @@ void PreCollectGarbageImpl(EObjectFlags KeepFlags)
 	// Reset GC skip counter
 	GNumAttemptsSinceLastGC = 0;
 
-	// Flush streaming before GC if requested
-	if (!GIsIncrementalReachabilityPending && GFlushStreamingOnGC && IsAsyncLoading())
+	if (!GIsIncrementalReachabilityPending)
 	{
-		UE_LOG(LogGarbage, Log, TEXT("CollectGarbageInternal() is flushing async loading"));
+		// Do not hold the lock during flush async loading and calling into user code through a delegate
+		// because this could cause deadlocks easily.
 		ReleaseGCLock();
-		FlushAsyncLoading();
-		AcquireGCLock();
 
-		GGCStats.bFlushedAsyncLoading = true;
-	}
+		// Flush streaming before GC if requested
+		if (GFlushStreamingOnGC && IsAsyncLoading())
+		{
+			UE_LOG(LogGarbage, Log, TEXT("CollectGarbageInternal() is flushing async loading"));
+			FlushAsyncLoading();
+			GGCStats.bFlushedAsyncLoading = true;
+		}
+		{
+			TRACE_CPUPROFILER_EVENT_SCOPE(BroadcastPreGarbageCollect);
+			FCoreUObjectDelegates::GetPreGarbageCollectDelegate().Broadcast();
+		}
+		// User code in PreGarbageCollect could have triggered new loading... need to flush that too.
+		if (GFlushStreamingOnGC && IsAsyncLoading())
+		{
+			UE_LOG(LogGarbage, Log, TEXT("CollectGarbageInternal() is flushing async loading because PreGarbageCollect triggered new loads"));
+			FlushAsyncLoading();
+			GGCStats.bFlushedAsyncLoading = true;
+		}
 
-	// Route callbacks so we can ensure that we are e.g. not in the middle of loading something by flushing
-	// the async loading, etc...
-	if (!GIsIncrementalReachabilityPending)
-	{
-		TRACE_CPUPROFILER_EVENT_SCOPE(BroadcastPreGarbageCollect);
-		FCoreUObjectDelegates::GetPreGarbageCollectDelegate().Broadcast();
+		AcquireGCLock();
 	}
+
 	GLastGCFrame = GFrameCounter;
 
 	{
 		GIsGarbageCollecting = true;
 
+		if (!GIsIncrementalReachabilityPending)
+		{
+			TRACE_CPUPROFILER_EVENT_SCOPE(BroadcastGarbageCollectStarted);
+			FCoreUObjectDelegates::GetGarbageCollectStartedDelegate().Broadcast();
+		}
+
 		// Make sure previous incremental purge has finished or we do a full purge pass in case we haven't kicked one
 		// off yet since the last call to garbage collection.
 		if (IsIncrementalPurgePending())
diff --git a/Engine/Source/Runtime/CoreUObject/Private/UObject/UObjectGlobals.cpp b/Engine/Source/Runtime/CoreUObject/Private/UObject/UObjectGlobals.cpp
--- a/Engine/Source/Runtime/CoreUObject/Private/UObject/UObjectGlobals.cpp
+++ b/Engine/Source/Runtime/CoreUObject/Private/UObject/UObjectGlobals.cpp
@@ -188,6 +188,12 @@ FSimpleMulticastDelegate& FCoreUObjectDelegates::GetPreGarbageCollectDelegate()
 	return Delegate;
 }
 
+FSimpleMulticastDelegate& FCoreUObjectDelegates::GetGarbageCollectStartedDelegate()
+{
+	static FSimpleMulticastDelegate Delegate;
+	return Delegate;
+}
+
 FSimpleMulticastDelegate& FCoreUObjectDelegates::GetPostGarbageCollect()
 {
 	static FSimpleMulticastDelegate Delegate;
diff --git a/Engine/Source/Runtime/CoreUObject/Public/UObject/UObjectGlobals.h b/Engine/Source/Runtime/CoreUObject/Public/UObject/UObjectGlobals.h
--- a/Engine/Source/Runtime/CoreUObject/Public/UObject/UObjectGlobals.h
+++ b/Engine/Source/Runtime/CoreUObject/Public/UObject/UObjectGlobals.h
@@ -3244,9 +3244,12 @@ struct FCoreUObjectDelegates
 	/** Sent when a network replay has started */
 	static COREUOBJECT_API FSimpleMulticastDelegate PostDemoPlay;
 
-	/** Called before garbage collection */
+	/** Called before garbage collection, before the GC lock is acquired. */
 	static COREUOBJECT_API FSimpleMulticastDelegate& GetPreGarbageCollectDelegate();
 
+	/** Called at the very beginning of garbage collection, once the GC lock is held. */
+	static COREUOBJECT_API FSimpleMulticastDelegate& GetGarbageCollectStartedDelegate();
+
 	/** Delegate type for reachability analysis external roots callback. First parameter is FGarbageCollectionTracer to use for tracing, second is flags with which objects should be kept alive regardless, third is whether to force single threading */
 	DECLARE_MULTICAST_DELEGATE_ThreeParams(FTraceExternalRootsForReachabilityAnalysisDelegate, FGarbageCollectionTracer&, EObjectFlags, bool);
 
