Assert on FPropertyTypeName with too many names and minimize its index waste

#rnx

diff --git a/Engine/Source/Runtime/CoreUObject/Private/UObject/PropertyTypeName.cpp b/Engine/Source/Runtime/CoreUObject/Private/UObject/PropertyTypeName.cpp
--- a/Engine/Source/Runtime/CoreUObject/Private/UObject/PropertyTypeName.cpp
+++ b/Engine/Source/Runtime/CoreUObject/Private/UObject/PropertyTypeName.cpp
@@ -195,24 +195,28 @@ int32 FPropertyTypeNameTable::StoreByIndex(const FPropertyTypeNameNode* Nodes, i
 		return 0;
 	}
 
+	check(Count <= GPropertyTypeNameBlockOffsetCount);
+
 	int32 Index;
-	int32 BlockIndex;
-	int32 BlockOffset;
-	do
+	for (int32 BaseIndex = NextIndex.load(std::memory_order_relaxed);;)
 	{
-		Index = NextIndex.fetch_add(Count, std::memory_order_relaxed);
-		BlockIndex = Index >> GPropertyTypeNameBlockOffsetBits;
-		BlockOffset = Index & GPropertyTypeNameBlockOffsetMask;
+		const int32 RemainingCountInBlock = GPropertyTypeNameBlockOffsetCount - (BaseIndex & GPropertyTypeNameBlockOffsetMask);
+		Index = BaseIndex + (RemainingCountInBlock <= Count ? RemainingCountInBlock : 0);
+		if (LIKELY(NextIndex.compare_exchange_weak(BaseIndex, Index + Count, std::memory_order_relaxed)))
+		{
+			break;
+		}
 	}
-	while (UNLIKELY(BlockOffset + Count >= GPropertyTypeNameBlockOffsetCount));
 
-	check(Index + Count < GPropertyTypeNameBlockCount * GPropertyTypeNameBlockOffsetCount);
+	check(Index + Count <= GPropertyTypeNameBlockCount * GPropertyTypeNameBlockOffsetCount);
 
+	const int32 BlockIndex = Index >> GPropertyTypeNameBlockOffsetBits;
 	FPropertyTypeNameNode* Block = Blocks[BlockIndex].load(std::memory_order_acquire);
 	if (UNLIKELY(!Block))
 	{
 		Block = AllocateBlock(BlockIndex);
 	}
+	const int32 BlockOffset = Index & GPropertyTypeNameBlockOffsetMask;
 	for (FPropertyTypeNameNode* Target = Block + BlockOffset; Count > 0; --Count)
 	{
 		*Target++ = *Nodes++;
diff --git a/Engine/Source/Runtime/CoreUObject/Tests/UObject/PropertyTypeNameTest.cpp b/Engine/Source/Runtime/CoreUObject/Tests/UObject/PropertyTypeNameTest.cpp
--- a/Engine/Source/Runtime/CoreUObject/Tests/UObject/PropertyTypeNameTest.cpp
+++ b/Engine/Source/Runtime/CoreUObject/Tests/UObject/PropertyTypeNameTest.cpp
@@ -457,7 +457,7 @@ TEST_CASE_NAMED(FPropertyTypeNameSmokeTest, "CoreUObject::PropertyTypeName::Smok
 
 TEST_CASE_NAMED(FPropertyTypeNameConcurrencyTest, "CoreUObject::PropertyTypeName::Concurrency", "[Core][UObject][EngineFilter]")
 {
-	ParallelFor(TEXT("PropertyTypeNameConcurrencyTest"), 64 * 1024, 1, [](int32 Index)
+	ParallelFor(TEXT("PropertyTypeNameConcurrencyTest"), 64 * 1024, 1, [NAME_TestStruct = FName(ANSITEXTVIEW("TestStruct"))](int32 Index)
 	{
 		FPropertyTypeNameBuilder Builder;
 
@@ -474,7 +474,7 @@ TEST_CASE_NAMED(FPropertyTypeNameConcurrencyTest, "CoreUObject::PropertyTypeName
 
 		Builder.AddName(NAME_StructProperty);
 		Builder.BeginParameters();
-		Builder.AddName(FName(WriteToString<32>(TEXTVIEW("TestStruct"), Index)));
+		Builder.AddName(FName(NAME_TestStruct, NAME_EXTERNAL_TO_INTERNAL(Index)));
 		Builder.EndParameters();
 
 		if (Index % 4 == 0)
