autortfm: rework the way UObject construction works

The goal here is that if a UObject is constructed during a transaction, and then decides to abort, we leave the UObject in an intact-enough state that the GC can come along later to destruct and deallocate the object.

To that end, we move StaticConstructObject_Internal back to instrumented, but execute the UObject malloc and UObjectBase constructor uninstrumented. Then later when the derived constructor runs over the same memory, we run the constructor instrumented, except when it gets down to the UObject constructor level, we run that uninstrumented.

If we abort, the UObject memory will not be freed by the AutoRTFM runtime, the destructor will not be called, but the memory will be rolled back to what the UObject and base constructors initialize it as, so it should be a bare bones object ready for destruction and deallocation.

- added a new attribute to the compiler (autortm_always_open) so we can annotate functions themselves (like constructors) to ensure the entire function isn't instrumented (including code we can't normally wrap, like when the constructor writes the vtable pointer)

diff --git a/Engine/Source/Programs/AutoRTFMTests/Classes/MyAutoRTFMTestObject.h b/Engine/Source/Programs/AutoRTFMTests/Classes/MyAutoRTFMTestObject.h
--- a/Engine/Source/Programs/AutoRTFMTests/Classes/MyAutoRTFMTestObject.h
+++ b/Engine/Source/Programs/AutoRTFMTests/Classes/MyAutoRTFMTestObject.h
@@ -9,8 +9,17 @@ class UMyAutoRTFMTestObject : public UObject
 	GENERATED_BODY()
 
 public:
-	UMyAutoRTFMTestObject() : Value(42)
-	{}
+	UMyAutoRTFMTestObject(const FObjectInitializer& ObjectInitializer = FObjectInitializer::Get()) : Value(42)
+	{
+		UObject* const Obj = ObjectInitializer.GetObj();
+		UObject* const Outer = Obj->GetOuter();
+
+		if (Outer->IsA<UMyAutoRTFMTestObject>())
+		{
+			UMyAutoRTFMTestObject* const OuterAsType = static_cast<UMyAutoRTFMTestObject*>(Outer);
+			OuterAsType->Value += 13;
+		}
+	}
 
 	int Value;
 };
diff --git a/Engine/Source/Programs/AutoRTFMTests/Private/UObjectTests.cpp b/Engine/Source/Programs/AutoRTFMTests/Private/UObjectTests.cpp
--- a/Engine/Source/Programs/AutoRTFMTests/Private/UObjectTests.cpp
+++ b/Engine/Source/Programs/AutoRTFMTests/Private/UObjectTests.cpp
@@ -6,30 +6,67 @@
 #include "UObject/GCObject.h"
 #include "UObject/ReachabilityAnalysis.h"
 
-TEST_CASE("UObject.Create")
+TEST_CASE("UObject.NewObject")
 {
-	UMyAutoRTFMTestObject* Object = nullptr;
+	SECTION("Create")
+	{
+		UMyAutoRTFMTestObject* Object = nullptr;
+
+		AutoRTFM::Commit([&]
+			{
+				Object = NewObject<UMyAutoRTFMTestObject>();
+			});
 
-	AutoRTFM::Commit([&]
+		REQUIRE(nullptr != Object);
+		REQUIRE(42 == Object->Value);
+	}
+
+	SECTION("Abort")
 	{
-		Object = NewObject<UMyAutoRTFMTestObject>();
-	});
+		UMyAutoRTFMTestObject* Object = nullptr;
+
+		REQUIRE(AutoRTFM::ETransactionResult::AbortedByRequest == AutoRTFM::Transact([&]
+			{
+				Object = NewObject<UMyAutoRTFMTestObject>();
+				AutoRTFM::AbortTransaction();
+			}));
 
-	REQUIRE(nullptr != Object);
-	REQUIRE(42 == Object->Value);
+		REQUIRE(nullptr == Object);
+	}
 }
 
-TEST_CASE("UObject.Abort")
+TEST_CASE("UObject.NewObjectWithOuter")
 {
-	UMyAutoRTFMTestObject* Object = nullptr;
+	SECTION("Create")
+	{
+		UMyAutoRTFMTestObject* Outer = NewObject<UMyAutoRTFMTestObject>();
+		UMyAutoRTFMTestObject* Object = nullptr;
 
-	REQUIRE(AutoRTFM::ETransactionResult::AbortedByRequest == AutoRTFM::Transact([&]
+		AutoRTFM::Commit([&]
+			{
+				Object = NewObject<UMyAutoRTFMTestObject>(Outer);
+			});
+
+		REQUIRE(nullptr != Object);
+		REQUIRE(42 == Object->Value);
+		REQUIRE(Object->IsInOuter(Outer));
+		REQUIRE(55 == Outer->Value);
+	}
+
+	SECTION("Abort")
 	{
-		Object = NewObject<UMyAutoRTFMTestObject>();
-		AutoRTFM::AbortTransaction();
-	}));
+		UMyAutoRTFMTestObject* Outer = NewObject<UMyAutoRTFMTestObject>();
+		UMyAutoRTFMTestObject* Object = nullptr;
+
+		REQUIRE(AutoRTFM::ETransactionResult::AbortedByRequest == AutoRTFM::Transact([&]
+			{
+				Object = NewObject<UMyAutoRTFMTestObject>(Outer);
+				AutoRTFM::AbortTransaction();
+			}));
 
-	REQUIRE(nullptr == Object);
+		REQUIRE(nullptr == Object);
+		REQUIRE(42 == Outer->Value);
+	}
 }
 
 // This is a copy of the helper function in TestGarbageCollector.cpp.
diff --git a/Engine/Source/Runtime/Core/Public/AutoRTFM/AutoRTFM.h b/Engine/Source/Runtime/Core/Public/AutoRTFM/AutoRTFM.h
--- a/Engine/Source/Runtime/Core/Public/AutoRTFM/AutoRTFM.h
+++ b/Engine/Source/Runtime/Core/Public/AutoRTFM/AutoRTFM.h
@@ -11,9 +11,11 @@
 #if UE_AUTORTFM
 #define UE_AUTORTFM_AUTORTFM(F) [[clang::autortfm(F), clang::noinline]]
 #define UE_AUTORTFM_NOAUTORTFM [[clang::noautortfm, clang::noinline]]
+#define UE_AUTORTFM_ALWAYS_OPEN [[clang::autortfm_always_open]]
 #else
 #define UE_AUTORTFM_AUTORTFM(F)
 #define UE_AUTORTFM_NOAUTORTFM
+#define UE_AUTORTFM_ALWAYS_OPEN
 #endif
 
 #if defined(UE_AUTORTFM_STANDALONE)
diff --git a/Engine/Source/Runtime/CoreUObject/Private/Serialization/AsyncLoadingThread.h b/Engine/Source/Runtime/CoreUObject/Private/Serialization/AsyncLoadingThread.h
--- a/Engine/Source/Runtime/CoreUObject/Private/Serialization/AsyncLoadingThread.h
+++ b/Engine/Source/Runtime/CoreUObject/Private/Serialization/AsyncLoadingThread.h
@@ -19,6 +19,7 @@
 #include "Misc/ConfigCacheIni.h"
 #include "HAL/LowLevelMemTracker.h"
 #include "Serialization/AsyncPackageLoader.h"
+#include "AutoRTFM/AutoRTFM.h"
 
 #include <atomic>
 
@@ -473,7 +474,9 @@ public:
 	/** Gets the current state of async loading */
 	static bool GetIsInAsyncLoadingTick()
 	{
-		return !!AsyncLoadingTickCounter.GetValue();
+		bool Result = false;
+		UE_AUTORTFM_OPEN({ Result = !!AsyncLoadingTickCounter.GetValue(); });
+		return Result;
 	}
 
 	FORCEINLINE bool IsAsyncLoadingPackages() override
diff --git a/Engine/Source/Runtime/CoreUObject/Private/UObject/Obj.cpp b/Engine/Source/Runtime/CoreUObject/Private/UObject/Obj.cpp
--- a/Engine/Source/Runtime/CoreUObject/Private/UObject/Obj.cpp
+++ b/Engine/Source/Runtime/CoreUObject/Private/UObject/Obj.cpp
@@ -71,6 +71,7 @@
 #include "Containers/VersePath.h"
 #include "Serialization/LoadTimeTracePrivate.h"
 #include "UObject/PropertyBagRepository.h"
+#include "AutoRTFM/AutoRTFM.h"
 
 DEFINE_LOG_CATEGORY(LogObj);
 
@@ -140,12 +141,14 @@ static UPackage*			GObjTransientPkg								= NULL;
 	UObject::FOnGetPreviewPlatform UObject::OnGetPreviewPlatform;
 #endif // WITH_EDITOR
 
+UE_AUTORTFM_ALWAYS_OPEN
 UObject::UObject( EStaticConstructor, EObjectFlags InFlags )
 : UObjectBaseUtility(InFlags | (RF_MarkAsNative | RF_MarkAsRootSet))
 {
 	EnsureNotRetrievingVTablePtr();
 }
 
+UE_AUTORTFM_ALWAYS_OPEN
 UObject::UObject(FVTableHelper& Helper)
 {
 	EnsureRetrievingVTablePtrDuringCtor(TEXT("UObject(FVTableHelper& Helper)"));
@@ -291,6 +294,7 @@ bool UObject::Rename( const TCHAR* InName, UObject* NewOuter, ERenameFlags Flags
 	bool bCreateRedirector = false;
 	UObject* OldOuter = nullptr;
 
+	AutoRTFM::Open([&]()
 	{
 		// Make sure that for the remainder of the duration of the rename operation nothing else is going to modify the UObject hash tables.
 		FScopedUObjectHashTablesLock HashTablesLock;
@@ -357,7 +361,7 @@ bool UObject::Rename( const TCHAR* InName, UObject* NewOuter, ERenameFlags Flags
 		UE::CoreUObject::Private::UpdateRenamedObject(this, NewName, NewOuter);
 #endif
 		LowLevelRename(NewName, NewOuter);
-	}
+	});
 
 	// Create the redirector AFTER renaming the object. Two objects of different classes may not have the same fully qualified name.
 	if (bCreateRedirector)
diff --git a/Engine/Source/Runtime/CoreUObject/Private/UObject/UObjectAllocator.cpp b/Engine/Source/Runtime/CoreUObject/Private/UObject/UObjectAllocator.cpp
--- a/Engine/Source/Runtime/CoreUObject/Private/UObject/UObjectAllocator.cpp
+++ b/Engine/Source/Runtime/CoreUObject/Private/UObject/UObjectAllocator.cpp
@@ -80,8 +80,12 @@ UObjectBase* FUObjectAllocator::AllocateUObject(int32 Size, int32 Alignment, boo
 	}
 	else
 	{
-		// Allocate new memory of the appropriate size and alignment.
-		Result = (UObjectBase*)FMemory::Malloc( Size, Alignment );
+		// we want to perform this allocation uninstrumented so the GC can clean this up if the transaction is aborted
+		UE_AUTORTFM_OPEN(
+		{
+			// Allocate new memory of the appropriate size and alignment.
+			Result = (UObjectBase*)FMemory::Malloc( Size, Alignment );
+		});
 	}
 
 #if !UE_BUILD_TEST && !UE_BUILD_SHIPPING
diff --git a/Engine/Source/Runtime/CoreUObject/Private/UObject/UObjectArchetype.cpp b/Engine/Source/Runtime/CoreUObject/Private/UObject/UObjectArchetype.cpp
--- a/Engine/Source/Runtime/CoreUObject/Private/UObject/UObjectArchetype.cpp
+++ b/Engine/Source/Runtime/CoreUObject/Private/UObject/UObjectArchetype.cpp
@@ -85,68 +85,71 @@ UObject* GetArchetypeFromRequiredInfoImpl(const UClass* Class, const UObject* Ou
 		if (Outer
 			&& Outer->GetClass() != UPackage::StaticClass()) // packages cannot have subobjects
 		{
-			// Get a lock on the UObject hash tables for the duration of the GetArchetype operation
-			FScopedUObjectHashTablesLock HashTablesLock;
+			AutoRTFM::Open([&]()
+			{
+				// Get a lock on the UObject hash tables for the duration of the GetArchetype operation
+				FScopedUObjectHashTablesLock HashTablesLock;
 
-			UObject* ArchetypeToSearch = nullptr;
+				UObject* ArchetypeToSearch = nullptr;
 #if UE_CACHE_ARCHETYPE
-			ArchetypeToSearch = GetArchetypeImpl(Outer, Policy);
+				ArchetypeToSearch = GetArchetypeImpl(Outer, Policy);
 #if UE_VERIFY_CACHED_ARCHETYPE
-			{
-				UObject* VerifyArchetype = GetArchetypeFromRequiredInfoImpl(Outer->GetClass(), Outer->GetOuter(), Outer->GetFName(), Outer->GetFlags(), bUseUpToDateClass, Policy);
-				checkf(ArchetypeToSearch == VerifyArchetype, TEXT("Cached archetype mismatch, expected: %s, cached: %s"), *GetFullNameSafe(VerifyArchetype), *GetFullNameSafe(ArchetypeToSearch));
-			}
+				{
+					UObject* VerifyArchetype = GetArchetypeFromRequiredInfoImpl(Outer->GetClass(), Outer->GetOuter(), Outer->GetFName(), Outer->GetFlags(), bUseUpToDateClass, Policy);
+					checkf(ArchetypeToSearch == VerifyArchetype, TEXT("Cached archetype mismatch, expected: %s, cached: %s"), *GetFullNameSafe(VerifyArchetype), *GetFullNameSafe(ArchetypeToSearch));
+				}
 #endif // UE_VERIFY_CACHED_ARCHETYPE
 #else
 			
 #if WITH_EDITOR
-			ArchetypeToSearch = Policy ? Policy->GetArchetype(Outer) : nullptr;
+				ArchetypeToSearch = Policy ? Policy->GetArchetype(Outer) : nullptr;
 #endif
 
-			if (!ArchetypeToSearch)
-			{
-				ArchetypeToSearch = GetArchetypeFromRequiredInfoImpl(Outer->GetClass(), Outer->GetOuter(), Outer->GetFName(), Outer->GetFlags(), bUseUpToDateClass, Policy);
-			}
+				if (!ArchetypeToSearch)
+				{
+					ArchetypeToSearch = GetArchetypeFromRequiredInfoImpl(Outer->GetClass(), Outer->GetOuter(), Outer->GetFName(), Outer->GetFlags(), bUseUpToDateClass, Policy);
+				}
 #endif // UE_CACHE_ARCHETYPE
-			UObject* MyArchetype = static_cast<UObject*>(FindObjectWithOuter(ArchetypeToSearch, Class, Name));
-			if (MyArchetype)
-			{
-				Result = MyArchetype; // found that my outers archetype had a matching component, that must be my archetype
-			}
-			else if (!!(ObjectFlags & RF_InheritableComponentTemplate) && Outer->IsA<UClass>())
-			{
-				const UClass* OuterSuperClass = static_cast<const UClass*>(Outer)->GetSuperClass();
-				for (const UClass* SuperClassArchetype = bUseUpToDateClass && OuterSuperClass ? OuterSuperClass->GetAuthoritativeClass() : OuterSuperClass;
-					SuperClassArchetype && SuperClassArchetype->HasAllClassFlags(CLASS_CompiledFromBlueprint);
-					SuperClassArchetype = SuperClassArchetype->GetSuperClass())
+				UObject* MyArchetype = static_cast<UObject*>(FindObjectWithOuter(ArchetypeToSearch, Class, Name));
+				if (MyArchetype)
 				{
-					if (GEventDrivenLoaderEnabled && EVENT_DRIVEN_ASYNC_LOAD_ACTIVE_AT_RUNTIME)
+					Result = MyArchetype; // found that my outers archetype had a matching component, that must be my archetype
+				}
+				else if (!!(ObjectFlags & RF_InheritableComponentTemplate) && Outer->IsA<UClass>())
+				{
+					const UClass* OuterSuperClass = static_cast<const UClass*>(Outer)->GetSuperClass();
+					for (const UClass* SuperClassArchetype = bUseUpToDateClass && OuterSuperClass ? OuterSuperClass->GetAuthoritativeClass() : OuterSuperClass;
+						SuperClassArchetype && SuperClassArchetype->HasAllClassFlags(CLASS_CompiledFromBlueprint);
+						SuperClassArchetype = SuperClassArchetype->GetSuperClass())
 					{
-						if (SuperClassArchetype->HasAnyFlags(RF_NeedLoad))
+						if (GEventDrivenLoaderEnabled && EVENT_DRIVEN_ASYNC_LOAD_ACTIVE_AT_RUNTIME)
 						{
-							UE_LOG(LogClass, Fatal, TEXT("%s had RF_NeedLoad when searching supers for an archetype of %s in %s"), *GetFullNameSafe(ArchetypeToSearch), *GetFullNameSafe(Class), *GetFullNameSafe(Outer));
+							if (SuperClassArchetype->HasAnyFlags(RF_NeedLoad))
+							{
+								UE_LOG(LogClass, Fatal, TEXT("%s had RF_NeedLoad when searching supers for an archetype of %s in %s"), *GetFullNameSafe(ArchetypeToSearch), *GetFullNameSafe(Class), *GetFullNameSafe(Outer));
+							}
+						}
+						Result = static_cast<UObject*>(FindObjectWithOuter(SuperClassArchetype, Class, Name));
+						// We can have invalid archetypes halfway through the hierarchy, keep looking if it's pending kill or transient
+						if (IsValid(Result) && !Result->HasAnyFlags(RF_Transient))
+						{
+							break;
 						}
-					}
-					Result = static_cast<UObject*>(FindObjectWithOuter(SuperClassArchetype, Class, Name));
-					// We can have invalid archetypes halfway through the hierarchy, keep looking if it's pending kill or transient
-					if (IsValid(Result) && !Result->HasAnyFlags(RF_Transient))
-					{
-						break;
 					}
 				}
-			}
-			else
-			{
-				if (GEventDrivenLoaderEnabled && EVENT_DRIVEN_ASYNC_LOAD_ACTIVE_AT_RUNTIME)
+				else
 				{
-					if (ArchetypeToSearch->HasAnyFlags(RF_NeedLoad))
+					if (GEventDrivenLoaderEnabled && EVENT_DRIVEN_ASYNC_LOAD_ACTIVE_AT_RUNTIME)
 					{
-						UE_LOG(LogClass, Fatal, TEXT("%s had RF_NeedLoad when searching for an archetype of %s in %s"), *GetFullNameSafe(ArchetypeToSearch), *GetFullNameSafe(Class), *GetFullNameSafe(Outer));
+						if (ArchetypeToSearch->HasAnyFlags(RF_NeedLoad))
+						{
+							UE_LOG(LogClass, Fatal, TEXT("%s had RF_NeedLoad when searching for an archetype of %s in %s"), *GetFullNameSafe(ArchetypeToSearch), *GetFullNameSafe(Class), *GetFullNameSafe(Outer));
+						}
 					}
-				}
 
-				Result = ArchetypeToSearch->GetClass()->FindArchetype(Class, Name);
-			}
+					Result = ArchetypeToSearch->GetClass()->FindArchetype(Class, Name);
+				}
+			});
 		}
 
 		if (!Result)
diff --git a/Engine/Source/Runtime/CoreUObject/Private/UObject/UObjectGlobals.cpp b/Engine/Source/Runtime/CoreUObject/Private/UObject/UObjectGlobals.cpp
--- a/Engine/Source/Runtime/CoreUObject/Private/UObject/UObjectGlobals.cpp
+++ b/Engine/Source/Runtime/CoreUObject/Private/UObject/UObjectGlobals.cpp
@@ -3435,16 +3435,18 @@ UObject* StaticAllocateObject
 	UObject* Obj = NULL;
 	if(InName == NAME_None)
 	{
+		AutoRTFM::Open([&]{
 #if WITH_EDITOR
-		if ( GOutputCookingWarnings && GetTransientPackage() != InOuter->GetOutermost() )
-		{
-			InName = MakeUniqueObjectName(InOuter, InClass, NAME_UniqueObjectNameForCooking);
-		}
-		else
+			if ( GOutputCookingWarnings && GetTransientPackage() != InOuter->GetOutermost() )
+			{
+				InName = MakeUniqueObjectName(InOuter, InClass, NAME_UniqueObjectNameForCooking);
+			}
+			else
 #endif
-		{
-			InName = MakeUniqueObjectName(InOuter, InClass);
-		}
+			{
+				InName = MakeUniqueObjectName(InOuter, InClass);
+			}
+		});
 	}
 	else
 	{
@@ -3603,8 +3605,13 @@ UObject* StaticAllocateObject
 
 	if (!bSubObject)
 	{
-		FMemory::Memzero((void *)Obj, TotalSize);
-		new ((void *)Obj) UObjectBase(const_cast<UClass*>(InClass), InFlags|RF_NeedInitialization, InternalSetFlags, InOuter, InName, OldIndex, OldSerialNumber);
+		// perform the UObjectBase construction in the open - we expect the GC to invoke
+		// the destructor in the case of transaction abort
+		UE_AUTORTFM_OPEN(
+		{
+			FMemory::Memzero((void *)Obj, TotalSize);
+			new ((void *)Obj) UObjectBase(const_cast<UClass*>(InClass), InFlags|RF_NeedInitialization, InternalSetFlags, InOuter, InName, OldIndex, OldSerialNumber);
+		});
 	}
 	else
 	{
@@ -3677,6 +3684,7 @@ void UObject::PostInitProperties()
 	FOverridableManager::Get().ClearOverrides(*this);
 }
 
+UE_AUTORTFM_ALWAYS_OPEN
 UObject::UObject()
 {
 	EnsureNotRetrievingVTablePtr();
@@ -3689,6 +3697,7 @@ UObject::UObject()
 	const_cast<FObjectInitializer&>(ObjectInitializer).FinalizeSubobjectClassInitialization();
 }
 
+UE_AUTORTFM_ALWAYS_OPEN
 UObject::UObject(const FObjectInitializer& ObjectInitializer)
 {
 	EnsureNotRetrievingVTablePtr();
@@ -4016,20 +4025,9 @@ void FObjectInitializer::PostConstructInit()
 	// Allow custom property initialization to happen before PostInitProperties is called
 	if (PropertyInitCallback)
 	{
-		// autortfm todo: if this transaction aborts and we are in a transaction's open nest,
-		// we need to have a way of propagating out that abort
-		if(AutoRTFM::IsTransactional())
-		{
-			AutoRTFM::EContextStatus Status = AutoRTFM::Close([&]
-			{
-				PropertyInitCallback();
-			});
-		}
-		else
-		{
-			PropertyInitCallback();
-		}
+		PropertyInitCallback();
 	}
+
 	// After the call to `PropertyInitCallback` to allow the callback to modify the instancing graph
 	if (bNeedInstancing || bNeedSubobjectInstancing)
 	{
@@ -4076,7 +4074,7 @@ void FObjectInitializer::PostConstructInit()
 		Obj->CheckDefaultSubobjects();
 	}
 
-	Obj->ClearFlags(RF_NeedInitialization);
+	UE_AUTORTFM_OPEN({ Obj->ClearFlags(RF_NeedInitialization); });
 
 	// clear the object pointer so we can guard against running this function again
 	Obj = nullptr;
diff --git a/Engine/Source/Runtime/CoreUObject/Public/UObject/DynamicallyTypedValue.h b/Engine/Source/Runtime/CoreUObject/Public/UObject/DynamicallyTypedValue.h
--- a/Engine/Source/Runtime/CoreUObject/Public/UObject/DynamicallyTypedValue.h
+++ b/Engine/Source/Runtime/CoreUObject/Public/UObject/DynamicallyTypedValue.h
@@ -62,9 +62,9 @@ namespace UE
 	{
 		static COREUOBJECT_API FDynamicallyTypedValueType& NullType();
 
-		FDynamicallyTypedValue() { InitializeToNull(); }
-		FDynamicallyTypedValue(const FDynamicallyTypedValue& Copyee) { InitializeFromCopy(Copyee); }
-		FDynamicallyTypedValue(FDynamicallyTypedValue&& Movee) { InitializeFromMove(MoveTemp(Movee)); }
+		FDynamicallyTypedValue() { UE_AUTORTFM_OPEN({ InitializeToNull(); }); }
+		FDynamicallyTypedValue(const FDynamicallyTypedValue& Copyee) { UE_AUTORTFM_OPEN({ InitializeFromCopy(Copyee); }); }
+		FDynamicallyTypedValue(FDynamicallyTypedValue&& Movee) { UE_AUTORTFM_OPEN({ InitializeFromMove(MoveTemp(Movee)); }); }
 
 		~FDynamicallyTypedValue() { Deinit(); }
 
diff --git a/Engine/Source/Runtime/CoreUObject/Public/UObject/UObjectAnnotation.h b/Engine/Source/Runtime/CoreUObject/Public/UObject/UObjectAnnotation.h
--- a/Engine/Source/Runtime/CoreUObject/Public/UObject/UObjectAnnotation.h
+++ b/Engine/Source/Runtime/CoreUObject/Public/UObject/UObjectAnnotation.h
@@ -224,21 +224,26 @@ public:
 	FORCEINLINE TAnnotation GetAnnotation(const UObjectBase *Object)
 	{
 		check(Object);
-		FScopeLock AnnotationMapLock(&AnnotationMapCritical);
-		if (Object != AnnotationCacheKey)
-		{			
-			AnnotationCacheKey = Object;
-			TAnnotation* Entry = AnnotationMap.Find(AnnotationCacheKey);
-			if (Entry)
-			{
-				AnnotationCacheValue = *Entry;
-			}
-			else
-			{
-				AnnotationCacheValue = TAnnotation();
+		TAnnotation Result;
+		UE_AUTORTFM_OPEN(
+		{
+			FScopeLock AnnotationMapLock(&AnnotationMapCritical);
+			if (Object != AnnotationCacheKey)
+			{			
+				AnnotationCacheKey = Object;
+				TAnnotation* Entry = AnnotationMap.Find(AnnotationCacheKey);
+				if (Entry)
+				{
+					AnnotationCacheValue = *Entry;
+				}
+				else
+				{
+					AnnotationCacheValue = TAnnotation();
+				}
 			}
-		}
-		return AnnotationCacheValue;
+			Result = AnnotationCacheValue;
+		});
+		return Result;
 	}
 
 	/**
diff --git a/Engine/Source/Runtime/CoreUObject/Public/UObject/UObjectArray.h b/Engine/Source/Runtime/CoreUObject/Public/UObject/UObjectArray.h
--- a/Engine/Source/Runtime/CoreUObject/Public/UObject/UObjectArray.h
+++ b/Engine/Source/Runtime/CoreUObject/Public/UObject/UObjectArray.h
@@ -160,15 +160,21 @@ public:
 	 */
 	FORCEINLINE bool ThisThreadAtomicallyClearedFlag(EInternalObjectFlags FlagToClear)
 	{
-		FlagToClear &= ~UE::GC::GReachableObjectFlag; // reachability bit can only be cleared by GC through *_ForGC functions
-		if (!!(FlagToClear & EInternalObjectFlags_RootFlags))
+		bool Result = false;
+		UE_AUTORTFM_OPEN(
 		{
-			return ClearRootFlags(FlagToClear);
-		}
-		else
-		{
-			return ThisThreadAtomicallyClearedFlag_ForGC(FlagToClear);
-		}
+			FlagToClear &= ~UE::GC::GReachableObjectFlag; // reachability bit can only be cleared by GC through *_ForGC functions
+			if (!!(FlagToClear & EInternalObjectFlags_RootFlags))
+			{
+				Result = ClearRootFlags(FlagToClear);
+			}
+			else
+			{
+				Result = ThisThreadAtomicallyClearedFlag_ForGC(FlagToClear);
+			}
+		});
+
+		return Result;
 	}
 
 	/**
@@ -203,15 +209,20 @@ public:
 	 * @return True if this call set the flag, false if it has been set by another thread.
 	 */
 	FORCEINLINE bool ThisThreadAtomicallySetFlag(EInternalObjectFlags FlagToSet)
-	{		
-		if (!!(FlagToSet & EInternalObjectFlags_RootFlags))
-		{
-			return SetRootFlags(FlagToSet);
-		}
-		else
+	{
+		bool Result = false;
+		UE_AUTORTFM_OPEN(
 		{
-			return ThisThreadAtomicallySetFlag_ForGC(FlagToSet);
-		}
+			if (!!(FlagToSet & EInternalObjectFlags_RootFlags))
+			{
+				Result = SetRootFlags(FlagToSet);
+			}
+			else
+			{
+				Result = ThisThreadAtomicallySetFlag_ForGC(FlagToSet);
+			}
+		});
+		return Result;
 	}
 
 	FORCEINLINE bool HasAnyFlags(EInternalObjectFlags InFlags) const
diff --git a/Engine/Source/Runtime/CoreUObject/Public/UObject/UObjectGlobals.h b/Engine/Source/Runtime/CoreUObject/Public/UObject/UObjectGlobals.h
--- a/Engine/Source/Runtime/CoreUObject/Public/UObject/UObjectGlobals.h
+++ b/Engine/Source/Runtime/CoreUObject/Public/UObject/UObjectGlobals.h
@@ -1782,16 +1782,7 @@ FUNCTION_NON_NULL_RETURN_END
 	Params.InstanceGraph = InInstanceGraph;
 	Params.ExternalPackage = ExternalPackage;
 
-	T* Result = nullptr;
-
-	// AutoRTFM: the idea here is for us to run the entire UObject creation as uninstrumented, including
-	// the object allocation. If our transaction gets aborted, we leave it up to the GC to realize that this
-	// object is no longer reachable and should be destroyed.
-	UE_AUTORTFM_OPEN(
-	{
-		Result = static_cast<T*>(StaticConstructObject_Internal(Params));
-	});
-
+	T* Result = static_cast<T*>(StaticConstructObject_Internal(Params));
 	return Result;
 }
 
@@ -1806,13 +1797,7 @@ FUNCTION_NON_NULL_RETURN_END
 	FStaticConstructObjectParameters Params(T::StaticClass());
 	Params.Outer = Outer;
 
-	T* Result = nullptr;
-
-	UE_AUTORTFM_OPEN(
-	{
-		Result = static_cast<T*>(StaticConstructObject_Internal(Params));
-	});
-
+	T* Result = static_cast<T*>(StaticConstructObject_Internal(Params));
 	return Result;
 }
 
@@ -1834,13 +1819,7 @@ FUNCTION_NON_NULL_RETURN_END
 	Params.bCopyTransientsFromClassDefaults = bCopyTransientsFromClassDefaults;
 	Params.InstanceGraph = InInstanceGraph;
 
-	T* Result = nullptr;
-
-	UE_AUTORTFM_OPEN(
-	{
-		Result = static_cast<T*>(StaticConstructObject_Internal(Params));
-	});
-
+	T* Result = static_cast<T*>(StaticConstructObject_Internal(Params));
 	return Result;
 }
 
