Allow TObjectPtr instances to resolve directly to placeholder-typed references when T=UObject.

Notes:
- For placeholder-typed objects, TObjectPtr<UObject> is considered to be type safe since placeholder types are based on UObject.
- This means we can avoid having to check for and indirectly resolve placeholders via the ptr's internal object handle in these cases.

diff --git a/Engine/Source/Runtime/CoreUObject/Private/UObject/LinkerLoad.cpp b/Engine/Source/Runtime/CoreUObject/Private/UObject/LinkerLoad.cpp
--- a/Engine/Source/Runtime/CoreUObject/Private/UObject/LinkerLoad.cpp
+++ b/Engine/Source/Runtime/CoreUObject/Private/UObject/LinkerLoad.cpp
@@ -6171,11 +6171,15 @@ FArchive& FLinkerLoad::operator<<( UObject*& Object )
 #if WITH_EDITOR
 	if (Object && UE::FPropertyBagRepository::IsPropertyBagPlaceholderObject(Object))
 	{
-		// This is needed because the pointer's type is checked only at compile time, which may not match the property
-		// bag placeholder object's type at runtime, and so we can't allow it to be dereferenced as the wrong base type.
-		// Note: These currently won't be discovered for replacement at reinstancing time, so it will remain set to NULL.
-		UE_LOG(LogLinker, Warning, TEXT("Serializing reference to \"%s\" as NULL to ensure type safety."), *Object->GetPathName());
-		Object = nullptr;
+		const FObjectProperty* ObjectProperty = CastField<FObjectProperty>(GetSerializedProperty());
+		if (!ObjectProperty || !Object->GetClass()->IsChildOf(ObjectProperty->PropertyClass))
+		{
+			// This is needed because the pointer's type is checked only at compile time, which may not match the property
+			// bag placeholder object's type at runtime, and so we can't allow it to be dereferenced as the wrong base type.
+			// Note: These currently won't be discovered for replacement at reinstancing time, so it will remain set to NULL.
+			UE_LOG(LogLinker, Warning, TEXT("Serializing reference to \"%s\" as NULL to ensure type safety."), *Object->GetPathName());
+			Object = nullptr;
+		}
 	}
 #endif
 	return *this;
@@ -6190,18 +6194,13 @@ FArchive& FLinkerLoad::operator<<(FObjectPtr& ObjectPtr)
 	// Wrapper that only allows pointers to exports with placeholder types when type safety features are enabled.
 	auto AsTypeSafeObjectPtr_Lambda = [](UObject* ResolvedObject)
 	{
-#if WITH_EDITOR
-		// If we can't mask the pointer to the placeholder object instance at access time, resolve it now to NULL.
+#if WITH_EDITOR && !UE_WITH_OBJECT_HANDLE_TYPE_SAFETY
+		// If type safety features are disabled, resolve unsafe references to placeholder-typed objects now to NULL.
 		// Note: Similar to hard references above, this means we won't find it for replacement at reinstancing time.
 		if (ResolvedObject && UE::FPropertyBagRepository::IsPropertyBagPlaceholderObject(ResolvedObject))
 		{
-#if UE_WITH_OBJECT_HANDLE_TYPE_SAFETY
-			// Note: With type safety enabled, unlike other instances of a placeholder type, we won't mark the CDO
-			// as also being a placeholder instance. That's because there are certain paths that need to be able
-			// to resolve the pointer (e.g. - the object initialization path during class construction). However,
-			// it also means we can't resolve other references to a placeholder CDO, as they may not be type-safe.
-			if (UNLIKELY(ResolvedObject->HasAnyFlags(RF_ClassDefaultObject)))
-#endif
+			const FObjectProperty* ObjectProperty = CastField<FObjectProperty>(GetSerializedProperty());
+			if (!ObjectProperty || !ResolvedObject->GetClass()->IsChildOf(ObjectProperty->PropertyClass))
 			{
 				UE_LOG(LogLinker, Warning, TEXT("Serializing reference to \"%s\" as NULL to ensure type safety."), *ResolvedObject->GetPathName());
 				ResolvedObject = nullptr;
diff --git a/Engine/Source/Runtime/CoreUObject/Public/UObject/ObjectPtr.h b/Engine/Source/Runtime/CoreUObject/Public/UObject/ObjectPtr.h
--- a/Engine/Source/Runtime/CoreUObject/Public/UObject/ObjectPtr.h
+++ b/Engine/Source/Runtime/CoreUObject/Public/UObject/ObjectPtr.h
@@ -102,23 +102,11 @@ public:
 
 	FORCEINLINE UObject* Get() const
 	{
-#if UE_WITH_OBJECT_HANDLE_TYPE_SAFETY
-		// Ensure the handle is resolved first (for late resolve), even if it's not considered type safe.
-		UObject* ResolvedObject = UE::CoreUObject::Private::ResolveObjectHandle(Handle);
-		return IsObjectHandleTypeSafe(Handle) ? ResolvedObject : nullptr;
-#else
 		return UE::CoreUObject::Private::ResolveObjectHandle(Handle);
-#endif
 	}
 
 	FORCEINLINE UClass* GetClass() const
 	{
-#if UE_WITH_OBJECT_HANDLE_TYPE_SAFETY
-		if (!IsObjectHandleTypeSafe(Handle))
-		{
-			return nullptr;
-		}
-#endif
 		return UE::CoreUObject::Private::ResolveObjectHandleClass(Handle);
 	}
 
@@ -190,13 +178,13 @@ public:
 	FORCEINLINE UObject& operator*() const { return *Get(); }
 
 	UE_DEPRECATED(5.1, "IsNull is deprecated, please use operator bool instead.")
-	FORCEINLINE bool IsNull() const { return IsNullNoResolve_Internal() || (UE::CoreUObject::Private::ResolveObjectHandleNoRead(Handle) == nullptr); }
+	FORCEINLINE bool IsNull() const { return UE::CoreUObject::Private::ResolveObjectHandleNoRead(Handle) == nullptr; }
 	
 	UE_DEPRECATED(5.1, "IsNullNoResolve is deprecated, please use operator bool instead.")
-	FORCEINLINE bool IsNullNoResolve() const { return IsNullNoResolve_Internal(); }
+	FORCEINLINE bool IsNullNoResolve() const { return IsObjectHandleNull(Handle); }
 	
-	FORCEINLINE bool operator!() const { return IsNullNoResolve_Internal(); }
-	explicit FORCEINLINE operator bool() const { return !IsNullNoResolve_Internal(); }
+	FORCEINLINE bool operator!() const { return IsObjectHandleNull(Handle); }
+	explicit FORCEINLINE operator bool() const { return !IsObjectHandleNull(Handle); }
 
 	FORCEINLINE bool IsResolved() const { return IsObjectHandleResolved(Handle); }
 
@@ -311,15 +299,6 @@ private:
 		}
 	}
 #endif // UE_OBJECT_PTR_GC_BARRIER
-
-	FORCEINLINE bool IsNullNoResolve_Internal() const
-	{
-#if UE_WITH_OBJECT_HANDLE_TYPE_SAFETY
-		return IsObjectHandleNull(Handle) || !IsObjectHandleTypeSafe(Handle);
-#else
-		return IsObjectHandleNull(Handle);
-#endif
-	}
 };
 
 template <typename T>
@@ -406,6 +385,65 @@ namespace ObjectPtr_Private
 		return IsObjectPtrEqualToRawPtrOfRelatedType<T>(Ptr, ObjectPtr_Private::CoerceToPointer<T>(Other));
 	}
 
+	/** Check for NULL without resolving the handle. */
+	template <
+		typename T
+#if UE_WITH_OBJECT_HANDLE_TYPE_SAFETY
+		UE_REQUIRES(std::is_same_v<T, UObject>)
+#endif
+	>
+	FORCEINLINE bool IsObjectPtrNull(const FObjectPtr& ObjectPtr)
+	{
+		return !ObjectPtr.operator bool();
+	}
+
+	/** Resolve and return the underlying reference. */
+	template <
+		typename T
+#if UE_WITH_OBJECT_HANDLE_TYPE_SAFETY
+		UE_REQUIRES(std::is_same_v<T, UObject>)
+#endif
+	>
+	FORCEINLINE T* Get(const FObjectPtr& ObjectPtr)
+	{
+		return (T*)ObjectPtr.Get();
+	}
+
+#if UE_WITH_OBJECT_HANDLE_TYPE_SAFETY
+	/** Check for NULL without resolving the handle. Always returns true if the handle is not type safe (only when T != UObject). */
+	template <
+		typename T
+		UE_REQUIRES(!std::is_same_v<T, UObject>)
+	>
+	FORCEINLINE bool IsObjectPtrNull(const FObjectPtr& ObjectPtr)
+	{
+		if (!IsObjectHandleTypeSafe(ObjectPtr.GetHandle()))
+		{
+			// Type is unsafe; this pointer will resolve to NULL.
+			return true;
+		}
+
+		return !ObjectPtr.operator bool();
+	}
+
+	/** Resolve and return the underlying reference. Always returns NULL if the handle is not type safe (only when T != UObject). */
+	template <
+		typename T
+		UE_REQUIRES(!std::is_same_v<T, UObject>)
+	>
+	FORCEINLINE T* Get(const FObjectPtr& ObjectPtr)
+	{
+		// Always resolve the pointer first (for late resolve), even if the handle is not considered type safe.
+		UObject* Result = ObjectPtr.Get();
+		if (!IsObjectHandleTypeSafe(ObjectPtr.GetHandle()))
+		{
+			Result = nullptr;
+		}
+
+		return (T*)Result;
+	}
+#endif
+
 	template <typename T, int = sizeof(T)>
 	char (&ResolveTypeIsComplete(int))[2];
 
@@ -567,7 +605,7 @@ public:
 #if UE_WITH_OBJECT_HANDLE_TYPE_SAFETY
 		// Do a NULL test first before comparing the underlying handles, in case either side is
 		// a non-NULL, but unsafe type pointer (which would equate to NULL when Get() is called).
-		return !ObjectPtr ? !Other : ObjectPtr == Other.ObjectPtr;
+		return ObjectPtr_Private::IsObjectPtrNull<T>(ObjectPtr) ? ObjectPtr_Private::IsObjectPtrNull<U>(Other.ObjectPtr) : ObjectPtr == Other.ObjectPtr;
 #else
 		return ObjectPtr == Other.ObjectPtr;
 #endif
@@ -576,7 +614,7 @@ public:
 	// Equality/Inequality comparisons against nullptr
 	FORCEINLINE bool operator==(TYPE_OF_NULLPTR) const
 	{
-		return !ObjectPtr.operator bool();
+		return ObjectPtr_Private::IsObjectPtrNull<T>(ObjectPtr);
 	}
 
 	// Equality/Inequality comparisons against another type that can be implicitly converted to the pointer type kept in a TObjectPtr
@@ -603,7 +641,7 @@ public:
 	//			a reference to the wrong type of object which we'll just send back static_casted as the wrong type.  Doing
 	//			a check or checkSlow here could catch this, but it would be better if the check could happen elsewhere that
 	//			isn't called as frequently.
-	FORCEINLINE T* Get() const { return (T*)(ObjectPtr.Get()); }
+	FORCEINLINE T* Get() const { return ObjectPtr_Private::Get<T>(ObjectPtr); }
 	FORCEINLINE UClass* GetClass() const { return ObjectPtr.GetClass(); }
 	FORCEINLINE TObjectPtr<UObject> GetOuter() const
 	{ 
@@ -629,13 +667,13 @@ public:
 	explicit FORCEINLINE operator T*& () { return GetInternalRef(); }
 
 	UE_DEPRECATED(5.1, "IsNull is deprecated, please use operator bool instead.  if (!MyObjectPtr) { ... }")
-	FORCEINLINE bool IsNull() const { return !ObjectPtr.operator bool(); }
+	FORCEINLINE bool IsNull() const { return ObjectPtr_Private::IsObjectPtrNull<T>(ObjectPtr); }
 
 	UE_DEPRECATED(5.1, "IsNullNoResolve is deprecated, please use operator bool instead.  if (!MyObjectPtr) { ... }")
-	FORCEINLINE bool IsNullNoResolve() const { return !ObjectPtr.operator bool(); }
+	FORCEINLINE bool IsNullNoResolve() const { return ObjectPtr_Private::IsObjectPtrNull<T>(ObjectPtr); }
 
-	FORCEINLINE bool operator!() const { return ObjectPtr.operator!(); }
-	explicit FORCEINLINE operator bool() const { return ObjectPtr.operator bool(); }
+	FORCEINLINE bool operator!() const { return ObjectPtr_Private::IsObjectPtrNull<T>(ObjectPtr); }
+	explicit FORCEINLINE operator bool() const { return !ObjectPtr_Private::IsObjectPtrNull<T>(ObjectPtr); }
 	FORCEINLINE bool IsResolved() const { return ObjectPtr.IsResolved(); }
 	FORCEINLINE FString GetPath() const { return ObjectPtr.GetPath(); }
 	FORCEINLINE FString GetPathName() const { return ObjectPtr.GetPathName(); }
diff --git a/Engine/Source/Runtime/CoreUObject/Tests/ObjectHandleTest.cpp b/Engine/Source/Runtime/CoreUObject/Tests/ObjectHandleTest.cpp
--- a/Engine/Source/Runtime/CoreUObject/Tests/ObjectHandleTest.cpp
+++ b/Engine/Source/Runtime/CoreUObject/Tests/ObjectHandleTest.cpp
@@ -395,7 +395,7 @@ TEST_CASE_METHOD(FObjectHandleTestBase, "CoreUObject::FObjectHandle::Type Safety
 	CHECK(UE::CoreUObject::Private::HasAnyFlags(TestClassDefaults, RF_ClassDefaultObject));
 
 	// construct objects for testing
-	UObject* TestSafeObject = NewObject<UObjectPtrTestClass>(TestPackage, TEXT("TestSafeObject"), RF_Transient);
+	UObjectPtrTestClass* TestSafeObject = NewObject<UObjectPtrTestClass>(TestPackage, TEXT("TestSafeObject"), RF_Transient);
 	UObject* TestUnsafeObject = NewObject<UObject>(TestPackage, TestClass, TEXT("TestUnsafeObject"), RF_Transient | RF_HasPlaceholderType);
 
 	// construct object handles for testing
@@ -433,10 +433,22 @@ TEST_CASE_METHOD(FObjectHandleTestBase, "CoreUObject::FObjectHandle::Type Safety
 	CHECK(TestUnsafeObjectHandle != TestSafeObjectHandle);
 	CHECK(TestUnsafeObjectHandle == TestUnsafeObjectHandle);
 
+	// simulate a placeholder object for type-safe pointer tests
+	TestSafeObject->SetFlags(RF_HasPlaceholderType);
+
 	// construct object pointers for testing
 	TObjectPtr<UObject> NullObjectPtr(nullptr);
-	TObjectPtr<UObject> TestSafeObjectPtr(TestSafeObject);
-	TObjectPtr<UObject> TestUnsafeObjectPtr(TestUnsafeObject);
+	TObjectPtr<UObject> TestSafeObjectPtr(TestSafeObject);									// type safe pointer to placeholder (UObject type)
+	FObjectPtr TestUnsafeObjectPtr_Untyped(TestUnsafeObject);
+	TObjectPtr<UObjectPtrTestClass> TestUnsafeObjectPtr(TestUnsafeObjectPtr_Untyped);		// unsafe pointer to placeholder object (non-UObject type)
+
+	// type safe object pointers should evaluate to true/non-NULL
+	CHECK(TestSafeObjectPtr);
+	CHECK(!!TestSafeObjectPtr);
+	CHECK(NULL != TestSafeObjectPtr);
+	CHECK(TestSafeObjectPtr != NULL);
+	CHECK(nullptr != TestSafeObjectPtr);
+	CHECK(TestSafeObjectPtr != nullptr);
 
 	// unsafe type object pointers should evaluate to NULL/false (for type safety)
 	CHECK(!TestUnsafeObjectPtr);
@@ -459,10 +471,11 @@ TEST_CASE_METHOD(FObjectHandleTestBase, "CoreUObject::FObjectHandle::Type Safety
 	CHECK(TestUnsafeObjectPtr.GetPathName() == TestUnsafeObject->GetPathName());
 	CHECK(TestUnsafeObjectPtr.GetFullName() == TestUnsafeObject->GetFullName());
 	CHECK(TestUnsafeObjectPtr.GetOuter() == TestUnsafeObject->GetOuter());
+	CHECK(TestUnsafeObjectPtr.GetClass() == TestUnsafeObject->GetClass());
 	CHECK(TestUnsafeObjectPtr.GetPackage() == TestUnsafeObject->GetPackage());
 
-	// an unsafe type object should not allow direct access to the underlying type
-	CHECK(TestUnsafeObjectPtr.GetClass() == nullptr);
+	// a type safe object pointer should resolve to a non-NULL value when dereferenced
+	CHECK(TestSafeObjectPtr.Get() == TestSafeObject);
 
 	// an unsafe type object pointer should resolve to NULL when dereferenced (for type safety)
 	CHECK(TestUnsafeObjectPtr.Get() == nullptr);
diff --git a/Engine/Source/Runtime/CoreUObject/Tests/ObjectPtrTest.cpp b/Engine/Source/Runtime/CoreUObject/Tests/ObjectPtrTest.cpp
--- a/Engine/Source/Runtime/CoreUObject/Tests/ObjectPtrTest.cpp
+++ b/Engine/Source/Runtime/CoreUObject/Tests/ObjectPtrTest.cpp
@@ -346,6 +346,58 @@ TEST_CASE_METHOD(FObjectPtrTestBase, "CoreUObject::TObjectPtr::Long Path", "[Cor
 }
 #endif
 
+TEST_CASE_METHOD(FObjectPtrTestBase, "CoreUObject::TObjectPtr::Type Specialization", "[CoreUObject][ObjectPtr]")
+{
+	FSnapshotObjectRefMetrics ObjectRefMetrics(*this);
+
+	const FName TestPackageName(TEXT("/Engine/Test/ObjectPtrTypeSafety/Transient"));
+	UPackage* TestPackage = NewObject<UPackage>(nullptr, TestPackageName, RF_Transient);
+	TestPackage->AddToRoot();
+	ON_SCOPE_EXIT
+	{
+		TestPackage->RemoveFromRoot();
+	};
+
+	UObject* TestAsBaseUObject = NewObject<UObjectPtrTestClass>(TestPackage);
+	UObjectPtrTestClass* TestAsDerivedUObject = NewObject<UObjectPtrTestClass>(TestPackage);
+#if UE_WITH_OBJECT_HANDLE_LATE_RESOLVE
+	FObjectPtr UntypedBaseUObjectPtr(MakeUnresolvedHandle(TestAsBaseUObject));
+	FObjectPtr UntypedDerivedUObjectPtr(MakeUnresolvedHandle(TestAsDerivedUObject));
+#else
+	FObjectPtr UntypedBaseUObjectPtr(TestAsBaseUObject);
+	FObjectPtr UntypedDerivedUObjectPtr(TestAsDerivedUObject);
+#endif
+
+	TObjectPtr<UObject> TestUObjectTypePtr(UntypedBaseUObjectPtr);
+	TObjectPtr<UObjectPtrTestClass> TestUObjectSubTypePtr(UntypedDerivedUObjectPtr);
+
+	// Evaluating the UObject base pointer type multiple times should not resolve the underlying handle.
+	TEST_TRUE(TEXT("UObject base pointer types should evaluate to true"), !!TestUObjectTypePtr);
+	TEST_TRUE(TEXT("UObject base pointer types should not equate to NULL"), TestUObjectTypePtr != nullptr);
+
+	ObjectRefMetrics.TestNumResolves(TEXT("Unexpected resolve count after UObject base pointer evaluation tests"), 0);
+
+	// Evaluating the UObject subtype pointer multiple times should also not resolve the underlying handle.
+	TEST_TRUE(TEXT("UObject subtype pointers should evaluate to true"), !!TestUObjectSubTypePtr);
+	TEST_TRUE(TEXT("UObject subtype pointers should not equate to NULL"), TestUObjectSubTypePtr != nullptr);
+
+	ObjectRefMetrics.TestNumResolves(TEXT("Unexpected resolve count after UObject subtype pointer evaluation tests"), 0);
+
+	// Dereferencing the UObject base pointer type multiple times should only increment the test's overall resolve count one time (if late resolve is enabled).
+	const UObject* ResolvedUObject = TestUObjectTypePtr;
+	TEST_TRUE(TEXT("UObject base pointer types should resolve to the original object"), ResolvedUObject == TestAsBaseUObject);
+	TEST_TRUE(TEXT("UObject base pointer types should resolve to a non-NULL reference"), Cast<UObject>(TestUObjectTypePtr) != nullptr);
+
+	ObjectRefMetrics.TestNumResolves(TEXT("Unexpected resolve count after dereferencing a UObject base pointer multiple times"), UE_WITH_OBJECT_HANDLE_LATE_RESOLVE ? 1 : 0);
+
+	// Dereferencing the UObject subtype pointer multiple times should only increment the test's overall resolve count one additional time (if late resolve is enabled).
+	const UObjectPtrTestClass* ResolvedUObjectSubType = TestUObjectSubTypePtr;
+	TEST_TRUE(TEXT("UObject subtype pointers should resolve to the original object"), ResolvedUObjectSubType == TestAsDerivedUObject);
+	TEST_TRUE(TEXT("UObject subtype pointers should resolve to a non-NULL reference"), Cast<UObjectPtrTestClass>(TestUObjectSubTypePtr) != nullptr);
+
+	ObjectRefMetrics.TestNumResolves(TEXT("Unexpected resolve count after dereferencing a UObject subtype pointer multiple times"), UE_WITH_OBJECT_HANDLE_LATE_RESOLVE ? 2 : 0);
+}
+
 void ObjectPtrStressTest(int32 NumTestNodes, bool bIsContiguousTest, bool bIsEvalOnlyTest)
 {
 	const FName TestPackageName(TEXT("/Engine/Test/ObjectPtr/EvalStressTest/Transient"));
