[BugFix] Changed SetClassFlags in InstanceDataObjectUtils.cpp to copy less flags. Notibly CLASS_Config is no longer being copied since IDOs don't have full config support.

diff --git a/Engine/Source/Runtime/CoreUObject/Private/UObject/InstanceDataObjectUtils.cpp b/Engine/Source/Runtime/CoreUObject/Private/UObject/InstanceDataObjectUtils.cpp
--- a/Engine/Source/Runtime/CoreUObject/Private/UObject/InstanceDataObjectUtils.cpp
+++ b/Engine/Source/Runtime/CoreUObject/Private/UObject/InstanceDataObjectUtils.cpp
@@ -637,7 +637,9 @@ namespace UE
 		IDOClass->ClassFlags |= CLASS_NotPlaceable | CLASS_Hidden | CLASS_HideDropDown;
 		
 		// copy flags from OwnerClass
-		IDOClass->ClassFlags |= OwnerClass->ClassFlags & (CLASS_ScriptInherit | CLASS_CompiledFromBlueprint);
+		IDOClass->ClassFlags |= OwnerClass->ClassFlags & (
+        	CLASS_EditInlineNew | CLASS_EditInlineNew | CLASS_CollapseCategories | CLASS_Const |
+        	CLASS_CompiledFromBlueprint | CLASS_HasInstancedReference);
 	}
 	
 	UClass* CreateInstanceDataObjectClass(const FPropertyBag* PropertyBag, UClass* OwnerClass, UObject* Outer)
