[BugFix] Values in ReplacedObjects param of FPropertyBagRepository::ReassociateObjects can be nullptr and cause a crash

diff --git a/Engine/Source/Runtime/CoreUObject/Private/UObject/PropertyBagRepository.cpp b/Engine/Source/Runtime/CoreUObject/Private/UObject/PropertyBagRepository.cpp
--- a/Engine/Source/Runtime/CoreUObject/Private/UObject/PropertyBagRepository.cpp
+++ b/Engine/Source/Runtime/CoreUObject/Private/UObject/PropertyBagRepository.cpp
@@ -133,11 +133,14 @@ void FPropertyBagRepository::ReassociateObjects(const TMap<UObject*, UObject*>&
 	{
 		if(AssociatedData.RemoveAndCopyValue(Pair.Key, OldBagData))
 		{
-			FPropertyBagAssociationData& NewBagData = AssociatedData.FindChecked(Pair.Value);
-			
-			CopyPropertySetBySerializationData(
-				OldBagData.InstanceDataObject->GetClass(), OldBagData.InstanceDataObject,
-				NewBagData.InstanceDataObject->GetClass(), NewBagData.InstanceDataObject);
+			if (Pair.Value != nullptr) // Pair.Value can be nullptr when an object was destroyed like for example a UClass when it's deleted
+			{
+				FPropertyBagAssociationData& NewBagData = AssociatedData.FindChecked(Pair.Value);
+				
+				CopyPropertySetBySerializationData(
+					OldBagData.InstanceDataObject->GetClass(), OldBagData.InstanceDataObject,
+					NewBagData.InstanceDataObject->GetClass(), NewBagData.InstanceDataObject);
+			}
 			OldBagData.Destroy();
 		}
 		Namespaces.Remove(Pair.Key);
