[CIS Fix] Removed duplicate use of CLASS_EditInlineNew in logical-or operation

diff --git a/Engine/Source/Runtime/CoreUObject/Private/UObject/InstanceDataObjectUtils.cpp b/Engine/Source/Runtime/CoreUObject/Private/UObject/InstanceDataObjectUtils.cpp
--- a/Engine/Source/Runtime/CoreUObject/Private/UObject/InstanceDataObjectUtils.cpp
+++ b/Engine/Source/Runtime/CoreUObject/Private/UObject/InstanceDataObjectUtils.cpp
@@ -638,8 +638,7 @@ namespace UE
 		
 		// copy flags from OwnerClass
 		IDOClass->ClassFlags |= OwnerClass->ClassFlags & (
-        	CLASS_EditInlineNew | CLASS_EditInlineNew | CLASS_CollapseCategories | CLASS_Const |
-        	CLASS_CompiledFromBlueprint | CLASS_HasInstancedReference);
+        	CLASS_EditInlineNew | CLASS_CollapseCategories | CLASS_Const | CLASS_CompiledFromBlueprint | CLASS_HasInstancedReference);
 	}
 	
 	UClass* CreateInstanceDataObjectClass(const FPropertyBag* PropertyBag, UClass* OwnerClass, UObject* Outer)
