Cleaned up a bunch of typedefs which only existed because DECLARE_FIELD didn't support types with commas.

diff --git a/Engine/Source/Runtime/CoreUObject/Private/UObject/PropertyArray.cpp b/Engine/Source/Runtime/CoreUObject/Private/UObject/PropertyArray.cpp
--- a/Engine/Source/Runtime/CoreUObject/Private/UObject/PropertyArray.cpp
+++ b/Engine/Source/Runtime/CoreUObject/Private/UObject/PropertyArray.cpp
@@ -24,7 +24,7 @@ FArrayProperty::FArrayProperty(FFieldVariant InOwner, const UECodeGen_Private::F
 
 #if WITH_EDITORONLY_DATA
 FArrayProperty::FArrayProperty(UField* InField)
-	: FArrayProperty_Super(InField)
+	: Super(InField)
 	, ArrayFlags(EArrayPropertyFlags::None)
 {
 	UArrayProperty* SourceProperty = CastChecked<UArrayProperty>(InField);
diff --git a/Engine/Source/Runtime/CoreUObject/Private/UObject/PropertyDelegate.cpp b/Engine/Source/Runtime/CoreUObject/Private/UObject/PropertyDelegate.cpp
--- a/Engine/Source/Runtime/CoreUObject/Private/UObject/PropertyDelegate.cpp
+++ b/Engine/Source/Runtime/CoreUObject/Private/UObject/PropertyDelegate.cpp
@@ -16,14 +16,14 @@
 IMPLEMENT_FIELD(FDelegateProperty)
 
 FDelegateProperty::FDelegateProperty(FFieldVariant InOwner, const UECodeGen_Private::FDelegatePropertyParams& Prop)
-	: FDelegateProperty_Super(InOwner, (const UECodeGen_Private::FPropertyParamsBaseWithOffset&)Prop)
+	: Super(InOwner, (const UECodeGen_Private::FPropertyParamsBaseWithOffset&)Prop)
 {
 	SignatureFunction = Prop.SignatureFunctionFunc ? Prop.SignatureFunctionFunc() : nullptr;
 }
 
 #if WITH_EDITORONLY_DATA
 FDelegateProperty::FDelegateProperty(UField* InField)
-	: FDelegateProperty_Super(InField)
+	: Super(InField)
 {
 	UDelegateProperty* SourceProperty = CastChecked<UDelegateProperty>(InField);
 	SignatureFunction = SourceProperty->SignatureFunction;
diff --git a/Engine/Source/Runtime/CoreUObject/Private/UObject/PropertyMap.cpp b/Engine/Source/Runtime/CoreUObject/Private/UObject/PropertyMap.cpp
--- a/Engine/Source/Runtime/CoreUObject/Private/UObject/PropertyMap.cpp
+++ b/Engine/Source/Runtime/CoreUObject/Private/UObject/PropertyMap.cpp
@@ -150,7 +150,7 @@ static const FName NAME_MapValue(ANSITEXTVIEW("Value"));
 IMPLEMENT_FIELD(FMapProperty)
 
 FMapProperty::FMapProperty(FFieldVariant InOwner, const FName& InName, EObjectFlags InObjectFlags, EMapPropertyFlags InMapFlags)
-	: FMapProperty_Super(InOwner, InName, InObjectFlags)
+	: Super(InOwner, InName, InObjectFlags)
 {
 	// These are expected to be set post-construction by AddCppProperty
 	KeyProp = nullptr;
@@ -160,7 +160,7 @@ FMapProperty::FMapProperty(FFieldVariant InOwner, const FName& InName, EObjectFl
 }
 
 FMapProperty::FMapProperty(FFieldVariant InOwner, const UECodeGen_Private::FMapPropertyParams& Prop)
-	: FMapProperty_Super(InOwner, (const UECodeGen_Private::FPropertyParamsBaseWithOffset&)Prop)
+	: Super(InOwner, (const UECodeGen_Private::FPropertyParamsBaseWithOffset&)Prop)
 {
 	// These are expected to be set post-construction by AddCppProperty
 	KeyProp = nullptr;
@@ -171,7 +171,7 @@ FMapProperty::FMapProperty(FFieldVariant InOwner, const UECodeGen_Private::FMapP
 
 #if WITH_EDITORONLY_DATA
 FMapProperty::FMapProperty(UField* InField)
-	: FMapProperty_Super(InField)
+	: Super(InField)
 	, MapFlags(EMapPropertyFlags::None)
 {
 	UMapProperty* SourceProperty = CastChecked<UMapProperty>(InField);
diff --git a/Engine/Source/Runtime/CoreUObject/Private/UObject/PropertySet.cpp b/Engine/Source/Runtime/CoreUObject/Private/UObject/PropertySet.cpp
--- a/Engine/Source/Runtime/CoreUObject/Private/UObject/PropertySet.cpp
+++ b/Engine/Source/Runtime/CoreUObject/Private/UObject/PropertySet.cpp
@@ -129,14 +129,14 @@ namespace UESetProperty_Private
 IMPLEMENT_FIELD(FSetProperty)
 
 FSetProperty::FSetProperty(FFieldVariant InOwner, const FName& InName, EObjectFlags InObjectFlags)
-	: FSetProperty_Super(InOwner, InName, InObjectFlags)
+	: Super(InOwner, InName, InObjectFlags)
 {
 	// This is expected to be set post-construction by AddCppProperty
 	ElementProp = nullptr;
 }
 
 FSetProperty::FSetProperty(FFieldVariant InOwner, const UECodeGen_Private::FSetPropertyParams& Prop)
-	: FSetProperty_Super(InOwner, (const UECodeGen_Private::FPropertyParamsBaseWithOffset&)Prop)
+	: Super(InOwner, (const UECodeGen_Private::FPropertyParamsBaseWithOffset&)Prop)
 {
 	// This is expected to be set post-construction by AddCppProperty
 	ElementProp = nullptr;
@@ -144,7 +144,7 @@ FSetProperty::FSetProperty(FFieldVariant InOwner, const UECodeGen_Private::FSetP
 
 #if WITH_EDITORONLY_DATA
 FSetProperty::FSetProperty(UField* InField)
-	: FSetProperty_Super(InField)
+	: Super(InField)
 {
 	USetProperty* SourceProperty = CastChecked<USetProperty>(InField);
 	SetLayout = SourceProperty->SetLayout;
diff --git a/Engine/Source/Runtime/CoreUObject/Private/UObject/PropertyStr.cpp b/Engine/Source/Runtime/CoreUObject/Private/UObject/PropertyStr.cpp
--- a/Engine/Source/Runtime/CoreUObject/Private/UObject/PropertyStr.cpp
+++ b/Engine/Source/Runtime/CoreUObject/Private/UObject/PropertyStr.cpp
@@ -13,7 +13,7 @@
 IMPLEMENT_FIELD(FStrProperty)
 
 FStrProperty::FStrProperty(FFieldVariant InOwner, const UECodeGen_Private::FStrPropertyParams& Prop)
-	: FStrProperty_Super(InOwner, (const UECodeGen_Private::FPropertyParamsBaseWithOffset&)Prop)
+	: Super(InOwner, (const UECodeGen_Private::FPropertyParamsBaseWithOffset&)Prop)
 {
 }
 
diff --git a/Engine/Source/Runtime/CoreUObject/Public/UObject/Field.h b/Engine/Source/Runtime/CoreUObject/Public/UObject/Field.h
--- a/Engine/Source/Runtime/CoreUObject/Public/UObject/Field.h
+++ b/Engine/Source/Runtime/CoreUObject/Public/UObject/Field.h
@@ -13,7 +13,7 @@ Field.h: Declares FField property system fundamentals
 #include "CoreTypes.h"
 #include "Delegates/Delegate.h"
 #include "HAL/PlatformMath.h"
-#include "HAL/ThreadSafeCounter.h"
+#include "HAL/PreprocessorHelpers.h"
 #include "HAL/UnrealMemory.h"
 #include "Internationalization/Text.h"
 #include "Math/RandomStream.h"
@@ -194,8 +194,8 @@ private: \
 	TClass& operator=(TClass&&);   \
 	TClass& operator=(const TClass&);   \
 public: \
-	typedef TSuperClass Super;\
-	typedef TClass ThisClass;\
+	using Super = PREPROCESSOR_REMOVE_OPTIONAL_PARENS(TSuperClass);\
+	using ThisClass = TClass;\
 	TClass(EInternal InInernal, FFieldClass* InClass) \
 		: Super(EC_InternalUseOnlyConstructor, InClass) \
 	{ \
diff --git a/Engine/Source/Runtime/CoreUObject/Public/UObject/UnrealType.h b/Engine/Source/Runtime/CoreUObject/Public/UObject/UnrealType.h
--- a/Engine/Source/Runtime/CoreUObject/Public/UObject/UnrealType.h
+++ b/Engine/Source/Runtime/CoreUObject/Public/UObject/UnrealType.h
@@ -3287,19 +3287,15 @@ public:
 //
 // Describes a dynamic string variable.
 //
-
-// need to break this out a different type so that the DECLARE_CASTED_CLASS_INTRINSIC macro can digest the comma
-typedef TProperty_WithEqualityAndSerializer<FString, FProperty> FStrProperty_Super;
-
-class COREUOBJECT_API FStrProperty : public FStrProperty_Super
+class COREUOBJECT_API FStrProperty : public TProperty_WithEqualityAndSerializer<FString, FProperty>
 {
-	DECLARE_FIELD(FStrProperty, FStrProperty_Super, CASTCLASS_FStrProperty)
+	DECLARE_FIELD(FStrProperty, (TProperty_WithEqualityAndSerializer<FString, FProperty>), CASTCLASS_FStrProperty)
 public:
-	typedef FStrProperty_Super::TTypeFundamentals TTypeFundamentals;
-	typedef TTypeFundamentals::TCppType TCppType;
+	using TTypeFundamentals = Super::TTypeFundamentals;
+	using TCppType = TTypeFundamentals::TCppType;
 
 	FStrProperty(FFieldVariant InOwner, const FName& InName, EObjectFlags InObjectFlags)
-		: FStrProperty_Super(InOwner, InName, InObjectFlags)
+		: Super(InOwner, InName, InObjectFlags)
 	{
 	}
 
@@ -3312,7 +3308,7 @@ public:
 
 #if WITH_EDITORONLY_DATA
 	explicit FStrProperty(UField* InField)
-		: FStrProperty_Super(InField)
+		: Super(InField)
 	{
 	}
 #endif // WITH_EDITORONLY_DATA
@@ -3343,14 +3339,11 @@ using FFreezableScriptArray = TScriptArray<TMemoryImageAllocator<DEFAULT_ALIGNME
 	static_assert(sizeof(FScriptArray) == sizeof(FFreezableScriptArray) && alignof(FScriptArray) == alignof(FFreezableScriptArray), "FScriptArray and FFreezableScriptArray are expected to be layout-compatible");
 #endif
 
-
-// need to break this out a different type so that the DECLARE_CASTED_CLASS_INTRINSIC macro can digest the comma
-typedef TProperty<FScriptArray, FProperty> FArrayProperty_Super;
 class FScriptArrayHelper;
 
-class COREUOBJECT_API FArrayProperty : public FArrayProperty_Super
+class COREUOBJECT_API FArrayProperty : public TProperty<FScriptArray, FProperty>
 {
-	DECLARE_FIELD(FArrayProperty, FArrayProperty_Super, CASTCLASS_FArrayProperty)
+	DECLARE_FIELD(FArrayProperty, (TProperty<FScriptArray, FProperty>), CASTCLASS_FArrayProperty)
 
 	// Variables.
 	EArrayPropertyFlags ArrayFlags;
@@ -3365,11 +3358,11 @@ public:
 		CPPAlignment = alignof(FScriptArray)
 	};
 
-	typedef FArrayProperty_Super::TTypeFundamentals TTypeFundamentals;
-	typedef TTypeFundamentals::TCppType TCppType;
+	using TTypeFundamentals = Super::TTypeFundamentals;
+	using TCppType = TTypeFundamentals::TCppType;
 
 	FArrayProperty(FFieldVariant InOwner, const FName& InName, EObjectFlags InObjectFlags, EArrayPropertyFlags InArrayPropertyFlags=EArrayPropertyFlags::None)
-		: FArrayProperty_Super(InOwner, InName, InObjectFlags)
+		: Super(InOwner, InName, InObjectFlags)
 		, Inner(nullptr)
 	{
 		ArrayFlags = InArrayPropertyFlags;
@@ -3484,12 +3477,9 @@ using FFreezableScriptMap = TScriptMap<FMemoryImageSetAllocator>;
 //@todo stever sizeof(FScriptMap) is 80 bytes, while sizeof(FFreezableScriptMap) is 56 bytes atm
 //static_assert(sizeof(FScriptMap) == sizeof(FFreezableScriptMap) && alignof(FScriptMap) == alignof(FFreezableScriptMap), "FScriptMap and FFreezableScriptMap are expected to be layout-compatible");
 
-// need to break this out a different type so that the DECLARE_CASTED_CLASS_INTRINSIC macro can digest the comma
-typedef TProperty<FScriptMap, FProperty> FMapProperty_Super;
-
-class COREUOBJECT_API FMapProperty : public FMapProperty_Super
+class COREUOBJECT_API FMapProperty : public TProperty<FScriptMap, FProperty>
 {
-	DECLARE_FIELD(FMapProperty, FMapProperty_Super, CASTCLASS_FMapProperty)
+	DECLARE_FIELD(FMapProperty, (TProperty<FScriptMap, FProperty>), CASTCLASS_FMapProperty)
 
 	// Properties representing the key type and value type of the contained pairs
 	FProperty*       KeyProp;
@@ -3511,8 +3501,8 @@ class COREUOBJECT_API FMapProperty : public FMapProperty_Super
 	}
 
 public:
-	typedef FMapProperty_Super::TTypeFundamentals TTypeFundamentals;
-	typedef TTypeFundamentals::TCppType TCppType;
+	using TTypeFundamentals = Super::TTypeFundamentals;
+	using TCppType = TTypeFundamentals::TCppType;
 
 	FMapProperty(FFieldVariant InOwner, const FName& InName, EObjectFlags InObjectFlags, EMapPropertyFlags InMapFlags=EMapPropertyFlags::None);
 
@@ -3645,20 +3635,17 @@ public:
 #endif
 };
 
-// need to break this out a different type so that the DECLARE_CASTED_CLASS_INTRINSIC macro can digest the comma
-typedef TProperty<FScriptSet, FProperty> FSetProperty_Super;
-
-class COREUOBJECT_API FSetProperty : public FSetProperty_Super
+class COREUOBJECT_API FSetProperty : public TProperty<FScriptSet, FProperty>
 {
-	DECLARE_FIELD(FSetProperty, FSetProperty_Super, CASTCLASS_FSetProperty)
+	DECLARE_FIELD(FSetProperty, (TProperty<FScriptSet, FProperty>), CASTCLASS_FSetProperty)
 
 	// Properties representing the key type and value type of the contained pairs
 	FProperty*       ElementProp;
 	FScriptSetLayout SetLayout;
 
 public:
-	typedef FSetProperty_Super::TTypeFundamentals TTypeFundamentals;
-	typedef TTypeFundamentals::TCppType TCppType;
+	using TTypeFundamentals = Super::TTypeFundamentals;
+	using TCppType = TTypeFundamentals::TCppType;
 
 	FSetProperty(FFieldVariant InOwner, const FName& InName, EObjectFlags InObjectFlags);
 
@@ -5854,22 +5841,19 @@ private:
 /**
  * Describes a pointer to a function bound to an Object.
  */
-// need to break this out a different type so that the DECLARE_CASTED_CLASS_INTRINSIC macro can digest the comma
-typedef TProperty<FScriptDelegate, FProperty> FDelegateProperty_Super;
-
-class COREUOBJECT_API FDelegateProperty : public FDelegateProperty_Super
+class COREUOBJECT_API FDelegateProperty : public TProperty<FScriptDelegate, FProperty>
 {
-	DECLARE_FIELD(FDelegateProperty, FDelegateProperty_Super, CASTCLASS_FDelegateProperty)
+	DECLARE_FIELD(FDelegateProperty, (TProperty<FScriptDelegate, FProperty>), CASTCLASS_FDelegateProperty)
 
 	/** Points to the source delegate function (the function declared with the delegate keyword) used in the declaration of this delegate property. */
 	TObjectPtr<UFunction> SignatureFunction;
 public:
 
-	typedef FDelegateProperty_Super::TTypeFundamentals TTypeFundamentals;
-	typedef TTypeFundamentals::TCppType TCppType;
+	using TTypeFundamentals = Super::TTypeFundamentals;
+	using TCppType = TTypeFundamentals::TCppType;
 
 	FDelegateProperty(FFieldVariant InOwner, const FName& InName, EObjectFlags InObjectFlags)
-		: FDelegateProperty_Super(InOwner, InName, InObjectFlags)
+		: Super(InOwner, InName, InObjectFlags)
 		, SignatureFunction(nullptr)
 	{
 	}
