World Partition: fix LinkerLoad support for instancing context none entries required when streaming in actors with editor-only references to actors, or references to editor-only actors.

diff --git a/Engine/Source/Runtime/CoreUObject/Private/UObject/LinkerLoad.cpp b/Engine/Source/Runtime/CoreUObject/Private/UObject/LinkerLoad.cpp
--- a/Engine/Source/Runtime/CoreUObject/Private/UObject/LinkerLoad.cpp
+++ b/Engine/Source/Runtime/CoreUObject/Private/UObject/LinkerLoad.cpp
@@ -5733,7 +5733,8 @@ UObject* FLinkerLoad::CreateImport( int32 Index )
 							else if( OuterImport.OuterIndex.IsNull() )
 							{
 								FName ObjectName = InstancingContext.RemapPackage(OuterImport.ObjectName);
-								UPackage* Pkg = CreatePackage(*ObjectName.ToString());
+								// Instancing context supports remapping editor-only references to none, handle this case here as well.
+								UPackage* Pkg = !ObjectName.IsNone() ? CreatePackage(*ObjectName.ToString()) : nullptr;
 								if (IsPackageReferenceAllowed(Pkg))
 								{
 									FindOuter = Pkg;
