Back out CL#26465295 - It was causing crash issues because of the set-dirty, and the rename doesn't appear necessary anymore.

diff --git a/Engine/Source/Runtime/CoreUObject/Private/UObject/PropertyBaseObject.cpp b/Engine/Source/Runtime/CoreUObject/Private/UObject/PropertyBaseObject.cpp
--- a/Engine/Source/Runtime/CoreUObject/Private/UObject/PropertyBaseObject.cpp
+++ b/Engine/Source/Runtime/CoreUObject/Private/UObject/PropertyBaseObject.cpp
@@ -718,11 +718,6 @@ void FObjectPropertyBase::CheckValidObject(void* ValueAddress, TObjectPtr<UObjec
 		bool bIsReplacingClassRefs = PropertyClass && PropertyClass->HasAnyClassFlags(CLASS_NewerVersionExists) != ObjectClass->HasAnyClassFlags(CLASS_NewerVersionExists);
 		if (!bIsReplacingClassRefs && !IsDeferringValueLoad())
 		{
-			if (UPackage* ObjectPackage = Object->GetPackage())
-			{
-				ObjectPackage->SetDirtyFlag(true);
-			}
-
 			if (!HasAnyPropertyFlags(CPF_NonNullable))
 			{
 				UE_LOG(LogProperty, Warning,
@@ -736,7 +731,6 @@ void FObjectPropertyBase::CheckValidObject(void* ValueAddress, TObjectPtr<UObjec
 			}
 			else
 			{
-				Object->Rename(nullptr, GetTransientPackage(), REN_ForceNoResetLoaders | REN_DoNotDirty | REN_DontCreateRedirectors);
 				UObject* DefaultValue = ConstructDefaultObjectValueIfNecessary(OldValue);
 
 				UE_LOG(LogProperty, Warning,
