[BugFix] changes made to IDOs in the details panel correctly mark the package dirty and save now.
[Refactor] IDOs and their Instances now share the same override set
[Feature] added FPropertyBagRepository::FindOwnerForInstanceDataObject method for reverse lookup
[Refactor] chanced some uses of UObjectBase* to UObject* in PropertyBagRespository to simplify code

diff --git a/Engine/Plugins/Experimental/InstanceDataObjectFixupTool/Source/Private/Tests/InstanceDataObjectTests.cpp b/Engine/Plugins/Experimental/InstanceDataObjectFixupTool/Source/Private/Tests/InstanceDataObjectTests.cpp
--- a/Engine/Plugins/Experimental/InstanceDataObjectFixupTool/Source/Private/Tests/InstanceDataObjectTests.cpp
+++ b/Engine/Plugins/Experimental/InstanceDataObjectFixupTool/Source/Private/Tests/InstanceDataObjectTests.cpp
@@ -310,7 +310,7 @@ static UObject* CreateInstanceDataObjectWithSubObjects(UObject* Owner)
 {
 	TArray<UObject*> SubObjects = {Owner};
 	GetObjectsWithOuter(Owner, SubObjects);
-	for (const UObject* SubObject : SubObjects)
+	for (UObject* SubObject : SubObjects)
 	{
 		UE::FPropertyBagRepository::Get().CreateInstanceDataObject(SubObject);
 	}
diff --git a/Engine/Source/Runtime/CoreUObject/Internal/UObject/PropertyBagRepository.h b/Engine/Source/Runtime/CoreUObject/Internal/UObject/PropertyBagRepository.h
--- a/Engine/Source/Runtime/CoreUObject/Internal/UObject/PropertyBagRepository.h
+++ b/Engine/Source/Runtime/CoreUObject/Internal/UObject/PropertyBagRepository.h
@@ -7,7 +7,6 @@
 #include "HAL/CriticalSection.h"
 #include "UObject/GCObject.h"
 
-class UObjectBase;
 class UObject;
 
 namespace UE
@@ -41,14 +40,12 @@ private:
 
 	/** Map of objects/subobjects to their top level property bag. */
 	// TODO: Currently will only exist in editor world, but could tracking per world make some sense for teardown in future? We're relying on object destruction to occur properly to free these up. 
-	TMap<const UObjectBase*, FPropertyBagAssociationData> AssociatedData;
-	// /** Map of subobject/container/struct/etc. paths (needs FPropertyBagPath) to their property bag. Might not want to track subobjects here (they'll be in ObjectToPropertyBagMap already). */
-	// TMap<FSoftObjectPath, FPropertyBag*> ObjectPathToPropertySubBagMap;
+	TMap<const UObject*, FPropertyBagAssociationData> AssociatedData;
 	
-	//TMap<const UObjectBase*, UObject*> ObjectToInstanceDataObjectMap;
+	TMap<const UObject*, const UObject*> InstanceDataObjectToOwner;
 
 	// used to make sure IDOs don't have name overlap
-	TMap<const UObjectBase*, TObjectPtr<UObject>> Namespaces;
+	TMap<const UObject*, TObjectPtr<UObject>> Namespaces;
 
 	/** Internal registry that tracks the current set of types for property bag container objects instanced as placeholders for package exports that have invalid or missing class imports on load. */
 	TUniquePtr<class FPropertyBagPlaceholderTypeRegistry> PropertyBagPlaceholderTypeRegistry;
@@ -67,7 +64,7 @@ public:
 	
 	// TODO: Restrict bag creation to actor creation and UStruct::SerializeVersionedTaggedProperties?
 	// Object owner is tracked internally
-	FPropertyBag* CreateOuterBag(const UObjectBase* Owner);
+	FPropertyBag* CreateOuterBag(const UObject* Owner);
 
 	// Future version for reworked InstanceDataObjects - track InstanceDataObject rather than bag (directly):
 	/**
@@ -76,11 +73,11 @@ public:
 	 * @param Archive		- used to read value of new archive from. Leave this set to nullptr to use the object's linker or copy Owner
 	 * @return				- Custom InstanceDataObject object, UClass derived from associated bag.
 	 */
-	COREUOBJECT_API UObject* CreateInstanceDataObject(const UObjectBase* Owner, FArchive* Archive = nullptr);
+	COREUOBJECT_API UObject* CreateInstanceDataObject(UObject* Owner, FArchive* Archive = nullptr);
 
 	// TODO: Restrict property bag  destruction to within UObject::BeginDestroy() & FPropertyBagProperty destructor.
 	// Removes bag, InstanceDataObject, and all associated data for this object.
-	void DestroyOuterBag(const UObjectBase* Owner);
+	void DestroyOuterBag(const UObject* Owner);
 
 	/**
 	 * ReassociateObjects
@@ -99,16 +96,18 @@ public:
 	 * @param Object	- Object to test.
 	 * @return			- Does the object's InstanceDataObject contain any loose properties requiring user fixup before the object may be published?
 	 */
-	COREUOBJECT_API bool RequiresFixup(const UObjectBase* Object) const;
+	COREUOBJECT_API bool RequiresFixup(const UObject* Object) const;
 	
 	// Accessors
-	COREUOBJECT_API bool HasBag(const UObjectBase* Owner) const;
-	COREUOBJECT_API FPropertyBag* FindBag(const UObjectBase* Owner);
-	COREUOBJECT_API const FPropertyBag* FindBag(const UObjectBase* Owner) const;
+	COREUOBJECT_API bool HasBag(const UObject* Owner) const;
+	COREUOBJECT_API FPropertyBag* FindBag(const UObject* Owner);
+	COREUOBJECT_API const FPropertyBag* FindBag(const UObject* Owner) const;
 
-	COREUOBJECT_API bool HasInstanceDataObject(const UObjectBase* Owner) const;
-	COREUOBJECT_API UObject* FindInstanceDataObject(const UObjectBase* Owner);
-	COREUOBJECT_API const UObject* FindInstanceDataObject(const UObjectBase* Owner) const;
+	COREUOBJECT_API bool HasInstanceDataObject(const UObject* Owner) const;
+	COREUOBJECT_API UObject* FindInstanceDataObject(const UObject* Owner);
+	COREUOBJECT_API const UObject* FindInstanceDataObject(const UObject* Owner) const;
+
+	COREUOBJECT_API const UObject* FindInstanceForDataObject(const UObject* InstanceDataObject) const;
 	
 	// query whether a property in an object was set when the object was deserialized
 	COREUOBJECT_API static bool WasPropertySetBySerialization(UObject* Object, const FPropertyPathName& Path);
@@ -144,10 +143,10 @@ private:
 	// Internal functions requiring the repository to be locked before being called
 
 	// Delete owner reference and disassociate all data. Returns success.
-	bool RemoveAssociationUnsafe(const UObjectBase* Owner);
+	bool RemoveAssociationUnsafe(const UObject* Owner);
 	
 	// Instantiate InstanceDataObject within BagData. Returns InstanceDataObject object. 
-	void CreateInstanceDataObjectUnsafe(const UObjectBase* Owner, FPropertyBagAssociationData& BagData, FArchive* Archive = nullptr);
+	void CreateInstanceDataObjectUnsafe(UObject* Owner, FPropertyBagAssociationData& BagData, FArchive* Archive = nullptr);
 };
 
 } // UE
diff --git a/Engine/Source/Runtime/CoreUObject/Private/UObject/InstanceDataObjectUtils.cpp b/Engine/Source/Runtime/CoreUObject/Private/UObject/InstanceDataObjectUtils.cpp
--- a/Engine/Source/Runtime/CoreUObject/Private/UObject/InstanceDataObjectUtils.cpp
+++ b/Engine/Source/Runtime/CoreUObject/Private/UObject/InstanceDataObjectUtils.cpp
@@ -12,6 +12,7 @@
 #include "UObject/UnrealType.h"
 
 static const FName NAME_ValuesSetBySerialization(ANSITEXTVIEW("_ValuesSetBySerialization"));
+static const FName NAME_InstanceDataObject(ANSITEXTVIEW("InstanceDataObject"));
 
 /** Type used for InstanceDataObject structs to provide support for hashing and custom guids. */
 class UInstanceDataObjectStruct final : public UScriptStruct
@@ -684,6 +685,7 @@ namespace UE
 		{
 			Result->SetMetaData(TEXT("DisplayName"), *DisplayName);
 		}
+		Result->SetMetaData(NAME_InstanceDataObject, TEXT("true"));
 #endif
 
 		SetClassFlags(Result, OwnerClass);
diff --git a/Engine/Source/Runtime/CoreUObject/Private/UObject/PropertyBagRepository.cpp b/Engine/Source/Runtime/CoreUObject/Private/UObject/PropertyBagRepository.cpp
--- a/Engine/Source/Runtime/CoreUObject/Private/UObject/PropertyBagRepository.cpp
+++ b/Engine/Source/Runtime/CoreUObject/Private/UObject/PropertyBagRepository.cpp
@@ -12,7 +12,6 @@
 #include "UObject/UObjectThreadContext.h"
 #include "UObject/LinkerLoad.h"
 #include "UObject/InstanceDataObjectUtils.h"
-#include "UObject/OverridableManager.h"
 #include "UObject/Package.h"
 
 #if WITH_EDITOR
@@ -123,6 +122,16 @@ FPropertyBagRepository& FPropertyBagRepository::Get()
 FPropertyBagRepository::FPropertyBagRepository()
 {
 	PropertyBagPlaceholderTypeRegistry = MakeUnique<FPropertyBagPlaceholderTypeRegistry>();
+#if WITH_EDITOR
+	FCoreUObjectDelegates::OnObjectModified.AddLambda([](const UObject* Object)
+	{
+		// if this object is an InstanceDataObject, modify it's owner as well
+		if (const UObject* Owner = Get().FindInstanceForDataObject(Object))
+		{
+			const_cast<UObject*>(Owner)->Modify();
+		}
+	});
+#endif
 }
 
 void FPropertyBagRepository::ReassociateObjects(const TMap<UObject*, UObject*>& ReplacedObjects)
@@ -133,10 +142,13 @@ void FPropertyBagRepository::ReassociateObjects(const TMap<UObject*, UObject*>&
 	{
 		if(AssociatedData.RemoveAndCopyValue(Pair.Key, OldBagData))
 		{
+			InstanceDataObjectToOwner.Remove(OldBagData.InstanceDataObject);
 			if (Pair.Value != nullptr) // Pair.Value can be nullptr when an object was destroyed like for example a UClass when it's deleted
 			{
 				FPropertyBagAssociationData& NewBagData = AssociatedData.FindChecked(Pair.Value);
 				
+				InstanceDataObjectToOwner.Add(NewBagData.InstanceDataObject, Pair.Value);
+				
 				CopyPropertySetBySerializationData(
 					OldBagData.InstanceDataObject->GetClass(), OldBagData.InstanceDataObject,
 					NewBagData.InstanceDataObject->GetClass(), NewBagData.InstanceDataObject);
@@ -159,7 +171,7 @@ void FPropertyBagRepository::CleanupLevel(const UObject* Level)
 }
 
 // TODO: Create these by class on construction?
-FPropertyBag* FPropertyBagRepository::CreateOuterBag(const UObjectBase* Owner)
+FPropertyBag* FPropertyBagRepository::CreateOuterBag(const UObject* Owner)
 {
 	FPropertyBagRepositoryLock LockRepo(this);
 	const FPropertyBagAssociationData* BagData = AssociatedData.Find(Owner);
@@ -172,7 +184,7 @@ FPropertyBag* FPropertyBagRepository::CreateOuterBag(const UObjectBase* Owner)
 	return BagData->Bag;
 }
 
-UObject* FPropertyBagRepository::CreateInstanceDataObject(const UObjectBase* Owner, FArchive* Archive)
+UObject* FPropertyBagRepository::CreateInstanceDataObject(UObject* Owner, FArchive* Archive)
 {
 	FPropertyBagRepositoryLock LockRepo(this);
 	FPropertyBagAssociationData& BagData = AssociatedData.FindOrAdd(Owner);
@@ -184,23 +196,24 @@ UObject* FPropertyBagRepository::CreateInstanceDataObject(const UObjectBase* Own
 }
 
 // TODO: Remove this? Bag destruction to be handled entirely via UObject::BeginDestroy() (+ FPropertyBagProperty destructor)?
-void FPropertyBagRepository::DestroyOuterBag(const UObjectBase* Owner)
+void FPropertyBagRepository::DestroyOuterBag(const UObject* Owner)
 {
 	FPropertyBagRepositoryLock LockRepo(this);
 	RemoveAssociationUnsafe(Owner);
 }
 
-bool FPropertyBagRepository::RequiresFixup(const UObjectBase* Object) const
+bool FPropertyBagRepository::RequiresFixup(const UObject* Object) const
 {
 	const FPropertyBag* PropertyBag = FindBag(Object);
 	return !PropertyBag || PropertyBag->IsEmpty();
 }
 
-bool FPropertyBagRepository::RemoveAssociationUnsafe(const UObjectBase* Owner)
+bool FPropertyBagRepository::RemoveAssociationUnsafe(const UObject* Owner)
 {
 	FPropertyBagAssociationData OldData;
 	if(AssociatedData.RemoveAndCopyValue(Owner, OldData))
 	{
+		InstanceDataObjectToOwner.Remove(OldData.InstanceDataObject);
 		OldData.Destroy();
 		return true;
 	}
@@ -211,7 +224,7 @@ bool FPropertyBagRepository::RemoveAssociationUnsafe(const UObjectBase* Owner)
 	return false;
 }
 
-bool FPropertyBagRepository::HasBag(const UObjectBase* Object) const
+bool FPropertyBagRepository::HasBag(const UObject* Object) const
 {
 	// TODO: Should be consistent across all objects of a given type, so handle via TStructOpsTypeTraits or similar?
 	FPropertyBagRepositoryLock LockRepo(this);
@@ -219,37 +232,44 @@ bool FPropertyBagRepository::HasBag(const UObjectBase* Object) const
 	return FindBag(Object) != nullptr;
 }
 
-FPropertyBag* FPropertyBagRepository::FindBag(const UObjectBase* Object)
+FPropertyBag* FPropertyBagRepository::FindBag(const UObject* Object)
 {
 	FPropertyBagRepositoryLock LockRepo(this);
 	const FPropertyBagAssociationData* BagData = AssociatedData.Find(Object);
 	return BagData ? BagData->Bag : nullptr;
 }
 
-const FPropertyBag* FPropertyBagRepository::FindBag(const UObjectBase* Object) const
+const FPropertyBag* FPropertyBagRepository::FindBag(const UObject* Object) const
 {
 	return const_cast<FPropertyBagRepository*>(this)->FindBag(Object);
 }
 
-bool FPropertyBagRepository::HasInstanceDataObject(const UObjectBase* Object) const
+bool FPropertyBagRepository::HasInstanceDataObject(const UObject* Object) const
 {
 	FPropertyBagRepositoryLock LockRepo(this);
 	// May be lazily instantiated, but implied from existence of object data.
 	return AssociatedData.Contains(Object);
 }
 
-UObject* FPropertyBagRepository::FindInstanceDataObject(const UObjectBase* Object)
+UObject* FPropertyBagRepository::FindInstanceDataObject(const UObject* Object)
 {
 	FPropertyBagRepositoryLock LockRepo(this);
 	const FPropertyBagAssociationData* BagData = AssociatedData.Find(Object);
 	return BagData ? BagData->InstanceDataObject : nullptr;
 }
 
-const UObject* FPropertyBagRepository::FindInstanceDataObject(const UObjectBase* Object) const
+const UObject* FPropertyBagRepository::FindInstanceDataObject(const UObject* Object) const
 {
 	return const_cast<FPropertyBagRepository*>(this)->FindInstanceDataObject(Object);
 }
 
+const UObject* FPropertyBagRepository::FindInstanceForDataObject(const UObject* InstanceDataObject) const
+{
+	FPropertyBagRepositoryLock LockRepo(this);
+	const UObject* const* Owner = InstanceDataObjectToOwner.Find(InstanceDataObject);
+	return Owner ? *Owner : nullptr;
+}
+
 bool FPropertyBagRepository::WasPropertySetBySerialization(UObject* Object, const FPropertyPathName& Path)
 {
 	return UE::WasPropertySetBySerialization(Object, Path);
@@ -262,11 +282,11 @@ bool FPropertyBagRepository::WasPropertySetBySerialization(const UStruct* Struct
 
 void FPropertyBagRepository::AddReferencedObjects(FReferenceCollector& Collector)
 {
-	for (TPair<const UObjectBase*, FPropertyBagAssociationData>& Element : AssociatedData)
+	for (TPair<const UObject*, FPropertyBagAssociationData>& Element : AssociatedData)
 	{
 		Collector.AddReferencedObject(Element.Value.InstanceDataObject);
 	}
-	for (TPair<const UObjectBase*, TObjectPtr<UObject>>& Element : Namespaces)
+	for (TPair<const UObject*, TObjectPtr<UObject>>& Element : Namespaces)
 	{
 		Collector.AddReferencedObject(Element.Value);
 	}
@@ -279,7 +299,7 @@ FString FPropertyBagRepository::GetReferencerName() const
 	return TEXT("FPropertyBagRepository");
 }
 
-void FPropertyBagRepository::CreateInstanceDataObjectUnsafe(const UObjectBase* Owner, FPropertyBagAssociationData& BagData, FArchive* Archive)
+void FPropertyBagRepository::CreateInstanceDataObjectUnsafe(UObject* Owner, FPropertyBagAssociationData& BagData, FArchive* Archive)
 {
 	check(!BagData.InstanceDataObject);	// No repeated calls
 	const FPropertyBag* PropertyBag = BagData.Bag;
@@ -314,6 +334,7 @@ void FPropertyBagRepository::CreateInstanceDataObjectUnsafe(const UObjectBase* O
 	Params.Outer = *OuterPtr;
 	UObject* InstanceDataObjectObject = StaticConstructObject_Internal(Params);
 	BagData.InstanceDataObject = InstanceDataObjectObject;
+	InstanceDataObjectToOwner.Add(InstanceDataObjectObject, Owner);
 	
 	// setup load context to mark properties the that were set by serialization
 	FUObjectSerializeContext* LoadContext = FUObjectThreadContext::Get().GetSerializeContext();
@@ -331,13 +352,12 @@ void FPropertyBagRepository::CreateInstanceDataObjectUnsafe(const UObjectBase* O
 		Dest->GetClass()->SerializeTaggedProperties(Reader, (uint8*)Dest, Dest->GetClass(), (uint8*)Dest->GetArchetype());
 	};
 	
-	UObject* OwnerAsObject = (UObject*)Owner;
 	if (Archive)
 	{
 		// re-deserialize Owner but redirect it into the IDO instead using impersonation
-		OwnerAsObject->Serialize(*Archive);
+		Owner->Serialize(*Archive);
 	}
-	else if (FLinkerLoad* Linker = OwnerAsObject->GetLinker())
+	else if (FLinkerLoad* Linker = Owner->GetLinker())
 	{
 		const FDelegateHandle OnTaggedPropertySerializeHandle = LoadContext->OnTaggedPropertySerialize.AddLambda(
 			[&BagData](const FUObjectSerializeContext& Context)
@@ -351,18 +371,18 @@ void FPropertyBagRepository::CreateInstanceDataObjectUnsafe(const UObjectBase* O
 
 		// TODO: @jordan.hoffmann - this is very inefficient! We should remove this call to Preload. To do so, we'd need to change MarkPropertySetBySerialization
 		// to cache the serialized property list in the property bag instead of the structs. We'd also need to copy the property bag values to the IDO
-		OwnerAsObject->SetFlags(RF_NeedLoad);
-		Linker->Preload(OwnerAsObject);
+		Owner->SetFlags(RF_NeedLoad);
+		Linker->Preload(Owner);
 		LoadContext->OnTaggedPropertySerialize.Remove(OnTaggedPropertySerializeHandle);
 		
 		// copy data from owner to IDO
-		CopyTaggedProperties(OwnerAsObject, BagData.InstanceDataObject);
+		CopyTaggedProperties(Owner, BagData.InstanceDataObject);
 	}
 	else
 	{
 		ensureMsgf(BagData.Bag == nullptr, TEXT("Linker missing when generating IDO for an object with loose properties. Loose properties will be lost"));
 		// copy data from owner to IDO
-		CopyTaggedProperties(OwnerAsObject, BagData.InstanceDataObject);
+		CopyTaggedProperties(Owner, BagData.InstanceDataObject);
 	}
 }
 
@@ -371,6 +391,7 @@ void FPropertyBagRepository::ShrinkMaps()
 {
 	FPropertyBagRepositoryLock LockRepo(this);
 	AssociatedData.Compact();
+	InstanceDataObjectToOwner.Compact();
 }
 
 bool FPropertyBagRepository::IsPropertyBagPlaceholderType(UStruct* Type)
diff --git a/Engine/Source/Runtime/CoreUObject/Public/UObject/OverridableManager.h b/Engine/Source/Runtime/CoreUObject/Public/UObject/OverridableManager.h
--- a/Engine/Source/Runtime/CoreUObject/Public/UObject/OverridableManager.h
+++ b/Engine/Source/Runtime/CoreUObject/Public/UObject/OverridableManager.h
@@ -5,6 +5,7 @@
 #include "UObject/Object.h"
 #include "UObject/UObjectAnnotation.h"
 #include "UObject/OverriddenPropertySet.h"
+#include "UObject/PropertyBagRepository.h"
 
 /*
  *************************************************************************************
@@ -38,28 +39,45 @@ public:
 
 	FOverriddenPropertySet* Find(const UObject& Object)
 	{
-		FOverriddenPropertyAnnotation Annotation = GetAnnotation(&Object);
+		const UObjectBase* KeyObject = GetKeyObject(Object);
+		FOverriddenPropertyAnnotation Annotation = GetAnnotation(KeyObject);
 		return Annotation.OverriddenProperties.Get();
 	}
 
 	FOverriddenPropertySet& FindChecked(const UObject& Object)
 	{
-		FOverriddenPropertyAnnotation Annotation = GetAnnotation(&Object);
+		const UObjectBase* KeyObject = GetKeyObject(Object);
+		FOverriddenPropertyAnnotation Annotation = GetAnnotation(KeyObject);
 		checkf(!Annotation.IsDefault(), TEXT("Caller is expecting the object to have overridable serialization enabled"));
 		return *Annotation.OverriddenProperties.Get();
 	}
 
 	FOverriddenPropertySet& FindOrAdd(UObject& Object)
 	{
-		FOverriddenPropertyAnnotation Annotation = GetAnnotation(&Object);
+		UObject* KeyObject = GetKeyObject(Object);
+		FOverriddenPropertyAnnotation Annotation = GetAnnotation(KeyObject);
 		if (Annotation.IsDefault())
 		{
-			Annotation.OverriddenProperties = MakeShared<FOverriddenPropertySet>(Object);
+			Annotation.OverriddenProperties = MakeShared<FOverriddenPropertySet>(*KeyObject);
 			AddAnnotation(&Object, Annotation);
 		}
 
 		return *Annotation.OverriddenProperties;
 	}
+private:
+	static UObject* GetKeyObject(const UObject& Object)
+	{
+		const UObject* Found = nullptr;
+#if WITH_EDITOR
+		static const FName NAME_InstanceDataObject(ANSITEXTVIEW("InstanceDataObject"));
+		const UClass* Class = Object.GetClass();
+		if (Class && !Class->HasAnyClassFlags(EClassFlags::CLASS_Native) && Class->GetBoolMetaData(NAME_InstanceDataObject))
+		{
+			Found = UE::FPropertyBagRepository::Get().FindInstanceForDataObject(&Object);
+		}
+#endif
+		return const_cast<UObject*>(Found ? Found : &Object);
+	}
 };
 
 enum class EOverriddenState : uint8
