[Backout] - CL32554750
Original CL Desc
-----------------------------------------------------------------
[BugFix] changes made to IDOs in the details panel correctly mark the package dirty and save now.
[Refactor] IDOs and their Instances now share the same override set
[Feature] added FPropertyBagRepository::FindOwnerForInstanceDataObject method for reverse lookup
[Refactor] chanced some uses of UObjectBase* to UObject* in PropertyBagRespository to simplify code

diff --git a/Engine/Plugins/Experimental/InstanceDataObjectFixupTool/Source/Private/Tests/InstanceDataObjectTests.cpp b/Engine/Plugins/Experimental/InstanceDataObjectFixupTool/Source/Private/Tests/InstanceDataObjectTests.cpp
--- a/Engine/Plugins/Experimental/InstanceDataObjectFixupTool/Source/Private/Tests/InstanceDataObjectTests.cpp
+++ b/Engine/Plugins/Experimental/InstanceDataObjectFixupTool/Source/Private/Tests/InstanceDataObjectTests.cpp
@@ -310,7 +310,7 @@ static UObject* CreateInstanceDataObjectWithSubObjects(UObject* Owner)
 {
 	TArray<UObject*> SubObjects = {Owner};
 	GetObjectsWithOuter(Owner, SubObjects);
-	for (UObject* SubObject : SubObjects)
+	for (const UObject* SubObject : SubObjects)
 	{
 		UE::FPropertyBagRepository::Get().CreateInstanceDataObject(SubObject);
 	}
diff --git a/Engine/Source/Runtime/CoreUObject/Internal/UObject/PropertyBagRepository.h b/Engine/Source/Runtime/CoreUObject/Internal/UObject/PropertyBagRepository.h
--- a/Engine/Source/Runtime/CoreUObject/Internal/UObject/PropertyBagRepository.h
+++ b/Engine/Source/Runtime/CoreUObject/Internal/UObject/PropertyBagRepository.h
@@ -7,6 +7,7 @@
 #include "HAL/CriticalSection.h"
 #include "UObject/GCObject.h"
 
+class UObjectBase;
 class UObject;
 
 namespace UE
@@ -40,12 +41,14 @@ private:
 
 	/** Map of objects/subobjects to their top level property bag. */
 	// TODO: Currently will only exist in editor world, but could tracking per world make some sense for teardown in future? We're relying on object destruction to occur properly to free these up. 
-	TMap<const UObject*, FPropertyBagAssociationData> AssociatedData;
+	TMap<const UObjectBase*, FPropertyBagAssociationData> AssociatedData;
+	// /** Map of subobject/container/struct/etc. paths (needs FPropertyBagPath) to their property bag. Might not want to track subobjects here (they'll be in ObjectToPropertyBagMap already). */
+	// TMap<FSoftObjectPath, FPropertyBag*> ObjectPathToPropertySubBagMap;
 	
-	TMap<const UObject*, const UObject*> InstanceDataObjectToOwner;
+	//TMap<const UObjectBase*, UObject*> ObjectToInstanceDataObjectMap;
 
 	// used to make sure IDOs don't have name overlap
-	TMap<const UObject*, TObjectPtr<UObject>> Namespaces;
+	TMap<const UObjectBase*, TObjectPtr<UObject>> Namespaces;
 
 	/** Internal registry that tracks the current set of types for property bag container objects instanced as placeholders for package exports that have invalid or missing class imports on load. */
 	TUniquePtr<class FPropertyBagPlaceholderTypeRegistry> PropertyBagPlaceholderTypeRegistry;
@@ -64,7 +67,7 @@ public:
 	
 	// TODO: Restrict bag creation to actor creation and UStruct::SerializeVersionedTaggedProperties?
 	// Object owner is tracked internally
-	FPropertyBag* CreateOuterBag(const UObject* Owner);
+	FPropertyBag* CreateOuterBag(const UObjectBase* Owner);
 
 	// Future version for reworked InstanceDataObjects - track InstanceDataObject rather than bag (directly):
 	/**
@@ -73,11 +76,11 @@ public:
 	 * @param Archive		- used to read value of new archive from. Leave this set to nullptr to use the object's linker or copy Owner
 	 * @return				- Custom InstanceDataObject object, UClass derived from associated bag.
 	 */
-	COREUOBJECT_API UObject* CreateInstanceDataObject(UObject* Owner, FArchive* Archive = nullptr);
+	COREUOBJECT_API UObject* CreateInstanceDataObject(const UObjectBase* Owner, FArchive* Archive = nullptr);
 
 	// TODO: Restrict property bag  destruction to within UObject::BeginDestroy() & FPropertyBagProperty destructor.
 	// Removes bag, InstanceDataObject, and all associated data for this object.
-	void DestroyOuterBag(const UObject* Owner);
+	void DestroyOuterBag(const UObjectBase* Owner);
 
 	/**
 	 * ReassociateObjects
@@ -96,18 +99,16 @@ public:
 	 * @param Object	- Object to test.
 	 * @return			- Does the object's InstanceDataObject contain any loose properties requiring user fixup before the object may be published?
 	 */
-	COREUOBJECT_API bool RequiresFixup(const UObject* Object) const;
+	COREUOBJECT_API bool RequiresFixup(const UObjectBase* Object) const;
 	
 	// Accessors
-	COREUOBJECT_API bool HasBag(const UObject* Owner) const;
-	COREUOBJECT_API FPropertyBag* FindBag(const UObject* Owner);
-	COREUOBJECT_API const FPropertyBag* FindBag(const UObject* Owner) const;
+	COREUOBJECT_API bool HasBag(const UObjectBase* Owner) const;
+	COREUOBJECT_API FPropertyBag* FindBag(const UObjectBase* Owner);
+	COREUOBJECT_API const FPropertyBag* FindBag(const UObjectBase* Owner) const;
 
-	COREUOBJECT_API bool HasInstanceDataObject(const UObject* Owner) const;
-	COREUOBJECT_API UObject* FindInstanceDataObject(const UObject* Owner);
-	COREUOBJECT_API const UObject* FindInstanceDataObject(const UObject* Owner) const;
-
-	COREUOBJECT_API const UObject* FindInstanceForDataObject(const UObject* InstanceDataObject) const;
+	COREUOBJECT_API bool HasInstanceDataObject(const UObjectBase* Owner) const;
+	COREUOBJECT_API UObject* FindInstanceDataObject(const UObjectBase* Owner);
+	COREUOBJECT_API const UObject* FindInstanceDataObject(const UObjectBase* Owner) const;
 	
 	// query whether a property in an object was set when the object was deserialized
 	COREUOBJECT_API static bool WasPropertySetBySerialization(UObject* Object, const FPropertyPathName& Path);
@@ -143,10 +144,10 @@ private:
 	// Internal functions requiring the repository to be locked before being called
 
 	// Delete owner reference and disassociate all data. Returns success.
-	bool RemoveAssociationUnsafe(const UObject* Owner);
+	bool RemoveAssociationUnsafe(const UObjectBase* Owner);
 	
 	// Instantiate InstanceDataObject within BagData. Returns InstanceDataObject object. 
-	void CreateInstanceDataObjectUnsafe(UObject* Owner, FPropertyBagAssociationData& BagData, FArchive* Archive = nullptr);
+	void CreateInstanceDataObjectUnsafe(const UObjectBase* Owner, FPropertyBagAssociationData& BagData, FArchive* Archive = nullptr);
 };
 
 } // UE
diff --git a/Engine/Source/Runtime/CoreUObject/Private/UObject/InstanceDataObjectUtils.cpp b/Engine/Source/Runtime/CoreUObject/Private/UObject/InstanceDataObjectUtils.cpp
--- a/Engine/Source/Runtime/CoreUObject/Private/UObject/InstanceDataObjectUtils.cpp
+++ b/Engine/Source/Runtime/CoreUObject/Private/UObject/InstanceDataObjectUtils.cpp
@@ -12,7 +12,6 @@
 #include "UObject/UnrealType.h"
 
 static const FName NAME_ValuesSetBySerialization(ANSITEXTVIEW("_ValuesSetBySerialization"));
-static const FName NAME_InstanceDataObject(ANSITEXTVIEW("InstanceDataObject"));
 
 /** Type used for InstanceDataObject structs to provide support for hashing and custom guids. */
 class UInstanceDataObjectStruct final : public UScriptStruct
@@ -685,7 +684,6 @@ namespace UE
 		{
 			Result->SetMetaData(TEXT("DisplayName"), *DisplayName);
 		}
-		Result->SetMetaData(NAME_InstanceDataObject, TEXT("true"));
 #endif
 
 		SetClassFlags(Result, OwnerClass);
diff --git a/Engine/Source/Runtime/CoreUObject/Private/UObject/PropertyBagRepository.cpp b/Engine/Source/Runtime/CoreUObject/Private/UObject/PropertyBagRepository.cpp
--- a/Engine/Source/Runtime/CoreUObject/Private/UObject/PropertyBagRepository.cpp
+++ b/Engine/Source/Runtime/CoreUObject/Private/UObject/PropertyBagRepository.cpp
@@ -12,6 +12,7 @@
 #include "UObject/UObjectThreadContext.h"
 #include "UObject/LinkerLoad.h"
 #include "UObject/InstanceDataObjectUtils.h"
+#include "UObject/OverridableManager.h"
 #include "UObject/Package.h"
 
 #if WITH_EDITOR
@@ -122,16 +123,6 @@ FPropertyBagRepository& FPropertyBagRepository::Get()
 FPropertyBagRepository::FPropertyBagRepository()
 {
 	PropertyBagPlaceholderTypeRegistry = MakeUnique<FPropertyBagPlaceholderTypeRegistry>();
-#if WITH_EDITOR
-	FCoreUObjectDelegates::OnObjectModified.AddLambda([](const UObject* Object)
-	{
-		// if this object is an InstanceDataObject, modify it's owner as well
-		if (const UObject* Owner = Get().FindInstanceForDataObject(Object))
-		{
-			const_cast<UObject*>(Owner)->Modify();
-		}
-	});
-#endif
 }
 
 void FPropertyBagRepository::ReassociateObjects(const TMap<UObject*, UObject*>& ReplacedObjects)
@@ -142,13 +133,10 @@ void FPropertyBagRepository::ReassociateObjects(const TMap<UObject*, UObject*>&
 	{
 		if(AssociatedData.RemoveAndCopyValue(Pair.Key, OldBagData))
 		{
-			InstanceDataObjectToOwner.Remove(OldBagData.InstanceDataObject);
 			if (Pair.Value != nullptr) // Pair.Value can be nullptr when an object was destroyed like for example a UClass when it's deleted
 			{
 				FPropertyBagAssociationData& NewBagData = AssociatedData.FindChecked(Pair.Value);
 				
-				InstanceDataObjectToOwner.Add(NewBagData.InstanceDataObject, Pair.Value);
-				
 				CopyPropertySetBySerializationData(
 					OldBagData.InstanceDataObject->GetClass(), OldBagData.InstanceDataObject,
 					NewBagData.InstanceDataObject->GetClass(), NewBagData.InstanceDataObject);
@@ -171,7 +159,7 @@ void FPropertyBagRepository::CleanupLevel(const UObject* Level)
 }
 
 // TODO: Create these by class on construction?
-FPropertyBag* FPropertyBagRepository::CreateOuterBag(const UObject* Owner)
+FPropertyBag* FPropertyBagRepository::CreateOuterBag(const UObjectBase* Owner)
 {
 	FPropertyBagRepositoryLock LockRepo(this);
 	const FPropertyBagAssociationData* BagData = AssociatedData.Find(Owner);
@@ -184,7 +172,7 @@ FPropertyBag* FPropertyBagRepository::CreateOuterBag(const UObject* Owner)
 	return BagData->Bag;
 }
 
-UObject* FPropertyBagRepository::CreateInstanceDataObject(UObject* Owner, FArchive* Archive)
+UObject* FPropertyBagRepository::CreateInstanceDataObject(const UObjectBase* Owner, FArchive* Archive)
 {
 	FPropertyBagRepositoryLock LockRepo(this);
 	FPropertyBagAssociationData& BagData = AssociatedData.FindOrAdd(Owner);
@@ -196,24 +184,23 @@ UObject* FPropertyBagRepository::CreateInstanceDataObject(UObject* Owner, FArchi
 }
 
 // TODO: Remove this? Bag destruction to be handled entirely via UObject::BeginDestroy() (+ FPropertyBagProperty destructor)?
-void FPropertyBagRepository::DestroyOuterBag(const UObject* Owner)
+void FPropertyBagRepository::DestroyOuterBag(const UObjectBase* Owner)
 {
 	FPropertyBagRepositoryLock LockRepo(this);
 	RemoveAssociationUnsafe(Owner);
 }
 
-bool FPropertyBagRepository::RequiresFixup(const UObject* Object) const
+bool FPropertyBagRepository::RequiresFixup(const UObjectBase* Object) const
 {
 	const FPropertyBag* PropertyBag = FindBag(Object);
 	return !PropertyBag || PropertyBag->IsEmpty();
 }
 
-bool FPropertyBagRepository::RemoveAssociationUnsafe(const UObject* Owner)
+bool FPropertyBagRepository::RemoveAssociationUnsafe(const UObjectBase* Owner)
 {
 	FPropertyBagAssociationData OldData;
 	if(AssociatedData.RemoveAndCopyValue(Owner, OldData))
 	{
-		InstanceDataObjectToOwner.Remove(OldData.InstanceDataObject);
 		OldData.Destroy();
 		return true;
 	}
@@ -224,7 +211,7 @@ bool FPropertyBagRepository::RemoveAssociationUnsafe(const UObject* Owner)
 	return false;
 }
 
-bool FPropertyBagRepository::HasBag(const UObject* Object) const
+bool FPropertyBagRepository::HasBag(const UObjectBase* Object) const
 {
 	// TODO: Should be consistent across all objects of a given type, so handle via TStructOpsTypeTraits or similar?
 	FPropertyBagRepositoryLock LockRepo(this);
@@ -232,44 +219,37 @@ bool FPropertyBagRepository::HasBag(const UObject* Object) const
 	return FindBag(Object) != nullptr;
 }
 
-FPropertyBag* FPropertyBagRepository::FindBag(const UObject* Object)
+FPropertyBag* FPropertyBagRepository::FindBag(const UObjectBase* Object)
 {
 	FPropertyBagRepositoryLock LockRepo(this);
 	const FPropertyBagAssociationData* BagData = AssociatedData.Find(Object);
 	return BagData ? BagData->Bag : nullptr;
 }
 
-const FPropertyBag* FPropertyBagRepository::FindBag(const UObject* Object) const
+const FPropertyBag* FPropertyBagRepository::FindBag(const UObjectBase* Object) const
 {
 	return const_cast<FPropertyBagRepository*>(this)->FindBag(Object);
 }
 
-bool FPropertyBagRepository::HasInstanceDataObject(const UObject* Object) const
+bool FPropertyBagRepository::HasInstanceDataObject(const UObjectBase* Object) const
 {
 	FPropertyBagRepositoryLock LockRepo(this);
 	// May be lazily instantiated, but implied from existence of object data.
 	return AssociatedData.Contains(Object);
 }
 
-UObject* FPropertyBagRepository::FindInstanceDataObject(const UObject* Object)
+UObject* FPropertyBagRepository::FindInstanceDataObject(const UObjectBase* Object)
 {
 	FPropertyBagRepositoryLock LockRepo(this);
 	const FPropertyBagAssociationData* BagData = AssociatedData.Find(Object);
 	return BagData ? BagData->InstanceDataObject : nullptr;
 }
 
-const UObject* FPropertyBagRepository::FindInstanceDataObject(const UObject* Object) const
+const UObject* FPropertyBagRepository::FindInstanceDataObject(const UObjectBase* Object) const
 {
 	return const_cast<FPropertyBagRepository*>(this)->FindInstanceDataObject(Object);
 }
 
-const UObject* FPropertyBagRepository::FindInstanceForDataObject(const UObject* InstanceDataObject) const
-{
-	FPropertyBagRepositoryLock LockRepo(this);
-	const UObject* const* Owner = InstanceDataObjectToOwner.Find(InstanceDataObject);
-	return Owner ? *Owner : nullptr;
-}
-
 bool FPropertyBagRepository::WasPropertySetBySerialization(UObject* Object, const FPropertyPathName& Path)
 {
 	return UE::WasPropertySetBySerialization(Object, Path);
@@ -282,11 +262,11 @@ bool FPropertyBagRepository::WasPropertySetBySerialization(const UStruct* Struct
 
 void FPropertyBagRepository::AddReferencedObjects(FReferenceCollector& Collector)
 {
-	for (TPair<const UObject*, FPropertyBagAssociationData>& Element : AssociatedData)
+	for (TPair<const UObjectBase*, FPropertyBagAssociationData>& Element : AssociatedData)
 	{
 		Collector.AddReferencedObject(Element.Value.InstanceDataObject);
 	}
-	for (TPair<const UObject*, TObjectPtr<UObject>>& Element : Namespaces)
+	for (TPair<const UObjectBase*, TObjectPtr<UObject>>& Element : Namespaces)
 	{
 		Collector.AddReferencedObject(Element.Value);
 	}
@@ -299,7 +279,7 @@ FString FPropertyBagRepository::GetReferencerName() const
 	return TEXT("FPropertyBagRepository");
 }
 
-void FPropertyBagRepository::CreateInstanceDataObjectUnsafe(UObject* Owner, FPropertyBagAssociationData& BagData, FArchive* Archive)
+void FPropertyBagRepository::CreateInstanceDataObjectUnsafe(const UObjectBase* Owner, FPropertyBagAssociationData& BagData, FArchive* Archive)
 {
 	check(!BagData.InstanceDataObject);	// No repeated calls
 	const FPropertyBag* PropertyBag = BagData.Bag;
@@ -334,7 +314,6 @@ void FPropertyBagRepository::CreateInstanceDataObjectUnsafe(UObject* Owner, FPro
 	Params.Outer = *OuterPtr;
 	UObject* InstanceDataObjectObject = StaticConstructObject_Internal(Params);
 	BagData.InstanceDataObject = InstanceDataObjectObject;
-	InstanceDataObjectToOwner.Add(InstanceDataObjectObject, Owner);
 	
 	// setup load context to mark properties the that were set by serialization
 	FUObjectSerializeContext* LoadContext = FUObjectThreadContext::Get().GetSerializeContext();
@@ -352,12 +331,13 @@ void FPropertyBagRepository::CreateInstanceDataObjectUnsafe(UObject* Owner, FPro
 		Dest->GetClass()->SerializeTaggedProperties(Reader, (uint8*)Dest, Dest->GetClass(), (uint8*)Dest->GetArchetype());
 	};
 	
+	UObject* OwnerAsObject = (UObject*)Owner;
 	if (Archive)
 	{
 		// re-deserialize Owner but redirect it into the IDO instead using impersonation
-		Owner->Serialize(*Archive);
+		OwnerAsObject->Serialize(*Archive);
 	}
-	else if (FLinkerLoad* Linker = Owner->GetLinker())
+	else if (FLinkerLoad* Linker = OwnerAsObject->GetLinker())
 	{
 		const FDelegateHandle OnTaggedPropertySerializeHandle = LoadContext->OnTaggedPropertySerialize.AddLambda(
 			[&BagData](const FUObjectSerializeContext& Context)
@@ -371,18 +351,18 @@ void FPropertyBagRepository::CreateInstanceDataObjectUnsafe(UObject* Owner, FPro
 
 		// TODO: @jordan.hoffmann - this is very inefficient! We should remove this call to Preload. To do so, we'd need to change MarkPropertySetBySerialization
 		// to cache the serialized property list in the property bag instead of the structs. We'd also need to copy the property bag values to the IDO
-		Owner->SetFlags(RF_NeedLoad);
-		Linker->Preload(Owner);
+		OwnerAsObject->SetFlags(RF_NeedLoad);
+		Linker->Preload(OwnerAsObject);
 		LoadContext->OnTaggedPropertySerialize.Remove(OnTaggedPropertySerializeHandle);
 		
 		// copy data from owner to IDO
-		CopyTaggedProperties(Owner, BagData.InstanceDataObject);
+		CopyTaggedProperties(OwnerAsObject, BagData.InstanceDataObject);
 	}
 	else
 	{
 		ensureMsgf(BagData.Bag == nullptr, TEXT("Linker missing when generating IDO for an object with loose properties. Loose properties will be lost"));
 		// copy data from owner to IDO
-		CopyTaggedProperties(Owner, BagData.InstanceDataObject);
+		CopyTaggedProperties(OwnerAsObject, BagData.InstanceDataObject);
 	}
 }
 
@@ -391,7 +371,6 @@ void FPropertyBagRepository::ShrinkMaps()
 {
 	FPropertyBagRepositoryLock LockRepo(this);
 	AssociatedData.Compact();
-	InstanceDataObjectToOwner.Compact();
 }
 
 bool FPropertyBagRepository::IsPropertyBagPlaceholderType(UStruct* Type)
diff --git a/Engine/Source/Runtime/CoreUObject/Public/UObject/OverridableManager.h b/Engine/Source/Runtime/CoreUObject/Public/UObject/OverridableManager.h
--- a/Engine/Source/Runtime/CoreUObject/Public/UObject/OverridableManager.h
+++ b/Engine/Source/Runtime/CoreUObject/Public/UObject/OverridableManager.h
@@ -5,7 +5,6 @@
 #include "UObject/Object.h"
 #include "UObject/UObjectAnnotation.h"
 #include "UObject/OverriddenPropertySet.h"
-#include "UObject/PropertyBagRepository.h"
 
 /*
  *************************************************************************************
@@ -39,45 +38,28 @@ public:
 
 	FOverriddenPropertySet* Find(const UObject& Object)
 	{
-		const UObjectBase* KeyObject = GetKeyObject(Object);
-		FOverriddenPropertyAnnotation Annotation = GetAnnotation(KeyObject);
+		FOverriddenPropertyAnnotation Annotation = GetAnnotation(&Object);
 		return Annotation.OverriddenProperties.Get();
 	}
 
 	FOverriddenPropertySet& FindChecked(const UObject& Object)
 	{
-		const UObjectBase* KeyObject = GetKeyObject(Object);
-		FOverriddenPropertyAnnotation Annotation = GetAnnotation(KeyObject);
+		FOverriddenPropertyAnnotation Annotation = GetAnnotation(&Object);
 		checkf(!Annotation.IsDefault(), TEXT("Caller is expecting the object to have overridable serialization enabled"));
 		return *Annotation.OverriddenProperties.Get();
 	}
 
 	FOverriddenPropertySet& FindOrAdd(UObject& Object)
 	{
-		UObject* KeyObject = GetKeyObject(Object);
-		FOverriddenPropertyAnnotation Annotation = GetAnnotation(KeyObject);
+		FOverriddenPropertyAnnotation Annotation = GetAnnotation(&Object);
 		if (Annotation.IsDefault())
 		{
-			Annotation.OverriddenProperties = MakeShared<FOverriddenPropertySet>(*KeyObject);
+			Annotation.OverriddenProperties = MakeShared<FOverriddenPropertySet>(Object);
 			AddAnnotation(&Object, Annotation);
 		}
 
 		return *Annotation.OverriddenProperties;
 	}
-private:
-	static UObject* GetKeyObject(const UObject& Object)
-	{
-		const UObject* Found = nullptr;
-#if WITH_EDITOR
-		static const FName NAME_InstanceDataObject(ANSITEXTVIEW("InstanceDataObject"));
-		const UClass* Class = Object.GetClass();
-		if (Class && !Class->HasAnyClassFlags(EClassFlags::CLASS_Native) && Class->GetBoolMetaData(NAME_InstanceDataObject))
-		{
-			Found = UE::FPropertyBagRepository::Get().FindInstanceForDataObject(&Object);
-		}
-#endif
-		return const_cast<UObject*>(Found ? Found : &Object);
-	}
 };
 
 enum class EOverriddenState : uint8
