Minor optimization when serializing FBulkData during cooks.

#rnx
#test SavePackages containing bulkdata via both cook and editor save flow

- Do not add any FBulkData objects to FLinkerSave::SerializedBulkData if the serialization being performed is a cook (bUpdatingLoadedPath is false)
- FLinkerSave::SerializedBulkData only needs to exist in editor builds.
- In theory there is no reason to continue checking bUpdatingLoadedPath in FLinkerSave::OnPostSaveBulkData but to be safe I left an ensure to make sure that if bUpdatingLoadedPath is false, that we have no entries in the array. We can remove this once we are confident about the change.

diff --git a/Engine/Source/Runtime/CoreUObject/Private/UObject/LinkerSave.cpp b/Engine/Source/Runtime/CoreUObject/Private/UObject/LinkerSave.cpp
--- a/Engine/Source/Runtime/CoreUObject/Private/UObject/LinkerSave.cpp
+++ b/Engine/Source/Runtime/CoreUObject/Private/UObject/LinkerSave.cpp
@@ -781,7 +781,12 @@ bool FLinkerSave::SerializeBulkData(FBulkData& BulkData, const FBulkDataSerializ
 	DataResource.LegacyBulkDataFlags	= SerializedMeta.Flags;
 	DataResource.OuterIndex				= ObjectIndicesMap.FindRef(Params.Owner);
 
-	SerializedBulkData.Add(&BulkData, ResourceIndex);
+#if WITH_EDITOR
+	if (bUpdatingLoadedPath)
+	{
+		SerializedBulkData.Add(&BulkData, ResourceIndex);
+	}
+#endif //WITH_EDITOR
 
 	return true;
 }
@@ -789,16 +794,15 @@ bool FLinkerSave::SerializeBulkData(FBulkData& BulkData, const FBulkDataSerializ
 void FLinkerSave::OnPostSaveBulkData()
 {
 #if WITH_EDITOR
-	if (bUpdatingLoadedPath)
+	ensure(SerializedBulkData.IsEmpty() || bUpdatingLoadedPath == true);
+ 
+	for (TPair<FBulkData*, int32>& Kv : SerializedBulkData)
 	{
-		for (TPair<FBulkData*, int32>& Kv : SerializedBulkData)
-		{
-			FBulkData& BulkData = *Kv.Key;
-			const FObjectDataResource& DataResource = DataResourceMap[Kv.Value];
-			BulkData.SetFlagsFromDiskWrittenValues(static_cast<EBulkDataFlags>(DataResource.LegacyBulkDataFlags), DataResource.SerialOffset, DataResource.SerialSize, Summary.BulkDataStartOffset);
-		}
+		FBulkData& BulkData = *Kv.Key;
+		const FObjectDataResource& DataResource = DataResourceMap[Kv.Value];
+		BulkData.SetFlagsFromDiskWrittenValues(static_cast<EBulkDataFlags>(DataResource.LegacyBulkDataFlags), DataResource.SerialOffset, DataResource.SerialSize, Summary.BulkDataStartOffset);
 	}
-#endif
-
+	
 	SerializedBulkData.Empty();
+#endif //WITH_EDITOR
 }
diff --git a/Engine/Source/Runtime/CoreUObject/Public/UObject/LinkerSave.h b/Engine/Source/Runtime/CoreUObject/Public/UObject/LinkerSave.h
--- a/Engine/Source/Runtime/CoreUObject/Public/UObject/LinkerSave.h
+++ b/Engine/Source/Runtime/CoreUObject/Public/UObject/LinkerSave.h
@@ -315,8 +315,12 @@ private:
 	/** The index of the last derived data chunk added to the package. */
 	int32 LastDerivedDataIndex = -1;
 #endif
+
+#if WITH_EDITOR
 	/** Map from bulk data object to resource index. */
 	TMap<FBulkData*, int32> SerializedBulkData;
+#endif //WITH_EDITOR
+
 	/** Default bulk data archive. */
 	FFileRegionMemoryWriter BulkDataAr;
 	/** Optional bulk data archive. */
