Added FUtf8StrProperty and FAnsiStrProperty.

diff --git a/Engine/Plugins/Enterprise/VariantManager/Source/VariantManager/Private/SVariantManager.cpp b/Engine/Plugins/Enterprise/VariantManager/Source/VariantManager/Private/SVariantManager.cpp
--- a/Engine/Plugins/Enterprise/VariantManager/Source/VariantManager/Private/SVariantManager.cpp
+++ b/Engine/Plugins/Enterprise/VariantManager/Source/VariantManager/Private/SVariantManager.cpp
@@ -52,6 +52,7 @@
 #include "Widgets/Layout/SSeparator.h"
 #include "Widgets/Notifications/SNotificationList.h"
 #include "Widgets/Text/SRichTextBlock.h"
+#include "UObject/StrProperty.h"
 
 #define LOCTEXT_NAMESPACE "SVariantManager"
 
diff --git a/Engine/Plugins/Enterprise/VariantManager/Source/VariantManager/Private/VariantManagerBlueprintLibrary.cpp b/Engine/Plugins/Enterprise/VariantManager/Source/VariantManager/Private/VariantManagerBlueprintLibrary.cpp
--- a/Engine/Plugins/Enterprise/VariantManager/Source/VariantManager/Private/VariantManagerBlueprintLibrary.cpp
+++ b/Engine/Plugins/Enterprise/VariantManager/Source/VariantManager/Private/VariantManagerBlueprintLibrary.cpp
@@ -7,6 +7,7 @@
 #include "LevelVariantSetsActor.h"
 #include "Trace/Trace.inl"
 #include "Variant.h"
+#include "UObject/StrProperty.h"
 #include "UObject/TextProperty.h"
 #include "VariantManager.h"
 #include "VariantManagerContentEditorModule.h"
diff --git a/Engine/Plugins/Runtime/StateTree/Source/StateTreeModule/Public/StateTreePropertyRefHelpers.h b/Engine/Plugins/Runtime/StateTree/Source/StateTreeModule/Public/StateTreePropertyRefHelpers.h
--- a/Engine/Plugins/Runtime/StateTree/Source/StateTreeModule/Public/StateTreePropertyRefHelpers.h
+++ b/Engine/Plugins/Runtime/StateTree/Source/StateTreeModule/Public/StateTreePropertyRefHelpers.h
@@ -4,6 +4,7 @@
 #include "UObject/Field.h"
 #include "UObject/TextProperty.h"
 #include "UObject/EnumProperty.h"
+#include "UObject/StrProperty.h"
 
 #if WITH_EDITOR
 #include "Containers/ContainersFwd.h"
diff --git a/Engine/Source/Editor/BlueprintGraph/Private/EdGraphSchema_K2.cpp b/Engine/Source/Editor/BlueprintGraph/Private/EdGraphSchema_K2.cpp
--- a/Engine/Source/Editor/BlueprintGraph/Private/EdGraphSchema_K2.cpp
+++ b/Engine/Source/Editor/BlueprintGraph/Private/EdGraphSchema_K2.cpp
@@ -7,6 +7,7 @@
 #include "Modules/ModuleManager.h"
 #include "UObject/Interface.h"
 #include "UObject/UnrealType.h"
+#include "UObject/StrProperty.h"
 #include "UObject/TextProperty.h"
 #include "UObject/ConstructorHelpers.h"
 #include "Engine/Blueprint.h"
diff --git a/Engine/Source/Editor/KismetCompiler/Private/KismetCompilerMisc.cpp b/Engine/Source/Editor/KismetCompiler/Private/KismetCompilerMisc.cpp
--- a/Engine/Source/Editor/KismetCompiler/Private/KismetCompilerMisc.cpp
+++ b/Engine/Source/Editor/KismetCompiler/Private/KismetCompilerMisc.cpp
@@ -10,6 +10,7 @@
 #include "Misc/CoreMisc.h"
 #include "UObject/MetaData.h"
 #include "UObject/UnrealType.h"
+#include "UObject/StrProperty.h"
 #include "UObject/TextProperty.h"
 #include "UObject/FieldPathProperty.h"
 #include "UObject/ObjectRedirector.h"
diff --git a/Engine/Source/Editor/KismetCompiler/Private/KismetCompilerVMBackend.cpp b/Engine/Source/Editor/KismetCompiler/Private/KismetCompilerVMBackend.cpp
--- a/Engine/Source/Editor/KismetCompiler/Private/KismetCompilerVMBackend.cpp
+++ b/Engine/Source/Editor/KismetCompiler/Private/KismetCompilerVMBackend.cpp
@@ -13,6 +13,7 @@
 #include "UObject/SoftObjectPtr.h"
 #include "UObject/Interface.h"
 #include "UObject/UnrealType.h"
+#include "UObject/StrProperty.h"
 #include "UObject/TextProperty.h"
 #include "UObject/FieldPathProperty.h"
 #include "UObject/PropertyPortFlags.h"
diff --git a/Engine/Source/Editor/PropertyEditor/Private/UserInterface/PropertyEditor/SPropertyEditorCombo.cpp b/Engine/Source/Editor/PropertyEditor/Private/UserInterface/PropertyEditor/SPropertyEditorCombo.cpp
--- a/Engine/Source/Editor/PropertyEditor/Private/UserInterface/PropertyEditor/SPropertyEditorCombo.cpp
+++ b/Engine/Source/Editor/PropertyEditor/Private/UserInterface/PropertyEditor/SPropertyEditorCombo.cpp
@@ -4,6 +4,7 @@
 #include "IDocumentation.h"
 
 #include "PropertyEditorHelpers.h"
+#include "UObject/StrProperty.h"
 #include "UserInterface/PropertyEditor/SPropertyComboBox.h"
 
 
diff --git a/Engine/Source/Editor/PropertyEditor/Private/UserInterface/PropertyEditor/SPropertyEditorText.cpp b/Engine/Source/Editor/PropertyEditor/Private/UserInterface/PropertyEditor/SPropertyEditorText.cpp
--- a/Engine/Source/Editor/PropertyEditor/Private/UserInterface/PropertyEditor/SPropertyEditorText.cpp
+++ b/Engine/Source/Editor/PropertyEditor/Private/UserInterface/PropertyEditor/SPropertyEditorText.cpp
@@ -16,6 +16,7 @@
 #include "PropertyNode.h"
 #include "SlotBase.h"
 #include "UObject/NameTypes.h"
+#include "UObject/StrProperty.h"
 #include "UObject/TextProperty.h"
 #include "UObject/UnrealNames.h"
 #include "UObject/UnrealType.h"
diff --git a/Engine/Source/Editor/UnrealEd/Private/Cooker/DiffWriterArchive.cpp b/Engine/Source/Editor/UnrealEd/Private/Cooker/DiffWriterArchive.cpp
--- a/Engine/Source/Editor/UnrealEd/Private/Cooker/DiffWriterArchive.cpp
+++ b/Engine/Source/Editor/UnrealEd/Private/Cooker/DiffWriterArchive.cpp
@@ -21,6 +21,7 @@
 #include "UObject/LinkerLoad.h"
 #include "UObject/ObjectMacros.h"
 #include "UObject/Package.h"
+#include "UObject/StrProperty.h"
 #include "UObject/PropertyOptional.h"
 #include "UObject/PropertyTempVal.h"
 #include "UObject/UnrealType.h"
diff --git a/Engine/Source/Editor/UnrealEd/Private/DataTableEditorUtils.cpp b/Engine/Source/Editor/UnrealEd/Private/DataTableEditorUtils.cpp
--- a/Engine/Source/Editor/UnrealEd/Private/DataTableEditorUtils.cpp
+++ b/Engine/Source/Editor/UnrealEd/Private/DataTableEditorUtils.cpp
@@ -2,6 +2,7 @@
 
 #include "DataTableEditorUtils.h"
 #include "UObject/UObjectHash.h"
+#include "UObject/StrProperty.h"
 #include "UObject/UObjectIterator.h"
 #include "Styling/SlateTypes.h"
 #include "Fonts/FontMeasure.h"
diff --git a/Engine/Source/Editor/UnrealEd/Private/Kismet2/BlueprintEditorUtils.cpp b/Engine/Source/Editor/UnrealEd/Private/Kismet2/BlueprintEditorUtils.cpp
--- a/Engine/Source/Editor/UnrealEd/Private/Kismet2/BlueprintEditorUtils.cpp
+++ b/Engine/Source/Editor/UnrealEd/Private/Kismet2/BlueprintEditorUtils.cpp
@@ -19,6 +19,7 @@
 #include "UObject/UObjectIterator.h"
 #include "UObject/StructOnScope.h"
 #include "UObject/MetaData.h"
+#include "UObject/StrProperty.h"
 #include "UObject/TextProperty.h"
 #include "Widgets/DeclarativeSyntaxSupport.h"
 #include "Styling/CoreStyle.h"
diff --git a/Engine/Source/Editor/UnrealEd/Private/Kismet2/KismetDebugUtilities.cpp b/Engine/Source/Editor/UnrealEd/Private/Kismet2/KismetDebugUtilities.cpp
--- a/Engine/Source/Editor/UnrealEd/Private/Kismet2/KismetDebugUtilities.cpp
+++ b/Engine/Source/Editor/UnrealEd/Private/Kismet2/KismetDebugUtilities.cpp
@@ -4,6 +4,7 @@
 #include "Engine/BlueprintGeneratedClass.h"
 #include "GameFramework/Actor.h"
 #include "UObject/PropertyPortFlags.h"
+#include "UObject/StrProperty.h"
 #include "UObject/TextProperty.h"
 #include "Widgets/DeclarativeSyntaxSupport.h"
 #include "Widgets/SWidget.h"
diff --git a/Engine/Source/Editor/UnrealEd/Private/Kismet2/StructureEditorUtils.cpp b/Engine/Source/Editor/UnrealEd/Private/Kismet2/StructureEditorUtils.cpp
--- a/Engine/Source/Editor/UnrealEd/Private/Kismet2/StructureEditorUtils.cpp
+++ b/Engine/Source/Editor/UnrealEd/Private/Kismet2/StructureEditorUtils.cpp
@@ -6,6 +6,7 @@
 #include "Modules/ModuleManager.h"
 #include "UObject/UObjectHash.h"
 #include "UObject/UnrealType.h"
+#include "UObject/StrProperty.h"
 #include "UObject/TextProperty.h"
 #include "Engine/Blueprint.h"
 #include "Engine/DataTable.h"
diff --git a/Engine/Source/Runtime/Core/Public/UObject/UObjectHierarchyFwd.h b/Engine/Source/Runtime/Core/Public/UObject/UObjectHierarchyFwd.h
--- a/Engine/Source/Runtime/Core/Public/UObject/UObjectHierarchyFwd.h
+++ b/Engine/Source/Runtime/Core/Public/UObject/UObjectHierarchyFwd.h
@@ -52,3 +52,5 @@ class		FMulticastDelegateProperty;
 class		FMapProperty;
 class		FSetProperty;
 class		FEnumProperty;
+class		FUtf8StrProperty;
+class		FAnsiStrProperty;
diff --git a/Engine/Source/Runtime/Core/Public/UObject/UnrealNames.inl b/Engine/Source/Runtime/Core/Public/UObject/UnrealNames.inl
--- a/Engine/Source/Runtime/Core/Public/UObject/UnrealNames.inl
+++ b/Engine/Source/Runtime/Core/Public/UObject/UnrealNames.inl
@@ -51,6 +51,8 @@ REGISTER_NAME(33,CoreUObject)
 // More class properties
 REGISTER_NAME(34,EnumProperty)
 REGISTER_NAME(35,OptionalProperty)
+REGISTER_NAME(36,Utf8StrProperty)
+REGISTER_NAME(37,AnsiStrProperty)
 
 
 // Special types.
diff --git a/Engine/Source/Runtime/CoreUObject/Private/Serialization/ArchiveStackTrace.cpp b/Engine/Source/Runtime/CoreUObject/Private/Serialization/ArchiveStackTrace.cpp
--- a/Engine/Source/Runtime/CoreUObject/Private/Serialization/ArchiveStackTrace.cpp
+++ b/Engine/Source/Runtime/CoreUObject/Private/Serialization/ArchiveStackTrace.cpp
@@ -6,6 +6,7 @@
 #include "UObject/ObjectMacros.h"
 #include "UObject/UObjectThreadContext.h"
 #include "UObject/UnrealType.h"
+#include "UObject/StrProperty.h"
 #include "HAL/PlatformStackWalk.h"
 #include "Serialization/AsyncLoading.h"
 #include "HAL/PlatformFileManager.h"
diff --git a/Engine/Source/Runtime/CoreUObject/Private/UObject/AnsiStrProperty.cpp b/Engine/Source/Runtime/CoreUObject/Private/UObject/AnsiStrProperty.cpp
new file mode 100644
--- /dev/null
+++ b/Engine/Source/Runtime/CoreUObject/Private/UObject/AnsiStrProperty.cpp
@@ -0,0 +1,16 @@
+// Copyright Epic Games, Inc. All Rights Reserved.
+
+#include "UObject/AnsiStrProperty.h"
+
+// Include StrProperty.cpp.inl's includes before defining the macros, in case the macros 'poison' other headers or there are re-entrant includes.
+#include "UObject/StrPropertyIncludes.cpp.inl"
+
+#define UE_STRPROPERTY_CLASS                FAnsiStrProperty
+#define UE_STRPROPERTY_STRINGTYPE           FAnsiString
+#define UE_STRPROPERTY_CASTCLASSFLAG        CASTCLASS_FAnsiStrProperty
+#define UE_STRPROPERTY_PROPERTYPARAMSSTRUCT FAnsiStrPropertyParams
+	#include "UObject/StrProperty.cpp.inl"
+#undef UE_STRPROPERTY_PROPERTYPARAMSSTRUCT
+#undef UE_STRPROPERTY_CASTCLASSFLAG
+#undef UE_STRPROPERTY_STRINGTYPE
+#undef UE_STRPROPERTY_CLASS
diff --git a/Engine/Source/Runtime/CoreUObject/Private/UObject/Class.cpp b/Engine/Source/Runtime/CoreUObject/Private/UObject/Class.cpp
--- a/Engine/Source/Runtime/CoreUObject/Private/UObject/Class.cpp
+++ b/Engine/Source/Runtime/CoreUObject/Private/UObject/Class.cpp
@@ -26,6 +26,7 @@
 #include "Templates/Casts.h"
 #include "UObject/DebugSerializationFlags.h"
 #include "UObject/PropertyTag.h"
+#include "UObject/StrProperty.h"
 #include "UObject/UnrealType.h"
 #include "UObject/UnrealTypePrivate.h"
 #include "UObject/Reload.h"
diff --git a/Engine/Source/Runtime/CoreUObject/Private/UObject/Field.cpp b/Engine/Source/Runtime/CoreUObject/Private/UObject/Field.cpp
--- a/Engine/Source/Runtime/CoreUObject/Private/UObject/Field.cpp
+++ b/Engine/Source/Runtime/CoreUObject/Private/UObject/Field.cpp
@@ -24,6 +24,7 @@ Field.cpp: Defines FField property system fundamentals
 #include "Templates/Casts.h"
 #include "UObject/DebugSerializationFlags.h"
 #include "UObject/PropertyTag.h"
+#include "UObject/StrProperty.h"
 #include "UObject/UnrealType.h"
 #include "UObject/UnrealTypePrivate.h"
 #include "UObject/Stack.h"
diff --git a/Engine/Source/Runtime/CoreUObject/Private/UObject/PropertyAccessUtil.cpp b/Engine/Source/Runtime/CoreUObject/Private/UObject/PropertyAccessUtil.cpp
--- a/Engine/Source/Runtime/CoreUObject/Private/UObject/PropertyAccessUtil.cpp
+++ b/Engine/Source/Runtime/CoreUObject/Private/UObject/PropertyAccessUtil.cpp
@@ -2,6 +2,7 @@
 
 #include "UObject/PropertyAccessUtil.h"
 #include "UObject/EnumProperty.h"
+#include "UObject/StrProperty.h"
 #include "UObject/TextProperty.h"
 #include "UObject/Object.h"
 #include "UObject/Class.h"
diff --git a/Engine/Source/Runtime/CoreUObject/Private/UObject/PropertyStr.cpp b/Engine/Source/Runtime/CoreUObject/Private/UObject/PropertyStr.cpp
deleted file mode 100644
diff --git a/Engine/Source/Runtime/CoreUObject/Private/UObject/ScriptCore.cpp b/Engine/Source/Runtime/CoreUObject/Private/UObject/ScriptCore.cpp
--- a/Engine/Source/Runtime/CoreUObject/Private/UObject/ScriptCore.cpp
+++ b/Engine/Source/Runtime/CoreUObject/Private/UObject/ScriptCore.cpp
@@ -24,6 +24,7 @@
 #include "Serialization/NullArchive.h"
 #include "UObject/SoftObjectPtr.h"
 #include "UObject/PropertyPortFlags.h"
+#include "UObject/StrProperty.h"
 #include "UObject/UnrealType.h"
 #include "UObject/Stack.h"
 #include "UObject/Reload.h"
diff --git a/Engine/Source/Runtime/CoreUObject/Private/UObject/StrProperty.cpp b/Engine/Source/Runtime/CoreUObject/Private/UObject/StrProperty.cpp
new file mode 100644
--- /dev/null
+++ b/Engine/Source/Runtime/CoreUObject/Private/UObject/StrProperty.cpp
@@ -0,0 +1,16 @@
+// Copyright Epic Games, Inc. All Rights Reserved.
+
+#include "UObject/StrProperty.h"
+
+// Include StrProperty.cpp.inl's includes before defining the macros, in case the macros 'poison' other headers or there are re-entrant includes.
+#include "UObject/StrPropertyIncludes.cpp.inl"
+
+#define UE_STRPROPERTY_CLASS                FStrProperty
+#define UE_STRPROPERTY_STRINGTYPE           FString
+#define UE_STRPROPERTY_CASTCLASSFLAG        CASTCLASS_FStrProperty
+#define UE_STRPROPERTY_PROPERTYPARAMSSTRUCT FStrPropertyParams
+	#include "UObject/StrProperty.cpp.inl"
+#undef UE_STRPROPERTY_PROPERTYPARAMSSTRUCT
+#undef UE_STRPROPERTY_CASTCLASSFLAG
+#undef UE_STRPROPERTY_STRINGTYPE
+#undef UE_STRPROPERTY_CLASS
diff --git a/Engine/Source/Runtime/CoreUObject/Private/UObject/StrProperty.cpp.inl b/Engine/Source/Runtime/CoreUObject/Private/UObject/StrProperty.cpp.inl
new file mode 100644
--- /dev/null
+++ b/Engine/Source/Runtime/CoreUObject/Private/UObject/StrProperty.cpp.inl
@@ -0,0 +1,153 @@
+// Copyright Epic Games, Inc. All Rights Reserved.
+
+/*********************************************************************************************************************
+ * NOTICE                                                                                                            *
+ *                                                                                                                   *
+ * This file is not intended to be included directly - it will be included by other .cpp files which have predefined *
+ * some macros to be expanded within this file.  As such, it does not have a #pragma once as it is intended to be    *
+ * included multiple times with different macro definitions.                                                         *
+ *                                                                                                                   *
+ * #includes needed to compile this file need to be specified in StringIncludes.cpp.inl file rather than here.       *
+ *********************************************************************************************************************/
+#define UE_INCLUDETOOL_IGNORE_INCONSISTENT_STATE
+
+#ifndef UE_STRPROPERTY_CLASS
+	#error "StrProperty.cpp.inl should only be included after defining UE_STRPROPERTY_CLASS"
+#endif
+#ifndef UE_STRPROPERTY_STRINGTYPE
+	#error "StrProperty.cpp.inl should only be included after defining UE_STRPROPERTY_STRINGTYPE"
+#endif
+#ifndef UE_STRPROPERTY_CASTCLASSFLAG
+	#error "StrProperty.cpp.inl should only be included after defining UE_STRPROPERTY_CASTCLASSFLAG"
+#endif
+#ifndef UE_STRPROPERTY_PROPERTYPARAMSSTRUCT
+	#error "StrProperty.cpp.inl should only be included after defining UE_STRPROPERTY_PROPERTYPARAMSSTRUCT"
+#endif
+
+IMPLEMENT_FIELD(UE_STRPROPERTY_CLASS)
+
+UE_STRPROPERTY_CLASS::UE_STRPROPERTY_CLASS(FFieldVariant InOwner, const UECodeGen_Private::UE_STRPROPERTY_PROPERTYPARAMSSTRUCT& Prop)
+	: Super(InOwner, (const UECodeGen_Private::FPropertyParamsBaseWithOffset&)Prop)
+{
+}
+
+EConvertFromTypeResult UE_STRPROPERTY_CLASS::ConvertFromType(const FPropertyTag& Tag, FStructuredArchive::FSlot Slot, uint8* Data, UStruct* DefaultsStruct, const uint8* Defaults)
+{
+	// Convert serialized text to string.
+	if (Tag.Type==NAME_TextProperty)
+	{ 
+		FText Text;
+		Slot << Text;
+		const UE_STRPROPERTY_STRINGTYPE String(FTextInspector::GetSourceString(Text) ? *FTextInspector::GetSourceString(Text) : TEXT(""));
+		SetPropertyValue_InContainer(Data, String, Tag.ArrayIndex);
+
+		return EConvertFromTypeResult::Converted;
+	}
+	else if (Tag.Type==NAME_StrProperty)
+	{
+		FString SavedStr;
+		Slot << SavedStr;
+		const UE_STRPROPERTY_STRINGTYPE ConvertedString(SavedStr);
+		SetPropertyValue_InContainer(Data, ConvertedString, Tag.ArrayIndex);
+
+		return EConvertFromTypeResult::Converted;
+	}
+	else if (Tag.Type==NAME_Utf8StrProperty)
+	{
+		FUtf8String SavedStr;
+		Slot << SavedStr;
+		const UE_STRPROPERTY_STRINGTYPE ConvertedString(SavedStr);
+		SetPropertyValue_InContainer(Data, ConvertedString, Tag.ArrayIndex);
+
+		return EConvertFromTypeResult::Converted;
+	}
+	else if (Tag.Type==NAME_AnsiStrProperty)
+	{
+		FAnsiString SavedStr;
+		Slot << SavedStr;
+		const UE_STRPROPERTY_STRINGTYPE ConvertedString(SavedStr);
+		SetPropertyValue_InContainer(Data, ConvertedString, Tag.ArrayIndex);
+
+		return EConvertFromTypeResult::Converted;
+	}
+
+	return EConvertFromTypeResult::UseSerializeItem;
+}
+
+FString UE_STRPROPERTY_CLASS::GetCPPTypeForwardDeclaration() const
+{
+	return FString();
+}
+
+void UE_STRPROPERTY_CLASS::ExportText_Internal( FString& ValueStr, const void* PropertyValueOrContainer, EPropertyPointerType PropertyPointerType, const void* DefaultValue, UObject* Parent, int32 PortFlags, UObject* ExportRootScope ) const
+{
+	UE_STRPROPERTY_STRINGTYPE StringValue;
+	if (PropertyPointerType == EPropertyPointerType::Container && HasGetter())
+	{
+		GetValue_InContainer(PropertyValueOrContainer, &StringValue);
+	}
+	else
+	{
+		StringValue = *(UE_STRPROPERTY_STRINGTYPE*)PointerToValuePtr(PropertyValueOrContainer, PropertyPointerType);
+	}
+
+	if (!(PortFlags & PPF_Delimited))
+	{
+		ValueStr += StringValue;
+	}
+	else if ( StringValue.Len() > 0 )
+	{
+		ValueStr += UE_STRPROPERTY_STRINGTYPE::Printf( CHARTEXT(UE_STRPROPERTY_STRINGTYPE::ElementType, "\"%s\""), *StringValue.ReplaceCharWithEscapedChar() );
+	}
+	else
+	{
+		ValueStr += CHARTEXT(UE_STRPROPERTY_STRINGTYPE::ElementType, "\"\"");
+	}
+}
+const TCHAR* UE_STRPROPERTY_CLASS::ImportText_Internal( const TCHAR* Buffer, void* ContainerOrPropertyPtr, EPropertyPointerType PropertyPointerType, UObject* Parent, int32 PortFlags, FOutputDevice* ErrorText ) const
+{
+	UE_STRPROPERTY_STRINGTYPE ImportedText;
+	if( !(PortFlags & PPF_Delimited) )
+	{
+		ImportedText = Buffer;
+
+		// in order to indicate that the value was successfully imported, advance the buffer past the last character that was imported
+		Buffer += FCString::Strlen(Buffer);
+	}
+	else
+	{
+		// require quoted string here
+		if (*Buffer != TCHAR('"'))
+		{
+			ErrorText->Logf(TEXT("Missing opening '\"' in string property value: %s"), Buffer);
+			return NULL;
+		}
+		const TCHAR* Start = Buffer;
+		FString Temp;
+		Buffer = FPropertyHelpers::ReadToken(Buffer, Temp);
+		if (Buffer == NULL)
+		{
+			return NULL;
+		}
+		if (Buffer > Start && Buffer[-1] != TCHAR('"'))
+		{
+			ErrorText->Logf(TEXT("Missing terminating '\"' in string property value: %s"), Start);
+			return NULL;
+		}
+		ImportedText = UE_STRPROPERTY_STRINGTYPE(MoveTemp(Temp));
+	}
+	if (PropertyPointerType == EPropertyPointerType::Container && HasSetter())
+	{
+		SetValue_InContainer(ContainerOrPropertyPtr, ImportedText);
+	}
+	else
+	{
+		*(UE_STRPROPERTY_STRINGTYPE*)PointerToValuePtr(ContainerOrPropertyPtr, PropertyPointerType) = MoveTemp(ImportedText);
+	}
+	return Buffer;
+}
+
+uint32 UE_STRPROPERTY_CLASS::GetValueTypeHashInternal(const void* Src) const
+{
+	return GetTypeHash(*(const UE_STRPROPERTY_STRINGTYPE*)Src);
+}
diff --git a/Engine/Source/Runtime/CoreUObject/Private/UObject/StrPropertyIncludes.cpp.inl b/Engine/Source/Runtime/CoreUObject/Private/UObject/StrPropertyIncludes.cpp.inl
new file mode 100644
--- /dev/null
+++ b/Engine/Source/Runtime/CoreUObject/Private/UObject/StrPropertyIncludes.cpp.inl
@@ -0,0 +1,26 @@
+// Copyright Epic Games, Inc. All Rights Reserved.
+
+#pragma once
+
+/*******************************************************************************************************
+ * NOTICE                                                                                              *
+ *                                                                                                     *
+ * This file is not intended to be included directly - it is only intended to contain the includes for *
+ * StrPropetryIncludes.cpp.inl.                                                                        *
+ *******************************************************************************************************/
+
+#ifdef UE_STRPROPERTY_CLASS
+	#error "StrPropertyIncludes.cpp.inl should not be included after defining UE_STRPROPERTY_CLASS"
+#endif
+#ifdef UE_STRPROPERTY_STRINGTYPE
+	#error "StrPropertyIncludes.cpp.inl should not be included after defining UE_STRPROPERTY_STRINGTYPE"
+#endif
+
+#include "CoreMinimal.h"
+#include "Containers/AnsiString.h"
+#include "Containers/UnrealString.h"
+#include "Containers/Utf8String.h"
+#include "UObject/ObjectMacros.h"
+#include "UObject/PropertyPortFlags.h"
+#include "Misc/StringBuilder.h"
+#include "Misc/AsciiSet.h"
diff --git a/Engine/Source/Runtime/CoreUObject/Private/UObject/UObjectGlobals.cpp b/Engine/Source/Runtime/CoreUObject/Private/UObject/UObjectGlobals.cpp
--- a/Engine/Source/Runtime/CoreUObject/Private/UObject/UObjectGlobals.cpp
+++ b/Engine/Source/Runtime/CoreUObject/Private/UObject/UObjectGlobals.cpp
@@ -60,6 +60,9 @@
 #include "Modules/ModuleManager.h"
 #include "UObject/EnumProperty.h"
 #include "UObject/TextProperty.h"
+#include "UObject/StrProperty.h"
+#include "UObject/AnsiStrProperty.h"
+#include "UObject/Utf8StrProperty.h"
 #include "UObject/FieldPathProperty.h"
 #include "UObject/MetaData.h"
 #include "HAL/LowLevelMemTracker.h"
@@ -6063,6 +6066,18 @@ namespace UECodeGen_Private
 			}
 			break;
 #endif
+
+			case EPropertyGenFlags::Utf8Str:
+			{
+				NewProp = NewFProperty<FUtf8StrProperty, FUtf8StrPropertyParams>(Outer, *PropBase);
+			}
+			break;
+
+			case EPropertyGenFlags::AnsiStr:
+			{
+				NewProp = NewFProperty<FAnsiStrProperty, FAnsiStrPropertyParams>(Outer, *PropBase);
+			}
+			break;
 		}
 
 		NewProp->ArrayDim = PropBase->ArrayDim;
diff --git a/Engine/Source/Runtime/CoreUObject/Private/UObject/Utf8StrProperty.cpp b/Engine/Source/Runtime/CoreUObject/Private/UObject/Utf8StrProperty.cpp
new file mode 100644
--- /dev/null
+++ b/Engine/Source/Runtime/CoreUObject/Private/UObject/Utf8StrProperty.cpp
@@ -0,0 +1,16 @@
+// Copyright Epic Games, Inc. All Rights Reserved.
+
+#include "UObject/Utf8StrProperty.h"
+
+// Include StrProperty.cpp.inl's includes before defining the macros, in case the macros 'poison' other headers or there are re-entrant includes.
+#include "UObject/StrPropertyIncludes.cpp.inl"
+
+#define UE_STRPROPERTY_CLASS                FUtf8StrProperty
+#define UE_STRPROPERTY_STRINGTYPE           FUtf8String
+#define UE_STRPROPERTY_CASTCLASSFLAG        CASTCLASS_FUtf8StrProperty
+#define UE_STRPROPERTY_PROPERTYPARAMSSTRUCT FUtf8StrPropertyParams
+	#include "UObject/StrProperty.cpp.inl"
+#undef UE_STRPROPERTY_PROPERTYPARAMSSTRUCT
+#undef UE_STRPROPERTY_CASTCLASSFLAG
+#undef UE_STRPROPERTY_STRINGTYPE
+#undef UE_STRPROPERTY_CLASS
diff --git a/Engine/Source/Runtime/CoreUObject/Public/Templates/Casts.h b/Engine/Source/Runtime/CoreUObject/Public/Templates/Casts.h
--- a/Engine/Source/Runtime/CoreUObject/Public/Templates/Casts.h
+++ b/Engine/Source/Runtime/CoreUObject/Public/Templates/Casts.h
@@ -445,6 +445,8 @@ DECLARE_CAST_BY_FLAG(FClassProperty)
 DECLARE_CAST_BY_FLAG(FInterfaceProperty)
 DECLARE_CAST_BY_FLAG(FNameProperty)
 DECLARE_CAST_BY_FLAG(FStrProperty)
+DECLARE_CAST_BY_FLAG(FUtf8StrProperty)
+DECLARE_CAST_BY_FLAG(FAnsiStrProperty)
 DECLARE_CAST_BY_FLAG(FTextProperty)
 DECLARE_CAST_BY_FLAG(FArrayProperty)
 DECLARE_CAST_BY_FLAG(FDelegateProperty)
diff --git a/Engine/Source/Runtime/CoreUObject/Public/UObject/AnsiStrProperty.h b/Engine/Source/Runtime/CoreUObject/Public/UObject/AnsiStrProperty.h
new file mode 100644
--- /dev/null
+++ b/Engine/Source/Runtime/CoreUObject/Public/UObject/AnsiStrProperty.h
@@ -0,0 +1,18 @@
+// Copyright Epic Games, Inc. All Rights Reserved.
+
+#pragma once
+
+// Include StrProperty.h.inl's includes before defining the macros, in case the macros 'poison' other headers or there are re-entrant includes.
+#include "UObject/StrPropertyIncludes.h.inl"
+
+#include "Containers/AnsiString.h"
+
+#define UE_STRPROPERTY_CLASS                FAnsiStrProperty
+#define UE_STRPROPERTY_STRINGTYPE           FAnsiString
+#define UE_STRPROPERTY_CASTCLASSFLAG        CASTCLASS_FAnsiStrProperty
+#define UE_STRPROPERTY_PROPERTYPARAMSSTRUCT FAnsiStrPropertyParams
+	#include "UObject/StrProperty.h.inl"
+#undef UE_STRPROPERTY_PROPERTYPARAMSSTRUCT
+#undef UE_STRPROPERTY_CASTCLASSFLAG
+#undef UE_STRPROPERTY_STRINGTYPE
+#undef UE_STRPROPERTY_CLASS
diff --git a/Engine/Source/Runtime/CoreUObject/Public/UObject/Field.h b/Engine/Source/Runtime/CoreUObject/Public/UObject/Field.h
--- a/Engine/Source/Runtime/CoreUObject/Public/UObject/Field.h
+++ b/Engine/Source/Runtime/CoreUObject/Public/UObject/Field.h
@@ -247,7 +247,7 @@ FField* TClass::Construct(const FFieldVariant& InOwner, const FName& InName, EOb
 } \
 FFieldClass* TClass::StaticClass() \
 { \
-	static FFieldClass StaticFieldClass(TEXT(#TClass), TClass::StaticClassCastFlagsPrivate(), TClass::StaticClassCastFlags(), TClass::Super::StaticClass(), &TClass::Construct); \
+	static FFieldClass StaticFieldClass(TEXT(PREPROCESSOR_TO_STRING(TClass)), TClass::StaticClassCastFlagsPrivate(), TClass::StaticClassCastFlags(), TClass::Super::StaticClass(), &TClass::Construct); \
 	return &StaticFieldClass; \
 } \
 
diff --git a/Engine/Source/Runtime/CoreUObject/Public/UObject/GeneratedCppIncludes.h b/Engine/Source/Runtime/CoreUObject/Public/UObject/GeneratedCppIncludes.h
--- a/Engine/Source/Runtime/CoreUObject/Public/UObject/GeneratedCppIncludes.h
+++ b/Engine/Source/Runtime/CoreUObject/Public/UObject/GeneratedCppIncludes.h
@@ -10,6 +10,9 @@
 #include "UObject/Package.h"
 #include "UObject/UnrealType.h"
 #include "UObject/EnumProperty.h"
+#include "UObject/StrProperty.h"
+#include "UObject/AnsiStrProperty.h"
+#include "UObject/Utf8StrProperty.h"
 #include "UObject/TextProperty.h"
 #include "UObject/FieldPathProperty.h"
 
diff --git a/Engine/Source/Runtime/CoreUObject/Public/UObject/ObjectMacros.h b/Engine/Source/Runtime/CoreUObject/Public/UObject/ObjectMacros.h
--- a/Engine/Source/Runtime/CoreUObject/Public/UObject/ObjectMacros.h
+++ b/Engine/Source/Runtime/CoreUObject/Public/UObject/ObjectMacros.h
@@ -378,6 +378,8 @@ enum EClassCastFlags : uint64
 	CASTCLASS_FVValueProperty				= 0x0200000000000000,
 	CASTCLASS_UVerseVMClass					= 0x0400000000000000,
 	CASTCLASS_FVRestValueProperty			= 0x0800000000000000,
+	CASTCLASS_FUtf8StrProperty				= 0x1000000000000000,
+	CASTCLASS_FAnsiStrProperty				= 0x2000000000000000,
 };
 
 #define CASTCLASS_AllFlags ((EClassCastFlags)0xFFFFFFFFFFFFFFFF)
diff --git a/Engine/Source/Runtime/CoreUObject/Public/UObject/StrProperty.h b/Engine/Source/Runtime/CoreUObject/Public/UObject/StrProperty.h
new file mode 100644
--- /dev/null
+++ b/Engine/Source/Runtime/CoreUObject/Public/UObject/StrProperty.h
@@ -0,0 +1,18 @@
+// Copyright Epic Games, Inc. All Rights Reserved.
+
+#pragma once
+
+// Include StrProperty.h.inl's includes before defining the macros, in case the macros 'poison' other headers or there are re-entrant includes.
+#include "UObject/StrPropertyIncludes.h.inl"
+
+#include "Containers/UnrealString.h"
+
+#define UE_STRPROPERTY_CLASS                FStrProperty
+#define UE_STRPROPERTY_STRINGTYPE           FString
+#define UE_STRPROPERTY_CASTCLASSFLAG        CASTCLASS_FStrProperty
+#define UE_STRPROPERTY_PROPERTYPARAMSSTRUCT FStrPropertyParams
+	#include "UObject/StrProperty.h.inl"
+#undef UE_STRPROPERTY_PROPERTYPARAMSSTRUCT
+#undef UE_STRPROPERTY_CASTCLASSFLAG
+#undef UE_STRPROPERTY_STRINGTYPE
+#undef UE_STRPROPERTY_CLASS
diff --git a/Engine/Source/Runtime/CoreUObject/Public/UObject/StrProperty.h.inl b/Engine/Source/Runtime/CoreUObject/Public/UObject/StrProperty.h.inl
new file mode 100644
--- /dev/null
+++ b/Engine/Source/Runtime/CoreUObject/Public/UObject/StrProperty.h.inl
@@ -0,0 +1,63 @@
+// Copyright Epic Games, Inc. All Rights Reserved.
+
+/*******************************************************************************************************************
+ * NOTICE                                                                                                          *
+ *                                                                                                                 *
+ * This file is not intended to be included directly - it will be included by other .h files which have predefined *
+ * some macros to be expanded within this file.  As such, it does not have a #pragma once as it is intended to be  *
+ * included multiple times with different macro definitions.                                                       *
+ *                                                                                                                 *
+ * #includes needed to compile this file need to be specified in StrPropertyIncludes.h.inl file rather than here.  *
+ *******************************************************************************************************************/
+#define UE_INCLUDETOOL_IGNORE_INCONSISTENT_STATE
+
+#ifndef UE_STRPROPERTY_CLASS
+	#error "StrProperty.h.inl should only be included after defining UE_STRPROPERTY_CLASS"
+#endif
+#ifndef UE_STRPROPERTY_STRINGTYPE
+	#error "StrProperty.h.inl should only be included after defining UE_STRPROPERTY_STRINGTYPE"
+#endif
+#ifndef UE_STRPROPERTY_CASTCLASSFLAG
+	#error "StrProperty.h.inl should only be included after defining UE_STRPROPERTY_CASTCLASSFLAG"
+#endif
+#ifndef UE_STRPROPERTY_PROPERTYPARAMSSTRUCT
+	#error "StrProperty.h.inl should only be included after defining UE_STRPROPERTY_PROPERTYPARAMSSTRUCT"
+#endif
+
+class COREUOBJECT_API UE_STRPROPERTY_CLASS : public TProperty_WithEqualityAndSerializer<UE_STRPROPERTY_STRINGTYPE, FProperty>
+{
+	DECLARE_FIELD(UE_STRPROPERTY_CLASS, (TProperty_WithEqualityAndSerializer<UE_STRPROPERTY_STRINGTYPE, FProperty>), UE_STRPROPERTY_CASTCLASSFLAG)
+public:
+	using TTypeFundamentals = Super::TTypeFundamentals;
+	using TCppType = TTypeFundamentals::TCppType;
+
+	UE_STRPROPERTY_CLASS(FFieldVariant InOwner, const FName& InName, EObjectFlags InObjectFlags)
+		: Super(InOwner, InName, InObjectFlags)
+	{
+	}
+
+	/**
+	 * Constructor used for constructing compiled in properties
+	 * @param InOwner Owner of the property
+	 * @param PropBase Pointer to the compiled in structure describing the property
+	 **/
+	UE_STRPROPERTY_CLASS(FFieldVariant InOwner, const UECodeGen_Private::UE_STRPROPERTY_PROPERTYPARAMSSTRUCT& Prop);
+
+#if WITH_EDITORONLY_DATA
+	explicit UE_STRPROPERTY_CLASS(UField* InField)
+		: Super(InField)
+	{
+	}
+#endif // WITH_EDITORONLY_DATA
+
+	// FProperty interface
+protected:
+	virtual void ExportText_Internal( FString& ValueStr, const void* PropertyValueOrContainer, EPropertyPointerType PropertyPointerType, const void* DefaultValue, UObject* Parent, int32 PortFlags, UObject* ExportRootScope ) const override;
+	virtual const TCHAR* ImportText_Internal(const TCHAR* Buffer, void* ContainerOrPropertyPtr, EPropertyPointerType PropertyPointerType, UObject* OwnerObject, int32 PortFlags, FOutputDevice* ErrorText) const override;
+public:
+	virtual EConvertFromTypeResult ConvertFromType(const FPropertyTag& Tag, FStructuredArchive::FSlot Slot, uint8* Data, UStruct* DefaultsStruct, const uint8* Defaults) override;
+	UE_DEPRECATED(5.4, "UnrealHeaderTool only API.  No replacement available.")
+	virtual FString GetCPPTypeForwardDeclaration() const override;
+	uint32 GetValueTypeHashInternal(const void* Src) const override;
+	// End of FProperty interface
+};
diff --git a/Engine/Source/Runtime/CoreUObject/Public/UObject/StrPropertyIncludes.h.inl b/Engine/Source/Runtime/CoreUObject/Public/UObject/StrPropertyIncludes.h.inl
new file mode 100644
--- /dev/null
+++ b/Engine/Source/Runtime/CoreUObject/Public/UObject/StrPropertyIncludes.h.inl
@@ -0,0 +1,17 @@
+// Copyright Epic Games, Inc. All Rights Reserved.
+
+/*=============================================================================
+	UnrealType.h: Unreal engine base type definitions.
+=============================================================================*/
+
+#pragma once
+
+#include "CoreGlobals.h"
+#include "Serialization/StructuredArchive.h"
+#include "UObject/Field.h"
+#include "UObject/NameTypes.h"
+#include "UObject/ObjectMacros.h"
+#include "UObject/PropertyPortFlags.h"
+#include "UObject/PropertyTag.h"
+#include "UObject/UObjectGlobals.h"
+#include "UObject/UnrealType.h"
diff --git a/Engine/Source/Runtime/CoreUObject/Public/UObject/UObjectGlobals.h b/Engine/Source/Runtime/CoreUObject/Public/UObject/UObjectGlobals.h
--- a/Engine/Source/Runtime/CoreUObject/Public/UObject/UObjectGlobals.h
+++ b/Engine/Source/Runtime/CoreUObject/Public/UObject/UObjectGlobals.h
@@ -3447,6 +3447,8 @@ namespace UECodeGen_Private
 		LargeWorldCoordinatesReal = 0x20,
 		Optional          = 0x21,
 		VValue            = 0x22,
+		Utf8Str           = 0x23,
+		AnsiStr           = 0x24,
 
 		// Property-specific flags
 		NativeBool        = 0x40,
@@ -3824,6 +3826,8 @@ namespace UECodeGen_Private
 	typedef FClassPropertyParams   FClassPtrPropertyParams;
 	typedef FObjectPropertyParams  FSoftObjectPropertyParams;
 	typedef FGenericPropertyParams FVerseValuePropertyParams;
+	typedef FGenericPropertyParams FUtf8StrPropertyParams;
+	typedef FGenericPropertyParams FAnsiStrPropertyParams;
 
 	struct FFunctionParams
 	{
diff --git a/Engine/Source/Runtime/CoreUObject/Public/UObject/UnrealType.h b/Engine/Source/Runtime/CoreUObject/Public/UObject/UnrealType.h
--- a/Engine/Source/Runtime/CoreUObject/Public/UObject/UnrealType.h
+++ b/Engine/Source/Runtime/CoreUObject/Public/UObject/UnrealType.h
@@ -3280,51 +3280,6 @@ public:
 	// End of FProperty interface
 };
 
-/*-----------------------------------------------------------------------------
-	FStrProperty.
------------------------------------------------------------------------------*/
-
-//
-// Describes a dynamic string variable.
-//
-class COREUOBJECT_API FStrProperty : public TProperty_WithEqualityAndSerializer<FString, FProperty>
-{
-	DECLARE_FIELD(FStrProperty, (TProperty_WithEqualityAndSerializer<FString, FProperty>), CASTCLASS_FStrProperty)
-public:
-	using TTypeFundamentals = Super::TTypeFundamentals;
-	using TCppType = TTypeFundamentals::TCppType;
-
-	FStrProperty(FFieldVariant InOwner, const FName& InName, EObjectFlags InObjectFlags)
-		: Super(InOwner, InName, InObjectFlags)
-	{
-	}
-
-	/**
-	 * Constructor used for constructing compiled in properties
-	 * @param InOwner Owner of the property
-	 * @param PropBase Pointer to the compiled in structure describing the property
-	 **/
-	FStrProperty(FFieldVariant InOwner, const UECodeGen_Private::FStrPropertyParams& Prop);
-
-#if WITH_EDITORONLY_DATA
-	explicit FStrProperty(UField* InField)
-		: Super(InField)
-	{
-	}
-#endif // WITH_EDITORONLY_DATA
-
-	// FProperty interface
-protected:
-	virtual void ExportText_Internal( FString& ValueStr, const void* PropertyValueOrContainer, EPropertyPointerType PropertyPointerType, const void* DefaultValue, UObject* Parent, int32 PortFlags, UObject* ExportRootScope ) const override;
-	virtual const TCHAR* ImportText_Internal(const TCHAR* Buffer, void* ContainerOrPropertyPtr, EPropertyPointerType PropertyPointerType, UObject* OwnerObject, int32 PortFlags, FOutputDevice* ErrorText) const override;
-public:
-	virtual EConvertFromTypeResult ConvertFromType(const FPropertyTag& Tag, FStructuredArchive::FSlot Slot, uint8* Data, UStruct* DefaultsStruct, const uint8* Defaults) override;
-	UE_DEPRECATED(5.4, "UnrealHeaderTool only API.  No replacement available.")
-	virtual FString GetCPPTypeForwardDeclaration() const override;
-	uint32 GetValueTypeHashInternal(const void* Src) const override;
-	// End of FProperty interface
-};
-
 /*-----------------------------------------------------------------------------
 	FArrayProperty.
 -----------------------------------------------------------------------------*/
@@ -7141,4 +7096,5 @@ class COREUOBJECT_API UMulticastInlineDelegatePropertyWrapper : public UMulticas
 
 #if UE_ENABLE_INCLUDE_ORDER_DEPRECATED_IN_5_5
 #include "Templates/IsTriviallyDestructible.h"
+#include "UObject/StrProperty.h"
 #endif
diff --git a/Engine/Source/Runtime/CoreUObject/Public/UObject/Utf8StrProperty.h b/Engine/Source/Runtime/CoreUObject/Public/UObject/Utf8StrProperty.h
new file mode 100644
--- /dev/null
+++ b/Engine/Source/Runtime/CoreUObject/Public/UObject/Utf8StrProperty.h
@@ -0,0 +1,18 @@
+// Copyright Epic Games, Inc. All Rights Reserved.
+
+#pragma once
+
+// Include StrProperty.h.inl's includes before defining the macros, in case the macros 'poison' other headers or there are re-entrant includes.
+#include "UObject/StrPropertyIncludes.h.inl"
+
+#include "Containers/Utf8String.h"
+
+#define UE_STRPROPERTY_CLASS                FUtf8StrProperty
+#define UE_STRPROPERTY_STRINGTYPE           FUtf8String
+#define UE_STRPROPERTY_CASTCLASSFLAG        CASTCLASS_FUtf8StrProperty
+#define UE_STRPROPERTY_PROPERTYPARAMSSTRUCT FUtf8StrPropertyParams
+	#include "UObject/StrProperty.h.inl"
+#undef UE_STRPROPERTY_PROPERTYPARAMSSTRUCT
+#undef UE_STRPROPERTY_CASTCLASSFLAG
+#undef UE_STRPROPERTY_STRINGTYPE
+#undef UE_STRPROPERTY_CLASS
diff --git a/Engine/Source/Runtime/PropertyPath/Public/PropertyTypeCompatibility.h b/Engine/Source/Runtime/PropertyPath/Public/PropertyTypeCompatibility.h
--- a/Engine/Source/Runtime/PropertyPath/Public/PropertyTypeCompatibility.h
+++ b/Engine/Source/Runtime/PropertyPath/Public/PropertyTypeCompatibility.h
@@ -7,6 +7,7 @@
 #include "UObject/WeakObjectPtr.h"
 #include "UObject/UnrealType.h"
 #include "UObject/EnumProperty.h"
+#include "UObject/StrProperty.h"
 #include "UObject/TextProperty.h"
 #include "Math/InterpCurvePoint.h"
 
