Added fatal logs to FPropertyTypeName when StoreByIndex fails

#rnx

diff --git a/Engine/Source/Runtime/CoreUObject/Private/UObject/PropertyTypeName.cpp b/Engine/Source/Runtime/CoreUObject/Private/UObject/PropertyTypeName.cpp
--- a/Engine/Source/Runtime/CoreUObject/Private/UObject/PropertyTypeName.cpp
+++ b/Engine/Source/Runtime/CoreUObject/Private/UObject/PropertyTypeName.cpp
@@ -195,7 +195,8 @@ int32 FPropertyTypeNameTable::StoreByIndex(const FPropertyTypeNameNode* Nodes, i
 		return 0;
 	}
 
-	check(Count <= GPropertyTypeNameBlockOffsetCount);
+	UE_CLOG(Count > GPropertyTypeNameBlockOffsetCount, LogCore, Fatal,
+		TEXT("Invalid property type name with %d nodes. This can happen when serializing from a corrupt or invalid archive."), Count);
 
 	int32 Index;
 	for (int32 BaseIndex = NextIndex.load(std::memory_order_relaxed);;)
@@ -208,7 +209,9 @@ int32 FPropertyTypeNameTable::StoreByIndex(const FPropertyTypeNameNode* Nodes, i
 		}
 	}
 
-	check(Index + Count <= GPropertyTypeNameBlockCount * GPropertyTypeNameBlockOffsetCount);
+	UE_CLOG(Index + Count > GPropertyTypeNameBlockCount * GPropertyTypeNameBlockOffsetCount, LogCore, Fatal,
+		TEXT("Exceeded property type name capacity of %d nodes when storing %d nodes."),
+		GPropertyTypeNameBlockCount * GPropertyTypeNameBlockOffsetCount, Count);
 
 	const int32 BlockIndex = Index >> GPropertyTypeNameBlockOffsetBits;
 	FPropertyTypeNameNode* Block = Blocks[BlockIndex].load(std::memory_order_acquire);
