Manual integration of 32610366 by [REDACTED_PERSON]
Potential fix when issuing empty texture streaming requests with batch API.

diff --git a/Engine/Source/Runtime/CoreUObject/Private/Serialization/BulkDataStreaming.cpp b/Engine/Source/Runtime/CoreUObject/Private/Serialization/BulkDataStreaming.cpp
--- a/Engine/Source/Runtime/CoreUObject/Private/Serialization/BulkDataStreaming.cpp
+++ b/Engine/Source/Runtime/CoreUObject/Private/Serialization/BulkDataStreaming.cpp
@@ -1040,6 +1040,11 @@ FBulkDataBatchRequest::FBatchBuilder::FBatchBuilder(int32 MaxCount)
 {
 }
 
+bool FBulkDataBatchRequest::FBatchBuilder::IsEmpty() const
+{
+	return BatchCount == 0;
+}
+
 FBulkDataBatchRequest::FBatchBuilder& FBulkDataBatchRequest::FBatchBuilder::Read(FBulkData& BulkData, EAsyncIOPriorityAndFlags Priority)
 {
 	if (BulkData.IsBulkDataLoaded())
diff --git a/Engine/Source/Runtime/CoreUObject/Public/Serialization/BulkData.h b/Engine/Source/Runtime/CoreUObject/Public/Serialization/BulkData.h
--- a/Engine/Source/Runtime/CoreUObject/Public/Serialization/BulkData.h
+++ b/Engine/Source/Runtime/CoreUObject/Public/Serialization/BulkData.h
@@ -1447,6 +1447,8 @@ public:
 	{
 	public:
 		COREUOBJECT_API FBatchBuilder(int32 MaxCount);
+		/** Returns whether the batch is empty. */
+		COREUOBJECT_API bool IsEmpty() const;
 		/** Read the entire bulk data and copy the result to the specified instance. */
 		COREUOBJECT_API FBatchBuilder& Read(FBulkData& BulkData, EAsyncIOPriorityAndFlags Priority = DefaultPriority);
 		/**
diff --git a/Engine/Source/Runtime/Engine/Private/Streaming/Texture2DStreamIn_IO.cpp b/Engine/Source/Runtime/Engine/Private/Streaming/Texture2DStreamIn_IO.cpp
--- a/Engine/Source/Runtime/Engine/Private/Streaming/Texture2DStreamIn_IO.cpp
+++ b/Engine/Source/Runtime/Engine/Private/Streaming/Texture2DStreamIn_IO.cpp
@@ -118,7 +118,7 @@ void FTexture2DStreamIn_IO::SetIORequests(const FContext& Context)
 			}
 		}
 
-		if (bFailedOnIOError)
+		if (bFailedOnIOError || Batch.IsEmpty())
 		{
 			return;
 		}
