Removed unused UVerseVMClass casting flag

#rnx

diff --git a/Engine/Source/Programs/Shared/EpicGames.Core/UnrealEngineTypes.cs b/Engine/Source/Programs/Shared/EpicGames.Core/UnrealEngineTypes.cs
--- a/Engine/Source/Programs/Shared/EpicGames.Core/UnrealEngineTypes.cs
+++ b/Engine/Source/Programs/Shared/EpicGames.Core/UnrealEngineTypes.cs
@@ -391,8 +391,7 @@ namespace EpicGames.Core
 		FLargeWorldCoordinatesRealProperty = 0x0080000000000000,
 		FOptionalProperty = 0x0100000000000000,
 		FVValueProperty = 0x0200000000000000,
-		UVerseVMClass = 0x0400000000000000,
-		FVRestValueProperty = 0x0800000000000000,
+		FVRestValueProperty = 0x0400000000000000,
 		AllFlags = UInt64.MaxValue,
 	};
 
diff --git a/Engine/Source/Runtime/CoreUObject/Public/UObject/ObjectMacros.h b/Engine/Source/Runtime/CoreUObject/Public/UObject/ObjectMacros.h
--- a/Engine/Source/Runtime/CoreUObject/Public/UObject/ObjectMacros.h
+++ b/Engine/Source/Runtime/CoreUObject/Public/UObject/ObjectMacros.h
@@ -376,8 +376,7 @@ enum EClassCastFlags : uint64
 	CASTCLASS_FLargeWorldCoordinatesRealProperty = 0x0080000000000000,
 	CASTCLASS_FOptionalProperty				= 0x0100000000000000,
 	CASTCLASS_FVValueProperty				= 0x0200000000000000,
-	CASTCLASS_UVerseVMClass					= 0x0400000000000000,
-	CASTCLASS_FVRestValueProperty			= 0x0800000000000000,
+	CASTCLASS_FVRestValueProperty			= 0x0400000000000000,
 };
 
 #define CASTCLASS_AllFlags ((EClassCastFlags)0xFFFFFFFFFFFFFFFF)
