Create a placeholder type object on failure to resolve a type import through an existing redirector asset at load time.

This allows a package export to still load its data to a property bag in this scenario.

diff --git a/Engine/Source/Runtime/CoreUObject/Internal/UObject/PropertyBagRepository.h b/Engine/Source/Runtime/CoreUObject/Internal/UObject/PropertyBagRepository.h
--- a/Engine/Source/Runtime/CoreUObject/Internal/UObject/PropertyBagRepository.h
+++ b/Engine/Source/Runtime/CoreUObject/Internal/UObject/PropertyBagRepository.h
@@ -129,7 +129,19 @@ public:
 	// query whether an object supports IDO generation
 	static COREUOBJECT_API bool IsInstanceDataObjectSupportEnabled(UObject* InObject = nullptr);
 
-	// create a new placeholder type object to swap in for a missing class/struct; this will be associated with a property bag when objects are serialized so we don't lose data
+	/**
+	 * Create a new placeholder type object to swap in for a missing class/struct. An object of
+	 * this type will be associated with a property bag when serialized so it doesn't lose data.
+	 * 
+	 * @param Outer			Scope at which to create the placeholder type object (e.g. UPackage).
+	 * @param Class			Type object class (or derivative type). For example, UClass::StaticClass().
+	 * @param Name			Optional object name. If not specified, a unique object name will be created.
+	 * @param Flags			Additional object flags. These will be appended to the default set of type object flags.
+	 *						(Note: All placeholder types are transient by definition and internally default to 'RF_Transient'.)
+	 * @param SuperStruct	Optional super type. By default, placeholder types are derivatives of UObject (NULL implies default).
+	 * 
+	 * @return A reference to a new placeholder type object.
+	 */ 
 	static COREUOBJECT_API UStruct* CreatePropertyBagPlaceholderType(UObject* Outer, UClass* Class, FName Name = NAME_None, EObjectFlags Flags = RF_NoFlags, UStruct* SuperStruct = nullptr);
 	template<typename T = UObject>
 	static UClass* CreatePropertyBagPlaceholderClass(UObject* Outer, UClass* Class, FName Name = NAME_None, EObjectFlags Flags = RF_NoFlags)
diff --git a/Engine/Source/Runtime/CoreUObject/Private/UObject/LinkerLoad.cpp b/Engine/Source/Runtime/CoreUObject/Private/UObject/LinkerLoad.cpp
--- a/Engine/Source/Runtime/CoreUObject/Private/UObject/LinkerLoad.cpp
+++ b/Engine/Source/Runtime/CoreUObject/Private/UObject/LinkerLoad.cpp
@@ -3698,25 +3698,47 @@ bool FLinkerLoad::VerifyImportInner(const int32 ImportIndex, FString& WarningSuf
 		{
 			VerifyImport(Import.OuterIndex.ToImport());
 
-			// if the import outer object has been resolved but not linker has been found, we am import to a memory only package (i.e. compiled in)
 			FObjectImport& OuterImport = Imp(Import.OuterIndex);
-			if (!OuterImport.SourceLinker && OuterImport.XObject)
+			if (!OuterImport.SourceLinker)
 			{
-				FObjectImport* Top;
-				for (Top = &OuterImport; Top->OuterIndex.IsImport(); Top = &Imp(Top->OuterIndex))
+				// if the import outer object has been resolved but no linker has been found, we import to a memory only package (i.e. compiled in)
+				if (OuterImport.XObject)
 				{
-					// for loop does what we need
-				}
+					FObjectImport* Top;
+					for (Top = &OuterImport; Top->OuterIndex.IsImport(); Top = &Imp(Top->OuterIndex))
+					{
+						// for loop does what we need
+					}
 
-				UPackage* Package = Cast<UPackage>(Top->XObject);
-				if (Package &&
-					// Assign TmpPkg to resolve the object in memory when there is no source linker available only if the package is MemoryOnly
-					// or we are loading an instanced package in which case the import package might be a duplicated pie package for example for which no linker exists
-					(Package->HasAnyPackageFlags(PKG_InMemoryOnly) || IsContextInstanced()))
+					UPackage* Package = Cast<UPackage>(Top->XObject);
+					if (Package &&
+						// Assign TmpPkg to resolve the object in memory when there is no source linker available only if the package is MemoryOnly
+						// or we are loading an instanced package in which case the import package might be a duplicated pie package for example for which no linker exists
+						(Package->HasAnyPackageFlags(PKG_InMemoryOnly) || IsContextInstanced()))
+					{
+						// This is an import to a memory-only package, just search for it in the package.
+						TmpPkg = Package;
+					}
+				}
+#if WITH_EDITOR
+				else
 				{
-					// This is an import to a memory-only package, just search for it in the package.
-					TmpPkg = Package;
+					const FUObjectSerializeContext* SerializeContext = GetSerializeContext();
+					if (ensure(SerializeContext))
+					{
+						// If we're serializing a redirector's destination object, validate/create the outer package object if it's a missing type. If this import
+						// represents a non-native type object that's no longer valid, this will allow exports of that type to still be serialized to a property bag,
+						// by creating a placeholder type object in its place. This way we won't lose any previously-serialized data for exports missing their type.
+						const UObjectRedirector* Redirector = Cast<UObjectRedirector>(SerializeContext->SerializedObject);
+						if (Redirector && Redirector->IsSerializingDestinationObject() && TryCreatePlaceholderClassImport(ImportIndex))
+						{
+							// We don't need to do a package search for the import below since we've just created it. Return false to signal that there is no failure.
+							check(Import.XObject);
+							return false;
+						}
+					}
 				}
+#endif
 			}
 
 			// Copy the SourceLinker from the FObjectImport for our Outer if the SourceLinker hasn't been set yet,
@@ -4054,7 +4076,7 @@ UClass* FLinkerLoad::GetExportLoadClass(int32 Index)
 }
 
 #if WITH_EDITOR
-UClass* FLinkerLoad::TryCreatePlaceholderTypeForExport(int32 ExportIndex)
+UClass* FLinkerLoad::TryCreatePlaceholderClassImport(int32 ImportIndex)
 {
 	const bool bAllowPlaceholderImportTypes = UE::FPropertyBagRepository::IsPropertyBagPlaceholderObjectSupportEnabled();
 	if (!bAllowPlaceholderImportTypes)
@@ -4062,45 +4084,62 @@ UClass* FLinkerLoad::TryCreatePlaceholderTypeForExport(int32 ExportIndex)
 		return nullptr;
 	}
 
-	UClass* LoadClass = nullptr;
-	FObjectExport& Export = ExportMap[ExportIndex];
+	UClass* ClassObject = nullptr;
+	FObjectImport& Import = ImportMap[ImportIndex];
 
-	// If the class import is missing, create a placeholder for this export. This will allow us to instance and redirect its data into a property bag.
-	if (Export.ClassIndex.IsImport() && !GEventDrivenLoaderEnabled)
+	// If the import is already set, return NULL to indicate that we didn't create a placeholder type object.
+	if (Import.XObject)
 	{
-		FObjectImport& LoadClassImport = Imp(Export.ClassIndex);
+		return nullptr;
+	}
 
-		// If the outer package import is also missing, create it now so that the full path remains the same. 
-		UObject* LoadClassParent = IndexToObject(LoadClassImport.OuterIndex);
-		if (!LoadClassParent && LoadClassImport.OuterIndex.IsImport())
+	if (UObject* ImportClassPackage = FindObjectFast<UPackage>(nullptr, Import.ClassPackage, /*bExactClass =*/ true))
+	{
+		if (UClass* ImportClass = FindObjectFast<UClass>(ImportClassPackage, Import.ClassName, /*bExactClass =*/ false))
 		{
-			FObjectImport& LoadClassParentImport = Imp(LoadClassImport.OuterIndex);
-			if (LoadClassParentImport.OuterIndex.IsNull())
+			if (ImportClass->IsChildOf<UClass>())
 			{
-				LoadClassParent = CreatePackage(*LoadClassParentImport.ObjectName.ToString());
+				// If the outer package import is also missing, create it now so that the full path remains the same. 
+				UPackage* ClassObjectPackage = Cast<UPackage>(IndexToObject(Import.OuterIndex));
+				if (!ClassObjectPackage && Import.OuterIndex.IsImport())
+				{
+					FObjectImport& OuterImport = Imp(Import.OuterIndex);
+					if (OuterImport.OuterIndex.IsNull())
+					{
+						ClassObjectPackage = CreatePackage(*OuterImport.ObjectName.ToString());
 
-				// Patch it into the import table so that we resolve to this package for future reference.
-				LoadClassParentImport.XObject = LoadClassParent;
-			}
-		}
+						// Patch it into the import table so that we resolve to this package for future reference.
+						OuterImport.XObject = ClassObjectPackage;
+					}
+				}
 
-		if (LoadClassParent)
-		{
-			if (UObject* LoadClassTypePackage = FindObjectFast<UPackage>(nullptr, LoadClassImport.ClassPackage, /*bExactClass =*/ true))
-			{
-				if (UClass* LoadClassType = FindObjectFast<UClass>(LoadClassTypePackage, LoadClassImport.ClassName, /*bExactClass =*/ false))
+				if (ClassObjectPackage)
 				{
 					// Create an opaque, non-native transient type object that has no reflected properties.
-					LoadClass = UE::FPropertyBagRepository::CreatePropertyBagPlaceholderClass(LoadClassParent, LoadClassType, LoadClassImport.ObjectName, RF_Transient);
+					ClassObject = UE::FPropertyBagRepository::CreatePropertyBagPlaceholderClass(ClassObjectPackage, ImportClass, Import.ObjectName);
 
 					// Patch it into the import table so that we resolve to this class for any future exports of this type.
-					LoadClassImport.XObject = LoadClass;
+					Import.XObject = ClassObject;
 				}
 			}
 		}
 	}
 
-	return LoadClass;
+	return ClassObject;
+}
+
+UClass* FLinkerLoad::TryCreatePlaceholderClassForExport(int32 ExportIndex)
+{
+	UClass* ClassObject = nullptr;
+	FObjectExport& Export = ExportMap[ExportIndex];
+
+	// If the class import is missing, create a placeholder for this export. This will allow us to instance and redirect its data into a property bag.
+	if (Export.ClassIndex.IsImport())
+	{
+		ClassObject = TryCreatePlaceholderClassImport(Export.ClassIndex.ToImport());
+	}
+
+	return ClassObject;
 }
 #endif
 
@@ -5044,7 +5083,7 @@ UObject* FLinkerLoad::CreateExport( int32 Index )
 		{
 #if WITH_EDITOR
 			// Try creating a placeholder type for it. This may allow us to instance and redirect its data into a property bag (to avoid data loss).
-			LoadClass = TryCreatePlaceholderTypeForExport(Index);
+			LoadClass = TryCreatePlaceholderClassForExport(Index);
 			if (!LoadClass)
 #endif
 			{
diff --git a/Engine/Source/Runtime/CoreUObject/Private/UObject/ObjectRedirector.cpp b/Engine/Source/Runtime/CoreUObject/Private/UObject/ObjectRedirector.cpp
--- a/Engine/Source/Runtime/CoreUObject/Private/UObject/ObjectRedirector.cpp
+++ b/Engine/Source/Runtime/CoreUObject/Private/UObject/ObjectRedirector.cpp
@@ -53,7 +53,12 @@ void UObjectRedirector::Serialize( FStructuredArchive::FRecord Record )
 {
 	Super::Serialize(Record);
 
-	Record << SA_VALUE(TEXT("DestinationObject"), DestinationObject);
+	{
+#if WITH_EDITOR
+		TGuardValue<bool> IsSerializingDestinationObjectScope(bIsSerializingDestinationObject, true);
+#endif
+		Record << SA_VALUE(TEXT("DestinationObject"), DestinationObject);
+	}
 }
 
 bool UObjectRedirector::NeedsLoadForEditorGame() const
diff --git a/Engine/Source/Runtime/CoreUObject/Private/UObject/PropertyBagRepository.cpp b/Engine/Source/Runtime/CoreUObject/Private/UObject/PropertyBagRepository.cpp
--- a/Engine/Source/Runtime/CoreUObject/Private/UObject/PropertyBagRepository.cpp
+++ b/Engine/Source/Runtime/CoreUObject/Private/UObject/PropertyBagRepository.cpp
@@ -413,7 +413,7 @@ bool FPropertyBagRepository::IsPropertyBagPlaceholderObjectSupportEnabled()
 	static bool bIsInitialized = false;
 	if (!bIsInitialized)
 	{
-		Private::bEnablePropertyBagPlaceholderObjectSupport = FParse::Param(FCommandLine::Get(), TEXT("WithPropertyBagPlaceholderObjects"));
+		Private::bEnablePropertyBagPlaceholderObjectSupport |= FParse::Param(FCommandLine::Get(), TEXT("WithPropertyBagPlaceholderObjects"));
 		bIsInitialized = true;
 	}
 	
@@ -434,7 +434,8 @@ bool FPropertyBagRepository::IsInstanceDataObjectSupportEnabled(UObject* InObjec
 
 UStruct* FPropertyBagRepository::CreatePropertyBagPlaceholderType(UObject* Outer, UClass* Class, FName Name, EObjectFlags Flags, UStruct* SuperStruct)
 {
-	UStruct* PlaceholderType = NewObject<UClass>(Outer, Class, Name, Flags);
+	// Generate and link a transient type object using the given SuperStruct as its base.
+	UStruct* PlaceholderType = NewObject<UClass>(Outer, Class, Name, Flags | RF_Transient);
 	PlaceholderType->SetSuperStruct(SuperStruct);
 	PlaceholderType->Bind();
 	PlaceholderType->StaticLink(/*bRelinkExistingProperties =*/ true);
diff --git a/Engine/Source/Runtime/CoreUObject/Public/UObject/LinkerLoad.h b/Engine/Source/Runtime/CoreUObject/Public/UObject/LinkerLoad.h
--- a/Engine/Source/Runtime/CoreUObject/Public/UObject/LinkerLoad.h
+++ b/Engine/Source/Runtime/CoreUObject/Public/UObject/LinkerLoad.h
@@ -864,14 +864,28 @@ private:
 
 #if WITH_EDITOR
 	/**
-	 * Utility function to create a placeholder type for the specified export.
+	 * Utility function for callers to try and create a placeholder class
+	 * object in the specified import slot if it represents a class object.
+	 * This allows for objects to still be serialized even if missing its
+	 * class type on load (the data will be serialized into a property bag).
+	 * 
+	 * Note: If the import is already non-NULL, this will not create a
+	 * placeholder class (and the function will return NULL as a result).
+	 *
+	 * @param  ImportIndex    Index of the import that may be a missing class.
+	 * @return A reference to the placeholder class (if created), otherwise NULL.
+	 */
+	UClass* TryCreatePlaceholderClassImport(int32 ImportIndex);
+
+	/**
+	 * Utility function to create a placeholder class for the specified export.
 	 * This will be called if the export could not resolve its LoadClass import.
 	 * In that case, its data might still be serializable (into a property bag).
 	 *
 	 * @param  ExportIndex    Index of the export that's missing its type.
 	 * @return The placeholder type that will be used to load the export's data.
 	 */
-	UClass* TryCreatePlaceholderTypeForExport(int32 ExportIndex);
+	UClass* TryCreatePlaceholderClassForExport(int32 ExportIndex);
 #endif
 
 #if WITH_EDITORONLY_DATA
diff --git a/Engine/Source/Runtime/CoreUObject/Public/UObject/ObjectRedirector.h b/Engine/Source/Runtime/CoreUObject/Public/UObject/ObjectRedirector.h
--- a/Engine/Source/Runtime/CoreUObject/Public/UObject/ObjectRedirector.h
+++ b/Engine/Source/Runtime/CoreUObject/Public/UObject/ObjectRedirector.h
@@ -59,4 +59,15 @@ class UObjectRedirector : public UObject
 	 * @return	return true if property values were added to the map.
 	 */
 	virtual bool GetNativePropertyValues( TMap<FString,FString>& out_PropertyValues, uint32 ExportFlags=0 ) const override;
+
+#if WITH_EDITOR
+	/** @return TRUE when serializing the reference to the destination object. */
+	bool IsSerializingDestinationObject() const
+	{
+		return bIsSerializingDestinationObject;
+	}
+
+private:
+	bool bIsSerializingDestinationObject = false;
+#endif
 };
