Added FinishDestroy() handling for properties.
- Added struct trait, WithClearOnFinishDestroy which cause the struct to be cleared during owning UObject's FinishDestroy()
- Added FProperty::ContainsFinishDestroy(), which redursively checks if the property or its children requires FinishDestroy() to be called
- Added FProperty::FinishDestroy() which recursively calls FinishDestroy(), struct propertys FinishDestroy() clears the struct if specified in the traits
- UStruct::Link() adds properties requiring finish destrou to the DestructLink
- Changed UObject::DestroyNonNativeProperties() to call FinishDestroy() on native properties in DestructorLink
- Added WithClearOnFinishDestroy to FInstancedStruct so that the contained value is destroyed safely in the case that the owning UObject and UScriptStruct are GC'd at the same time.

diff --git a/Engine/Plugins/Experimental/StructUtils/Source/StructUtils/Public/InstancedStruct.h b/Engine/Plugins/Experimental/StructUtils/Source/StructUtils/Public/InstancedStruct.h
--- a/Engine/Plugins/Experimental/StructUtils/Source/StructUtils/Public/InstancedStruct.h
+++ b/Engine/Plugins/Experimental/StructUtils/Source/StructUtils/Public/InstancedStruct.h
@@ -272,6 +272,7 @@ struct TStructOpsTypeTraits<FInstancedStruct> : public TStructOpsTypeTraitsBase2
 		WithGetPreloadDependencies = true,
 		WithNetSerializer = true,
 		WithFindInnerPropertyInstance = true,
+		WithClearOnFinishDestroy = true,
 	};
 };
 
diff --git a/Engine/Plugins/Experimental/StructUtils/Source/StructUtilsTestSuite/Private/InstancedStructTest.cpp b/Engine/Plugins/Experimental/StructUtils/Source/StructUtilsTestSuite/Private/InstancedStructTest.cpp
--- a/Engine/Plugins/Experimental/StructUtils/Source/StructUtilsTestSuite/Private/InstancedStructTest.cpp
+++ b/Engine/Plugins/Experimental/StructUtils/Source/StructUtilsTestSuite/Private/InstancedStructTest.cpp
@@ -110,6 +110,44 @@ namespace FInstancedStructTest
 	};
 	IMPLEMENT_AI_INSTANT_TEST(FTest_InstancedStructBasic, "System.StructUtils.InstancedStruct.Basic");
 
+	struct FTest_InstancedStructCustomScriptStruct : FAITestBase
+	{
+		virtual bool InstantTest() override
+		{
+			// Create the TestObject before UScriptStruct, so that CustomStruct gets destroyed first.
+			TWeakObjectPtr<UTestObjectWithInstanceStruct> TestObject = NewObject<UTestObjectWithInstanceStruct>();
+			check(TestObject.IsValid());
+
+			TWeakObjectPtr<UScriptStruct> CustomStruct = NewObject<UScriptStruct>();
+			check(CustomStruct.IsValid());
+
+			FIntProperty* IntProp = new FIntProperty(CustomStruct.Get(), FName("Int"), RF_Public);
+			check(IntProp);
+			CustomStruct->AddCppProperty(IntProp);
+
+			FStrProperty* StrProp = new FStrProperty(CustomStruct.Get(), FName("String"), RF_Public);
+			check(StrProp);
+			CustomStruct->AddCppProperty(StrProp);
+
+			CustomStruct->SetSuperStruct(nullptr);
+			CustomStruct->Bind();
+			CustomStruct->StaticLink(/*RelinkExistingProperties*/true);
+
+			
+			TestObject->Value.InitializeAs(CustomStruct.Get());
+			AITEST_TRUE("FInstancedStruct created to a specific struct type should be IsValid()", TestObject->Value.IsValid());
+
+			// CustomStruct and TestObject should both get collected.
+			CollectGarbage(GARBAGE_COLLECTION_KEEPFLAGS);
+
+			AITEST_FALSE("CustomStruct should not be valid", CustomStruct.IsValid());
+			AITEST_FALSE("TestObject should not be valid", TestObject.IsValid());
+			
+			return true;
+		}
+	};
+	IMPLEMENT_AI_INSTANT_TEST(FTest_InstancedStructCustomScriptStruct, "System.StructUtils.InstancedStruct.CustomScriptStruct");
+
 }
 
 
diff --git a/Engine/Plugins/Experimental/StructUtils/Source/StructUtilsTestSuite/Public/StructUtilsTestTypes.h b/Engine/Plugins/Experimental/StructUtils/Source/StructUtilsTestSuite/Public/StructUtilsTestTypes.h
--- a/Engine/Plugins/Experimental/StructUtils/Source/StructUtilsTestSuite/Public/StructUtilsTestTypes.h
+++ b/Engine/Plugins/Experimental/StructUtils/Source/StructUtilsTestSuite/Public/StructUtilsTestTypes.h
@@ -185,6 +185,17 @@ public:
 	FInstancedPropertyBag Bag;
 };
 
+UCLASS()
+class UTestObjectWithInstanceStruct : public UObject
+{
+	GENERATED_BODY()
+
+public:
+	UPROPERTY()
+	FInstancedStruct Value;
+};
+
+
 #if UE_ENABLE_INCLUDE_ORDER_DEPRECATED_IN_5_2
 #include "CoreMinimal.h"
 #endif
diff --git a/Engine/Source/Runtime/CoreUObject/Private/Blueprint/BlueprintSupport.cpp b/Engine/Source/Runtime/CoreUObject/Private/Blueprint/BlueprintSupport.cpp
--- a/Engine/Source/Runtime/CoreUObject/Private/Blueprint/BlueprintSupport.cpp
+++ b/Engine/Source/Runtime/CoreUObject/Private/Blueprint/BlueprintSupport.cpp
@@ -2616,7 +2616,17 @@ void UObject::DestroyNonNativeProperties()
 
 	for (FProperty* P = GetClass()->DestructorLink; P; P = P->DestructorLinkNext)
 	{
-		P->DestroyValue_InContainer(this);
+		if (!P->GetOwnerClass()->HasAnyClassFlags(CLASS_Native | CLASS_Intrinsic))
+		{
+			// Non-native value, destroy.
+			P->DestroyValue_InContainer(this);
+		}
+		else
+		{
+			// Native value, call to handle finish destroy.
+			// Native properties appear here if they report true on ContainsFinishDestroy() during UStruct::Link().
+			P->FinishDestroy_InContainer(this);
+		}
 	}
 }
 
diff --git a/Engine/Source/Runtime/CoreUObject/Private/UObject/Class.cpp b/Engine/Source/Runtime/CoreUObject/Private/UObject/Class.cpp
--- a/Engine/Source/Runtime/CoreUObject/Private/UObject/Class.cpp
+++ b/Engine/Source/Runtime/CoreUObject/Private/UObject/Class.cpp
@@ -959,10 +959,12 @@ void UStruct::Link(FArchive& Ar, bool bRelinkExistingProperties)
 
 		const UClass* OwnerClass = Property->GetOwnerClass();
 		bool bOwnedByNativeClass = OwnerClass && OwnerClass->HasAnyClassFlags(CLASS_Native | CLASS_Intrinsic);
+		bool bShouldHandleFinishDestroy = Property->ContainsFinishDestroy(EncounteredStructProps);
 
-		if (!Property->HasAnyPropertyFlags(CPF_IsPlainOldData | CPF_NoDestructor) &&
-			!bOwnedByNativeClass) // these would be covered by the native destructor
-		{	
+		if ((!Property->HasAnyPropertyFlags(CPF_IsPlainOldData | CPF_NoDestructor)
+			&& !bOwnedByNativeClass) // these would be covered by the native destructor
+			|| bShouldHandleFinishDestroy)
+		{
 			// things in a struct that need a destructor will still be in here, even though in many cases they will also be destroyed by a native destructor on the whole struct
 			*DestructorLinkPtr = Property;
 			DestructorLinkPtr = &(*DestructorLinkPtr)->DestructorLinkNext;
diff --git a/Engine/Source/Runtime/CoreUObject/Private/UObject/Property.cpp b/Engine/Source/Runtime/CoreUObject/Private/UObject/Property.cpp
--- a/Engine/Source/Runtime/CoreUObject/Private/UObject/Property.cpp
+++ b/Engine/Source/Runtime/CoreUObject/Private/UObject/Property.cpp
@@ -928,6 +928,16 @@ void FProperty::InitializeValueInternal( void* Dest ) const
 	checkf(0, TEXT("%s failed to handle InitializeValueInternal, but it was not CPF_ZeroConstructor"), *GetFullName());
 }
 
+bool FProperty::ContainsClearOnFinishDestroyInternal( TArray<const FStructProperty*>& EncounteredStructProps ) const
+{
+	return false;
+}
+
+void FProperty::FinishDestroyInternal( void* Data ) const
+{
+	// Empty
+}
+
 /**
  * Verify that modifying this property's value via ImportText is allowed.
  * 
diff --git a/Engine/Source/Runtime/CoreUObject/Private/UObject/PropertyArray.cpp b/Engine/Source/Runtime/CoreUObject/Private/UObject/PropertyArray.cpp
--- a/Engine/Source/Runtime/CoreUObject/Private/UObject/PropertyArray.cpp
+++ b/Engine/Source/Runtime/CoreUObject/Private/UObject/PropertyArray.cpp
@@ -1111,6 +1111,29 @@ bool FArrayProperty::PassCPPArgsByRef() const
 	return true;
 }
 
+bool FArrayProperty::ContainsClearOnFinishDestroyInternal(TArray<const FStructProperty*>& EncounteredStructProps) const
+{
+	check(Inner);
+	return Inner->ContainsFinishDestroy(EncounteredStructProps);
+}
+
+void FArrayProperty::FinishDestroyInternal( void* Data ) const
+{
+	if (!Data)
+	{
+		return;
+	}
+	check(Inner);
+	if ((Inner->PropertyFlags & (CPF_IsPlainOldData | CPF_NoDestructor)) == 0)
+	{
+		FScriptArrayHelper ArrayHelper(this, Data);
+		for (int32 ElementIndex = 0; ElementIndex < ArrayHelper.Num(); ElementIndex++)
+		{
+			Inner->FinishDestroy(ArrayHelper.GetRawPtr(ElementIndex));
+		}
+	}	
+}
+
 /**
  * Creates new copies of components
  *
diff --git a/Engine/Source/Runtime/CoreUObject/Private/UObject/PropertyMap.cpp b/Engine/Source/Runtime/CoreUObject/Private/UObject/PropertyMap.cpp
--- a/Engine/Source/Runtime/CoreUObject/Private/UObject/PropertyMap.cpp
+++ b/Engine/Source/Runtime/CoreUObject/Private/UObject/PropertyMap.cpp
@@ -1353,6 +1353,45 @@ bool FMapProperty::PassCPPArgsByRef() const
 	return true;
 }
 
+bool FMapProperty::ContainsClearOnFinishDestroyInternal(TArray<const FStructProperty*>& EncounteredStructProps) const
+{
+	check(KeyProp);
+	check(ValueProp);
+	return KeyProp->ContainsFinishDestroy(EncounteredStructProps) || ValueProp->ContainsFinishDestroy(EncounteredStructProps);
+}
+
+void FMapProperty::FinishDestroyInternal( void* Data ) const
+{
+	if (!Data)
+	{
+		return;
+	}
+	
+	check(KeyProp);
+	check(ValueProp);
+
+	const bool bMayHaveFinishDestroyKey   = (KeyProp->PropertyFlags & (CPF_IsPlainOldData | CPF_NoDestructor)) == 0;
+	const bool bMayHaveFinishDestroyValue = (ValueProp->PropertyFlags & (CPF_IsPlainOldData | CPF_NoDestructor)) == 0;
+
+	if (bMayHaveFinishDestroyKey
+		|| bMayHaveFinishDestroyValue)
+	{
+		FScriptMapHelper MapHelper(this, Data);
+		for (FScriptMapHelper::FIterator It(MapHelper.CreateIterator()); It; ++It)
+		{
+			uint8* PairPtr = MapHelper.GetPairPtr(It);
+			if (bMayHaveFinishDestroyKey)
+			{
+				KeyProp->FinishDestroy(PairPtr);
+			}
+			if (bMayHaveFinishDestroyValue)
+			{
+				ValueProp->FinishDestroy(PairPtr + MapLayout.ValueOffset);
+			}
+		}
+	}
+}
+
 /**
  * Creates new copies of components
  * 
diff --git a/Engine/Source/Runtime/CoreUObject/Private/UObject/PropertyOptional.cpp b/Engine/Source/Runtime/CoreUObject/Private/UObject/PropertyOptional.cpp
--- a/Engine/Source/Runtime/CoreUObject/Private/UObject/PropertyOptional.cpp
+++ b/Engine/Source/Runtime/CoreUObject/Private/UObject/PropertyOptional.cpp
@@ -444,6 +444,21 @@ void FOptionalProperty::DestroyValueInternal(void* Data) const
 	MarkUnset(Data);
 }
 
+bool FOptionalProperty::ContainsClearOnFinishDestroyInternal(TArray<const FStructProperty*>& EncounteredStructProps) const
+{
+	check(ValueProperty);
+	return ValueProperty->ContainsFinishDestroy(EncounteredStructProps);
+}
+
+void FOptionalProperty::FinishDestroyInternal( void* Data ) const
+{
+	check(ValueProperty);
+	if (void* Value = GetValuePointerForReplaceIfSet(Data))
+	{
+		ValueProperty->FinishDestroy(Value);
+	}
+}
+
 void FOptionalProperty::InstanceSubobjects(void* Data, void const* DefaultData, UObject* InOwner, struct FObjectInstancingGraph* InstanceGraph)
 {
 	if (Data && IsSet(Data) && ValueProperty->ContainsInstancedObjectProperty())
diff --git a/Engine/Source/Runtime/CoreUObject/Private/UObject/PropertySet.cpp b/Engine/Source/Runtime/CoreUObject/Private/UObject/PropertySet.cpp
--- a/Engine/Source/Runtime/CoreUObject/Private/UObject/PropertySet.cpp
+++ b/Engine/Source/Runtime/CoreUObject/Private/UObject/PropertySet.cpp
@@ -809,6 +809,25 @@ bool FSetProperty::PassCPPArgsByRef() const
 	return true;
 }
 
+bool FSetProperty::ContainsClearOnFinishDestroyInternal(TArray<const FStructProperty*>& EncounteredStructProps) const
+{
+	check(ElementProp);
+	return ElementProp->ContainsFinishDestroy(EncounteredStructProps);
+}
+
+void FSetProperty::FinishDestroyInternal( void* Data ) const
+{
+	if ((ElementProp->PropertyFlags & (CPF_IsPlainOldData | CPF_NoDestructor)) == 0)
+	{
+		FScriptSetHelper SetHelper(this, Data);
+		for (FScriptSetHelper::FIterator It(SetHelper); It; ++It)
+		{
+			uint8* ElementPtr = SetHelper.GetElementPtr(It);
+			ElementProp->FinishDestroy(ElementPtr);
+		}
+	}
+}
+
 /**
  * Creates new copies of components
  * 
diff --git a/Engine/Source/Runtime/CoreUObject/Private/UObject/PropertyStruct.cpp b/Engine/Source/Runtime/CoreUObject/Private/UObject/PropertyStruct.cpp
--- a/Engine/Source/Runtime/CoreUObject/Private/UObject/PropertyStruct.cpp
+++ b/Engine/Source/Runtime/CoreUObject/Private/UObject/PropertyStruct.cpp
@@ -365,6 +365,68 @@ void FStructProperty::DestroyValueInternal( void* Dest ) const
 	Struct->DestroyStruct(Dest, ArrayDim);
 }
 
+bool FStructProperty::ContainsClearOnFinishDestroyInternal(TArray<const FStructProperty*>& EncounteredStructProps) const
+{
+	// Skip if already being processed.
+	if (EncounteredStructProps.Contains(this))
+	{
+		return false;
+	}
+
+	if (!Struct)
+	{
+		UE_LOG(LogGarbage, Warning, TEXT("Broken FStructProperty does not have a UStruct: %s"), *GetFullName() );
+		return false;
+	}
+
+	if (const UScriptStruct::ICppStructOps* CppStructOps = Struct->GetCppStructOps())
+	{
+		if (CppStructOps->HasClearOnFinishDestroy())
+		{
+			return true;
+		}
+	}
+
+	EncounteredStructProps.Add(this);
+
+	bool bValue = false;
+	for (const FProperty* Property = Struct->PropertyLink; Property; Property = Property->PropertyLinkNext)
+	{
+		if (Property->ContainsFinishDestroy(EncounteredStructProps))
+		{
+			bValue = true;
+			break;
+		}
+	}
+
+	EncounteredStructProps.RemoveSingleSwap(this, EAllowShrinking::No);
+	
+	return bValue;
+}
+
+void FStructProperty::FinishDestroyInternal( void* Data ) const
+{
+	const int32 Stride = Struct->GetStructureSize();
+	
+	if (UScriptStruct::ICppStructOps* CppStructOps = Struct->GetCppStructOps())
+	{
+		if (CppStructOps->HasClearOnFinishDestroy())
+		{
+			Struct->ClearScriptStruct(Data, ArrayDim);
+			return;
+		}
+	}
+	
+	for (int32 ArrayIndex = 0; ArrayIndex < ArrayDim; ArrayIndex++)
+	{
+		uint8* ItemData = (uint8*)Data + ArrayIndex * Stride;
+		for (const FProperty* Property = Struct->PropertyLink; Property; Property = Property->PropertyLinkNext)
+		{
+			Property->FinishDestroy(Property->ContainerPtrToValuePtr<void>(ItemData));
+		}
+	}
+}
+
 /**
  * Creates new copies of components
  * 
diff --git a/Engine/Source/Runtime/CoreUObject/Public/UObject/Class.h b/Engine/Source/Runtime/CoreUObject/Public/UObject/Class.h
--- a/Engine/Source/Runtime/CoreUObject/Public/UObject/Class.h
+++ b/Engine/Source/Runtime/CoreUObject/Public/UObject/Class.h
@@ -855,12 +855,13 @@ struct TStructOpsTypeTraitsBase2
 		WithNetDeltaSerializer         = false,                         // struct has a NetDeltaSerialize function for serializing differences in state from a previous NetSerialize operation.
 		WithSerializeFromMismatchedTag = false,                         // struct has a SerializeFromMismatchedTag function for converting from other property tags.
 		WithStructuredSerializeFromMismatchedTag = false,               // struct has an FStructuredArchive-based SerializeFromMismatchedTag function for converting from other property tags.
-		WithPostScriptConstruct        = false,                         // struct has a PostScriptConstruct function which is called after it is constructed in blueprints
+		WithPostScriptConstruct        = false,                         // struct has a PostScriptConstruct func tion which is called after it is constructed in blueprints
 		WithNetSharedSerialization     = false,                         // struct has a NetSerialize function that does not require the package map to serialize its state.
 		WithGetPreloadDependencies     = false,                         // struct has a GetPreloadDependencies function to return all objects that will be Preload()ed when the struct is serialized at load time.
 		WithPureVirtual                = false,                         // struct has PURE_VIRTUAL functions and cannot be constructed when CHECK_PUREVIRTUALS is true
 		WithFindInnerPropertyInstance  = false,							// struct has a FindInnerPropertyInstance function that can provide an FProperty and data pointer when given a property FName
 		WithCanEditChange			   = false,							// struct has an editor-only CanEditChange function that can conditionally make child properties read-only in the details panel (same idea as UObject::CanEditChange)
+		WithClearOnFinishDestroy	   = false,							// struct should be cleared during owner UObject's FinishDestroy. Clearing calls destructor and initializes again to default value. This is intended for structs which may need to access UObject pointer members during destruction. Referenced objects may already have their FinishDestroy() called. Clearing should ensure that no UObject pointer members are used during the final destruction.
 	};
 
 	static constexpr EPropertyObjectReferenceType WithSerializerObjectReferences = EPropertyObjectReferenceType::Conservative; // struct's Serialize method(s) may serialize object references of these types - default Conservative means unknown and object reference collector archives should serialize this struct 
@@ -927,6 +928,7 @@ public:
 			bool HasGetTypeHash : 1;
 			bool IsAbstract : 1;
 			bool HasFindInnerPropertyInstance : 1;
+			bool ClearOnFinishDestroy : 1;
 #if WITH_EDITOR
 			bool HasCanEditChange : 1;
 #endif
@@ -1165,6 +1167,12 @@ public:
 			return GetCapabilities().IsAbstract;
 		}
 
+		/** return true if this struct has a ClearOnFinishDestroy */
+		bool HasClearOnFinishDestroy() const
+		{
+			return GetCapabilities().ClearOnFinishDestroy;
+		}
+
 #if WITH_EDITOR
 		/** Returns true if this struct wants to indicate whether a property can be edited in the details panel */
 		bool HasCanEditChange() const
@@ -1224,6 +1232,7 @@ public:
 				TModels_V<CGetTypeHashable, CPPSTRUCT>,
 				TIsAbstract<CPPSTRUCT>::Value,
 				TTraits::WithFindInnerPropertyInstance,
+				TTraits::WithClearOnFinishDestroy,
 #if WITH_EDITOR
 				TTraits::WithCanEditChange,
 #endif
diff --git a/Engine/Source/Runtime/CoreUObject/Public/UObject/PropertyOptional.h b/Engine/Source/Runtime/CoreUObject/Public/UObject/PropertyOptional.h
--- a/Engine/Source/Runtime/CoreUObject/Public/UObject/PropertyOptional.h
+++ b/Engine/Source/Runtime/CoreUObject/Public/UObject/PropertyOptional.h
@@ -222,6 +222,8 @@ public:
 	virtual void ClearValueInternal(void* Data) const override;
 	virtual void InitializeValueInternal(void* Data) const override;
 	virtual void DestroyValueInternal(void* Data) const override;
+	virtual bool ContainsClearOnFinishDestroyInternal(TArray<const FStructProperty*>& EncounteredStructProps) const override;
+	virtual void FinishDestroyInternal(void* Data) const override;
 	virtual void InstanceSubobjects(void* Data, void const* DefaultData, UObject* Owner, struct FObjectInstancingGraph* InstanceGraph) override;
 	virtual int32 GetMinAlignment() const override;
 	virtual bool ContainsObjectReference(TArray<const FStructProperty*>& EncounteredStructProps, EPropertyObjectReferenceType InReferenceType = EPropertyObjectReferenceType::Strong) const override;
diff --git a/Engine/Source/Runtime/CoreUObject/Public/UObject/UnrealType.h b/Engine/Source/Runtime/CoreUObject/Public/UObject/UnrealType.h
--- a/Engine/Source/Runtime/CoreUObject/Public/UObject/UnrealType.h
+++ b/Engine/Source/Runtime/CoreUObject/Public/UObject/UnrealType.h
@@ -929,10 +929,62 @@ public:
 			DestroyValueInternal(ContainerPtrToValuePtr<void>(Dest));
 		}
 	}
+
 protected:
 	COREUOBJECT_API virtual void DestroyValueInternal( void* Dest ) const;
 public:
 
+	/**
+	 * Returns true if the property or any of the child properties should be cleared on FinishDestroy.
+	 */
+	COREUOBJECT_API bool ContainsFinishDestroy(TArray<const FStructProperty*>& EncounteredStructProps) const
+	{
+		// Skip if the property does not need any destroying. 
+		if (PropertyFlags & (CPF_IsPlainOldData | CPF_NoDestructor))
+		{
+			return false;
+		}
+		return ContainsClearOnFinishDestroyInternal(EncounteredStructProps);
+	}
+
+	/**
+	 * Applies appropriate finish destroy actions for the property if needed.
+	 * This is used during UObject destruction to e.g. safely clear values which rely on UScriptStructs. 
+	 * This does the entire fixed size array.
+	 *
+	 * @param	Data		the address of the value for this property that should be handled for finish destroy.
+	 */
+	COREUOBJECT_API void FinishDestroy( void* Data ) const
+	{
+		// Skip if the property does not need any destroying. 
+		if (PropertyFlags & (CPF_IsPlainOldData | CPF_NoDestructor))
+		{
+			return;
+		}
+		FinishDestroyInternal(Data);
+	}
+
+	/**
+	 * Applies appropriate finish destroy actions for the property if needed.
+	 * This is used during UObject destruction to e.g. safely clear values which rely on UScriptStructs. 
+	 * This does the entire fixed size array.
+	 *
+	 * @param	Data		the address of the container containing the value that should be handled for finish destroy.
+	 */
+	COREUOBJECT_API void FinishDestroy_InContainer( void* Data ) const
+	{
+		// Skip if the property does not need any destroying. 
+		if (PropertyFlags & (CPF_IsPlainOldData | CPF_NoDestructor))
+		{
+			return;
+		}
+		FinishDestroyInternal(ContainerPtrToValuePtr<void>(Data));
+	}
+protected:
+	COREUOBJECT_API virtual bool ContainsClearOnFinishDestroyInternal(TArray<const FStructProperty*>& EncounteredStructProps) const;
+	COREUOBJECT_API virtual void FinishDestroyInternal( void* Data ) const;
+public:
+
 	/**
 	 * Zeros, copies from the default, or calls the constructor for on the value for this property. 
 	 * The existing data is assumed invalid (so for example this might indirectly call FString::FString,
@@ -3410,6 +3462,8 @@ public:
 protected:
 	virtual void ExportText_Internal( FString& ValueStr, const void* PropertyValueOrContainer, EPropertyPointerType PropertyPointerType, const void* DefaultValue, UObject* Parent, int32 PortFlags, UObject* ExportRootScope ) const override;
 	virtual const TCHAR* ImportText_Internal(const TCHAR* Buffer, void* ContainerOrPropertyPtr, EPropertyPointerType PropertyPointerType, UObject* OwnerObject, int32 PortFlags, FOutputDevice* ErrorText) const override;
+	virtual bool ContainsClearOnFinishDestroyInternal(TArray<const FStructProperty*>& EncounteredStructProps) const override;
+	virtual void FinishDestroyInternal(void* Data) const override;
 public:
 	virtual void InitializeValueInternal(void* Dest) const override
 	{
@@ -3546,6 +3600,8 @@ public:
 protected:
 	virtual void ExportText_Internal(FString& ValueStr, const void* PropertyValueOrContainer, EPropertyPointerType PropertyPointerType, const void* DefaultValue, UObject* Parent, int32 PortFlags, UObject* ExportRootScope) const override;
 	virtual const TCHAR* ImportText_Internal(const TCHAR* Buffer, void* ContainerOrPropertyPtr, EPropertyPointerType PropertyPointerType, UObject* OwnerObject, int32 PortFlags, FOutputDevice* ErrorText) const override;
+	virtual bool ContainsClearOnFinishDestroyInternal(TArray<const FStructProperty*>& EncounteredStructProps) const override;
+	virtual void FinishDestroyInternal(void* Data) const override;
 public:
 	virtual void InitializeValueInternal(void* Dest) const override
 	{
@@ -3689,6 +3745,8 @@ public:
 protected:
 	virtual void ExportText_Internal(FString& ValueStr, const void* PropertyValueOrContainer, EPropertyPointerType PropertyPointerType, const void* DefaultValue, UObject* Parent, int32 PortFlags, UObject* ExportRootScope) const override;
 	virtual const TCHAR* ImportText_Internal(const TCHAR* Buffer, void* ContainerOrPropertyPtr, EPropertyPointerType PropertyPointerType, UObject* OwnerObject, int32 PortFlags, FOutputDevice* ErrorText) const override;
+	virtual bool ContainsClearOnFinishDestroyInternal(TArray<const FStructProperty*>& EncounteredStructProps) const override;
+	virtual void FinishDestroyInternal(void* Data) const override;
 public:
 	virtual void CopyValuesInternal(void* Dest, void const* Src, int32 Count) const override;
 	virtual void ClearValueInternal(void* Data) const override;
@@ -5808,6 +5866,8 @@ public:
 protected:
 	virtual void ExportText_Internal( FString& ValueStr, const void* PropertyValueOrContainer, EPropertyPointerType PropertyPointerType, const void* DefaultValue, UObject* Parent, int32 PortFlags, UObject* ExportRootScope ) const override;
 	virtual const TCHAR* ImportText_Internal(const TCHAR* Buffer, void* ContainerOrPropertyPtr, EPropertyPointerType PropertyPointerType, UObject* OwnerObject, int32 PortFlags, FOutputDevice* ErrorText) const override;
+	virtual bool ContainsClearOnFinishDestroyInternal(TArray<const FStructProperty*>& EncounteredStructProps) const override;
+	virtual void FinishDestroyInternal(void* Data) const override;
 public:
 	virtual void CopyValuesInternal( void* Dest, void const* Src, int32 Count  ) const override;
 	virtual void ClearValueInternal( void* Data ) const override;
