Adding UObject count and UObject array capacity stats when dumping detailed UObject hash table overhead to log

diff --git a/Engine/Source/Runtime/CoreUObject/Private/UObject/UObjectHash.cpp b/Engine/Source/Runtime/CoreUObject/Private/UObject/UObjectHash.cpp
--- a/Engine/Source/Runtime/CoreUObject/Private/UObject/UObjectHash.cpp
+++ b/Engine/Source/Runtime/CoreUObject/Private/UObject/UObjectHash.cpp
@@ -2018,7 +2018,7 @@ void LogHashMemoryOverheadStatistics(FOutputDevice& Ar, const bool bShowIndividu
 		const SIZE_T Size = GUObjectArray.GetAllocatedSize();
 		if (bShowIndividualStats)
 		{
-			Ar.Logf(TEXT("Memory used by UObjectArray: %" SIZE_T_FMT " bytes."), Size);
+			Ar.Logf(TEXT("Memory used by UObjectArray: %" SIZE_T_FMT " bytes. (%d UObjects, %d slots) "), Size, GUObjectArray.GetObjectArrayNumMinusAvailable(), GUObjectArray.GetObjectArrayCapacity());
 		}
 		TotalSize += Size;
 	}
