Fixed subtree parsing in FJsonArchiveInputFormatter.

diff --git a/Engine/Source/Runtime/CoreUObject/Private/Serialization/Formatters/JsonArchiveInputFormatter.cpp b/Engine/Source/Runtime/CoreUObject/Private/Serialization/Formatters/JsonArchiveInputFormatter.cpp
--- a/Engine/Source/Runtime/CoreUObject/Private/Serialization/Formatters/JsonArchiveInputFormatter.cpp
+++ b/Engine/Source/Runtime/CoreUObject/Private/Serialization/Formatters/JsonArchiveInputFormatter.cpp
@@ -34,6 +34,14 @@ FJsonArchiveInputFormatter::FJsonArchiveInputFormatter(FArchive& InInner, TFunct
 	ArrayValuesRemainingStack.Reserve(64);
 }
 
+FJsonArchiveInputFormatter::FJsonArchiveInputFormatter(FArchive& InInner, TFunction<UObject* (const FPackageIndex)>&& InResolveObject, TSharedPtr<FJsonValue>&& RootObject) 
+	: Inner(InInner)
+	, ResolveObject(MoveTemp(InResolveObject))
+{
+	ValueStack.Reserve(64);
+	ValueStack.Push(MoveTemp(RootObject));
+}
+
 FJsonArchiveInputFormatter::~FJsonArchiveInputFormatter()
 {
 }
@@ -45,9 +53,7 @@ FArchive& FJsonArchiveInputFormatter::GetUnderlyingArchive()
 
 FStructuredArchiveFormatter* FJsonArchiveInputFormatter::CreateSubtreeReader()
 {
-	FJsonArchiveInputFormatter* Cloned = new FJsonArchiveInputFormatter(this->Inner, this->ResolveObject);
-	Cloned->ValueStack.Push(ValueStack.Top());
-
+	FJsonArchiveInputFormatter* Cloned = new FJsonArchiveInputFormatter(this->Inner, CopyTemp(this->ResolveObject), CopyTemp(ValueStack.Top()));
 	return Cloned;
 }
 
diff --git a/Engine/Source/Runtime/CoreUObject/Public/Serialization/Formatters/JsonArchiveInputFormatter.h b/Engine/Source/Runtime/CoreUObject/Public/Serialization/Formatters/JsonArchiveInputFormatter.h
--- a/Engine/Source/Runtime/CoreUObject/Public/Serialization/Formatters/JsonArchiveInputFormatter.h
+++ b/Engine/Source/Runtime/CoreUObject/Public/Serialization/Formatters/JsonArchiveInputFormatter.h
@@ -100,6 +100,9 @@ public:
 private:
 	FArchive& Inner;
 
+	// Used for the subtree parser
+	FJsonArchiveInputFormatter(FArchive& InInner, TFunction<UObject* (const FPackageIndex)>&& InResolveObject, TSharedPtr<FJsonValue>&& RootObject);
+
 	struct FObjectRecord
 	{
 		FObjectRecord(TSharedPtr<FJsonObject> InJsonObject, int64 InValueCount)
