Always write absolute paths for [AssetLog] entries. Relative paths are fine for rendering, but we need absolute paths to be able to use them for any VCS lookup / build health purposes in Horde.

diff --git a/Engine/Source/Runtime/CoreUObject/Private/UObject/UObjectGlobals.cpp b/Engine/Source/Runtime/CoreUObject/Private/UObject/UObjectGlobals.cpp
--- a/Engine/Source/Runtime/CoreUObject/Private/UObject/UObjectGlobals.cpp
+++ b/Engine/Source/Runtime/CoreUObject/Private/UObject/UObjectGlobals.cpp
@@ -5653,27 +5653,34 @@ void UE::SerializeForLog(FCbWriter& Writer, const FAssetLog& AssetLog)
 		return bDefault;
 	};
 
+	FString AbsLocalPath;
 	if (!LocalPath.IsEmpty())
 	{
+		AbsLocalPath = IFileManager::Get().ConvertToAbsolutePathForExternalAppForRead(*LocalPath);
+		FPaths::MakePlatformFilename(AbsLocalPath);
+
 		static bool bShowDiskPath = GetConfigBool(TEXT("Core.System"), TEXT("AssetLogShowsDiskPath"), true);
-		static bool bShowAbsolutePath = GetConfigBool(TEXT("Core.System"), TEXT("AssetLogShowsAbsolutePath"), false);
-		if (bShowAbsolutePath)
-		{
-			LocalPath = IFileManager::Get().ConvertToAbsolutePathForExternalAppForRead(*LocalPath);
-		}
-		FPaths::MakePlatformFilename(LocalPath);
 		if (bShowDiskPath)
 		{
-			ObjectPath = LocalPath;
+			static bool bShowAbsolutePath = GetConfigBool(TEXT("Core.System"), TEXT("AssetLogShowsAbsolutePath"), false);
+			if (bShowAbsolutePath)
+			{
+				ObjectPath = AbsLocalPath;
+			}
+			else
+			{
+				ObjectPath = LocalPath;
+				FPaths::MakePlatformFilename(ObjectPath);
+			}
 		}
 	}
 
 	Writer.BeginObject();
 	Writer.AddString(ANSITEXTVIEW("$type"), ANSITEXTVIEW("Asset"));
 	Writer.AddString(ANSITEXTVIEW("$text"), ObjectPath);
-	if (!LocalPath.IsEmpty())
+	if (!AbsLocalPath.IsEmpty())
 	{
-		Writer.AddString(ANSITEXTVIEW("file"), LocalPath);
+		Writer.AddString(ANSITEXTVIEW("file"), AbsLocalPath);
 	}
 	Writer.EndObject();
 }
