Moved not-power-of-two ContainerMask out of the enum so that the flags can be viewed in the MSVC debugger.

diff --git a/Engine/Source/Runtime/CoreUObject/Private/UObject/UnrealType.cpp b/Engine/Source/Runtime/CoreUObject/Private/UObject/UnrealType.cpp
--- a/Engine/Source/Runtime/CoreUObject/Private/UObject/UnrealType.cpp
+++ b/Engine/Source/Runtime/CoreUObject/Private/UObject/UnrealType.cpp
@@ -8,6 +8,8 @@ DEFINE_LOG_CATEGORY(LogType);
 
 ENUM_CLASS_FLAGS(FPropertyValueIterator::EPropertyValueFlags)
 
+#define EPropertyValueFlags_ContainerMask (EPropertyValueFlags::IsOptional | EPropertyValueFlags::IsArray | EPropertyValueFlags::IsMap | EPropertyValueFlags::IsSet | EPropertyValueFlags::IsStruct)
+
 FPropertyValueIterator::FPropertyValueIterator(
 	FFieldClass* InPropertyClass,
 	const UStruct* InStruct,
@@ -89,7 +91,7 @@ FORCEINLINE_DEBUGGABLE bool FPropertyValueIterator::NextValue(EPropertyValueIter
 		// Handle matching properties
 		if (!bIsPropertyMatchProcessed && EnumHasAnyFlags(PropertyValueFlags, EPropertyValueFlags::IsMatch))
 		{
-			if (EnumHasAnyFlags(PropertyValueFlags, EPropertyValueFlags::ContainerMask))
+			if (EnumHasAnyFlags(PropertyValueFlags, EPropertyValueFlags_ContainerMask))
 			{
 				// this match is also a container/struct, so recurse into it next time
 				Entry.NextValueIndex = Entry.ValueIndex;
@@ -98,7 +100,7 @@ FORCEINLINE_DEBUGGABLE bool FPropertyValueIterator::NextValue(EPropertyValueIter
 		}
 
 		// Handle container properties
-		check(EnumHasAnyFlags(PropertyValueFlags, EPropertyValueFlags::ContainerMask));
+		check(EnumHasAnyFlags(PropertyValueFlags, EPropertyValueFlags_ContainerMask));
 		if (InRecursionFlags == EPropertyValueIteratorFlags::FullRecursion)
 		{
 			FPropertyValueStackEntry NewEntry(PropertyValue);
diff --git a/Engine/Source/Runtime/CoreUObject/Public/UObject/UnrealType.h b/Engine/Source/Runtime/CoreUObject/Public/UObject/UnrealType.h
--- a/Engine/Source/Runtime/CoreUObject/Public/UObject/UnrealType.h
+++ b/Engine/Source/Runtime/CoreUObject/Public/UObject/UnrealType.h
@@ -6970,12 +6970,13 @@ private:
 		None = 0x0,
 		IsMatch = 0x01,
 
-		ContainerMask = 0xF8,
 		IsOptional = 0x08,
 		IsArray = 0x10,
 		IsMap = 0x20,
 		IsSet = 0x40,
 		IsStruct = 0x80,
+
+		// When adding a new 'container' (needs to be recursed into) flag here, add it to the EPropertyValueFlags_ContainerMask macro too.
 	};
 	FRIEND_ENUM_CLASS_FLAGS(EPropertyValueFlags)
 
