UStruct::IsChildOf<T>: Support classes which provide StaticStruct, in addition to those which provide StaticClass.

diff --git a/Engine/Source/Runtime/CoreUObject/Public/UObject/Class.h b/Engine/Source/Runtime/CoreUObject/Public/UObject/Class.h
--- a/Engine/Source/Runtime/CoreUObject/Public/UObject/Class.h
+++ b/Engine/Source/Runtime/CoreUObject/Public/UObject/Class.h
@@ -7,6 +7,7 @@
 #pragma once
 
 #include "Concepts/GetTypeHashable.h"
+#include "Concepts/StaticClassProvider.h"
 #include "Concepts/StaticStructProvider.h"
 #include "Containers/Array.h"
 #include "Containers/EnumAsByte.h"
@@ -610,7 +611,19 @@ public:
 	template<class T>
 	bool IsChildOf() const
 	{
-		return IsChildOf(T::StaticClass());
+		if constexpr (TModels_V<CStaticClassProvider, T>)
+		{
+			return IsChildOf(T::StaticClass());
+		}
+		else if constexpr (TModels_V<CStaticStructProvider, T>)
+		{
+			return IsChildOf(T::StaticStruct());
+		}
+		else
+		{
+			static_assert(sizeof(T) == 0, "Unsupported type - requires StaticClass or StaticStruct.");
+			return false;
+		}
 	}
 
 	/** Returns true if this struct either is SomeBase, or is a child of SomeBase. This will not crash on null structs */
