New Property Visitor
Visit this property and allows recursion into the inner properties
This allows you to visit inner properties without knowing about its container type as opposed to TPropertyIterator.
This visit property pattern facilitate the recursion into users defined properties and allows users to add specific visit logic on UStruct via traits.

diff --git a/Engine/Plugins/Experimental/StructUtils/Source/StructUtils/Private/InstancedStruct.cpp b/Engine/Plugins/Experimental/StructUtils/Source/StructUtils/Private/InstancedStruct.cpp
--- a/Engine/Plugins/Experimental/StructUtils/Source/StructUtils/Private/InstancedStruct.cpp
+++ b/Engine/Plugins/Experimental/StructUtils/Source/StructUtils/Private/InstancedStruct.cpp
@@ -587,3 +587,13 @@ bool FInstancedStruct::FindInnerPropertyInstance(FName PropertyName, const FProp
 
 	return false;
 }
+
+EPropertyVisitorControlFlow FInstancedStruct::Visit(FPropertyVisitorPath& Path, const TFunctionRef<EPropertyVisitorControlFlow(const FPropertyVisitorPath& /*Path*/, void* /*Data*/)> InFunc)
+{
+	if (!ScriptStruct || !StructMemory)
+	{
+		return EPropertyVisitorControlFlow::StepOver;
+	}
+
+	return ScriptStruct->Visit(Path, StructMemory, InFunc);
+}
diff --git a/Engine/Plugins/Experimental/StructUtils/Source/StructUtils/Public/InstancedStruct.h b/Engine/Plugins/Experimental/StructUtils/Source/StructUtils/Public/InstancedStruct.h
--- a/Engine/Plugins/Experimental/StructUtils/Source/StructUtils/Public/InstancedStruct.h
+++ b/Engine/Plugins/Experimental/StructUtils/Source/StructUtils/Public/InstancedStruct.h
@@ -154,6 +154,7 @@ public:
 	void GetPreloadDependencies(TArray<UObject*>& OutDeps);
 	bool NetSerialize(FArchive& Ar, class UPackageMap* Map, bool& bOutSuccess);
 	bool FindInnerPropertyInstance(FName PropertyName, const FProperty*& OutProp, const void*& OutData) const;
+	EPropertyVisitorControlFlow Visit(FPropertyVisitorPath& Path, const TFunctionRef<EPropertyVisitorControlFlow(const FPropertyVisitorPath& /*Path*/, void* /*Data*/)> InFunc);
 
 	/** Returns struct type. */
 	const UScriptStruct* GetScriptStruct() const
@@ -273,6 +274,7 @@ struct TStructOpsTypeTraits<FInstancedStruct> : public TStructOpsTypeTraitsBase2
 		WithNetSerializer = true,
 		WithFindInnerPropertyInstance = true,
 		WithClearOnFinishDestroy = true,
+		WithVisitor = true,
 	};
 };
 
diff --git a/Engine/Source/Runtime/CoreUObject/Private/UObject/Class.cpp b/Engine/Source/Runtime/CoreUObject/Private/UObject/Class.cpp
--- a/Engine/Source/Runtime/CoreUObject/Private/UObject/Class.cpp
+++ b/Engine/Source/Runtime/CoreUObject/Private/UObject/Class.cpp
@@ -757,6 +757,54 @@ void UStruct::CollectBytecodeAndPropertyReferencedObjectsRecursively()
 	}
 }
 
+EPropertyVisitorControlFlow UStruct::Visit(void* Data, const TFunctionRef<EPropertyVisitorControlFlow(const FPropertyVisitorPath& /*Path*/, void* /*Data*/)> InFunc) const
+{
+	FPropertyVisitorPath Path;
+	return Visit(Path, Data, InFunc);
+}
+
+EPropertyVisitorControlFlow UStruct::Visit(FPropertyVisitorPath& Path, void* Data, const TFunctionRef<EPropertyVisitorControlFlow(const FPropertyVisitorPath& /*Path*/, void* /*Data*/)> InFunc) const
+{
+	EPropertyVisitorControlFlow RetVal = EPropertyVisitorControlFlow::StepOver;
+	for (const FProperty* Property = PropertyLink; Property; Property = Property->PropertyLinkNext)
+	{
+		FPropertyVisitorScope Scope(Path, FPropertyVisitorInfo(Property));
+
+		// ArrayDim of one means it is just a single property, not a static array
+		if(Property->ArrayDim == 1)
+		{
+			RetVal = Property->Visit(Path, Property->ContainerPtrToValuePtr<void>(Data), InFunc);
+			if (RetVal == EPropertyVisitorControlFlow::Stop)
+			{
+				return EPropertyVisitorControlFlow::Stop;
+			}
+			if (RetVal == EPropertyVisitorControlFlow::StepOut)
+			{
+				return EPropertyVisitorControlFlow::StepOver;
+			}
+		}
+		else
+		{
+			// In the case of a static array, we need to setup the EPropertyVisitorInfoType
+			for (int32 StaticArrayIndex = 0; StaticArrayIndex < Property->ArrayDim; ++StaticArrayIndex)
+			{
+				Path.Top().SetIndex(StaticArrayIndex, EPropertyVisitorInfoType::StaticArrayIndex);
+				RetVal = Property->Visit(Path, Property->ContainerPtrToValuePtr<void>(Data, StaticArrayIndex), InFunc);
+
+				if (RetVal == EPropertyVisitorControlFlow::Stop)
+				{
+					return EPropertyVisitorControlFlow::Stop;
+				}
+				if (RetVal == EPropertyVisitorControlFlow::StepOut)
+				{
+					return EPropertyVisitorControlFlow::StepOver;
+				}
+			}
+		}
+	}
+	return RetVal;
+}
+
 void UStruct::PreloadChildren(FArchive& Ar)
 {
 	for (UField* Field = Children; Field; Field = Field->Next)
@@ -2956,6 +3004,12 @@ void UScriptStruct::PrepareCppStructOps()
 	}
 #endif
 
+	if (CppStructOps->HasVisitor())
+	{
+		UE_LOG(LogClass, Verbose, TEXT("Native struct %s has native property Visit."), *GetName());
+		StructFlags = EStructFlags(StructFlags | STRUCT_Visitor);
+	}
+
 	check(!bPrepareCppStructOpsCompleted); // recursion is unacceptable
 	bPrepareCppStructOpsCompleted = true;
 }
@@ -3432,6 +3486,25 @@ bool UScriptStruct::FindInnerPropertyInstance(FName PropertyName, const void* Da
 	return false;
 }
 
+EPropertyVisitorControlFlow UScriptStruct::Visit(FPropertyVisitorPath& Path, void* Data, const TFunctionRef<EPropertyVisitorControlFlow(const FPropertyVisitorPath& /*Path*/, void* /*Data*/)> InFunc) const
+{
+	const EPropertyVisitorControlFlow RetVal = UStruct::Visit(Path, Data, InFunc);
+	if (RetVal == EPropertyVisitorControlFlow::Stop)
+	{
+		return EPropertyVisitorControlFlow::Stop;
+	}
+	if (RetVal == EPropertyVisitorControlFlow::StepOut)
+	{
+		return EPropertyVisitorControlFlow::StepOver;
+	}
+	if (StructFlags & STRUCT_Visitor)
+	{
+		checkf(CppStructOps && CppStructOps->HasVisitor(), TEXT("Expecting to have a visitor implementation when STRUCT_Visitor is set"));
+		return CppStructOps->Visit(Path, Data, InFunc);
+	}
+	return RetVal;
+}
+
 void UScriptStruct::ClearScriptStruct(void* Dest, int32 ArrayDim) const
 {
 	uint8 *Data = (uint8*)Dest;
diff --git a/Engine/Source/Runtime/CoreUObject/Private/UObject/Property.cpp b/Engine/Source/Runtime/CoreUObject/Private/UObject/Property.cpp
--- a/Engine/Source/Runtime/CoreUObject/Private/UObject/Property.cpp
+++ b/Engine/Source/Runtime/CoreUObject/Private/UObject/Property.cpp
@@ -19,6 +19,7 @@
 #include "UObject/Package.h"
 #include "UObject/PropertyHelper.h"
 #include "UObject/PropertyTypeName.h"
+#include "UObject/PropertyVisitor.h"
 #include "UObject/SoftObjectPath.h"
 #include "UObject/UnrealTypePrivate.h"
 #include "UObject/UObjectGlobals.h"
@@ -1336,6 +1337,19 @@ namespace UE::CoreUObject::Private
 	}
 }
 
+EPropertyVisitorControlFlow FProperty::Visit(void* Data, const TFunctionRef<EPropertyVisitorControlFlow(const FPropertyVisitorPath&, void*)> InFunc) const
+{
+	FPropertyVisitorPath Path(FPropertyVisitorInfo(this));
+	return Visit(Path, Data, InFunc);
+}
+
+EPropertyVisitorControlFlow FProperty::Visit(FPropertyVisitorPath& Path, void* Data, const TFunctionRef<EPropertyVisitorControlFlow(const FPropertyVisitorPath& /*Path*/, void* /*Data*/)> InFunc) const
+{
+	checkf(Path.Top().Property == this, TEXT("The property set in the info has to match to this property"));
+
+	return InFunc(Path, Data);
+}
+
 int32 FProperty::SetupOffset()
 {
 	UObject* OwnerUObject = GetOwner<UObject>();
diff --git a/Engine/Source/Runtime/CoreUObject/Private/UObject/PropertyArray.cpp b/Engine/Source/Runtime/CoreUObject/Private/UObject/PropertyArray.cpp
--- a/Engine/Source/Runtime/CoreUObject/Private/UObject/PropertyArray.cpp
+++ b/Engine/Source/Runtime/CoreUObject/Private/UObject/PropertyArray.cpp
@@ -1345,3 +1345,35 @@ bool FArrayProperty::CanSerializeFromTypeName(UE::FPropertyTypeName Type) const
 	check(LocalInner);
 	return LocalInner->CanSerializeFromTypeName(Type.GetParameter(0));
 }
+
+EPropertyVisitorControlFlow FArrayProperty::Visit(FPropertyVisitorPath& Path, void* Data, const TFunctionRef<EPropertyVisitorControlFlow(const FPropertyVisitorPath& /*Path*/, void* /*Data*/)> InFunc) const
+{
+	// Indicate in the path that this property contains inner properties
+	Path.Top().bContainsInnerProperties = true;
+
+	EPropertyVisitorControlFlow RetVal = Super::Visit(Path, Data, InFunc);
+
+	if (RetVal == EPropertyVisitorControlFlow::StepInto)
+	{
+		checkf(Inner, TEXT("Expecting a valid inner property type"));
+		FScriptArrayHelper ArrayHelper(this, Data);
+
+		FPropertyVisitorScope Scope(Path, FPropertyVisitorInfo(Inner));
+
+		const int32 ArrayNum = ArrayHelper.Num();
+		for (int32 ContainerIndex = 0; ContainerIndex < ArrayNum; ContainerIndex++)
+		{
+			Path.Top().SetIndex(ContainerIndex, EPropertyVisitorInfoType::ContainerIndex);
+			RetVal = Inner->Visit(Path, ArrayHelper.GetRawPtr(ContainerIndex), InFunc);
+			if (RetVal == EPropertyVisitorControlFlow::Stop)
+			{
+				return EPropertyVisitorControlFlow::Stop;
+			}
+			if (RetVal == EPropertyVisitorControlFlow::StepOut)
+			{
+				return EPropertyVisitorControlFlow::StepOver;
+			}
+		}
+	}
+	return RetVal;
+}
diff --git a/Engine/Source/Runtime/CoreUObject/Private/UObject/PropertyBaseObject.cpp b/Engine/Source/Runtime/CoreUObject/Private/UObject/PropertyBaseObject.cpp
--- a/Engine/Source/Runtime/CoreUObject/Private/UObject/PropertyBaseObject.cpp
+++ b/Engine/Source/Runtime/CoreUObject/Private/UObject/PropertyBaseObject.cpp
@@ -753,3 +753,20 @@ bool FObjectPropertyBase::SameType(const FProperty* Other) const
 	return (Super::SameType(Other) || (Other && Other->IsA<FObjectPropertyBase>() && ((FObjectPropertyBase*)Other)->AllowObjectTypeReinterpretationTo(this))) && 
 			 (PropertyClass == ((FObjectPropertyBase*)Other)->PropertyClass);
 }
+
+EPropertyVisitorControlFlow FObjectPropertyBase::Visit(FPropertyVisitorPath& Path, void* Data, const TFunctionRef<EPropertyVisitorControlFlow(const FPropertyVisitorPath& /*Path*/, void* /*Data*/)> InFunc) const
+{
+	// Indicate in the path that this property contains inner properties
+	Path.Top().bContainsInnerProperties = true;
+
+	EPropertyVisitorControlFlow RetVal = Super::Visit(Path, Data, InFunc);
+
+	if (RetVal == EPropertyVisitorControlFlow::StepInto)
+	{
+		if (const TObjectPtr<UObject> Object = GetObjectPropertyValue(Data))
+		{
+			RetVal = Object->GetClass()->Visit(Path, Object, InFunc);
+		}
+	}
+	return RetVal;
+}
diff --git a/Engine/Source/Runtime/CoreUObject/Private/UObject/PropertyMap.cpp b/Engine/Source/Runtime/CoreUObject/Private/UObject/PropertyMap.cpp
--- a/Engine/Source/Runtime/CoreUObject/Private/UObject/PropertyMap.cpp
+++ b/Engine/Source/Runtime/CoreUObject/Private/UObject/PropertyMap.cpp
@@ -1842,3 +1842,51 @@ bool FMapProperty::CanSerializeFromTypeName(UE::FPropertyTypeName Type) const
 	check(LocalValueProp);
 	return LocalKeyProp->CanSerializeFromTypeName(Type.GetParameter(0)) && LocalValueProp->CanSerializeFromTypeName(Type.GetParameter(1));
 }
+
+EPropertyVisitorControlFlow FMapProperty::Visit(FPropertyVisitorPath& Path, void* Data, const TFunctionRef<EPropertyVisitorControlFlow(const FPropertyVisitorPath& /*Path*/, void* /*Data*/)> InFunc) const
+{
+	// Indicate in the path that this property contains inner properties
+	Path.Top().bContainsInnerProperties = true;
+
+	EPropertyVisitorControlFlow RetVal = Super::Visit(Path, Data, InFunc);
+
+	if (RetVal == EPropertyVisitorControlFlow::StepInto)
+	{
+		checkf(KeyProp && ValueProp, TEXT("Expecting a valid inner property type"));
+		FScriptMapHelper MapHelper(this, Data);
+
+		for (FScriptMapHelper::FIterator It(MapHelper); It; ++It)
+		{
+			{
+				// Visit Key
+				FPropertyVisitorScope Scope(Path, FPropertyVisitorInfo(KeyProp, It.GetLogicalIndex(), EPropertyVisitorInfoType::MapKey));
+
+				RetVal = KeyProp->Visit(Path, MapHelper.GetKeyPtr(It), InFunc);
+				if (RetVal == EPropertyVisitorControlFlow::Stop)
+				{
+					return EPropertyVisitorControlFlow::Stop;
+				}
+				if (RetVal == EPropertyVisitorControlFlow::StepOut)
+				{
+					return EPropertyVisitorControlFlow::StepOver;
+				}
+			}
+
+			{
+				// Visit Value
+				FPropertyVisitorScope Scope(Path, FPropertyVisitorInfo(ValueProp, It.GetLogicalIndex(), EPropertyVisitorInfoType::MapValue));
+
+				RetVal = ValueProp->Visit(Path, MapHelper.GetValuePtr(It), InFunc);
+				if (RetVal == EPropertyVisitorControlFlow::Stop)
+				{
+					return EPropertyVisitorControlFlow::Stop;
+				}
+				if (RetVal == EPropertyVisitorControlFlow::StepOut)
+				{
+					return EPropertyVisitorControlFlow::StepOver;
+				}
+			}
+		}
+	}
+	return RetVal;
+}
diff --git a/Engine/Source/Runtime/CoreUObject/Private/UObject/PropertyOptional.cpp b/Engine/Source/Runtime/CoreUObject/Private/UObject/PropertyOptional.cpp
--- a/Engine/Source/Runtime/CoreUObject/Private/UObject/PropertyOptional.cpp
+++ b/Engine/Source/Runtime/CoreUObject/Private/UObject/PropertyOptional.cpp
@@ -681,3 +681,20 @@ bool FOptionalProperty::CanSerializeFromTypeName(UE::FPropertyTypeName Type) con
 	check(LocalValueProperty);
 	return LocalValueProperty->CanSerializeFromTypeName(Type.GetParameter(0));
 }
+
+EPropertyVisitorControlFlow FOptionalProperty::Visit(FPropertyVisitorPath& Path, void* Data, const TFunctionRef<EPropertyVisitorControlFlow(const FPropertyVisitorPath& /*Path*/, void* /*Data*/)> InFunc) const
+{
+	// Indicate in the path that this property contains inner properties
+	Path.Top().bContainsInnerProperties = true;
+
+	EPropertyVisitorControlFlow RetVal = Super::Visit(Path, Data, InFunc);
+
+	if (RetVal == EPropertyVisitorControlFlow::StepInto && IsSet(Data))
+	{
+		checkf(ValueProperty, TEXT("Expecting a valid property value"));
+
+		FPropertyVisitorScope Scope(Path, FPropertyVisitorInfo(ValueProperty));
+		RetVal = ValueProperty->Visit(Path, Data, InFunc);
+	}
+	return RetVal;
+}
diff --git a/Engine/Source/Runtime/CoreUObject/Private/UObject/PropertySet.cpp b/Engine/Source/Runtime/CoreUObject/Private/UObject/PropertySet.cpp
--- a/Engine/Source/Runtime/CoreUObject/Private/UObject/PropertySet.cpp
+++ b/Engine/Source/Runtime/CoreUObject/Private/UObject/PropertySet.cpp
@@ -1151,3 +1151,34 @@ bool FSetProperty::CanSerializeFromTypeName(UE::FPropertyTypeName Type) const
 	check(LocalElementProp);
 	return LocalElementProp->CanSerializeFromTypeName(Type.GetParameter(0));
 }
+
+EPropertyVisitorControlFlow FSetProperty::Visit(FPropertyVisitorPath& Path, void* Data, const TFunctionRef<EPropertyVisitorControlFlow(const FPropertyVisitorPath& /*Path*/, void* /*Data*/)> InFunc) const
+{
+	// Indicate in the path that this property contains inner properties
+	Path.Top().bContainsInnerProperties = true;
+
+	EPropertyVisitorControlFlow RetVal = Super::Visit(Path, Data, InFunc);
+
+	if (RetVal == EPropertyVisitorControlFlow::StepInto)
+	{
+		checkf(ElementProp, TEXT("Expecting a valid inner property type"));
+		FScriptSetHelper SetHelper(this, Data);
+
+		FPropertyVisitorScope Scope(Path, FPropertyVisitorInfo(ElementProp));
+		for (FScriptSetHelper::FIterator It(SetHelper); It; ++It)
+		{
+			// visit element
+			Path.Top().SetIndex(It.GetLogicalIndex(), EPropertyVisitorInfoType::ContainerIndex);
+			RetVal = ElementProp->Visit(Path, SetHelper.GetElementPtr(It), InFunc);
+			if (RetVal == EPropertyVisitorControlFlow::Stop)
+			{
+				return EPropertyVisitorControlFlow::Stop;
+			}
+			if (RetVal == EPropertyVisitorControlFlow::StepOut)
+			{
+				return EPropertyVisitorControlFlow::StepOver;
+			}
+		}
+	}
+	return RetVal;
+}
diff --git a/Engine/Source/Runtime/CoreUObject/Private/UObject/PropertyStruct.cpp b/Engine/Source/Runtime/CoreUObject/Private/UObject/PropertyStruct.cpp
--- a/Engine/Source/Runtime/CoreUObject/Private/UObject/PropertyStruct.cpp
+++ b/Engine/Source/Runtime/CoreUObject/Private/UObject/PropertyStruct.cpp
@@ -595,3 +595,17 @@ bool FStructProperty::CanSerializeFromTypeName(UE::FPropertyTypeName Type) const
 
 	return false;
 }
+
+EPropertyVisitorControlFlow FStructProperty::Visit(FPropertyVisitorPath& Path, void* Data, const TFunctionRef<EPropertyVisitorControlFlow(const FPropertyVisitorPath& /*Path*/, void* /*Data*/)> InFunc) const
+{
+	// Indicate in the path that this property contains inner properties
+	Path.Top().bContainsInnerProperties = true;
+
+	EPropertyVisitorControlFlow RetVal = Super::Visit(Path, Data, InFunc);
+
+	if (RetVal == EPropertyVisitorControlFlow::StepInto)
+	{
+		RetVal = Struct->Visit(Path, Data, InFunc);
+	}
+	return RetVal;
+}
\ No newline at end of file
diff --git a/Engine/Source/Runtime/CoreUObject/Private/UObject/PropertyVisitor.cpp b/Engine/Source/Runtime/CoreUObject/Private/UObject/PropertyVisitor.cpp
new file mode 100644
--- /dev/null
+++ b/Engine/Source/Runtime/CoreUObject/Private/UObject/PropertyVisitor.cpp
@@ -0,0 +1,130 @@
+// Copyright Epic Games, Inc. All Rights Reserved.
+
+#include "UObject/PropertyVisitor.h"
+#include "Containers/UnrealString.h"
+#include "UObject/NameTypes.h"
+#include "UObject/UnrealType.h"
+
+FString FPropertyVisitorPath::ToString(const TCHAR* Separator /*= TEXT(".")*/) const
+{
+	TStringBuilder<FName::StringBufferSize> PropertyPath;
+	bool bFirstEntry = true;
+	for (const FPropertyVisitorInfo& Entry : Path)
+	{
+		bool bDisplayPropertyName = false;
+		bool bDisplayIndex = false;
+		FString Suffix(TEXT(""));
+		switch (Entry.PropertyInfo)
+		{
+			case EPropertyVisitorInfoType::None:
+				bDisplayPropertyName = true;
+				break;
+			case EPropertyVisitorInfoType::StaticArrayIndex:
+				bDisplayPropertyName = true;
+				bDisplayIndex = true;
+				break;
+			case EPropertyVisitorInfoType::ContainerIndex:
+				bDisplayIndex = true;
+				break;
+			case EPropertyVisitorInfoType::MapKey:
+				bDisplayIndex = true;
+				Suffix = TEXT("Key");
+				break;
+			case EPropertyVisitorInfoType::MapValue:
+				bDisplayIndex = true;
+				Suffix = TEXT("Value");
+				break;
+			default:
+				checkf(false, TEXT("Unsupported enum value"));
+				break;
+		}
+		if (bDisplayPropertyName)
+		{
+			if (bFirstEntry)
+			{
+				bFirstEntry = false;
+			}
+			else
+			{
+				PropertyPath.Append(Separator);
+			}
+			PropertyPath.Append(Entry.Property->GetAuthoredName());
+		}
+		if (bDisplayIndex)
+		{
+			checkf(Entry.Index != INDEX_NONE, TEXT("Expecting the index to be valid"));
+			PropertyPath.Appendf(TEXT("[%d]"), Entry.Index);
+		}
+		if (Suffix.Len())
+		{
+			PropertyPath.Append(Separator);
+			PropertyPath.Append(Suffix);
+		}
+	}
+
+	return PropertyPath.ToString();
+}
+
+bool FPropertyVisitorPath::Contained(const FPropertyVisitorPath& Other, bool* bIsEqual) const
+{
+	int32 i;
+	for (i = 0; i < Path.Num(); ++i)
+	{
+		if (i >= Other.Path.Num())
+		{
+			// The other path is smaller than this one, so not contained in
+			break;
+		}
+
+		const FPropertyVisitorInfo& PathInfo = Path[i];
+		const FPropertyVisitorInfo& OtherPathInfo = Other.Path[i];
+
+		if (PathInfo.Property != OtherPathInfo.Property)
+		{
+			// The property is different, so not contained in
+			break;
+		}
+
+		if (PathInfo.PropertyInfo != OtherPathInfo.PropertyInfo)
+		{
+			if (PathInfo.PropertyInfo != EPropertyVisitorInfoType::None)
+			{
+				// The property info type is different and it is not none, so not contained in
+				break;
+			}
+		}
+		else if (PathInfo.Index != OtherPathInfo.Index)
+		{
+			// The index is different, so not contained in
+			break;
+		}
+	}
+	if (bIsEqual)
+	{
+		*bIsEqual = i == Other.Path.Num();
+	}
+	return i == Path.Num();
+}
+
+void* FPropertyVisitorPath::GetPropertyDataPtr(UObject* Object) const
+{
+	void* DataPtr = nullptr;
+	checkf(Object, TEXT("Expecting an valid object"));
+
+	Object->GetClass()->Visit(Object, [this, &DataPtr](const FPropertyVisitorPath& InPath, void* Data)
+	{
+		bool bIsEqual = false;
+		if (InPath.Contained(*this, &bIsEqual))
+		{
+			if(bIsEqual)
+			{
+				DataPtr = Data;
+				return EPropertyVisitorControlFlow::Stop;
+			}
+			return EPropertyVisitorControlFlow::StepInto;
+		}
+		return EPropertyVisitorControlFlow::StepOver;
+	});
+
+	return DataPtr;
+}
\ No newline at end of file
diff --git a/Engine/Source/Runtime/CoreUObject/Public/UObject/Class.h b/Engine/Source/Runtime/CoreUObject/Public/UObject/Class.h
--- a/Engine/Source/Runtime/CoreUObject/Public/UObject/Class.h
+++ b/Engine/Source/Runtime/CoreUObject/Public/UObject/Class.h
@@ -65,6 +65,7 @@
 #include "UObject/Object.h"
 #include "UObject/ObjectMacros.h"
 #include "UObject/PropertyTag.h"
+#include "UObject/PropertyVisitor.h"
 #include "UObject/ReflectedTypeAccessors.h"
 #include "UObject/Script.h"
 #include "UObject/TopLevelAssetPath.h"
@@ -724,6 +725,27 @@ public:
 	 */
 	COREUOBJECT_API void CollectBytecodeAndPropertyReferencedObjectsRecursively();
 
+	/**
+	 * Visits this property and allows recursion into the inner properties
+	 * This method allows callers to visit inner properties without knowing about its container type as opposed to TPropertyIterator.
+	 * This visit property pattern facilitates the recursion into user defined properties and allows users to add specific visit logic on UStruct via traits. 
+	 * @param Data to the property to visit
+	 * @param InFunc to call on each visited property, the return value controls what is the next behavior once this property has been visited
+	 * @return the new action to take one visited this property
+	 */
+	COREUOBJECT_API EPropertyVisitorControlFlow Visit(void* Data, const TFunctionRef<EPropertyVisitorControlFlow(const FPropertyVisitorPath& /*Path*/, void* /*Data*/)> InFunc) const;
+
+	/**
+	 * Visits this property and allows recursion into the inner properties
+	 * This method allows callers to visit inner properties without knowing about its container type as opposed to TPropertyIterator.
+	 * This visit property pattern facilitates the recursion into user defined properties and allows users to add specific visit logic on UStruct via traits. 
+	 * @param Path that was computed until we reached this property
+	 * @param Data to the property to visit
+	 * @param InFunc to call on each visited property, the return value controls what is the next behavior once this property has been visited
+	 * @return the new action to take one visited this property
+	 */
+	COREUOBJECT_API virtual EPropertyVisitorControlFlow Visit(FPropertyVisitorPath& Path, void* Data, const TFunctionRef<EPropertyVisitorControlFlow(const FPropertyVisitorPath& /*Path*/, void* /*Data*/)> InFunc) const;
+
 protected:
 
 	/** Returns the property name from the guid */
@@ -837,6 +859,9 @@ enum EStructFlags
 
 	/** If set, this struct will have CanEditChange on it in the editor to determine if a child property can be edited */
 	STRUCT_CanEditChange = 0x02000000,
+
+	/** If set, this struct will have Visit on it to allow custom property visiting implementation */
+	STRUCT_Visitor = 0x04000000,
 	
 	/** Struct flags that are automatically inherited */
 	STRUCT_Inherit				= STRUCT_HasInstancedReference|STRUCT_Atomic,
@@ -875,6 +900,7 @@ struct TStructOpsTypeTraitsBase2
 		WithFindInnerPropertyInstance  = false,							// struct has a FindInnerPropertyInstance function that can provide an FProperty and data pointer when given a property FName
 		WithCanEditChange			   = false,							// struct has an editor-only CanEditChange function that can conditionally make child properties read-only in the details panel (same idea as UObject::CanEditChange)
 		WithClearOnFinishDestroy	   = false,							// struct should be cleared during owner UObject's FinishDestroy. Clearing calls destructor and initializes again to default value. This is intended for structs which may need to access UObject pointer members during destruction. Referenced objects may already have their FinishDestroy() called. Clearing should ensure that no UObject pointer members are used during the final destruction.
+		WithVisitor					   = false,							// struct has Visit function that allows to visit additional properties
 	};
 
 	static constexpr EPropertyObjectReferenceType WithSerializerObjectReferences = EPropertyObjectReferenceType::Conservative; // struct's Serialize method(s) may serialize object references of these types - default Conservative means unknown and object reference collector archives should serialize this struct 
@@ -945,6 +971,7 @@ public:
 #if WITH_EDITOR
 			bool HasCanEditChange : 1;
 #endif
+			bool HasVisitor : 1;
 		};
 
 		/**
@@ -1195,7 +1222,15 @@ public:
 		/** Returns true if this struct would allow the given property to be edited in the details panel. */
 		virtual bool CanEditChange(const FEditPropertyChain& PropertyChain, const void* Data) const = 0;
 #endif
-		
+
+		/** Returns true if this struct wants to indicate whether it has a custom impl for visiting properties */
+		bool HasVisitor() const
+		{
+			return GetCapabilities().HasVisitor;
+		}
+		/** Structs property visitor signature */
+		virtual EPropertyVisitorControlFlow Visit(FPropertyVisitorPath& Path, void* Data, const TFunctionRef<EPropertyVisitorControlFlow(const FPropertyVisitorPath& /*Path*/, void* /*Data*/)> InFunc) const = 0;
+
 	private:
 		/** sizeof() of the structure **/
 		const int32 Size;
@@ -1249,6 +1284,7 @@ public:
 #if WITH_EDITOR
 				TTraits::WithCanEditChange,
 #endif
+				TTraits::WithVisitor,
 			};
 			return Capabilities;
 		}
@@ -1552,6 +1588,18 @@ public:
 			}
 		}
 #endif // WITH_EDITOR
+
+		virtual EPropertyVisitorControlFlow Visit(FPropertyVisitorPath& Path, void* Data, const TFunctionRef<EPropertyVisitorControlFlow(const FPropertyVisitorPath& /*Path*/, void* /*Data*/)> InFunc) const override
+		{
+			if constexpr (TStructOpsTypeTraits<CPPSTRUCT>::WithVisitor)
+			{
+				return ((CPPSTRUCT*)Data)->Visit(Path, InFunc);
+			}
+			else
+			{
+				return EPropertyVisitorControlFlow::StepOver;
+			}
+		}
 	};
 
 	/** Template for noexport classes to autoregister before main starts **/
@@ -1799,6 +1847,9 @@ public:
 	 * @return				whether the property instance was found
 	 */
 	virtual COREUOBJECT_API bool FindInnerPropertyInstance(FName PropertyName, const void* Data, const FProperty*& OutProp, const void*& OutData) const;
+
+	/* Custom visit implementation for structs */
+	virtual COREUOBJECT_API EPropertyVisitorControlFlow Visit(FPropertyVisitorPath& Path, void* Data, const TFunctionRef<EPropertyVisitorControlFlow(const FPropertyVisitorPath& /*Path*/, void* /*Data*/)> InFunc) const override;
 };
 
 /*-----------------------------------------------------------------------------
diff --git a/Engine/Source/Runtime/CoreUObject/Public/UObject/PropertyOptional.h b/Engine/Source/Runtime/CoreUObject/Public/UObject/PropertyOptional.h
--- a/Engine/Source/Runtime/CoreUObject/Public/UObject/PropertyOptional.h
+++ b/Engine/Source/Runtime/CoreUObject/Public/UObject/PropertyOptional.h
@@ -235,5 +235,7 @@ public:
 	virtual bool LoadTypeName(UE::FPropertyTypeName Type, const FPropertyTag* Tag = nullptr) override;
 	virtual void SaveTypeName(UE::FPropertyTypeNameBuilder& Type) const override;
 	virtual bool CanSerializeFromTypeName(UE::FPropertyTypeName Type) const override;
+
+	virtual EPropertyVisitorControlFlow Visit(FPropertyVisitorPath& Path, void* Data, const TFunctionRef<EPropertyVisitorControlFlow(const FPropertyVisitorPath& /*Path*/, void* /*Data*/)> InFunc) const override;
 	// End of FProperty interface
 };
diff --git a/Engine/Source/Runtime/CoreUObject/Public/UObject/PropertyVisitor.h b/Engine/Source/Runtime/CoreUObject/Public/UObject/PropertyVisitor.h
new file mode 100644
--- /dev/null
+++ b/Engine/Source/Runtime/CoreUObject/Public/UObject/PropertyVisitor.h
@@ -0,0 +1,141 @@
+// Copyright Epic Games, Inc. All Rights Reserved.
+
+#pragma once
+
+#include "Containers/Array.h"
+
+class FProperty;
+
+enum class EPropertyVisitorControlFlow : uint8
+{
+	Stop, // Stop the visit
+	StepOver, // Skip over to the next property or item
+	StepOut, // Stop iteration at this level and continue on the outer on the next property or item
+	StepInto, // Introspect the inner properties if any
+};
+
+enum class EPropertyVisitorInfoType : uint8
+{
+	None, // Property is not inside a container
+	StaticArrayIndex, // Property is a static array and has a valid index
+	ContainerIndex, // Property is inside a container and has a valid index
+	MapKey, // Property represents a key of the map container and has a valid index
+	MapValue, // Property represents a value of the map container and has a valid index
+};
+
+struct FPropertyVisitorInfo
+{
+	explicit FPropertyVisitorInfo( const FProperty* InProperty, int32 InIndex = INDEX_NONE, EPropertyVisitorInfoType InPropertyInfo = EPropertyVisitorInfoType::None)
+		: Property(InProperty)
+		, Index(InIndex)
+		, PropertyInfo(InPropertyInfo)
+		, bContainsInnerProperties(false)
+	{}
+	FPropertyVisitorInfo(const FPropertyVisitorInfo& Other) = default;
+	FPropertyVisitorInfo(FPropertyVisitorInfo&& Other) = default;
+
+	void SetIndex(int32 InIndex, EPropertyVisitorInfoType InPropertyInfo)
+	{
+		Index = InIndex;
+		PropertyInfo = InPropertyInfo;
+	}
+
+	/** The property currently being visited */
+	const FProperty* Property;
+
+	/**
+     * Index of the element being visited in the container, otherwise INDEX_NONE.
+     * For maps and sets it indicates the logical index. */
+	int32 Index;
+
+	/** Whether this property is inside a container and if it is key or a value of a map */
+	EPropertyVisitorInfoType PropertyInfo;
+
+	/* Indicate that this property contains inner properties */
+	bool bContainsInnerProperties;
+};
+
+struct FPropertyVisitorPath
+{
+public:
+	FPropertyVisitorPath() = default;
+	FPropertyVisitorPath(const FPropertyVisitorPath& Other) = default;
+	FPropertyVisitorPath(FPropertyVisitorPath&& Other) = default;
+	explicit FPropertyVisitorPath(const FPropertyVisitorInfo& Info)
+	{
+		Push(Info);
+	}
+
+	void Push(const FPropertyVisitorInfo& Info)
+	{
+		Path.Push(Info);
+	}
+
+	void Pop()
+	{
+		Path.Pop(EAllowShrinking::No);
+	}
+
+	int32 Num() const
+	{
+		return Path.Num();
+	}
+
+	FPropertyVisitorInfo& Top()
+	{
+		return Path.Top();
+	}
+
+	const FPropertyVisitorInfo& Top() const
+	{
+		return Path.Top();
+	}
+
+	const TArray<FPropertyVisitorInfo>& GetPath() const
+	{
+		return Path;
+	}
+
+	COREUOBJECT_API FString ToString(const TCHAR* Separator = TEXT(".")) const;
+
+	/**
+	 * Is this property path contained in the specified one
+	 * @param OtherPath property path to check if it is contianed in
+	 * @param bIsEqual optional parameter to know if it fully match
+	 * @return true if it is contained in the specified path
+	 */
+	COREUOBJECT_API bool Contained(const FPropertyVisitorPath& OtherPath, bool* bIsEqual) const;
+
+	/**
+	 * Retrieves the data using the specified root object
+	 * @param Object to search into
+	 * @return 
+	 */
+	COREUOBJECT_API void* GetPropertyDataPtr(UObject* Object) const;
+
+	TArray<FPropertyVisitorInfo>::TConstIterator GetRootIterator() const
+	{
+		return Path.CreateConstIterator();
+	}
+
+protected:
+	TArray<FPropertyVisitorInfo> Path;
+};
+
+struct FPropertyVisitorScope
+{
+public:
+	FPropertyVisitorScope(FPropertyVisitorPath& InPath, const FPropertyVisitorInfo& Info)
+		: Path(InPath)
+	{
+		Path.Push(Info);
+	}
+
+	~FPropertyVisitorScope()
+	{
+		Path.Pop();
+	}
+
+protected:
+	FPropertyVisitorPath& Path;
+};
\ No newline at end of file
diff --git a/Engine/Source/Runtime/CoreUObject/Public/UObject/UnrealType.h b/Engine/Source/Runtime/CoreUObject/Public/UObject/UnrealType.h
--- a/Engine/Source/Runtime/CoreUObject/Public/UObject/UnrealType.h
+++ b/Engine/Source/Runtime/CoreUObject/Public/UObject/UnrealType.h
@@ -83,6 +83,8 @@ namespace UE::GC
 	class FPropertyStack;
 	class FSchemaBuilder;
 }
+enum class EPropertyVisitorControlFlow : uint8;
+struct FPropertyVisitorPath;
 
 COREUOBJECT_API DECLARE_LOG_CATEGORY_EXTERN(LogType, Log, All);
 
@@ -313,6 +315,27 @@ public:
 		checkf(HasGetter(), TEXT("Calling a getter on %s but it doesn't have one"), *GetFullName());
 	}
 
+	/**
+	 * Visits this property and allows recursion into the inner properties
+	 * This method allows callers to visit inner properties without knowing about its container type as opposed to TPropertyIterator.
+	 * This visit property pattern facilitates the recursion into user defined properties and allows users to add specific visit logic on UStruct via traits. 
+	 * @param Data to the property to visit
+	 * @param InFunc to call on each visited property, the return value controls what is the next behavior once this property has been visited
+	 * @return the new action to take one visited this property
+	 */
+	COREUOBJECT_API EPropertyVisitorControlFlow Visit(void* Data, const TFunctionRef<EPropertyVisitorControlFlow(const FPropertyVisitorPath& /*Path*/, void* /*Data*/)> InFunc) const;
+
+	/**
+	 * Visits this property and allows recursion into the inner properties
+	 * This method allows callers to visit inner properties without knowing about its container type as opposed to TPropertyIterator.
+	 * This visit property pattern facilitates the recursion into user defined properties and allows users to add specific visit logic on UStruct via traits. 
+	 * @param Path that was computed until we reached this property
+	 * @param Data to the property to visit
+	 * @param InFunc to call on each visited property, the return value controls what is the next behavior once this property has been visited
+	 * @return the new action to take one visited this property
+	 */
+	COREUOBJECT_API virtual EPropertyVisitorControlFlow Visit(FPropertyVisitorPath& Path, void* Data, const TFunctionRef<EPropertyVisitorControlFlow(const FPropertyVisitorPath& /*Path*/, void* /*Data*/)> InFunc) const;
+
 private:
 	/** Set the alignment offset for this property 
 	 * @return the size of the structure including this newly added property
@@ -2495,6 +2518,7 @@ public:
 	virtual FName GetID() const override;
 	virtual void InstanceSubobjects( void* Data, void const* DefaultData, UObject* Owner, struct FObjectInstancingGraph* InstanceGraph ) override;
 	virtual bool SameType(const FProperty* Other) const override;
+	virtual EPropertyVisitorControlFlow Visit(FPropertyVisitorPath& Path, void* Data, const TFunctionRef<EPropertyVisitorControlFlow(const FPropertyVisitorPath& /*Path*/, void* /*Data*/)> InFunc) const override;
 	// End of FProperty interface
 
 	// FObjectPropertyBase interface
@@ -3506,6 +3530,7 @@ public:
 	virtual bool LoadTypeName(UE::FPropertyTypeName Type, const FPropertyTag* Tag = nullptr) override;
 	virtual void SaveTypeName(UE::FPropertyTypeNameBuilder& Type) const override;
 	virtual bool CanSerializeFromTypeName(UE::FPropertyTypeName Type) const override;
+	virtual EPropertyVisitorControlFlow Visit(FPropertyVisitorPath& Path, void* Data, const TFunctionRef<EPropertyVisitorControlFlow(const FPropertyVisitorPath& /*Path*/, void* /*Data*/)> InFunc) const override;
 	// End of FProperty interface
 
 	FString GetCPPTypeCustom(FString* ExtendedTypeText, uint32 CPPExportFlags, const FString& InnerTypeText, const FString& InInnerExtendedTypeText) const;
@@ -3637,6 +3662,7 @@ public:
 	virtual bool LoadTypeName(UE::FPropertyTypeName Type, const FPropertyTag* Tag = nullptr) override;
 	virtual void SaveTypeName(UE::FPropertyTypeNameBuilder& Type) const override;
 	virtual bool CanSerializeFromTypeName(UE::FPropertyTypeName Type) const override;
+	virtual EPropertyVisitorControlFlow Visit(FPropertyVisitorPath& Path, void* Data, const TFunctionRef<EPropertyVisitorControlFlow(const FPropertyVisitorPath& /*Path*/, void* /*Data*/)> InFunc) const override;
 	// End of FProperty interface
 
 	FString GetCPPTypeCustom(FString* ExtendedTypeText, uint32 CPPExportFlags, const FString& KeyTypeText, const FString& InKeyExtendedTypeText, const FString& ValueTypeText, const FString& InValueExtendedTypeText) const;
@@ -3763,6 +3789,7 @@ public:
 	virtual bool LoadTypeName(UE::FPropertyTypeName Type, const FPropertyTag* Tag = nullptr) override;
 	virtual void SaveTypeName(UE::FPropertyTypeNameBuilder& Type) const override;
 	virtual bool CanSerializeFromTypeName(UE::FPropertyTypeName Type) const override;
+	virtual EPropertyVisitorControlFlow Visit(FPropertyVisitorPath& Path, void* Data, const TFunctionRef<EPropertyVisitorControlFlow(const FPropertyVisitorPath& /*Path*/, void* /*Data*/)> InFunc) const override;
 	// End of FProperty interface
 
 	FString GetCPPTypeCustom(FString* ExtendedTypeText, uint32 CPPExportFlags, const FString& ElementTypeText, const FString& InElementExtendedTypeText) const;
@@ -5886,6 +5913,7 @@ public:
 	virtual bool LoadTypeName(UE::FPropertyTypeName Type, const FPropertyTag* Tag = nullptr) override;
 	virtual void SaveTypeName(UE::FPropertyTypeNameBuilder& Type) const override;
 	virtual bool CanSerializeFromTypeName(UE::FPropertyTypeName Type) const override;
+	virtual EPropertyVisitorControlFlow Visit(FPropertyVisitorPath& Path, void* Data, const TFunctionRef<EPropertyVisitorControlFlow(const FPropertyVisitorPath& /*Path*/, void* /*Data*/)> InFunc) const override;
 	// End of FProperty interface
 
 	bool FindInnerPropertyInstance(FName PropertyName, const void* Data, const FProperty*& OutProp, const void*& OutData) const;
