Converted the PropertyValueIterator to use new property visitor

diff --git a/Engine/Source/Runtime/CoreUObject/Private/UObject/UnrealType.cpp b/Engine/Source/Runtime/CoreUObject/Private/UObject/UnrealType.cpp
--- a/Engine/Source/Runtime/CoreUObject/Private/UObject/UnrealType.cpp
+++ b/Engine/Source/Runtime/CoreUObject/Private/UObject/UnrealType.cpp
@@ -3,6 +3,7 @@
 #include "UObject/UnrealType.h"
 #include "UObject/PropertyOptional.h"
 #include "Serialization/ArchiveUObjectFromStructuredArchive.h"
+#include "UObject/PropertyVisitor.h"
 
 DEFINE_LOG_CATEGORY(LogType);
 
@@ -53,25 +54,24 @@ FPropertyValueIterator::EPropertyValueFlags FPropertyValueIterator::GetPropertyV
 
 void FPropertyValueIterator::FillStructProperties(const UStruct* Struct, FPropertyValueStackEntry& Entry)
 {
-	FPropertyValueStackEntry::FValueArrayType& ValueArray = Entry.ValueArray;
-	for (TFieldIterator<FProperty> It(Struct, EFieldIteratorFlags::IncludeSuper, DeprecatedPropertyFlags, EFieldIteratorFlags::ExcludeInterfaces); It; ++It)
+	Struct->Visit(const_cast<void*>(Entry.Owner), [this, &Entry](const FPropertyVisitorPath& Path, void* Data)
 	{
-		const FProperty* Property  = *It;
-		EPropertyValueFlags PropertyValueFlags = GetPropertyValueFlags(Property);
-		if (PropertyValueFlags != EPropertyValueFlags::None)
+		if (const FProperty* InnerProperty = Path.Top().Property)
 		{
-			int32 Num = Property->ArrayDim;
-			ValueArray.Reserve(ValueArray.Num() + Num);
-			for (int32 StaticIndex = 0; StaticIndex < Num; ++StaticIndex)
+			if ((DeprecatedPropertyFlags & EFieldIteratorFlags::IncludeDeprecated) != 0 || !InnerProperty->HasAllPropertyFlags(CPF_Deprecated))
 			{
-				const void* PropertyValue = Property->ContainerPtrToValuePtr<void>(Entry.Owner, StaticIndex);
-				ValueArray.Emplace(BasePairType(Property, PropertyValue), PropertyValueFlags);
+				EPropertyValueFlags InnerFlags = GetPropertyValueFlags(InnerProperty);
+				if (InnerFlags != EPropertyValueFlags::None)
+				{
+					Entry.ValueArray.Emplace(BasePairType(InnerProperty, Data), InnerFlags);
+				}
 			}
 		}
-	}
+		return EPropertyVisitorControlFlow::StepOver;
+	});
 }
 
-FORCEINLINE_DEBUGGABLE bool FPropertyValueIterator::NextValue(EPropertyValueIteratorFlags InRecursionFlags)
+bool FPropertyValueIterator::NextValue(EPropertyValueIteratorFlags InRecursionFlags)
 {
 	check(PropertyIteratorStack.Num() > 0)
 	FPropertyValueStackEntry& Entry = PropertyIteratorStack.Last();
@@ -104,79 +104,28 @@ FORCEINLINE_DEBUGGABLE bool FPropertyValueIterator::NextValue(EPropertyValueIter
 		if (InRecursionFlags == EPropertyValueIteratorFlags::FullRecursion)
 		{
 			FPropertyValueStackEntry NewEntry(PropertyValue);
-			
-			if (EnumHasAnyFlags(PropertyValueFlags, EPropertyValueFlags::IsOptional))
-			{
-				const FOptionalProperty* OptionalProperty = CastFieldChecked<FOptionalProperty>(Property);
-				const FProperty* InnerProperty = OptionalProperty->GetValueProperty();
-				EPropertyValueFlags InnerFlags = GetPropertyValueFlags(InnerProperty);
-				if (InnerFlags != EPropertyValueFlags::None)
-				{
-					if (const void* InnerValue = OptionalProperty->GetValuePointerForReadIfSet(PropertyValue))
-					{
-						NewEntry.ValueArray.Emplace(BasePairType(InnerProperty, InnerValue), InnerFlags);
-					}
-				}
-			}
-			else if (EnumHasAnyFlags(PropertyValueFlags, EPropertyValueFlags::IsArray))
+			Property->Visit(const_cast<void*>(PropertyValue), [this, &NewEntry](const FPropertyVisitorPath& Path, void* Data)
 			{
-				const FArrayProperty* ArrayProperty = CastFieldChecked<FArrayProperty>(Property);
-				const FProperty* InnerProperty = ArrayProperty->Inner;
-				EPropertyValueFlags InnerFlags = GetPropertyValueFlags(InnerProperty);
-				if (InnerFlags != EPropertyValueFlags::None)
+				if (const FProperty* InnerProperty = Path.Top().Property)
 				{
-					FScriptArrayHelper Helper(ArrayProperty, PropertyValue);
-					const int32 Num = Helper.Num();
-					NewEntry.ValueArray.Reserve(Num);
-					for (int32 DynamicIndex = 0; DynamicIndex < Num; ++DynamicIndex)
+					if ((DeprecatedPropertyFlags & EFieldIteratorFlags::IncludeDeprecated) != 0 || !InnerProperty->HasAllPropertyFlags(CPF_Deprecated))
 					{
-						NewEntry.ValueArray.Emplace(
-							BasePairType(InnerProperty, Helper.GetRawPtr(DynamicIndex)), InnerFlags);
-					}
-				}
-			}
-			else if (EnumHasAnyFlags(PropertyValueFlags, EPropertyValueFlags::IsMap))
-			{
-				const FMapProperty* MapProperty = CastFieldChecked<FMapProperty>(Property);
-				const FProperty* KeyProperty = MapProperty->KeyProp;
-				const FProperty* ValueProperty = MapProperty->ValueProp;
-				EPropertyValueFlags KeyFlags = GetPropertyValueFlags(KeyProperty);
-				EPropertyValueFlags ValueFlags = GetPropertyValueFlags(ValueProperty);
-				if ((KeyFlags | ValueFlags) != EPropertyValueFlags::None)
-				{
-					FScriptMapHelper Helper(MapProperty, PropertyValue);
-					for (FScriptMapHelper::FIterator It(Helper); It; ++It)
-					{
-						if (KeyFlags != EPropertyValueFlags::None)
+						// Visit any properties at the top level that contains inner properties and are not object references
+						if (Path.Num() == 1 && Path.Top().bContainsInnerProperties && !InnerProperty->IsA<FObjectPropertyBase>())
 						{
-							NewEntry.ValueArray.Emplace(BasePairType(KeyProperty, Helper.GetKeyPtr(It)), KeyFlags);
+							return EPropertyVisitorControlFlow::StepInto;
 						}
-						if (ValueFlags != EPropertyValueFlags::None)
+
+						EPropertyValueFlags InnerFlags = GetPropertyValueFlags(InnerProperty);
+						if (InnerFlags != EPropertyValueFlags::None)
 						{
-							NewEntry.ValueArray.Emplace(BasePairType(ValueProperty, Helper.GetValuePtr(It)), ValueFlags);
+							NewEntry.ValueArray.Emplace(BasePairType(InnerProperty, Data), InnerFlags);
 						}
 					}
 				}
-			}
-			else if (EnumHasAnyFlags(PropertyValueFlags, EPropertyValueFlags::IsSet))
-			{
-				const FSetProperty* SetProperty = CastFieldChecked<FSetProperty>(Property);
-				const FProperty* InnerProperty = SetProperty->ElementProp;
-				EPropertyValueFlags InnerFlags = GetPropertyValueFlags(InnerProperty);
-				if (InnerFlags != EPropertyValueFlags::None)
-				{
-					FScriptSetHelper Helper(SetProperty, PropertyValue);
-					for (FScriptSetHelper::FIterator It(Helper); It; ++It)
-					{
-						NewEntry.ValueArray.Emplace(BasePairType(InnerProperty, Helper.GetElementPtr(It)), InnerFlags);
-					}
-				}
-			}
-			else if (EnumHasAnyFlags(PropertyValueFlags, EPropertyValueFlags::IsStruct))
-			{
-				const FStructProperty* StructProperty = CastFieldChecked<FStructProperty>(Property);
-				FillStructProperties(StructProperty->Struct, NewEntry);
-			}
+				return EPropertyVisitorControlFlow::StepOver;
+			});
+			
 			if (NewEntry.ValueArray.Num() > 0)
 			{
 				PropertyIteratorStack.Emplace(MoveTemp(NewEntry));
