Adding missing GC stats related to initial reference collection and pre-reachability callbacks

diff --git a/Engine/Source/Runtime/CoreUObject/Private/UObject/GarbageCollection.cpp b/Engine/Source/Runtime/CoreUObject/Private/UObject/GarbageCollection.cpp
--- a/Engine/Source/Runtime/CoreUObject/Private/UObject/GarbageCollection.cpp
+++ b/Engine/Source/Runtime/CoreUObject/Private/UObject/GarbageCollection.cpp
@@ -3986,6 +3986,7 @@ class FRealtimeGC : public FGarbageCollectionTracer
 
 	TConstArrayView<UObject**> GetInitialReferences(EGCOptions Options)
 	{
+		const double StartTime = FPlatformTime::Seconds();
 		if (!!(Options & EGCOptions::Parallel))
 		{
 			InitialCollection.Wait();
@@ -3994,6 +3995,7 @@ class FRealtimeGC : public FGarbageCollectionTracer
 		{
 			FGCObject::GGCObjectReferencer->AddInitialReferences(InitialReferences);
 		}
+		GGCStats.InitialReferenceCollectionTime += FPlatformTime::Seconds() - StartTime;
 
 		return InitialReferences;
 	}
@@ -4437,7 +4439,9 @@ PRAGMA_DISABLE_DEPRECATION_WARNINGS
 		// because it may require tracing objects (via FGarbageCollectionTracer) multiple times
 		if (!GReachabilityState.IsSuspended())
 		{
+			const double StartTime = FPlatformTime::Seconds();
 			FCoreUObjectDelegates::TraceExternalRootsForReachabilityAnalysis.Broadcast(*this, KeepFlags, !(Options & EGCOptions::Parallel));
+			GGCStats.TraceExternalRootsTime += FPlatformTime::Seconds() - StartTime;
 		}
 PRAGMA_ENABLE_DEPRECATION_WARNINGS
 	}
diff --git a/Engine/Source/Runtime/CoreUObject/Public/UObject/ReachabilityAnalysis.h b/Engine/Source/Runtime/CoreUObject/Public/UObject/ReachabilityAnalysis.h
--- a/Engine/Source/Runtime/CoreUObject/Public/UObject/ReachabilityAnalysis.h
+++ b/Engine/Source/Runtime/CoreUObject/Public/UObject/ReachabilityAnalysis.h
@@ -73,7 +73,9 @@ struct FStats
 	double VerifyNoUnreachableTime = 0.0;
 	double GarbageTrackingTime = 0.0;
 
+	double InitialReferenceCollectionTime = 0.0;
 	double MarkObjectsAsUnreachableTime = 0.0;
+	double TraceExternalRootsTime = 0.0;
 	double NotifyUnreachableTime = 0.0;
 	double DissolveUnreachableClustersTime = 0.0;
 
