Replaced ImplicitConv usage in TObjectPtr.

diff --git a/Engine/Source/Runtime/CoreUObject/Public/UObject/ObjectPtr.h b/Engine/Source/Runtime/CoreUObject/Public/UObject/ObjectPtr.h
--- a/Engine/Source/Runtime/CoreUObject/Public/UObject/ObjectPtr.h
+++ b/Engine/Source/Runtime/CoreUObject/Public/UObject/ObjectPtr.h
@@ -516,8 +516,8 @@ public:
 	}
 
 	template <
-		typename U,
-		decltype(ImplicitConv<T*>(std::declval<U*>()))* = nullptr
+		typename U
+		UE_REQUIRES(std::is_convertible_v<U*, T*>)
 	>
 	FORCEINLINE TObjectPtr(const TObjectPtr<U>& Other)
 		: ObjectPtr(Other.ObjectPtr)
@@ -526,9 +526,8 @@ public:
 	}
 
 	template <
-		typename U,
-		decltype(ImplicitConv<T*>(std::declval<U>()))* = nullptr
-		UE_REQUIRES(!TIsTObjectPtr_V<std::decay_t<U>>)
+		typename U
+		UE_REQUIRES(!TIsTObjectPtr_V<std::decay_t<U>> && std::is_convertible_v<U, T*>)
 	>
 	FORCEINLINE TObjectPtr(const U& Object)
 		: ObjectPtr(const_cast<std::remove_const_t<T>*>(ImplicitConv<T*>(Object)))
@@ -568,8 +567,8 @@ public:
 	}
 
 	template <
-		typename U,
-		decltype(ImplicitConv<T*>(std::declval<U*>()))* = nullptr
+		typename U
+		UE_REQUIRES(std::is_convertible_v<U*, T*>)
 	>
 	FORCEINLINE TObjectPtr<T>& operator=(const TObjectPtr<U>& Other)
 	{
@@ -579,9 +578,8 @@ public:
 	}
 
 	template <
-		typename U,
-		decltype(ImplicitConv<T*>(std::declval<U>()))* = nullptr
-		UE_REQUIRES(!TIsTObjectPtr_V<std::decay_t<U>>)
+		typename U
+		UE_REQUIRES(!TIsTObjectPtr_V<std::decay_t<U>> && std::is_convertible_v<U, T*>)
 	>
 	FORCEINLINE TObjectPtr<T>& operator=(U&& Object)
 	{
@@ -1690,8 +1688,8 @@ public:
 	 * Constructs a non-null pointer from another non-null pointer
 	 */
 	template <
-		typename OtherObjectType,
-		decltype(ImplicitConv<ObjectType*>((OtherObjectType*)nullptr))* = nullptr
+		typename OtherObjectType
+		UE_REQUIRES(std::is_convertible_v<OtherObjectType*, ObjectType*>)
 	>
 	FORCEINLINE TNonNullPtr(const TNonNullPtr<OtherObjectType>& Other)
 		: Object(Other.Object)
@@ -1722,8 +1720,8 @@ public:
 	 * Assignment operator taking another TNonNullPtr
 	 */
 	template <
-		typename OtherObjectType,
-		decltype(ImplicitConv<ObjectType*>((OtherObjectType*)nullptr))* = nullptr
+		typename OtherObjectType
+		UE_REQUIRES(std::is_convertible_v<OtherObjectType*, ObjectType*>)
 	>
 	FORCEINLINE TNonNullPtr& operator=(const TNonNullPtr<OtherObjectType>& Other)
 	{
