Replaced ImplicitConv usage in T[Script/Weak]Interface.  Includes a fix to TScriptInterface::operator=(const TScriptInterface<OtherInterfaceType>&).

diff --git a/Engine/Source/Runtime/CoreUObject/Public/UObject/ScriptInterface.h b/Engine/Source/Runtime/CoreUObject/Public/UObject/ScriptInterface.h
--- a/Engine/Source/Runtime/CoreUObject/Public/UObject/ScriptInterface.h
+++ b/Engine/Source/Runtime/CoreUObject/Public/UObject/ScriptInterface.h
@@ -154,8 +154,8 @@ public:
 	 * Construction from an object type that may natively implement InterfaceType
 	 */
 	template <
-		typename U,
-		decltype(ImplicitConv<UObject*>(std::declval<U>()))* = nullptr
+		typename U
+		UE_REQUIRES(std::is_convertible_v<U, UObject*>)
 	>
 	FORCEINLINE TScriptInterface(U&& Source)
 	{
@@ -181,8 +181,8 @@ public:
 	 * Construction from another script interface of a compatible interface type
 	 */
 	template <
-		typename OtherInterfaceType,
-		decltype(ImplicitConv<InInterfaceType*>(std::declval<OtherInterfaceType*>()))* = nullptr
+		typename OtherInterfaceType
+		UE_REQUIRES(std::is_convertible_v<OtherInterfaceType*, InInterfaceType*>)
 	>
 	FORCEINLINE TScriptInterface(const TScriptInterface<OtherInterfaceType>& Other)
 	{
@@ -234,8 +234,8 @@ public:
 	 * Assignment from an object type that may natively implement InterfaceType
 	 */
 	template <
-		typename U,
-		decltype(ImplicitConv<UObject*>(std::declval<U>()))* = nullptr
+		typename U
+		UE_REQUIRES(std::is_convertible_v<U, UObject*>)
 	>
 	TScriptInterface& operator=(U&& Source)
 	{
@@ -247,8 +247,8 @@ public:
 	 * Assignment from another script interface of a compatible interface type
 	 */
 	template <
-		typename OtherInterfaceType,
-		decltype(ImplicitConv<InInterfaceType*>(std::declval<OtherInterfaceType>()))* = nullptr
+		typename OtherInterfaceType
+		UE_REQUIRES(std::is_convertible_v<OtherInterfaceType*, InInterfaceType*>)
 	>
 	TScriptInterface& operator=(const TScriptInterface<OtherInterfaceType>& Other)
 	{
@@ -269,12 +269,18 @@ public:
 	/**
 	 * Comparison operator, taking a pointer to InterfaceType
 	 */
-	template <typename OtherInterface, typename = decltype(ImplicitConv<InInterfaceType*>((OtherInterface*)nullptr))>
+	template <
+		typename OtherInterface
+		UE_REQUIRES(std::is_convertible_v<OtherInterface*, InInterfaceType*>)
+	>
 	FORCEINLINE bool operator==( const OtherInterface* Other ) const
 	{
 		return GetInterface() == Other;
 	}
-	template <typename OtherInterface, typename = decltype(ImplicitConv<InInterfaceType*>((OtherInterface*)nullptr))>
+	template <
+		typename OtherInterface
+		UE_REQUIRES(std::is_convertible_v<OtherInterface*, InInterfaceType*>)
+	>
 	FORCEINLINE bool operator!=( const OtherInterface* Other ) const
 	{
 		return GetInterface() != Other;
diff --git a/Engine/Source/Runtime/CoreUObject/Public/UObject/WeakInterfacePtr.h b/Engine/Source/Runtime/CoreUObject/Public/UObject/WeakInterfacePtr.h
--- a/Engine/Source/Runtime/CoreUObject/Public/UObject/WeakInterfacePtr.h
+++ b/Engine/Source/Runtime/CoreUObject/Public/UObject/WeakInterfacePtr.h
@@ -30,8 +30,8 @@ struct TWeakInterfacePtr
 	 * @param Object The object to create a weak pointer to. This object must implement interface T.
 	 */
 	template<
-		typename U,
-		decltype(ImplicitConv<typename TCopyQualifiersFromTo<U, UObject>::Type*>(std::declval<U>()))* = nullptr
+		typename U
+		UE_REQUIRES(std::is_convertible_v<U, typename TCopyQualifiersFromTo<U, UObject>::Type*>)
 	>
 	TWeakInterfacePtr(U&& Object)
 	{
