Removed ImplicitConv usage from TStridedView, TSoftObjectPtr, TStrongObjectPtr and TWeakFieldPtr.

diff --git a/Engine/Source/Runtime/Core/Public/Containers/StridedView.h b/Engine/Source/Runtime/Core/Public/Containers/StridedView.h
--- a/Engine/Source/Runtime/Core/Public/Containers/StridedView.h
+++ b/Engine/Source/Runtime/Core/Public/Containers/StridedView.h
@@ -51,7 +51,10 @@ public:
 
 	TStridedView() = default;
 
-	template <typename OtherElementType, decltype(ImplicitConv<ElementType* const*>((OtherElementType**)nullptr))* = nullptr>
+	template <
+		typename OtherElementType
+		UE_REQUIRES(std::is_convertible_v<OtherElementType**, ElementType* const*>)
+	>
 	FORCEINLINE TStridedView(SizeType InBytesBetweenElements, OtherElementType* InFirstElementPtr, SizeType InNumElements)
 		: FirstElementPtr(InFirstElementPtr)
 		, BytesBetweenElements(InBytesBetweenElements)
@@ -62,7 +65,10 @@ public:
 		check(BytesBetweenElements % alignof(ElementType) == 0);
 	}
 
-	template <typename OtherElementType, decltype(ImplicitConv<ElementType* const*>((OtherElementType**)nullptr))* = nullptr>
+	template <
+		typename OtherElementType
+		UE_REQUIRES(std::is_convertible_v<OtherElementType**, ElementType* const*>)
+	>
 	FORCEINLINE TStridedView(const TStridedView<OtherElementType, SizeType>& Other)
 		: FirstElementPtr(nullptr)
 		, BytesBetweenElements(Other.GetStride())
diff --git a/Engine/Source/Runtime/CoreUObject/Public/UObject/SoftObjectPtr.h b/Engine/Source/Runtime/CoreUObject/Public/UObject/SoftObjectPtr.h
--- a/Engine/Source/Runtime/CoreUObject/Public/UObject/SoftObjectPtr.h
+++ b/Engine/Source/Runtime/CoreUObject/Public/UObject/SoftObjectPtr.h
@@ -11,6 +11,8 @@
 #include "UObject/PersistentObjectPtr.h"
 #include "UObject/SoftObjectPath.h"
 
+#include <type_traits>
+
 /**
  * TIsSoftObjectPointerType
  * Trait for recognizing 'soft' (path-based) object pointer types
@@ -128,14 +130,20 @@ public:
 	FORCEINLINE TSoftObjectPtr& operator=(TSoftObjectPtr&& Other) = default;
 	
 	/** Construct from another soft pointer */
-	template <class U, class = decltype(ImplicitConv<T*>((U*)nullptr))>
+	template <
+		class U
+		UE_REQUIRES(std::is_convertible_v<U*, T*>)
+	>
 	FORCEINLINE TSoftObjectPtr(const TSoftObjectPtr<U>& Other)
 		: SoftObjectPtr(Other.SoftObjectPtr)
 	{
 	}
 
 	/** Construct from a moveable soft pointer */
-	template <class U, class = decltype(ImplicitConv<T*>((U*)nullptr))>
+	template <
+		class U
+		UE_REQUIRES(std::is_convertible_v<U*, T*>)
+	>
 	FORCEINLINE TSoftObjectPtr(TSoftObjectPtr<U>&& Other)
 		: SoftObjectPtr(MoveTemp(Other.SoftObjectPtr))
 	{
@@ -210,7 +218,10 @@ public:
 	}
 
 	/** Copy from a weak pointer to an object already in memory */
-	template <class U, class = decltype(ImplicitConv<T*>((U*)nullptr))>
+	template <
+		class U
+		UE_REQUIRES(std::is_convertible_v<U*, T*>)
+	>
 	FORCEINLINE TSoftObjectPtr& operator=(const TWeakObjectPtr<U>& Other)
 	{
 		SoftObjectPtr = Other;
@@ -218,7 +229,10 @@ public:
 	}
 
 	/** Copy from another soft pointer */
-	template <class U, class = decltype(ImplicitConv<T*>((U*)nullptr))>
+	template <
+		class U
+		UE_REQUIRES(std::is_convertible_v<U*, T*>)
+	>
 	FORCEINLINE TSoftObjectPtr& operator=(TSoftObjectPtr<U> Other)
 	{
 		SoftObjectPtr = MoveTemp(Other.SoftObjectPtr);
@@ -465,7 +479,10 @@ public:
 	FORCEINLINE TSoftClassPtr& operator=(TSoftClassPtr&& Other) = default;
 		
 	/** Construct from another soft pointer */
-	template <class TClassA, class = decltype(ImplicitConv<TClass*>((TClassA*)nullptr))>
+	template <
+		class TClassA
+		UE_REQUIRES(std::is_convertible_v<TClassA*, TClass*>)
+	>
 	FORCEINLINE TSoftClassPtr(const TSoftClassPtr<TClassA>& Other)
 		: SoftObjectPtr(Other.SoftObjectPtr)
 	{
@@ -508,7 +525,10 @@ public:
 	}
 
 	/** Copy from a weak pointer already in memory */
-	template<class TClassA, class = decltype(ImplicitConv<TClass*>((TClassA*)nullptr))>
+	template <
+		class TClassA
+		UE_REQUIRES(std::is_convertible_v<TClassA*, TClass*>)
+	>
 	FORCEINLINE TSoftClassPtr& operator=(const TWeakObjectPtr<TClassA>& Other)
 	{
 		SoftObjectPtr = Other;
@@ -516,7 +536,10 @@ public:
 	}
 
 	/** Copy from another soft pointer */
-	template<class TClassA, class = decltype(ImplicitConv<TClass*>((TClassA*)nullptr))>
+	template <
+		class TClassA
+		UE_REQUIRES(std::is_convertible_v<TClassA*, TClass*>)
+	>
 	FORCEINLINE TSoftClassPtr& operator=(const TSoftObjectPtr<TClassA>& Other)
 	{
 		SoftObjectPtr = Other.SoftObjectPtr;
diff --git a/Engine/Source/Runtime/CoreUObject/Public/UObject/StrongObjectPtr.h b/Engine/Source/Runtime/CoreUObject/Public/UObject/StrongObjectPtr.h
--- a/Engine/Source/Runtime/CoreUObject/Public/UObject/StrongObjectPtr.h
+++ b/Engine/Source/Runtime/CoreUObject/Public/UObject/StrongObjectPtr.h
@@ -96,8 +96,8 @@ public:
 
 	template <
 		typename OtherObjectType,
-		typename OtherReferencerNameProvider,
-		typename = decltype(ImplicitConv<ObjectType*>((OtherObjectType*)nullptr))
+		typename OtherReferencerNameProvider
+		UE_REQUIRES(std::is_convertible_v<OtherObjectType*, ObjectType*>)
 	>
 	FORCEINLINE_DEBUGGABLE TStrongObjectPtr(const TStrongObjectPtr<OtherObjectType, OtherReferencerNameProvider>& InOther)
 	{
@@ -112,8 +112,8 @@ public:
 
 	template <
 		typename OtherObjectType,
-		typename OtherReferencerNameProvider,
-		typename = decltype(ImplicitConv<ObjectType*>((OtherObjectType*)nullptr))
+		typename OtherReferencerNameProvider
+		UE_REQUIRES(std::is_convertible_v<OtherObjectType*, ObjectType*>)
 	>
 	FORCEINLINE_DEBUGGABLE TStrongObjectPtr& operator=(const TStrongObjectPtr<OtherObjectType, OtherReferencerNameProvider>& InOther)
 	{
diff --git a/Engine/Source/Runtime/CoreUObject/Public/UObject/WeakFieldPtr.h b/Engine/Source/Runtime/CoreUObject/Public/UObject/WeakFieldPtr.h
--- a/Engine/Source/Runtime/CoreUObject/Public/UObject/WeakFieldPtr.h
+++ b/Engine/Source/Runtime/CoreUObject/Public/UObject/WeakFieldPtr.h
@@ -10,6 +10,8 @@
 #include "UObject/WeakObjectPtr.h"
 #include "UObject/FieldPath.h"
 
+#include <type_traits>
+
 template<class T>
 struct TWeakFieldPtr
 {
@@ -43,10 +45,10 @@ public:
 	* @param Object object to create a weak pointer to
 	**/
 	template <
-		typename U,
-		typename = decltype(ImplicitConv<T*>((U*)nullptr))
-		>
-		FORCEINLINE TWeakFieldPtr(U* InField, EDummy1 = Dummy1)
+		typename U
+		UE_REQUIRES(std::is_convertible_v<U*, T*>)
+	>
+	FORCEINLINE TWeakFieldPtr(U* InField, EDummy1 = Dummy1)
 		: Owner(InField ? InField->GetOwnerUObject() : (UObject*)nullptr)
 		, Field(InField)
 	{
