Removed ImplicitConv usage from TNonNullPointer, TNonNullSubclassOf, TSubclassOf and TSubclassOfField.

diff --git a/Engine/Source/Runtime/Core/Public/Templates/NonNullPointer.h b/Engine/Source/Runtime/Core/Public/Templates/NonNullPointer.h
--- a/Engine/Source/Runtime/Core/Public/Templates/NonNullPointer.h
+++ b/Engine/Source/Runtime/Core/Public/Templates/NonNullPointer.h
@@ -50,8 +50,8 @@ public:
 	 * Constructs a non-null pointer from another non-null pointer
 	 */
 	template <
-		typename OtherObjectType,
-		decltype(ImplicitConv<ObjectType*>((OtherObjectType*)nullptr))* = nullptr
+		typename OtherObjectType
+		UE_REQUIRES(std::is_convertible_v<OtherObjectType*, ObjectType*>)
 	>
 	FORCEINLINE TNonNullPtr(const TNonNullPtr<OtherObjectType>& Other)
 		: Object(Other.Object)
@@ -82,8 +82,8 @@ public:
 	 * Assignment operator taking another TNonNullPtr
 	 */
 	template <
-		typename OtherObjectType,
-		decltype(ImplicitConv<ObjectType*>((OtherObjectType*)nullptr))* = nullptr
+		typename OtherObjectType
+		UE_REQUIRES(std::is_convertible_v<OtherObjectType*, ObjectType*>)
 	>
 	FORCEINLINE TNonNullPtr& operator=(const TNonNullPtr<OtherObjectType>& Other)
 	{
diff --git a/Engine/Source/Runtime/CoreUObject/Public/Templates/NonNullSubclassOf.h b/Engine/Source/Runtime/CoreUObject/Public/Templates/NonNullSubclassOf.h
--- a/Engine/Source/Runtime/CoreUObject/Public/Templates/NonNullSubclassOf.h
+++ b/Engine/Source/Runtime/CoreUObject/Public/Templates/NonNullSubclassOf.h
@@ -35,8 +35,8 @@ public:
 
 	/** Copy Constructor, will only compile if types are compatible */
 	template <
-		typename U,
-		decltype(ImplicitConv<T*>((U*)nullptr))* = nullptr
+		typename U
+		UE_REQUIRES(std::is_convertible_v<U*, T*>)
 	>
 	FORCEINLINE TNonNullSubclassOf(const TSubclassOf<U>& From)
 		: Super(From)
@@ -46,8 +46,8 @@ public:
 	/** Assignment operator, will only compile if types are compatible */
 	template
 	<
-		typename U,
-		decltype(ImplicitConv<T*>((U*)nullptr))* = nullptr
+		typename U
+		UE_REQUIRES(std::is_convertible_v<U*, T*>)
 	>
 	FORCEINLINE TNonNullSubclassOf& operator=(const TSubclassOf<U>& From)
 	{
diff --git a/Engine/Source/Runtime/CoreUObject/Public/Templates/SubclassOf.h b/Engine/Source/Runtime/CoreUObject/Public/Templates/SubclassOf.h
--- a/Engine/Source/Runtime/CoreUObject/Public/Templates/SubclassOf.h
+++ b/Engine/Source/Runtime/CoreUObject/Public/Templates/SubclassOf.h
@@ -50,11 +50,11 @@ public:
 
 	/** Construct from a UClass* (or something implicitly convertible to it) */
 	template <
-		typename U,
-		std::enable_if_t<
-			!TIsTSubclassOf<std::decay_t<U>>::Value,
-			decltype(ImplicitConv<UClass*>(std::declval<U>()))
-		>* = nullptr
+		typename U
+		UE_REQUIRES(
+			!TIsTSubclassOf<std::decay_t<U>>::Value &&
+			std::is_convertible_v<U, UClass*>
+		)
 	>
 	FORCEINLINE TSubclassOf(U&& From)
 		: Class(From)
@@ -63,8 +63,8 @@ public:
 
 	/** Construct from another TSubclassOf, only if types are compatible */
 	template <
-		typename OtherT,
-		decltype(ImplicitConv<T*>((OtherT*)nullptr))* = nullptr
+		typename OtherT
+		UE_REQUIRES(std::is_convertible_v<OtherT*, T*>)
 	>
 	FORCEINLINE TSubclassOf(const TSubclassOf<OtherT>& Other)
 		: Class(Other.Class)
@@ -74,8 +74,8 @@ public:
 
 	/** Assign from another TSubclassOf, only if types are compatible */
 	template <
-		typename OtherT,
-		decltype(ImplicitConv<T*>((OtherT*)nullptr))* = nullptr
+		typename OtherT
+		UE_REQUIRES(std::is_convertible_v<OtherT*, T*>)
 	>
 	FORCEINLINE TSubclassOf& operator=(const TSubclassOf<OtherT>& Other)
 	{
@@ -93,11 +93,11 @@ public:
 
 	/** Assign from a UClass* (or something implicitly convertible to it). */
 	template <
-		typename U,
-		std::enable_if_t<
-			!TIsTSubclassOf<std::decay_t<U>>::Value,
-			decltype(ImplicitConv<UClass*>(std::declval<U>()))
-		>* = nullptr
+		typename U
+		UE_REQUIRES(
+			!TIsTSubclassOf<std::decay_t<U>>::Value &&
+			std::is_convertible_v<U, UClass*>
+		)
 	>
 	FORCEINLINE TSubclassOf& operator=(U&& From)
 	{
diff --git a/Engine/Source/Runtime/CoreUObject/Public/Templates/SubclassOfField.h b/Engine/Source/Runtime/CoreUObject/Public/Templates/SubclassOfField.h
--- a/Engine/Source/Runtime/CoreUObject/Public/Templates/SubclassOfField.h
+++ b/Engine/Source/Runtime/CoreUObject/Public/Templates/SubclassOfField.h
@@ -47,11 +47,11 @@ public:
 
 	/** Construct from a FFieldClass* (or something implicitly convertible to it) */
 	template <
-		typename U,
-		std::enable_if_t<
-			!TIsTSubclassOfField<std::decay_t<U>>::Value,
-			decltype(ImplicitConv<FFieldClass*>(std::declval<U>()))
-		>* = nullptr
+		typename U
+		UE_REQUIRES(
+			!TIsTSubclassOfField<std::decay_t<U>>::Value &&
+			std::is_convertible_v<U, FFieldClass*>
+		)
 	>
 	FORCEINLINE TSubclassOfField(U&& From)
 		: Class(From)
@@ -60,8 +60,8 @@ public:
 
 	/** Construct from another TSubclassOfField, only if types are compatible */
 	template <
-		typename OtherT,
-		decltype(ImplicitConv<T*>((OtherT*)nullptr))* = nullptr
+		typename OtherT
+		UE_REQUIRES(std::is_convertible_v<OtherT*, T*>)
 	>
 	FORCEINLINE TSubclassOfField(const TSubclassOfField<OtherT>& Other)
 		: Class(Other.Class)
@@ -70,8 +70,8 @@ public:
 
 	/** Assign from another TSubclassOfField, only if types are compatible */
 	template <
-		typename OtherT,
-		decltype(ImplicitConv<T*>((OtherT*)nullptr))* = nullptr
+		typename OtherT
+		UE_REQUIRES(std::is_convertible_v<OtherT*, T*>)
 	>
 	FORCEINLINE TSubclassOfField& operator=(const TSubclassOfField<OtherT>& Other)
 	{
@@ -88,11 +88,11 @@ public:
 
 	/** Assign from a FFieldClass* (or something implicitly convertible to it). */
 	template <
-		typename U,
-		std::enable_if_t<
-			!TIsTSubclassOfField<std::decay_t<U>>::Value,
-			decltype(ImplicitConv<FFieldClass*>(std::declval<U>()))
-		>* = nullptr
+		typename U
+		UE_REQUIRES(
+			!TIsTSubclassOfField<std::decay_t<U>>::Value &&
+			std::is_convertible_v<U, FFieldClass*>
+		)
 	>
 	FORCEINLINE TSubclassOfField& operator=(U&& From)
 	{
