Updated field/class/enum meta-data localization queries to always return a localizable FText

This ensures that the 'keys' culture shows the text as localizable, even if there's no translation data available for it yet

#rnx

diff --git a/Engine/Source/Runtime/CoreUObject/Private/UObject/Class.cpp b/Engine/Source/Runtime/CoreUObject/Private/UObject/Class.cpp
--- a/Engine/Source/Runtime/CoreUObject/Private/UObject/Class.cpp
+++ b/Engine/Source/Runtime/CoreUObject/Private/UObject/Class.cpp
@@ -268,8 +268,6 @@ struct FDisplayNameHelper
  */
 FText UField::GetDisplayNameText() const
 {
-	FText LocalizedDisplayName;
-
 	static const FString Namespace = TEXT("UObjectDisplayNames");
 	static const FName NAME_DisplayName(TEXT("DisplayName"));
 
@@ -281,12 +279,7 @@ FText UField::GetDisplayNameText() const
 		NativeDisplayName = FName::NameToDisplayString(FDisplayNameHelper::Get(*this), false);
 	}
 
-	if ( !( FText::FindText( Namespace, Key, /*OUT*/LocalizedDisplayName, &NativeDisplayName ) ) )
-	{
-		LocalizedDisplayName = FText::FromString(NativeDisplayName );
-	}
-
-	return LocalizedDisplayName;
+	return FInternationalization::ForUseOnlyByLocMacroAndGraphNodeTextLiterals_CreateText(*NativeDisplayName, *Namespace, *Key);
 }
 
 /**
@@ -331,7 +324,7 @@ FText UField::GetToolTipText(bool bShortTooltip) const
 		{
 			FormatNativeToolTip(NativeToolTip, true);
 		}
-		LocalizedToolTip = FText::FromString(NativeToolTip);
+		LocalizedToolTip = FInternationalization::ForUseOnlyByLocMacroAndGraphNodeTextLiterals_CreateText(*NativeToolTip, *Namespace, *Key);
 	}
 
 	return LocalizedToolTip;
@@ -474,14 +467,10 @@ FText UField::GetMetaDataText(const TCHAR* MetaDataKey, const FString Localizati
 	}
 
 	FText LocalizedMetaData;
-	if ( !( FText::FindText( LocalizationNamespace, LocalizationKey, /*OUT*/LocalizedMetaData, &DefaultMetaData ) ) )
+	if (!DefaultMetaData.IsEmpty())
 	{
-		if (!DefaultMetaData.IsEmpty())
-		{
-			LocalizedMetaData = FText::AsCultureInvariant(DefaultMetaData);
-		}
+		LocalizedMetaData = FInternationalization::ForUseOnlyByLocMacroAndGraphNodeTextLiterals_CreateText(*DefaultMetaData, *LocalizationNamespace, *LocalizationKey);
 	}
-
 	return LocalizedMetaData;
 }
 
@@ -500,16 +489,11 @@ FText UField::GetMetaDataText(const FName& MetaDataKey, const FString Localizati
 		DefaultMetaData = FName::NameToDisplayString(GetName(), false);
 	}
 	
-
 	FText LocalizedMetaData;
-	if ( !( FText::FindText( LocalizationNamespace, LocalizationKey, /*OUT*/LocalizedMetaData, &DefaultMetaData ) ) )
+	if (!DefaultMetaData.IsEmpty())
 	{
-		if (!DefaultMetaData.IsEmpty())
-		{
-			LocalizedMetaData = FText::AsCultureInvariant(DefaultMetaData);
-		}
+		LocalizedMetaData = FInternationalization::ForUseOnlyByLocMacroAndGraphNodeTextLiterals_CreateText(*DefaultMetaData, *LocalizationNamespace, *LocalizationKey);
 	}
-
 	return LocalizedMetaData;
 }
 
diff --git a/Engine/Source/Runtime/CoreUObject/Private/UObject/Enum.cpp b/Engine/Source/Runtime/CoreUObject/Private/UObject/Enum.cpp
--- a/Engine/Source/Runtime/CoreUObject/Private/UObject/Enum.cpp
+++ b/Engine/Source/Runtime/CoreUObject/Private/UObject/Enum.cpp
@@ -672,7 +672,6 @@ FText UEnum::GetDisplayNameTextByIndex(int32 NameIndex) const
 	}
 
 #if WITH_EDITOR
-	FText LocalizedDisplayName;
 	// In the editor, use metadata and localization to look up names
 	static const FString Namespace = TEXT("UObjectDisplayNames");
 	const FString Key = GetFullGroupName(false) + TEXT(".") + RawName;
@@ -687,14 +686,9 @@ FText UEnum::GetDisplayNameTextByIndex(int32 NameIndex) const
 		NativeDisplayName = FName::NameToDisplayString(RawName, false);
 	}
 
-	if (!(FText::FindText(Namespace, Key, /*OUT*/LocalizedDisplayName, &NativeDisplayName)))
+	if (!NativeDisplayName.IsEmpty())
 	{
-		LocalizedDisplayName = FText::FromString(NativeDisplayName);
-	}
-
-	if (!LocalizedDisplayName.IsEmpty())
-	{
-		return LocalizedDisplayName;
+		return FInternationalization::ForUseOnlyByLocMacroAndGraphNodeTextLiterals_CreateText(*NativeDisplayName, *Namespace, *Key);
 	}
 #endif
 
@@ -931,7 +925,7 @@ FText UEnum::GetToolTipTextByIndex(int32 NameIndex) const
 			NativeToolTip.TrimEndInline();
 		}
 
-		LocalizedToolTip = FText::FromString(NativeToolTip);
+		LocalizedToolTip = FInternationalization::ForUseOnlyByLocMacroAndGraphNodeTextLiterals_CreateText(*NativeToolTip, *Namespace, *Key);
 	}
 
 	return LocalizedToolTip;
diff --git a/Engine/Source/Runtime/CoreUObject/Private/UObject/Field.cpp b/Engine/Source/Runtime/CoreUObject/Private/UObject/Field.cpp
--- a/Engine/Source/Runtime/CoreUObject/Private/UObject/Field.cpp
+++ b/Engine/Source/Runtime/CoreUObject/Private/UObject/Field.cpp
@@ -676,8 +676,6 @@ struct FFieldDisplayNameHelper
 */
 FText FField::GetDisplayNameText() const
 {
-	FText LocalizedDisplayName;
-
 	static const FString Namespace = TEXT("UObjectDisplayNames");
 	static const FName NAME_DisplayName(TEXT("DisplayName"));
 
@@ -693,12 +691,7 @@ FText FField::GetDisplayNameText() const
 		NativeDisplayName = FName::NameToDisplayString(FFieldDisplayNameHelper::Get(*this), IsA<FBoolProperty>());
 	}
 
-	if (!(FText::FindText(Namespace, Key, /*OUT*/LocalizedDisplayName, &NativeDisplayName)))
-	{
-		LocalizedDisplayName = FText::FromString(NativeDisplayName);
-	}
-
-	return LocalizedDisplayName;
+	return FInternationalization::ForUseOnlyByLocMacroAndGraphNodeTextLiterals_CreateText(*NativeDisplayName, *Namespace, *Key);
 }
 
 /**
@@ -744,7 +737,7 @@ FText FField::GetToolTipText(bool bShortTooltip) const
 				NativeToolTip.TrimEndInline();
 			}
 		}
-		LocalizedToolTip = FText::FromString(NativeToolTip);
+		LocalizedToolTip = FInternationalization::ForUseOnlyByLocMacroAndGraphNodeTextLiterals_CreateText(*NativeToolTip, *Namespace, *Key);
 	}
 
 	return LocalizedToolTip;
@@ -810,16 +803,11 @@ const FText FField::GetMetaDataText(const TCHAR* MetaDataKey, const FString Loca
 		DefaultMetaData = FName::NameToDisplayString(GetName(), IsA(FBoolProperty::StaticClass()));
 	}
 
-
 	FText LocalizedMetaData;
-	if (!(FText::FindText(LocalizationNamespace, LocalizationKey, /*OUT*/LocalizedMetaData, &DefaultMetaData)))
+	if (!DefaultMetaData.IsEmpty())
 	{
-		if (!DefaultMetaData.IsEmpty())
-		{
-			LocalizedMetaData = FText::AsCultureInvariant(DefaultMetaData);
-		}
+		LocalizedMetaData = FInternationalization::ForUseOnlyByLocMacroAndGraphNodeTextLiterals_CreateText(*DefaultMetaData, *LocalizationNamespace, *LocalizationKey);
 	}
-
 	return LocalizedMetaData;
 }
 
@@ -838,16 +826,11 @@ const FText FField::GetMetaDataText(const FName& MetaDataKey, const FString Loca
 		DefaultMetaData = FName::NameToDisplayString(GetName(), IsA(FBoolProperty::StaticClass()));
 	}
 
-
 	FText LocalizedMetaData;
-	if (!(FText::FindText(LocalizationNamespace, LocalizationKey, /*OUT*/LocalizedMetaData, &DefaultMetaData)))
+	if (!DefaultMetaData.IsEmpty())
 	{
-		if (!DefaultMetaData.IsEmpty())
-		{
-			LocalizedMetaData = FText::AsCultureInvariant(DefaultMetaData);
-		}
+		LocalizedMetaData = FInternationalization::ForUseOnlyByLocMacroAndGraphNodeTextLiterals_CreateText(*DefaultMetaData, *LocalizationNamespace, *LocalizationKey);
 	}
-
 	return LocalizedMetaData;
 }
 
