Add checks before allocating buffers during bulkdata serialization.

#rnx

- If archive is corrupt we could serialize an unreasonable value for the payload size from disk. This could cause very large allocations to be made which could cause the process to OOM or to allocate resources that it doesn't need.
- To avoid this we can check the archive we are serializing from and make sure that, at the very least, it will be possible to load all of the payload bytes from it. This way we will never make an allocation larger than the remaining space in the archive. It would then be up to the higher level code to validate that overall size of the provided file makes sense.
- Note that we are not able to do this is we are loading a payload from legacy bulkdata format that was marked as compressed as we only know the uncompressed payload size, not the compressed size on disk.

diff --git a/Engine/Source/Runtime/CoreUObject/Private/Serialization/EditorBulkData.cpp b/Engine/Source/Runtime/CoreUObject/Private/Serialization/EditorBulkData.cpp
--- a/Engine/Source/Runtime/CoreUObject/Private/Serialization/EditorBulkData.cpp
+++ b/Engine/Source/Runtime/CoreUObject/Private/Serialization/EditorBulkData.cpp
@@ -329,19 +329,25 @@ FArchive& operator<<(FArchive& Ar, FSharedBuffer& Buffer)
 
 	if (Ar.IsLoading())
 	{
+		Buffer.Reset();
+
 		int64 BufferLength;
 		Ar << BufferLength;
 
 		if (BufferLength >= 0)
 		{
-			FUniqueBuffer MutableBuffer = FUniqueBuffer::Alloc(BufferLength);
-			Ar.Serialize(MutableBuffer.GetData(), BufferLength);
+			if (const int64 SpaceRemaining = Ar.TotalSize() - Ar.Tell(); SpaceRemaining >= BufferLength)
+			{
+				FUniqueBuffer MutableBuffer = FUniqueBuffer::Alloc(BufferLength);
+				Ar.Serialize(MutableBuffer.GetData(), BufferLength);
 
-			Buffer = MutableBuffer.MoveToShared();
-		}
-		else
-		{
-			Buffer.Reset();
+				Buffer = MutableBuffer.MoveToShared();
+			}
+			else
+			{
+				Ar.SetError();
+				UE_LOG(LogSerialization, Error, TEXT("Not enough data remaining in archive (%s) to load the payload. Bytes required: %lld vs bytes remaining %lld"), *Ar.GetArchiveName(), BufferLength, SpaceRemaining);
+			}
 		}
 	}
 	else if (Ar.IsSaving())
@@ -1517,16 +1523,25 @@ bool FEditorBulkData::SerializeData(FArchive& Ar, FCompressedBuffer& InPayload,
 	{
 		// Loading from old bulkdata format
 		const int64 Size = GetPayloadSize();
-		FUniqueBuffer LoadPayload = FUniqueBuffer::Alloc(Size);
 
+		FUniqueBuffer LoadPayload;
 		if (EnumHasAnyFlags(PayloadFlags, EFlags::LegacyFileIsCompressed))
 		{
+			LoadPayload = FUniqueBuffer::Alloc(Size);
 			Ar.SerializeCompressed(LoadPayload.GetData(), Size, NAME_Zlib, COMPRESS_NoFlags, false);
 		}
-		else
+		else if(const int64 SpaceRemaining = Ar.TotalSize() - Ar.Tell(); SpaceRemaining >= Size)
 		{
+			LoadPayload = FUniqueBuffer::Alloc(Size);
 			Ar.Serialize(LoadPayload.GetData(), Size);
 		}
+		else
+		{
+			Ar.SetError();
+			UE_LOG(LogSerialization, Error, TEXT("Not enough data remaining in archive (%s) to load the payload. Bytes required: %lld vs bytes remaining %lld"), *Ar.GetArchiveName(), Size, SpaceRemaining);
+
+			return false;
+		}
 
 		InPayload = FCompressedBuffer::Compress(LoadPayload.MoveToShared(), ECompressedBufferCompressor::NotSet, ECompressedBufferCompressionLevel::None);
 
