EditorBulkData : serialize compressed data to the undo buffer without decompressing

diff --git a/Engine/Source/Runtime/CoreUObject/Private/Serialization/EditorBulkData.cpp b/Engine/Source/Runtime/CoreUObject/Private/Serialization/EditorBulkData.cpp
--- a/Engine/Source/Runtime/CoreUObject/Private/Serialization/EditorBulkData.cpp
+++ b/Engine/Source/Runtime/CoreUObject/Private/Serialization/EditorBulkData.cpp
@@ -913,14 +913,20 @@ void FEditorBulkData::Serialize(FArchive& Ar, UObject* Owner, bool bAllowRegiste
 			{
 				if (Payload.IsNull() && !IsDataVirtualized())
 				{
-					// We need to serialize in FSharedBuffer form, or otherwise we'd need to support
-					// multiple code paths here. Technically a bit wasteful but for general use it 
-					// shouldn't be noticeable. This will make it easier to do the real perf wins
-					// in the future.
-					Payload = GetDataInternal().Decompress();
-				}
+					// We load the data from disk to serialize it to the undo history
+					// since don't know that the data on disk will stay the same
 
-				Ar << Payload;
+					FCompressedBuffer CompressedPayload = GetDataInternal();
+					bool bIsCompressedBuffer = true;
+					Ar << bIsCompressedBuffer;
+					Ar << CompressedPayload;
+				}
+				else
+				{
+					bool bIsCompressedBuffer = false;
+					Ar << bIsCompressedBuffer;
+					Ar << Payload;
+				}
 			}
 			else
 			{
@@ -930,17 +936,48 @@ void FEditorBulkData::Serialize(FArchive& Ar, UObject* Owner, bool bAllowRegiste
 				{
 					BulkDataId = FGuid::NewGuid();
 				}
+				
+				bool bIsCompressedBuffer = false;
+				Ar << bIsCompressedBuffer;
+
+				if ( bIsCompressedBuffer )
+				{
+					FCompressedBuffer CompressedPayload;
+					Ar << CompressedPayload;
 
-				Ar << Payload;
+					// there's no way to just set a CompressedBuffer for Payload, so we have to Decompress
+					//	it might be nice if Payload could be either FSharedBuffer or FCompressedBuffer
+					//	then we could just store it without decompressing
+
+					if ( CanUnloadData() )
+					{
+						//	don't bother to Decompress if it will just be discarded immediately by UnloadData below
+						Payload.Reset();
+					}
+					else
+					{
+						Payload = CompressedPayload.Decompress();
+					}
+				}
+				else
+				{
+					Ar << Payload;
+				}
+
+				// the call to UnloadData below will then often discard the Payload we just read !?
+				// ? is it guaranteed that the Payload we just loaded matches what we will get back from PackagePath after UnloadData ?
 
 				Register(Owner, TEXT("Serialize/Transacting"), false /* bAllowUpdateId */);
 			}
 
 			// Try to unload the payload if possible, usually because we loaded it during the transaction in the
 			// first place and we don't want to keep it in memory anymore.
-			// This does mean if the owning asset is frequently edited we will be reloading the payload off disk
-			// a lot. But in practice this didn't show up as too much of a problem. If someone has found this to 
-			// be a perf issue, then remove the call to ::UnloadData and trade memory cost for perf gain.
+			// 
+			// BulkData currently does not ever cache the decompressed Payload when it is possible to reload from disk.
+			// that is, on typical BulkData, CanUnloadData is always false because the data is either on disk (so Payload is null)
+			//	or not on disk
+			// we may have gotten out of sync with that standard state here due to serialization,
+			//	so go ahead and UnloadData now.
 			UnloadData();
 		}
 	}
@@ -1205,6 +1242,7 @@ void FEditorBulkData::Serialize(FArchive& Ar, UObject* Owner, bool bAllowRegiste
 					// store the payload as an empty FSharedBuffer.
 					if (CompressedPayload.GetRawSize() > 0)
 					{
+						// an example of when this is hit is in UTexture::Serialize for object duplication
 						Payload = CompressedPayload.Decompress();
 					}
 					else
@@ -1686,7 +1724,16 @@ void FEditorBulkData::UnloadData()
 {
 	UE::TUniqueLock _(Mutex);
 
-	if (CanUnloadData())
+	// This only does anything in unusual cases.
+	// In standard use, if a BulkData has disk backing (so CanUnloadData), then the in-memory Payload is always null
+	//	 because decompression does NOT cache the decompressed result in Payload.
+	// If there is no disk backing (so CanUnloadData is false), then Payload is not null.
+	// In both cases, UnloadData is a nop.
+	// The only time it does anything is if you had in-memory payload, then created a disk backing,
+	//	so CanUnloadData transitions from false to true,
+	//  then the next call to UnloadData will act to drop the in-memory payload.
+
+	if ( ! Payload.IsNull() && CanUnloadData())
 	{
 		Payload.Reset();
 	}
@@ -1962,6 +2009,7 @@ FCompressedBuffer FEditorBulkData::GetDataInternal() const
 	if (Payload)
 	{
 		// Note that this doesn't actually compress the data!
+		//	it reinterprets Payload as an FCompressedBuffer, doesn't have to alloc or memcpy, just points at it
 		return FCompressedBuffer::Compress(Payload, ECompressedBufferCompressor::NotSet, ECompressedBufferCompressionLevel::None);
 	}
 
@@ -1994,12 +2042,20 @@ FCompressedBuffer FEditorBulkData::GetDataInternal() const
 
 bool FEditorBulkData::DoesPayloadNeedLoading() const
 {
+	TRACE_CPUPROFILER_EVENT_SCOPE(FEditorBulkData::DoesPayloadNeedLoading);
+	// this function is very simple, but can still be a long time scope
+	//	 because acquiring the mutex blocks if the main thread is serializing the bulkdata to the undo transaction buffer
+
 	UE::TUniqueLock _(Mutex);
 	return Payload.IsNull() && PayloadSize > 0;
 }
 
 TFuture<FSharedBuffer> FEditorBulkData::GetPayload() const
 {
+	TRACE_CPUPROFILER_EVENT_SCOPE(FEditorBulkData::GetPayload);
+	// this function is very simple, but can still be a long time scope
+	//	 because acquiring the mutex blocks if the main thread is serializing the bulkdata to the undo transaction buffer
+
 	UE::TUniqueLock _(Mutex);
 
 	TPromise<FSharedBuffer> Promise;
@@ -2018,8 +2074,14 @@ TFuture<FSharedBuffer> FEditorBulkData::GetPayload() const
 	{
 		FCompressedBuffer CompressedPayload = GetDataInternal();
 
+		FSharedBuffer DecompressedPayload = CompressedPayload.Decompress();
+
+		// NOTE: DecompressedPayload is *NOT* cached in the Payload variable
+		//	so UnloadData will do nothing
+		//	and we will reload from disk if GetPayload() is called twice in a row
+
 		// TODO: Not actually async yet!
-		Promise.SetValue(CompressedPayload.Decompress());
+		Promise.SetValue(DecompressedPayload);
 	}
 
 	return Promise.GetFuture();
