Fixes for TStrongObjectPtr's operators.

diff --git a/Engine/Source/Runtime/CoreUObject/Public/UObject/StrongObjectPtr.h b/Engine/Source/Runtime/CoreUObject/Public/UObject/StrongObjectPtr.h
--- a/Engine/Source/Runtime/CoreUObject/Public/UObject/StrongObjectPtr.h
+++ b/Engine/Source/Runtime/CoreUObject/Public/UObject/StrongObjectPtr.h
@@ -181,16 +181,75 @@ public:
 private:
 	TUniquePtr< UEStrongObjectPtr_Private::TInternalReferenceCollector<ReferencerNameProvider> > ReferenceCollector;
 
-	template <typename RHSObjectType, typename RHSReferencerNameProvider>
-	friend FORCEINLINE bool operator==(const TStrongObjectPtr<ObjectType, ReferencerNameProvider>& InLHS, const TStrongObjectPtr<RHSObjectType, RHSReferencerNameProvider>& InRHS)
+	[[nodiscard]] friend FORCEINLINE bool operator==(const TStrongObjectPtr& InLHS, const TStrongObjectPtr& InRHS)
 	{
 		return InLHS.Get() == InRHS.Get();
 	}
 
-	template <typename RHSObjectType, typename RHSReferencerNameProvider>
-	friend FORCEINLINE bool operator!=(const TStrongObjectPtr<ObjectType, ReferencerNameProvider>& InLHS, const TStrongObjectPtr<RHSObjectType, RHSReferencerNameProvider>& InRHS)
+	[[nodiscard]] friend FORCEINLINE bool operator==(const TStrongObjectPtr& InLHS, TYPE_OF_NULLPTR)
+	{
+		return !InLHS.IsValid();
+	}
+
+	[[nodiscard]] friend FORCEINLINE bool operator==(TYPE_OF_NULLPTR, const TStrongObjectPtr& InRHS)
+	{
+		return !InRHS.IsValid();
+	}
+
+	template <
+		typename RHSObjectType,
+		typename RHSReferencerNameProvider
+		UE_REQUIRES(UE_REQUIRES_EXPR(std::declval<ObjectType*>() == std::declval<RHSObjectType*>()))
+	>
+	[[nodiscard]] friend FORCEINLINE bool operator==(const TStrongObjectPtr& InLHS, const TStrongObjectPtr<RHSObjectType, RHSReferencerNameProvider>& InRHS)
+	{
+		return InLHS.Get() == InRHS.Get();
+	}
+
+	template <
+		typename LHSObjectType,
+		typename LHSReferencerNameProvider
+		UE_REQUIRES(UE_REQUIRES_EXPR(std::declval<LHSObjectType*>() == std::declval<ObjectType*>()))
+	>
+	[[nodiscard]] friend FORCEINLINE bool operator==(const TStrongObjectPtr<LHSObjectType, LHSReferencerNameProvider>& InLHS, const TStrongObjectPtr& InRHS)
+	{
+		return InLHS.Get() == InRHS.Get();
+	}
+
+#if !PLATFORM_COMPILER_HAS_GENERATED_COMPARISON_OPERATORS
+	[[nodiscard]] friend FORCEINLINE bool operator!=(const TStrongObjectPtr& InLHS, const TStrongObjectPtr& InRHS)
 	{
 		return InLHS.Get() != InRHS.Get();
 	}
 
+	[[nodiscard]] friend FORCEINLINE bool operator!=(const TStrongObjectPtr& InLHS, TYPE_OF_NULLPTR)
+	{
+		return InLHS.IsValid();
+	}
+
+	[[nodiscard]] friend FORCEINLINE bool operator!=(TYPE_OF_NULLPTR, const TStrongObjectPtr& InRHS)
+	{
+		return InRHS.IsValid();
+	}
+
+	template <
+		typename RHSObjectType,
+		typename RHSReferencerNameProvider
+		UE_REQUIRES(UE_REQUIRES_EXPR(std::declval<ObjectType*>() == std::declval<RHSObjectType*>()))
+	>
+	[[nodiscard]] friend FORCEINLINE bool operator!=(const TStrongObjectPtr& InLHS, const TStrongObjectPtr<RHSObjectType, RHSReferencerNameProvider>& InRHS)
+	{
+		return InLHS.Get() != InRHS.Get();
+	}
+
+	template <
+		typename LHSObjectType,
+		typename LHSReferencerNameProvider
+		UE_REQUIRES(UE_REQUIRES_EXPR(std::declval<LHSObjectType*>() == std::declval<ObjectType*>()))
+	>
+	[[nodiscard]] friend FORCEINLINE bool operator!=(const TStrongObjectPtr<LHSObjectType, LHSReferencerNameProvider>& InLHS, const TStrongObjectPtr& InRHS)
+	{
+		return InLHS.Get() != InRHS.Get();
+	}
+#endif
 };
