Only make ZenLoader the default loader in editor cases

diff --git a/Engine/Source/Runtime/CoreUObject/Private/Serialization/AsyncPackageLoader.cpp b/Engine/Source/Runtime/CoreUObject/Private/Serialization/AsyncPackageLoader.cpp
--- a/Engine/Source/Runtime/CoreUObject/Private/Serialization/AsyncPackageLoader.cpp
+++ b/Engine/Source/Runtime/CoreUObject/Private/Serialization/AsyncPackageLoader.cpp
@@ -188,11 +188,14 @@ void InitAsyncThread()
 	LLM_SCOPE(ELLMTag::AsyncLoading);
 	if (FIoDispatcher::IsInitialized())
 	{
-		bool bSettingsEnabled = true;
+		bool bSettingsEnabled = false;
 		bool bCommandLineEnabled = false;
 		bool bCommandLineDisabled = false;
 		bool bHasUseIoStoreParamInEditor = false;
 #if WITH_EDITOR
+		// enable zenloader by default in editor case
+		bSettingsEnabled = true;
+
 		bCommandLineEnabled = FParse::Param(FCommandLine::Get(), TEXT("ZenLoader"));
 		bCommandLineDisabled = FParse::Param(FCommandLine::Get(), TEXT("NoZenLoader"));
 		check(GConfig);
