Removed the cause of the some warning in our asset mirgration test.

One was caused by the import validation not accounting for the package remapping when validating the class of the import in the linkerload.

The other was caused by some code in the open world tech that asked the asset registry to scan for a file that doesn't exist. This one is resolved by putting the flag to tell the engine that the instanced package from migration don't have on disk file.

diff --git a/Engine/Source/Developer/AssetTools/Private/AssetTools.cpp b/Engine/Source/Developer/AssetTools/Private/AssetTools.cpp
--- a/Engine/Source/Developer/AssetTools/Private/AssetTools.cpp
+++ b/Engine/Source/Developer/AssetTools/Private/AssetTools.cpp
@@ -985,6 +985,9 @@ namespace UE::AssetTools::Private
 				 */
 				MigrationPackageData.InstancedPackage = LoadPackage(MigrationPackage, *MigrationPackageData.GetOriginalPackageName(), LOAD_Verify | LOAD_NoVerify, ReaderOverride, &MigrationImplContext.InstancingContext);
 
+				// Add the PKG_NewlyCreated flag here to help the engine code that expect the packages without it to have a file associated with it on disk.
+				MigrationPackage->SetPackageFlags(PKG_NewlyCreated);
+
 				PackagesToClean.Add(MigrationPackage);
 			}
 		}
diff --git a/Engine/Source/Runtime/CoreUObject/Private/UObject/LinkerLoad.cpp b/Engine/Source/Runtime/CoreUObject/Private/UObject/LinkerLoad.cpp
--- a/Engine/Source/Runtime/CoreUObject/Private/UObject/LinkerLoad.cpp
+++ b/Engine/Source/Runtime/CoreUObject/Private/UObject/LinkerLoad.cpp
@@ -5882,7 +5882,7 @@ UObject* FLinkerLoad::CreateImport( int32 Index )
 		else if (ImportsToVerifyOnCreate.Contains(Index))
 		{
 			const UClass* ExpectedImportClass = nullptr;
-			if (UPackage* ImportClassPackage = FindObjectFast<UPackage>(nullptr, Import.ClassPackage, false))
+			if (UPackage* ImportClassPackage = FindObjectFast<UPackage>(nullptr, InstancingContext.RemapPackage(Import.ClassPackage), false))
 			{
 				const UObject* FoundObject = FindObjectFast<UObject>(ImportClassPackage, Import.ClassName, false);
 				ExpectedImportClass = Cast<UClass>(FoundObject);
