Remove annotation if we abort, and open on it since it has a lock

diff --git a/Engine/Source/Runtime/CoreUObject/Private/UObject/UObjectArchetype.cpp b/Engine/Source/Runtime/CoreUObject/Private/UObject/UObjectArchetype.cpp
--- a/Engine/Source/Runtime/CoreUObject/Private/UObject/UObjectArchetype.cpp
+++ b/Engine/Source/Runtime/CoreUObject/Private/UObject/UObjectArchetype.cpp
@@ -223,7 +223,16 @@ UObject* GetArchetypeImpl(const UObject* InObject, const FObjectArchetypeHelper:
 		if (Archetype && !(InObject->GetOuter() && InObject->GetOuter()->HasAnyFlags(RF_NeedLoad)))
 		{
 			ArchetypeIndex = GUObjectArray.ObjectToIndex(Archetype);
-			ArchetypeAnnotation.AddAnnotation(InObject, FArchetypeInfo{ ArchetypeIndex, GUObjectArray.AllocateSerialNumber(ArchetypeIndex) });
+			UE_AUTORTFM_OPEN(
+			{
+				ArchetypeAnnotation.AddAnnotation(InObject, FArchetypeInfo{ ArchetypeIndex, GUObjectArray.AllocateSerialNumber(ArchetypeIndex) });
+			});
+
+			// If we abort we need to remove this annotation
+			UE_AUTORTFM_ONABORT(
+			{
+				ArchetypeAnnotation.RemoveAnnotation(InObject);
+			});
 		}
 	}
 	else
