A bit of cleanup in IDOs before removing property bags

#rnx

diff --git a/Engine/Source/Runtime/CoreUObject/Private/UObject/InstanceDataObjectUtils.cpp b/Engine/Source/Runtime/CoreUObject/Private/UObject/InstanceDataObjectUtils.cpp
--- a/Engine/Source/Runtime/CoreUObject/Private/UObject/InstanceDataObjectUtils.cpp
+++ b/Engine/Source/Runtime/CoreUObject/Private/UObject/InstanceDataObjectUtils.cpp
@@ -7,6 +7,7 @@
 #include "UObject/EnumProperty.h"
 #include "UObject/Field.h"
 #include "UObject/PropertyBag.h"
+#include "UObject/PropertyHelper.h"
 #include "UObject/PropertyBagRepository.h"
 #include "UObject/PropertyOptional.h"
 #include "UObject/UnrealType.h"
@@ -106,7 +107,7 @@ namespace UE
 	// typedef to help make it clearer when a pathName has indices and when the indices are wildcarded away
 	using FWildcardPropertyPathName = FPropertyPathName;
 
-	static const FName NAME_StructOriginalTypeMetadata(ANSITEXTVIEW("OriginalType"));
+	static const FName NAME_DisplayName(ANSITEXTVIEW("DisplayName"));
 	static const FName NAME_PresentAsTypeMetadata(ANSITEXTVIEW("PresentAsType"));
 	static const FName NAME_IsLooseMetadata(ANSITEXTVIEW("IsLoose"));
 	static const FName NAME_VerseClass(ANSITEXTVIEW("VerseClass"));
@@ -330,21 +331,18 @@ namespace UE
 		return Result;
 	}
 
-	static FString UnmanglePropertyName(const FName MaybeMangledName, bool* bOutNameWasMangled)
+	static FString UnmanglePropertyName(const FName MaybeMangledName, bool& bOutNameWasMangled)
 	{
 		FString Result = MaybeMangledName.ToString();
-		if (Result.StartsWith(TEXT("__verse_0x")))
+		if (Result.StartsWith(TEXTVIEW("__verse_0x")))
 		{
 			// chop "__verse_0x" (10 char) + CRC (8 char) + "_" (1 char)
 			Result = Result.RightChop(19);
-			if (bOutNameWasMangled)
-			{
-				*bOutNameWasMangled = true;
-			}
+			bOutNameWasMangled = true;
 		}
-		else if (bOutNameWasMangled)
+		else
 		{
-			*bOutNameWasMangled = false;
+			bOutNameWasMangled = false;
 		}
 		return Result;
 	}
@@ -354,18 +352,17 @@ namespace UE
 		const TMap<FWildcardPropertyPathName, TMap<FName, const FProperty*>>& LooseProperties, FWildcardPropertyPathName& Path)
 	{
 #if WITH_EDITORONLY_DATA
-		static const FName NAME_DisplayName(TEXT("DisplayName"));
 		if (!Property->HasMetaData(NAME_DisplayName))
 		{
 			bool bNeedsDisplayName = false;
-			const FString DisplayName = UnmanglePropertyName(Property->GetFName(), &bNeedsDisplayName);
+			FString DisplayName = UnmanglePropertyName(Property->GetFName(), bNeedsDisplayName);
 			if (bNeedsDisplayName)
 			{
-				Property->SetMetaData(NAME_DisplayName, *DisplayName);
+				Property->SetMetaData(NAME_DisplayName, MoveTemp(DisplayName));
 			}
 		}
 #endif
-		
+
 		if (FStructProperty* AsStructProperty = CastField<FStructProperty>(Property))
 		{
 			if (!AsStructProperty->Struct->UseNativeSerialization())
@@ -373,14 +370,14 @@ namespace UE
 #if WITH_EDITORONLY_DATA
 				//@note: Transfer existing metadata over as we build the InstanceDataObject from the struct or it owner, if any, this is useful for testing purposes
 				FString OriginalName;
-				if (const FString* OriginalType = AsStructProperty->FindMetaData(NAME_StructOriginalTypeMetadata))
+				if (const FString* OriginalType = AsStructProperty->FindMetaData(NAME_OriginalType))
 				{
 					OriginalName = *OriginalType;
 				}
 				//@note: To support metadata defined on array of struct in UPROPERTY for testing purposes
 				else if (FField* OwnerField = AsStructProperty->Owner.ToField())
 				{
-					if (const FString* OwnerOriginalType = OwnerField->FindMetaData(NAME_StructOriginalTypeMetadata))
+					if (const FString* OwnerOriginalType = OwnerField->FindMetaData(NAME_OriginalType))
 					{
 						OriginalName = *OwnerOriginalType;
 					}
@@ -400,7 +397,7 @@ namespace UE
 				}
 				AsStructProperty->Struct = Struct;
 #if WITH_EDITORONLY_DATA
-				AsStructProperty->SetMetaData(NAME_StructOriginalTypeMetadata, *OriginalName);
+				AsStructProperty->SetMetaData(NAME_OriginalType, *OriginalName);
 				AsStructProperty->SetMetaData(NAME_PresentAsTypeMetadata, *OriginalName);
 				AsStructProperty->Struct->SetMetaData(NAME_PresentAsTypeMetadata, *OriginalName);
 #endif
@@ -429,7 +426,7 @@ namespace UE
 			ConvertToInstanceDataObjectProperty(AsOptionalProperty->GetValueProperty(), PropertyType.GetParameter(0), Outer, LooseProperties, Path);
 		}
 	}
-	
+
 	// copy template property then convert it into an InstanceDataObject property by adding loose properties
 	static FProperty* CreateInstanceDataObjectProperty(const FProperty* TemplateProperty, UObject* Outer,
 		const TMap<FWildcardPropertyPathName, TMap<FName, const FProperty*>>& LooseProperties, FWildcardPropertyPathName& Path)
@@ -670,33 +667,32 @@ namespace UE
 		
 		// copy flags from OwnerClass
 		IDOClass->ClassFlags |= OwnerClass->ClassFlags & (
-        	CLASS_EditInlineNew | CLASS_CollapseCategories | CLASS_Const | CLASS_CompiledFromBlueprint | CLASS_HasInstancedReference);
+			CLASS_EditInlineNew | CLASS_CollapseCategories | CLASS_Const | CLASS_CompiledFromBlueprint | CLASS_HasInstancedReference);
 	}
-	
+
 	UClass* CreateInstanceDataObjectClass(const FPropertyBag* PropertyBag, UClass* OwnerClass, UObject* Outer)
 	{
 		const TMap<FWildcardPropertyPathName, TMap<FName, const FProperty*>> LooseProperties = GetWildcardedLooseProperties(PropertyBag);
 		FWildcardPropertyPathName ParentPath;
 		UClass* Result = CreateInstanceDataObjectStructRec<UClass>(OwnerClass, Outer, LooseProperties, ParentPath);
 #if WITH_EDITORONLY_DATA
-		const FString& DisplayName = OwnerClass->GetMetaData(TEXT("DisplayName"));
-		if (!DisplayName.IsEmpty())
+		if (const FString& DisplayName = OwnerClass->GetMetaData(NAME_DisplayName); !DisplayName.IsEmpty())
 		{
-			Result->SetMetaData(TEXT("DisplayName"), *DisplayName);
+			Result->SetMetaData(NAME_DisplayName, *DisplayName);
 		}
 #endif
 
 		SetClassFlags(Result, OwnerClass);
 
-		const UObject* OwnerCDO = OwnerClass->GetDefaultObject(true);
-		UObject* ResultCDO = Result->GetDefaultObject(true);
+		const UObject* OwnerCDO = OwnerClass->GetDefaultObject();
+		UObject* ResultCDO = Result->GetDefaultObject();
 		if (ensure(OwnerCDO && ResultCDO))
 		{
 			CopyCDO(OwnerCDO, ResultCDO);
 		}
 		return Result;
 	}
-	
+
 	static void MarkPropertySetBySerialization(const UStruct* Struct, const void* StructData, const void* PropertyDataPtr)
 	{
 		if (const FSetProperty* ValuesSetByPropertyBagProperty = CastField<FSetProperty>(Struct->FindPropertyByName(NAME_ValuesSetBySerialization)))
