Fix UScriptStruct::Serialize clobbering the saved struct flags on load for structs that don't have a matching CppStructOps.
This was causing unversioned property deserialization to break as some FStructProperties had different property flags from when they were serialized.

diff --git a/Engine/Source/Runtime/CoreUObject/Private/UObject/Class.cpp b/Engine/Source/Runtime/CoreUObject/Private/UObject/Class.cpp
--- a/Engine/Source/Runtime/CoreUObject/Private/UObject/Class.cpp
+++ b/Engine/Source/Runtime/CoreUObject/Private/UObject/Class.cpp
@@ -3004,13 +3004,24 @@ void UScriptStruct::Serialize( FArchive& Ar )
 {
 	Super::Serialize(Ar);
 
-	// serialize the struct's flags
-	Ar << (uint32&)StructFlags;
+	// Serialize only the non-computed struct flags.
+	const uint32 ComputedStructFlags = uint32(StructFlags) & STRUCT_ComputedFlags;
+	const uint32 NonComputedStructFlags = uint32(StructFlags) & ~STRUCT_ComputedFlags;
+	uint32 SavedStructFlags = NonComputedStructFlags;
+	Ar << SavedStructFlags;
 
 	if (Ar.IsLoading())
 	{
+		StructFlags = EStructFlags(SavedStructFlags);
+
 		ClearCppStructOps(); // we want to be sure to do this from scratch
 		PrepareCppStructOps();
+
+		// If there was no matching CppStructOps found, restore the computed struct flags that were reset by ClearCppStructOps.
+		if (!CppStructOps)
+		{
+			StructFlags = EStructFlags((uint32(StructFlags) & ~STRUCT_ComputedFlags) | ComputedStructFlags);
+		}
 	}
 }
 
