Fixing race condition when disabling disregard for GC set with multithreading enabled

diff --git a/Engine/Source/Runtime/CoreUObject/Private/UObject/GarbageCollection.cpp b/Engine/Source/Runtime/CoreUObject/Private/UObject/GarbageCollection.cpp
--- a/Engine/Source/Runtime/CoreUObject/Private/UObject/GarbageCollection.cpp
+++ b/Engine/Source/Runtime/CoreUObject/Private/UObject/GarbageCollection.cpp
@@ -636,19 +636,18 @@ void OnDisregardForGCSetDisabled(int32 NumObjects)
 {
 	using namespace UE::GC;
 	using namespace UE::GC::Private;
-	using FMarkMarkDisregardState = TThreadedGather<TArray<UObject*>>;
+	using FDisregardSetDisabledState = TThreadedGather<TSet<int32>>;
 
-	FMarkMarkDisregardState MarkDisregardState;
+	FDisregardSetDisabledState DisregardSetDisabledState;
 	
-	MarkDisregardState.Start(EGatherOptions::Parallel, NumObjects);
-	FMarkMarkDisregardState::FThreadIterators& ThreadIterators = MarkDisregardState.GetThreadIterators();
+	DisregardSetDisabledState.Start(EGatherOptions::Parallel, NumObjects);
+	FDisregardSetDisabledState::FThreadIterators& ThreadIterators = DisregardSetDisabledState.GetThreadIterators();
 
 	// Objects in the disregard for GC set do not have any of the reachability flags set (this way they never become (Maybe)Unreachable) nor are they added to the GRoots array 
-	FScopeLock RootsLock(&GRootsCritical);
-	ParallelFor(TEXT("GC.OnDisregardForGCSetDisabled"), MarkDisregardState.NumWorkerThreads(), 1, [&ThreadIterators](int32 ThreadIndex)
+	ParallelFor(TEXT("GC.OnDisregardForGCSetDisabled"), DisregardSetDisabledState.NumWorkerThreads(), 1, [&ThreadIterators](int32 ThreadIndex)
 	{
 		TRACE_CPUPROFILER_EVENT_SCOPE(OnDisregardForGCSetDisabledTask);
-		FMarkMarkDisregardState::FIterator& ThreadState = ThreadIterators[ThreadIndex];
+		FDisregardSetDisabledState::FIterator& ThreadState = ThreadIterators[ThreadIndex];
 
 		while (ThreadState.Index <= ThreadState.LastIndex)
 		{
@@ -659,11 +658,14 @@ void OnDisregardForGCSetDisabled(int32 NumObjects)
 				RootItem->ThisThreadAtomicallySetFlag_ForGC(GReachableObjectFlag);
 				if (RootItem->HasAnyFlags(EInternalObjectFlags_RootFlags))
 				{
-					GRoots.Add(ObjectIndex);
+					ThreadState.Payload.Add(ObjectIndex);
 				}
 			}
 		}
-	}, (MarkDisregardState.NumWorkerThreads() == 1) ? EParallelForFlags::ForceSingleThread : EParallelForFlags::None);
+	}, (DisregardSetDisabledState.NumWorkerThreads() == 1) ? EParallelForFlags::ForceSingleThread : EParallelForFlags::None);
+
+	FScopeLock RootsLock(&GRootsCritical);
+	DisregardSetDisabledState.Finish(GRoots);
 }
 
 /**
diff --git a/Engine/Source/Runtime/CoreUObject/Private/UObject/ReachabilityAnalysisState.h b/Engine/Source/Runtime/CoreUObject/Private/UObject/ReachabilityAnalysisState.h
--- a/Engine/Source/Runtime/CoreUObject/Private/UObject/ReachabilityAnalysisState.h
+++ b/Engine/Source/Runtime/CoreUObject/Private/UObject/ReachabilityAnalysisState.h
@@ -187,7 +187,7 @@ struct TDefaultPayloadOps
 
 	inline static void Append(const PayloadType& InSource, PayloadType& OutDest)
 	{
-		OutDest += InSource;
+		OutDest.Append(InSource);
 	}
 };
 
