Cooker SkipOnlyEditorOnly: AsyncLoading2: Fix imports to report their referencer as the package that imported them rather than the package that called the outermost loadpackage. This is necessary for us to detect that packages loaded by editoronly references from their sourcepackage should be skipped and not added to the cook.
#rnx

diff --git a/Engine/Source/Runtime/CoreUObject/Private/Serialization/AsyncLoading2.cpp b/Engine/Source/Runtime/CoreUObject/Private/Serialization/AsyncLoading2.cpp
--- a/Engine/Source/Runtime/CoreUObject/Private/Serialization/AsyncLoading2.cpp
+++ b/Engine/Source/Runtime/CoreUObject/Private/Serialization/AsyncLoading2.cpp
@@ -507,6 +507,19 @@ struct FPackageReferencer
 #if WITH_EDITOR
 	ECookLoadType CookLoadType = ECookLoadType::Unexpected;
 #endif
+
+	static FPackageReferencer FromImport(FName ReferencerName)
+	{
+		FPackageReferencer Result;
+#if UE_WITH_PACKAGE_ACCESS_TRACKING
+		Result.ReferencerPackageName = ReferencerName;
+		Result.ReferencerPackageOp = PackageAccessTrackingOps::NAME_Load;
+#endif
+#if WITH_EDITOR
+		Result.CookLoadType = ECookLoadType::Unexpected;
+#endif
+		return Result;
+	}
 };
 
 struct FPackageRequest
@@ -633,7 +646,7 @@ struct FAsyncPackageDesc2
 			PackageIdToLoad,
 			UPackageName,
 			MoveTemp(PackagePathToLoad),
-			ImportingPackageDesc.PackageReferencer,
+			FPackageReferencer::FromImport(ImportingPackageDesc.UPackageName),
 			true
 		};
 	}
