[GC] Fix wasting one cycle when starting async purge

diff --git a/Engine/Source/Runtime/CoreUObject/Private/UObject/GarbageCollection.cpp b/Engine/Source/Runtime/CoreUObject/Private/UObject/GarbageCollection.cpp
--- a/Engine/Source/Runtime/CoreUObject/Private/UObject/GarbageCollection.cpp
+++ b/Engine/Source/Runtime/CoreUObject/Private/UObject/GarbageCollection.cpp
@@ -764,7 +764,7 @@ class FAsyncPurge : public FRunnable
 
 		while (NumObjectsDestroyedOnGameThread < LocalNumObjectsToDestroyOnGameThread && ObjCurrentPurgeObjectIndexOnGameThread < GUnreachableObjects.Num())
 		{
-			FUObjectItem* ObjectItem = GUnreachableObjects[ObjCurrentPurgeObjectIndexOnGameThread];			
+			FUObjectItem* ObjectItem = GUnreachableObjects[ObjCurrentPurgeObjectIndexOnGameThread];
 			if (ObjectItem)
 			{
 				GUnreachableObjects[ObjCurrentPurgeObjectIndexOnGameThread] = nullptr;
@@ -797,6 +797,47 @@ class FAsyncPurge : public FRunnable
 		return bFinishedDestroyingObjects;
 	}
 
+	/** [GAME THREAD] Destroys objects that are unreachable for the first tick */
+	bool TickDestroyGameThreadObjectsFirst(bool bUseTimeLimit, double TimeLimit, double StartTime)
+	{
+		TRACE_CPUPROFILER_EVENT_SCOPE(FAsyncPurge::TickDestroyGameThreadObjects);
+		const int32 TimeLimitEnforcementGranularityForDeletion = 100;
+		int32 ProcessedObjectsCount = 0;
+		bool bFinishedDestroyingObjects = true;
+
+		// Lock once for the entire batch
+		GUObjectArray.LockInternalArray();
+
+		while (ObjCurrentPurgeObjectIndexOnGameThread < GUnreachableObjects.Num())
+		{
+			FUObjectItem* ObjectItem = GUnreachableObjects[ObjCurrentPurgeObjectIndexOnGameThread];
+			GUnreachableObjects[ObjCurrentPurgeObjectIndexOnGameThread] = nullptr;
+			UObject* Object = (UObject*)ObjectItem->Object;
+			Object->~UObject();
+			GUObjectAllocator.FreeUObject(Object);
+			++ProcessedObjectsCount;
+			++ObjCurrentPurgeObjectIndexOnGameThread;
+
+			if (bUseTimeLimit && (ProcessedObjectsCount == TimeLimitEnforcementGranularityForDeletion))
+			{
+				ProcessedObjectsCount = 0;
+				if ((FPlatformTime::Seconds() - StartTime) > TimeLimit)
+				{
+					bFinishedDestroyingObjects = false;
+					break;
+				}
+			}
+		}
+		NumObjectsDestroyedOnGameThread = ObjCurrentPurgeObjectIndexOnGameThread;
+		NumObjectsToDestroyOnGameThread.store(ObjCurrentPurgeObjectIndexOnGameThread, std::memory_order_release);
+		ObjCurrentPurgeObjectIndex = ObjCurrentPurgeObjectIndexOnGameThread;
+		ObjectsDestroyedSinceLastMarkPhase = ObjCurrentPurgeObjectIndexOnGameThread;
+
+		GUObjectArray.UnlockInternalArray();
+
+		return bFinishedDestroyingObjects;
+	}
+
 	/** Waits for the worker thread to finish destroying objects */
 	void WaitForAsyncDestructionToFinish()
 	{
@@ -860,7 +901,7 @@ public:
 	}
 
 	/** [MAIN THREAD] Adds objects to the purge queue */
-	void BeginPurge()
+	void BeginPurge(bool bUseTimeLimit, double TimeLimit, double StartTime)
 	{
 		check(IsFinished()); // In single-threaded mode we need to be finished or the condition below will hang
 		if (FinishedPurgeEvent->Wait())
@@ -873,7 +914,22 @@ public:
 			NumObjectsDestroyedOnGameThread = 0;
 			ObjCurrentPurgeObjectIndexOnGameThread = 0;
 
-			BeginPurgeEvent->Trigger();
+			// In MT case we need to consume all available time limit before triggering Async Purge.
+			// Otherwise Async thread almost always will not process enought objects for us to destroy on GT.
+			// That happens because it takes significant amount time to unpark the thread and switch context before
+			// it can process and increment NumObjectsToDestroyOnGameThread.
+			if (!!Thread)
+			{
+				if (!TickDestroyGameThreadObjectsFirst(bUseTimeLimit, TimeLimit, StartTime))
+				{
+					// Start async helper when the timelimit reached and there is still stuff to do
+					BeginPurgeEvent->Trigger();
+				}
+				else
+				{
+					FinishedPurgeEvent->Trigger();
+				}
+			}
 		}
 	}
 
@@ -4921,7 +4977,7 @@ bool IncrementalDestroyGarbage(bool bUseTimeLimit, double TimeLimit)
 		int32 ProcessCount = 0;
 		if (GObjCurrentPurgeObjectIndexNeedsReset)
 		{
-			GAsyncPurge->BeginPurge();
+			GAsyncPurge->BeginPurge(bUseTimeLimit, TimeLimit, GCStartTime);
 			// Reset the reset flag but don't reset the actual index yet for stat purposes
 			GObjCurrentPurgeObjectIndexNeedsReset = false;
 		}
