Unshelved from pending changelist '32714305':

[CrashFix] Array properties aren't idempotent when override serialization is enabled so we need to avoid deserializing them twice during IDO generation. This change will skip any non-loose properties during the Preload pass so that only the loose properties are deserialized, then the non-loose ones will get set during the CopyTaggedProperties pass

diff --git a/Engine/Source/Runtime/CoreUObject/Private/UObject/InstanceDataObjectUtils.cpp b/Engine/Source/Runtime/CoreUObject/Private/UObject/InstanceDataObjectUtils.cpp
--- a/Engine/Source/Runtime/CoreUObject/Private/UObject/InstanceDataObjectUtils.cpp
+++ b/Engine/Source/Runtime/CoreUObject/Private/UObject/InstanceDataObjectUtils.cpp
@@ -110,6 +110,7 @@ namespace UE
 	static const FName NAME_DisplayName(ANSITEXTVIEW("DisplayName"));
 	static const FName NAME_PresentAsTypeMetadata(ANSITEXTVIEW("PresentAsType"));
 	static const FName NAME_IsLooseMetadata(ANSITEXTVIEW("IsLoose"));
+	static const FName NAME_ContainsLoosePropertiesMetadata(ANSITEXTVIEW("ContainsLooseProperties"));
 	static const FName NAME_VerseClass(ANSITEXTVIEW("VerseClass"));
 	static const FName NAME_IDOMapKey(ANSITEXTVIEW("Key"));
 	static const FName NAME_IDOMapValue(ANSITEXTVIEW("Value"));
@@ -362,7 +363,16 @@ namespace UE
 			}
 		}
 #endif
-
+const auto TrySetContainsLooseProperties = [](FProperty* Property, const FFieldVariant& Inner)
+		{
+#if WITH_EDITORONLY_DATA
+			if (Inner.HasMetaData(NAME_ContainsLoosePropertiesMetadata))
+			{
+				Property->SetMetaData(NAME_ContainsLoosePropertiesMetadata, TEXT("True"));
+			}
+#endif
+		};
+		
 		if (FStructProperty* AsStructProperty = CastField<FStructProperty>(Property))
 		{
 			if (!AsStructProperty->Struct->UseNativeSerialization())
@@ -400,41 +410,86 @@ namespace UE
 				AsStructProperty->SetMetaData(NAME_OriginalType, *OriginalName);
 				AsStructProperty->SetMetaData(NAME_PresentAsTypeMetadata, *OriginalName);
 				AsStructProperty->Struct->SetMetaData(NAME_PresentAsTypeMetadata, *OriginalName);
+
+				TrySetContainsLooseProperties(AsStructProperty, AsStructProperty->Struct);
 #endif
 			}
 		}
-		else if (const FArrayProperty* AsArrayProperty = CastField<FArrayProperty>(Property))
+		else if (FArrayProperty* AsArrayProperty = CastField<FArrayProperty>(Property))
 		{
 			ConvertToInstanceDataObjectProperty(AsArrayProperty->Inner, PropertyType.GetParameter(0), Outer, LooseProperties, Path);
+			TrySetContainsLooseProperties(AsArrayProperty, AsArrayProperty->Inner);
 		}
-		else if (const FSetProperty* AsSetProperty = CastField<FSetProperty>(Property))
+		else if (FSetProperty* AsSetProperty = CastField<FSetProperty>(Property))
 		{
 			ConvertToInstanceDataObjectProperty(AsSetProperty->ElementProp, PropertyType.GetParameter(0), Outer, LooseProperties, Path);
+			TrySetContainsLooseProperties(AsSetProperty, AsSetProperty->ElementProp);
 		}
-		else if (const FMapProperty* AsMapProperty = CastField<FMapProperty>(Property))
+		else if (FMapProperty* AsMapProperty = CastField<FMapProperty>(Property))
 		{
 			Path.Push({NAME_IDOMapKey});
 			ConvertToInstanceDataObjectProperty(AsMapProperty->KeyProp, PropertyType.GetParameter(0), Outer, LooseProperties, Path);
 			Path.Pop();
+			TrySetContainsLooseProperties(AsMapProperty, AsMapProperty->KeyProp);
 			
 			Path.Push({NAME_IDOMapValue});
 			ConvertToInstanceDataObjectProperty(AsMapProperty->ValueProp, PropertyType.GetParameter(1), Outer, LooseProperties, Path);
 			Path.Pop();
+			TrySetContainsLooseProperties(AsMapProperty, AsMapProperty->ValueProp);
 		}
-		else if (const FOptionalProperty* AsOptionalProperty = CastField<FOptionalProperty>(Property))
+		else if (FOptionalProperty* AsOptionalProperty = CastField<FOptionalProperty>(Property))
 		{
 			ConvertToInstanceDataObjectProperty(AsOptionalProperty->GetValueProperty(), PropertyType.GetParameter(0), Outer, LooseProperties, Path);
+			TrySetContainsLooseProperties(AsOptionalProperty, AsOptionalProperty->GetValueProperty());
+		}
+		
+#if WITH_EDITORONLY_DATA
+		if (Property->GetBoolMetaData(NAME_IsLooseMetadata) || Property->GetBoolMetaData(NAME_ContainsLoosePropertiesMetadata))
+		{
+			Property->GetOwnerStruct()->SetMetaData(NAME_ContainsLoosePropertiesMetadata, TEXT("True"));
+		}
+#endif
+	}
+
+	// recursively gives a property the metadata and flags of a loose property
+	static void MarkPropertyAsLoose(FProperty* Property)
+	{
+#if WITH_EDITORONLY_DATA
+		Property->SetMetaData(NAME_IsLooseMetadata, TEXT("True"));
+#endif
+		Property->SetPropertyFlags(CPF_Edit | CPF_EditConst);
+		if (const FArrayProperty* AsArrayProperty = CastField<FArrayProperty>(Property))
+		{
+			MarkPropertyAsLoose(AsArrayProperty->Inner);
+		}
+		else if (const FSetProperty* AsSetProperty = CastField<FSetProperty>(Property))
+		{
+			MarkPropertyAsLoose(AsSetProperty->ElementProp);
+		}
+		else if (const FMapProperty* AsMapProperty = CastField<FMapProperty>(Property))
+		{
+			MarkPropertyAsLoose(AsMapProperty->KeyProp);
+			MarkPropertyAsLoose(AsMapProperty->ValueProp);
+		}
+		else if (const FOptionalProperty* AsOptionalProperty = CastField<FOptionalProperty>(Property))
+		{
+			MarkPropertyAsLoose(AsOptionalProperty->GetValueProperty());
 		}
 	}
 
 	// copy template property then convert it into an InstanceDataObject property by adding loose properties
 	static FProperty* CreateInstanceDataObjectProperty(const FProperty* TemplateProperty, UObject* Outer,
-		const TMap<FWildcardPropertyPathName, TMap<FName, const FProperty*>>& LooseProperties, FWildcardPropertyPathName& Path)
+		const TMap<FWildcardPropertyPathName, TMap<FName, const FProperty*>>& LooseProperties, FWildcardPropertyPathName& Path, bool bIsLoose)
 	{
 		FProperty* InstanceDataObjectProperty = CastFieldChecked<FProperty>(FField::Duplicate(TemplateProperty, Outer));
 #if WITH_EDITORONLY_DATA
 		FField::CopyMetaData(TemplateProperty, InstanceDataObjectProperty);
 #endif
+		if (bIsLoose)
+		{
+			MarkPropertyAsLoose(InstanceDataObjectProperty);
+		}
+		
 		ConvertToInstanceDataObjectProperty(InstanceDataObjectProperty, Path.GetSegment(Path.GetSegmentCount() - 1).Type, Outer, LooseProperties, Path);
 		return InstanceDataObjectProperty;
 	}
@@ -507,33 +562,7 @@ namespace UE
 		
 		return LooseProperties;
 	}
-
-	// recursively gives a property the metadata and flags of a loose property
-	static void MarkPropertyAsLoose(FProperty* Property)
-	{
-#if WITH_EDITORONLY_DATA
-		Property->SetMetaData(NAME_IsLooseMetadata, TEXT("True"));
-#endif
-		Property->SetPropertyFlags(CPF_Edit | CPF_EditConst);
-		if (const FArrayProperty* AsArrayProperty = CastField<FArrayProperty>(Property))
-		{
-			MarkPropertyAsLoose(AsArrayProperty->Inner);
-		}
-		else if (const FSetProperty* AsSetProperty = CastField<FSetProperty>(Property))
-{
-			MarkPropertyAsLoose(AsSetProperty->ElementProp);
-		}
-		else if (const FMapProperty* AsMapProperty = CastField<FMapProperty>(Property))
-		{
-			MarkPropertyAsLoose(AsMapProperty->KeyProp);
-			MarkPropertyAsLoose(AsMapProperty->ValueProp);
-		}
-		else if (const FOptionalProperty* AsOptionalProperty = CastField<FOptionalProperty>(Property))
-		{
-			MarkPropertyAsLoose(AsOptionalProperty->GetValueProperty());
-		}
-	}
-
+	
 	// constructs an InstanceDataObject struct by merging the properties in 
 	static UStruct* CreateInstanceDataObjectStructRec(const UClass* StructClass, UStruct* TemplateStruct,
 		UObject* Outer, const TMap<FWildcardPropertyPathName, TMap<FName, const FProperty*>>& LooseProperties, FWildcardPropertyPathName& Path)
@@ -570,7 +599,7 @@ namespace UE
 					continue;
 				}
 				Path.Push(CreateSegmentFromProperty(TemplateProperty));
-				FProperty* SuperProperty = CreateInstanceDataObjectProperty(TemplateProperty, Super, LooseProperties, Path);
+				FProperty* SuperProperty = CreateInstanceDataObjectProperty(TemplateProperty, Super, LooseProperties, Path, false);
 				
 				Path.Pop();
 				SuperProperties.Add(SuperProperty);
@@ -606,10 +635,8 @@ namespace UE
 			for (const TPair<FName, const FProperty*>& BagProperty : *BagProperties)
 			{
 				Path.Push(CreateSegmentFromProperty(BagProperty.Value));
-				FProperty* LooseProperty = CreateInstanceDataObjectProperty(BagProperty.Value, Result, LooseProperties, Path);
+				FProperty* LooseProperty = CreateInstanceDataObjectProperty(BagProperty.Value, Result, LooseProperties, Path, true);
 				Path.Pop();
-				
-				MarkPropertyAsLoose(LooseProperty);
 				LooseInstanceDataObjectProperties.Add(LooseProperty);
 			}
 		}
diff --git a/Engine/Source/Runtime/CoreUObject/Private/UObject/LinkerLoad.cpp b/Engine/Source/Runtime/CoreUObject/Private/UObject/LinkerLoad.cpp
--- a/Engine/Source/Runtime/CoreUObject/Private/UObject/LinkerLoad.cpp
+++ b/Engine/Source/Runtime/CoreUObject/Private/UObject/LinkerLoad.cpp
@@ -980,6 +980,7 @@ FLinkerLoad::FLinkerLoad(UPackage* InParent, const FPackagePath& InPackagePath,
 , TemplateForGetArchetypeFromLoader(nullptr)
 , bForceSimpleIndexToObject(false)
 , bLockoutLegacyOperations(false)
+, bSkipKnownProperties(false)
 , bIsAsyncLoader(false)
 , bIsDestroyingLoader(false)
 #if WITH_EDITOR
@@ -7188,6 +7189,22 @@ bool FLinkerLoad::SerializeBulkData(FBulkData& BulkData, const FBulkDataSerializ
 	return true;
 }
 
+bool FLinkerLoad::ShouldSkipProperty(const FProperty* InProperty) const
+{
+	if (bSkipKnownProperties)
+	{
+#if WITH_EDITORONLY_DATA
+		static const FName NAME_IsLooseMetadata(ANSITEXTVIEW("IsLoose"));
+		static const FName NAME_ContainsLoosePropertiesMetadata(ANSITEXTVIEW("ContainsLooseProperties"));
+		if (!InProperty->GetBoolMetaData(NAME_IsLooseMetadata) && !InProperty->GetBoolMetaData(NAME_ContainsLoosePropertiesMetadata))
+#endif
+		{
+			return true;
+		}
+	}
+	return FArchiveUObject::ShouldSkipProperty(InProperty);
+}
+
 void FLinkerLoad::SerializeBulkMeta(UE::BulkData::Private::FBulkMetaData& Meta, int64& DuplicateSerialOffset, int32 ElementSize)
 {
 	using namespace UE::BulkData::Private;
diff --git a/Engine/Source/Runtime/CoreUObject/Private/UObject/PropertyBagRepository.cpp b/Engine/Source/Runtime/CoreUObject/Private/UObject/PropertyBagRepository.cpp
--- a/Engine/Source/Runtime/CoreUObject/Private/UObject/PropertyBagRepository.cpp
+++ b/Engine/Source/Runtime/CoreUObject/Private/UObject/PropertyBagRepository.cpp
@@ -1,6 +1,7 @@
 // Copyright Epic Games, Inc. All Rights Reserved.
 
 #include "UObject/PropertyBagRepository.h"
+
 #include "Containers/Queue.h"
 #include "Serialization/ObjectReader.h"
 #include "Serialization/ObjectWriter.h"
@@ -12,7 +13,10 @@
 #include "UObject/UObjectThreadContext.h"
 #include "UObject/LinkerLoad.h"
 #include "UObject/InstanceDataObjectUtils.h"
+#include "UObject/OverridableManager.h"
+#include "UObject/OverriddenPropertySet.h"
 #include "UObject/Package.h"
+#include "Templates/UnrealTemplate.h"
 
 #if WITH_EDITOR
 #include "HAL/IConsoleManager.h"
@@ -125,10 +129,64 @@ FPropertyBagRepository::FPropertyBagRepository()
 #if WITH_EDITOR
 	FCoreUObjectDelegates::OnObjectModified.AddLambda([](const UObject* Object)
 	{
-		// if this object is an InstanceDataObject, modify it's owner as well
-		if (const UObject* Owner = Get().FindInstanceForDataObject(Object))
+		auto CopyChanges = [](const UObject* Source, UObject* Dest)
 		{
-			const_cast<UObject*>(Owner)->Modify();
+			TArray<uint8> Buffer;
+			Buffer.Reserve(Source->GetClass()->GetStructureSize());
+			FObjectWriter Writer(Buffer);
+			Source->GetClass()->SerializeTaggedProperties(Writer, (uint8*)Source, Source->GetClass(), (uint8*)Source->GetArchetype());
+		
+			FObjectReader Reader(Buffer);
+			Reader.ArMergeOverrides = true;
+			Dest->GetClass()->SerializeTaggedProperties(Reader, (uint8*)Dest, Dest->GetClass(), (uint8*)Dest->GetArchetype());	
+		};
+		
+		if (UObject* Instance = const_cast<UObject*>(Get().FindInstanceForDataObject(Object)))
+		{
+			// if this object is an instance, modify it's IDO as well
+			Instance->Modify();
+			CopyChanges(Object, Instance);
+		}
+	});
+
+	FCoreUObjectDelegates::OnObjectPropertyChanged.AddLambda([](const UObject* Object, FPropertyChangedEvent& ChangeEvent)
+	{
+		static TSet<TSoftObjectPtr<UObject>> ChangeCallbacksToSkip;
+		if (ChangeCallbacksToSkip.Remove(Object))
+		{
+			// avoids infinite recursion
+			return;
+		}
+		
+		auto CopyChanges = [&ChangeEvent](const UObject* Source, UObject* Dest)
+		{
+			//Dest->PreEditChange(nullptr);
+			const FOverriddenPropertySet* SourceOverriddenProperties = FOverridableManager::Get().GetOverriddenProperties(*Source);
+			const FOverriddenPropertySet* DestOverriddenProperties = FOverridableManager::Get().GetOverriddenProperties(*Dest);
+			TArray<TObjectPtr<UObject>>* Test = reinterpret_cast<TArray<TObjectPtr<UObject>>*>(reinterpret_cast<uint8*>(Dest) + 240);
+			TArray<uint8> Buffer;
+			Buffer.Reserve(Source->GetClass()->GetStructureSize());
+			FObjectWriter Writer(Buffer);
+			Source->GetClass()->SerializeTaggedProperties(Writer, (uint8*)Source, Source->GetClass(), (uint8*)Source->GetArchetype());
+		
+			FObjectReader Reader(Buffer);
+			Reader.ArMergeOverrides = true;
+			Dest->GetClass()->SerializeTaggedProperties(Reader, (uint8*)Dest, Dest->GetClass(), (uint8*)Dest->GetArchetype());	
+			(void)Test;
+			//Dest->PostEditChangeProperty(ChangeEvent);
+		};
+		
+		if (UObject* Ido = Get().FindInstanceDataObject(Object))
+		{
+			// if this object is an InstanceDataObject, modify it's owner as well
+			ChangeCallbacksToSkip.Add(Ido); // avoid infinite recursion
+			CopyChanges(Object, Ido);
+		}
+		else if (UObject* Instance = const_cast<UObject*>(Get().FindInstanceForDataObject(Object)))
+		{
+			// if this object is an instance, modify it's IDO as well
+			ChangeCallbacksToSkip.Add(Ido); // avoid infinite recursion
+			CopyChanges(Object, Instance);
 		}
 	});
 #endif
@@ -372,7 +430,11 @@ void FPropertyBagRepository::CreateInstanceDataObjectUnsafe(UObject* Owner, FPro
 		// TODO: @jordan.hoffmann - this is very inefficient! We should remove this call to Preload. To do so, we'd need to change MarkPropertySetBySerialization
 		// to cache the serialized property list in the property bag instead of the structs. We'd also need to copy the property bag values to the IDO
 		Owner->SetFlags(RF_NeedLoad);
-		Linker->Preload(Owner);
+		{
+			TGuardValue<bool> ScopedSkipKnownProperties(Linker->bSkipKnownProperties, true);
+			FGuardValue_Bitfield(Linker->ArMergeOverrides, true);
+            Linker->Preload(Owner);
+		}
 		LoadContext->OnTaggedPropertySerialize.Remove(OnTaggedPropertySerializeHandle);
 		
 		// copy data from owner to IDO
diff --git a/Engine/Source/Runtime/CoreUObject/Public/UObject/LinkerLoad.h b/Engine/Source/Runtime/CoreUObject/Public/UObject/LinkerLoad.h
--- a/Engine/Source/Runtime/CoreUObject/Public/UObject/LinkerLoad.h
+++ b/Engine/Source/Runtime/CoreUObject/Public/UObject/LinkerLoad.h
@@ -168,6 +168,10 @@ public:
 	UObject*				TemplateForGetArchetypeFromLoader;
 	bool					bForceSimpleIndexToObject;
 	bool					bLockoutLegacyOperations;
+	
+	/** Set to true during IDO generation to skip non-loose properties during serialization.
+	 * In any other context, this would end up skipping *every* property */
+	UE_INTERNAL bool bSkipKnownProperties;
 
 	/** True if Loader is FAsyncArchive  */
 	bool					bIsAsyncLoader;
@@ -406,6 +410,8 @@ public:
 
 	virtual bool SerializeBulkData(FBulkData& BulkData, const FBulkDataSerializationParams& Params) override;
 
+	virtual bool ShouldSkipProperty(const FProperty* InProperty) const override;
+
 private:
 
 	void SerializeBulkMeta(UE::BulkData::Private::FBulkMetaData& Meta, int64& DuplicateSerialOffset, int32 ElementSize);
@@ -1474,7 +1480,7 @@ private:
 
 
 	/** 
-	 * Creates the export hash.
+	0 *00 0C0r0e0ates the export hash.
 	 */
 	ELinkerStatus CreateExportHash();
 
