[Partial Backout] Accidentally included wip changes from a different shelf in CL 32863792. Backing them out.

diff --git a/Engine/Source/Runtime/CoreUObject/Private/UObject/PropertyBagRepository.cpp b/Engine/Source/Runtime/CoreUObject/Private/UObject/PropertyBagRepository.cpp
--- a/Engine/Source/Runtime/CoreUObject/Private/UObject/PropertyBagRepository.cpp
+++ b/Engine/Source/Runtime/CoreUObject/Private/UObject/PropertyBagRepository.cpp
@@ -13,8 +13,6 @@
 #include "UObject/UObjectThreadContext.h"
 #include "UObject/LinkerLoad.h"
 #include "UObject/InstanceDataObjectUtils.h"
-#include "UObject/OverridableManager.h"
-#include "UObject/OverriddenPropertySet.h"
 #include "UObject/Package.h"
 #include "Templates/UnrealTemplate.h"
 
@@ -129,64 +127,10 @@ FPropertyBagRepository::FPropertyBagRepository()
 #if WITH_EDITOR
 	FCoreUObjectDelegates::OnObjectModified.AddLambda([](const UObject* Object)
 	{
-		auto CopyChanges = [](const UObject* Source, UObject* Dest)
+		// if this object is an InstanceDataObject, modify it's owner as well
+		if (const UObject* Owner = Get().FindInstanceForDataObject(Object))
 		{
-			TArray<uint8> Buffer;
-			Buffer.Reserve(Source->GetClass()->GetStructureSize());
-			FObjectWriter Writer(Buffer);
-			Source->GetClass()->SerializeTaggedProperties(Writer, (uint8*)Source, Source->GetClass(), (uint8*)Source->GetArchetype());
-		
-			FObjectReader Reader(Buffer);
-			Reader.ArMergeOverrides = true;
-			Dest->GetClass()->SerializeTaggedProperties(Reader, (uint8*)Dest, Dest->GetClass(), (uint8*)Dest->GetArchetype());	
-		};
-		
-		if (UObject* Instance = const_cast<UObject*>(Get().FindInstanceForDataObject(Object)))
-		{
-			// if this object is an instance, modify it's IDO as well
-			Instance->Modify();
-			CopyChanges(Object, Instance);
-		}
-	});
-
-	FCoreUObjectDelegates::OnObjectPropertyChanged.AddLambda([](const UObject* Object, FPropertyChangedEvent& ChangeEvent)
-	{
-		static TSet<TSoftObjectPtr<UObject>> ChangeCallbacksToSkip;
-		if (ChangeCallbacksToSkip.Remove(Object))
-		{
-			// avoids infinite recursion
-			return;
-		}
-		
-		auto CopyChanges = [&ChangeEvent](const UObject* Source, UObject* Dest)
-		{
-			//Dest->PreEditChange(nullptr);
-			const FOverriddenPropertySet* SourceOverriddenProperties = FOverridableManager::Get().GetOverriddenProperties(*Source);
-			const FOverriddenPropertySet* DestOverriddenProperties = FOverridableManager::Get().GetOverriddenProperties(*Dest);
-			TArray<TObjectPtr<UObject>>* Test = reinterpret_cast<TArray<TObjectPtr<UObject>>*>(reinterpret_cast<uint8*>(Dest) + 240);
-			TArray<uint8> Buffer;
-			Buffer.Reserve(Source->GetClass()->GetStructureSize());
-			FObjectWriter Writer(Buffer);
-			Source->GetClass()->SerializeTaggedProperties(Writer, (uint8*)Source, Source->GetClass(), (uint8*)Source->GetArchetype());
-		
-			FObjectReader Reader(Buffer);
-			Reader.ArMergeOverrides = true;
-			Dest->GetClass()->SerializeTaggedProperties(Reader, (uint8*)Dest, Dest->GetClass(), (uint8*)Dest->GetArchetype());	
-			(void)Test;
-			//Dest->PostEditChangeProperty(ChangeEvent);
-		};
-		
-		if (UObject* Ido = Get().FindInstanceDataObject(Object))
-		{
-			// if this object is an InstanceDataObject, modify it's owner as well
-			ChangeCallbacksToSkip.Add(Ido); // avoid infinite recursion
-			CopyChanges(Object, Ido);
-		}
-		else if (UObject* Instance = const_cast<UObject*>(Get().FindInstanceForDataObject(Object)))
-		{
-			// if this object is an instance, modify it's IDO as well
-			ChangeCallbacksToSkip.Add(Ido); // avoid infinite recursion
-			CopyChanges(Object, Instance);
+			const_cast<UObject*>(Owner)->Modify();
 		}
 	});
 #endif
