Added FPropertyPathNameTree

This type will replace FPropertyBag as a means of gathering unique property paths during package loading.

#rnx

diff --git a/Engine/Source/Runtime/CoreUObject/Internal/UObject/PropertyPathNameTree.h b/Engine/Source/Runtime/CoreUObject/Internal/UObject/PropertyPathNameTree.h
new file mode 100644
--- /dev/null
+++ b/Engine/Source/Runtime/CoreUObject/Internal/UObject/PropertyPathNameTree.h
@@ -0,0 +1,110 @@
+// Copyright Epic Games, Inc. All Rights Reserved.
+
+#pragma once
+
+#include "CoreTypes.h"
+#include "Containers/Map.h"
+#include "Templates/TypeHash.h"
+#include "UObject/PropertyPathName.h"
+#include "UObject/PropertyTypeName.h"
+
+#define UE_API COREUOBJECT_API
+
+namespace UE
+{
+
+/**
+ * A tree of property path names and their associated types.
+ *
+ * A union of the paths that are added, ignoring any container index.
+ */
+class FPropertyPathNameTree
+{
+	/** Key corresponding to a FPropertyPathNameSegment with no index. */
+	struct FKey
+	{
+		FName Name;
+		FPropertyTypeName Type;
+
+		[[nodiscard]] friend inline bool operator==(const FKey& Lhs, const FKey& Rhs)
+		{
+			return Lhs.Name == Rhs.Name && Lhs.Type == Rhs.Type;
+		}
+
+		[[nodiscard]] friend inline bool operator<(const FKey& Lhs, const FKey& Rhs)
+		{
+			if (int32 Compare = Lhs.Name.Compare(Rhs.Name))
+			{
+				return Compare < 0;
+			}
+			return Lhs.Type < Rhs.Type;
+		}
+
+		[[nodiscard]] friend inline uint32 GetTypeHash(const FKey& Key)
+		{
+			return HashCombineFast(GetTypeHash(Key.Name), GetTypeHash(Key.Type));
+		}
+	};
+
+public:
+	FPropertyPathNameTree() = default;
+	FPropertyPathNameTree(FPropertyPathNameTree&&) = default;
+	FPropertyPathNameTree(const FPropertyPathNameTree&) = delete;
+	FPropertyPathNameTree& operator=(FPropertyPathNameTree&&) = default;
+	FPropertyPathNameTree& operator=(const FPropertyPathNameTree&) = delete;
+
+	/** True if the tree contains no property path names. */
+	inline bool IsEmpty() const { return Nodes.IsEmpty(); }
+
+	/** Remove every property path name from the tree. */
+	UE_API void Empty();
+
+	/** Adds the path to the tree. Keeps any existing nodes that match both name and type. */
+	UE_API void Add(const FPropertyPathName& Path, int32 StartIndex = 0);
+
+	/**
+	 * Finds the path within the tree.
+	 *
+	 * @param OutSubTree   If non-null, this is set to the maybe-null sub-tree at the path if the path is found.
+	 * @param Path         Path to search within the tree.
+	 * @param StartIndex   Index of the segment within Path at which to start searching.
+	 * @return true if the path exists within the tree, otherwise false.
+	 */
+	UE_API bool Find(FPropertyPathNameTree** OutSubTree, const FPropertyPathName& Path, int32 StartIndex = 0);
+	inline bool Find(const FPropertyPathNameTree** OutSubTree, const FPropertyPathName& Path, int32 StartIndex = 0) const
+	{
+		return const_cast<FPropertyPathNameTree*>(this)->Find(const_cast<FPropertyPathNameTree**>(OutSubTree), Path, StartIndex);
+	}
+
+	class FConstIterator
+	{
+	public:
+		inline explicit FConstIterator(const TMap<FKey, TUniquePtr<FPropertyPathNameTree>>::TConstIterator& InNodeIt)
+			: NodeIt(InNodeIt)
+		{
+		}
+
+		inline FConstIterator& operator++() { ++NodeIt; return *this; }
+
+		inline explicit operator bool() const { return (bool)NodeIt; }
+
+		inline bool operator==(const FConstIterator& Rhs) const { return NodeIt == Rhs.NodeIt; }
+		inline bool operator!=(const FConstIterator& Rhs) const { return NodeIt != Rhs.NodeIt; }
+
+		inline FName GetName() const { return NodeIt.Key().Name; }
+		inline FPropertyTypeName GetType() const { return NodeIt.Key().Type; }
+		inline const FPropertyPathNameTree* GetSubTree() const { return NodeIt.Value().Get(); }
+
+	private:
+		TMap<FKey, TUniquePtr<FPropertyPathNameTree>>::TConstIterator NodeIt;
+	};
+
+	inline FConstIterator CreateConstIterator() const { return FConstIterator(Nodes.CreateConstIterator()); }
+
+private:
+	TMap<FKey, TUniquePtr<FPropertyPathNameTree>> Nodes;
+};
+
+} // UE
+
+#undef UE_API
diff --git a/Engine/Source/Runtime/CoreUObject/Private/UObject/PropertyPathNameTree.cpp b/Engine/Source/Runtime/CoreUObject/Private/UObject/PropertyPathNameTree.cpp
new file mode 100644
--- /dev/null
+++ b/Engine/Source/Runtime/CoreUObject/Private/UObject/PropertyPathNameTree.cpp
@@ -0,0 +1,56 @@
+// Copyright Epic Games, Inc. All Rights Reserved.
+
+#include "UObject/PropertyPathNameTree.h"
+
+#include "Templates/TypeHash.h"
+
+namespace UE
+{
+
+void FPropertyPathNameTree::Empty()
+{
+	Nodes.Empty();
+}
+
+void FPropertyPathNameTree::Add(const FPropertyPathName& Path, int32 StartIndex)
+{
+	if (const int32 SegmentCount = Path.GetSegmentCount(); StartIndex < SegmentCount)
+	{
+		FPropertyPathNameSegment Segment = Path.GetSegment(StartIndex);
+		TUniquePtr<FPropertyPathNameTree>& Child = Nodes.FindOrAdd({Segment.Name, Segment.Type});
+		if (++StartIndex < SegmentCount)
+		{
+			if (!Child)
+			{
+				Child = MakeUnique<FPropertyPathNameTree>();
+			}
+			Child->Add(Path, StartIndex);
+		}
+	}
+}
+
+bool FPropertyPathNameTree::Find(FPropertyPathNameTree** OutSubTree, const FPropertyPathName& Path, int32 StartIndex)
+{
+	if (const int32 SegmentCount = Path.GetSegmentCount(); StartIndex < SegmentCount)
+	{
+		FPropertyPathNameSegment Segment = Path.GetSegment(StartIndex);
+		if (TUniquePtr<FPropertyPathNameTree>* Child = Nodes.Find({Segment.Name, Segment.Type}))
+		{
+			if (++StartIndex >= SegmentCount)
+			{
+				if (OutSubTree)
+				{
+					*OutSubTree = Child->Get();
+				}
+				return true;
+			}
+			if (*Child)
+			{
+				return (*Child)->Find(OutSubTree, Path, StartIndex);
+			}
+		}
+	}
+	return false;
+}
+
+} // UE
diff --git a/Engine/Source/Runtime/CoreUObject/Tests/UObject/PropertyPathNameTreeTest.cpp b/Engine/Source/Runtime/CoreUObject/Tests/UObject/PropertyPathNameTreeTest.cpp
new file mode 100644
--- /dev/null
+++ b/Engine/Source/Runtime/CoreUObject/Tests/UObject/PropertyPathNameTreeTest.cpp
@@ -0,0 +1,241 @@
+// Copyright Epic Games, Inc. All Rights Reserved.
+
+#include "UObject/PropertyPathNameTree.h"
+
+#if WITH_TESTS
+
+#include "Tests/TestHarnessAdapter.h"
+
+namespace UE
+{
+
+TEST_CASE_NAMED(FPropertyPathNameTreeTest, "CoreUObject::PropertyPathNameTree", "[Core][UObject][SmokeFilter]")
+{
+	const FName CountName(TEXTVIEW("Count"));
+	const FName SizeName(TEXTVIEW("Size"));
+
+	FPropertyTypeNameBuilder TypeBuilder;
+	TypeBuilder.AddName(NAME_IntProperty);
+	const FPropertyTypeName IntType = TypeBuilder.Build();
+
+	TypeBuilder.Reset();
+	TypeBuilder.AddName(NAME_FloatProperty);
+	const FPropertyTypeName FloatType = TypeBuilder.Build();
+
+	TypeBuilder.Reset();
+	TypeBuilder.AddName(NAME_StructProperty);
+	TypeBuilder.BeginParameters();
+	TypeBuilder.AddName(NAME_Vector);
+	TypeBuilder.EndParameters();
+	const FPropertyTypeName VectorType = TypeBuilder.Build();
+
+	SECTION("Empty")
+	{
+		FPropertyPathNameTree Tree;
+		CHECK(Tree.IsEmpty());
+
+		FPropertyPathName PathName;
+		PathName.Push({CountName});
+		Tree.Add(PathName);
+		CHECK_FALSE(Tree.IsEmpty());
+
+		Tree.Empty();
+		CHECK(Tree.IsEmpty());
+	}
+
+	SECTION("Name")
+	{
+		FPropertyPathName PathName;
+		PathName.Push({CountName});
+		FPropertyPathNameTree Tree;
+		Tree.Add(PathName);
+		FPropertyPathNameTree* SubTree = &Tree;
+		CHECK(Tree.Find(&SubTree, PathName));
+		CHECK_FALSE(SubTree);
+		const FPropertyPathNameTree::FConstIterator First = Tree.CreateConstIterator();
+		FPropertyPathNameTree::FConstIterator It = Tree.CreateConstIterator();
+		CHECK(First == It);
+		CHECK_FALSE(First != It);
+		CHECKED_IF(It)
+		{
+			CHECK(It.GetName() == CountName);
+			CHECK(It.GetType().IsEmpty());
+			CHECK_FALSE(It.GetSubTree());
+			++It;
+			CHECK_FALSE(It);
+			CHECK_FALSE(First == It);
+			CHECK(First != It);
+		}
+	}
+
+	SECTION("NameType")
+	{
+		FPropertyPathName PathName;
+		PathName.Push({CountName, IntType});
+		FPropertyPathNameTree Tree;
+		Tree.Add(PathName);
+		FPropertyPathNameTree* SubTree = &Tree;
+		CHECK(Tree.Find(&SubTree, PathName));
+		CHECK_FALSE(SubTree);
+		const FPropertyPathNameTree::FConstIterator First = Tree.CreateConstIterator();
+		FPropertyPathNameTree::FConstIterator It = Tree.CreateConstIterator();
+		CHECK(First == It);
+		CHECK_FALSE(First != It);
+		CHECKED_IF(It)
+		{
+			CHECK(It.GetName() == CountName);
+			CHECK(It.GetType() == IntType);
+			CHECK_FALSE(It.GetSubTree());
+			++It;
+			CHECK_FALSE(It);
+			CHECK_FALSE(First == It);
+			CHECK(First != It);
+		}
+	}
+
+	SECTION("NameTypeIndex")
+	{
+		FPropertyPathName PathName;
+		PathName.Push({CountName, IntType, 7});
+		FPropertyPathNameTree Tree;
+		Tree.Add(PathName);
+		CHECK(Tree.Find(nullptr, PathName));
+		PathName.SetIndex(5);
+		CHECK(Tree.Find(nullptr, PathName));
+		PathName.SetIndex(3);
+		Tree.Add(PathName);
+		PathName.SetIndex(INDEX_NONE);
+		const FPropertyPathNameTree* SubTree = &Tree;
+		CHECK(Tree.Find(&SubTree, PathName));
+		CHECK_FALSE(SubTree);
+		const FPropertyPathNameTree::FConstIterator First = Tree.CreateConstIterator();
+		FPropertyPathNameTree::FConstIterator It = Tree.CreateConstIterator();
+		CHECK(First == It);
+		CHECK_FALSE(First != It);
+		CHECKED_IF(It)
+		{
+			CHECK(It.GetName() == CountName);
+			CHECK(It.GetType() == IntType);
+			CHECK_FALSE(It.GetSubTree());
+			++It;
+			CHECK_FALSE(It);
+			CHECK_FALSE(First == It);
+			CHECK(First != It);
+		}
+	}
+
+	SECTION("SameNameDiffType")
+	{
+		FPropertyPathName PathNameInt;
+		PathNameInt.Push({CountName, IntType});
+		FPropertyPathName PathNameFloat;
+		PathNameFloat.Push({CountName, FloatType});
+		FPropertyPathNameTree Tree;
+		Tree.Add(PathNameInt);
+		Tree.Add(PathNameFloat);
+		CHECK(Tree.Find(nullptr, PathNameInt));
+		CHECK(Tree.Find(nullptr, PathNameFloat));
+
+		FPropertyPathNameTree::FConstIterator It = Tree.CreateConstIterator();
+		CHECKED_IF(It)
+		{
+			CHECK(It.GetType() == IntType);
+			CHECK_FALSE(It.GetSubTree());
+			++It;
+			CHECKED_IF(It)
+			{
+				CHECK(It.GetType() == FloatType);
+				CHECK_FALSE(It.GetSubTree());
+				++It;
+				CHECK_FALSE(It);
+			}
+		}
+	}
+
+	SECTION("DiffNameSameType")
+	{
+		FPropertyPathName PathNameCount;
+		PathNameCount.Push({CountName, IntType});
+		FPropertyPathName PathNameSize;
+		PathNameSize.Push({SizeName, IntType});
+		FPropertyPathNameTree Tree;
+		Tree.Add(PathNameCount);
+		Tree.Add(PathNameSize);
+		CHECK(Tree.Find(nullptr, PathNameCount));
+		CHECK(Tree.Find(nullptr, PathNameSize));
+
+		FPropertyPathNameTree::FConstIterator It = Tree.CreateConstIterator();
+		CHECKED_IF(It)
+		{
+			CHECK(It.GetName() == CountName);
+			CHECK_FALSE(It.GetSubTree());
+			++It;
+			CHECKED_IF(It)
+			{
+				CHECK(It.GetName() == SizeName);
+				CHECK_FALSE(It.GetSubTree());
+				++It;
+				CHECK_FALSE(It);
+			}
+		}
+	}
+
+	SECTION("Tree")
+	{
+		FPropertyPathName ParentPathName;
+		ParentPathName.Push({NAME_Vector, VectorType});
+		FPropertyPathName PathNameA = ParentPathName;
+		PathNameA.Push({CountName, IntType});
+		FPropertyPathName PathNameB = ParentPathName;
+		PathNameB.Push({SizeName, FloatType});
+		FPropertyPathNameTree Tree;
+		Tree.Add(PathNameA);
+		Tree.Add(PathNameB);
+		CHECK(Tree.Find(nullptr, ParentPathName));
+		CHECK(Tree.Find(nullptr, PathNameA));
+		CHECK(Tree.Find(nullptr, PathNameB));
+
+		FPropertyPathNameTree::FConstIterator ParentIt = Tree.CreateConstIterator();
+		CHECKED_IF(ParentIt)
+		{
+			CHECK(ParentIt.GetName() == NAME_Vector);
+			CHECK(ParentIt.GetType() == VectorType);
+			const FPropertyPathNameTree* ChildTree = ParentIt.GetSubTree();
+			++ParentIt;
+			CHECK_FALSE(ParentIt);
+			CHECKED_IF(ChildTree)
+			{
+				FPropertyPathNameTree::FConstIterator ChildIt = ChildTree->CreateConstIterator();
+				CHECKED_IF(ChildIt)
+				{
+					CHECK(ChildIt.GetName() == CountName);
+					CHECK_FALSE(ChildIt.GetSubTree());
+					++ChildIt;
+					CHECKED_IF(ChildIt)
+					{
+						CHECK(ChildIt.GetName() == SizeName);
+						CHECK_FALSE(ChildIt.GetSubTree());
+						++ChildIt;
+						CHECK_FALSE(ChildIt);
+					}
+				}
+			}
+		}
+
+		FPropertyPathNameTree* ChildTree = nullptr;
+		CHECK(Tree.Find(&ChildTree, ParentPathName));
+		CHECKED_IF(ChildTree)
+		{
+			CHECK(ChildTree->Find(nullptr, PathNameA, 1));
+			CHECK(ChildTree->Find(nullptr, PathNameB, 1));
+		}
+
+		FPropertyPathNameTree* MissingTree = &Tree;
+		CHECK_FALSE(Tree.Find(&MissingTree, PathNameA, 1));
+		CHECK(MissingTree == &Tree);
+	}
+}
+
+} // UE
+
+#endif // WITH_TESTS
