Replaced FPropertyBag with FPropertyPathNameTree for collecting unknown properties

#rnx

diff --git a/Engine/Source/Runtime/CoreUObject/Internal/UObject/PropertyBagRepository.h b/Engine/Source/Runtime/CoreUObject/Internal/UObject/PropertyBagRepository.h
--- a/Engine/Source/Runtime/CoreUObject/Internal/UObject/PropertyBagRepository.h
+++ b/Engine/Source/Runtime/CoreUObject/Internal/UObject/PropertyBagRepository.h
@@ -14,6 +14,7 @@ namespace UE
 
 class FPropertyBag;
 class FPropertyPathName;
+class FPropertyPathNameTree;
 
 // Singleton class tracking property bag association with objects
 class FPropertyBagRepository : public FGCObject
@@ -22,7 +23,8 @@ class FPropertyBagRepository : public FGCObject
 	{
 		void Destroy();
 		
-		FPropertyBag* Bag = nullptr;		// The existence of an association implies the existence of the bag. TODO: Ref bags via handle? Store as value?
+		FPropertyBag* Bag = nullptr; // TODO: Ref bags via handle? Store as value?
+		FPropertyPathNameTree* Tree = nullptr;
 		TObjectPtr<UObject> InstanceDataObject = nullptr;
 	};
 	// TODO: Make private throughout and extend access permissions here or in wrapper classes? Don't want engine code modifying bags outside of serializers and details panels.
@@ -66,6 +68,11 @@ public:
 	// Object owner is tracked internally
 	FPropertyBag* CreateOuterBag(const UObject* Owner);
 
+	/**
+	 * Finds or creates a property path name tree to collect unknown property paths within the owner.
+	 */
+	FPropertyPathNameTree* CreateUnknownPropertyTree(const UObject* Owner);
+
 	// Future version for reworked InstanceDataObjects - track InstanceDataObject rather than bag (directly):
 	/**
 	 * Instantiate an InstanceDataObject object representing all fields within the bag, tracked against the owner object.
diff --git a/Engine/Source/Runtime/CoreUObject/Private/UObject/Class.cpp b/Engine/Source/Runtime/CoreUObject/Private/UObject/Class.cpp
--- a/Engine/Source/Runtime/CoreUObject/Private/UObject/Class.cpp
+++ b/Engine/Source/Runtime/CoreUObject/Private/UObject/Class.cpp
@@ -39,6 +39,7 @@
 #include "UObject/PropertyBag.h"
 #include "UObject/PropertyBagRepository.h"
 #include "UObject/PropertyOptional.h"
+#include "UObject/PropertyPathNameTree.h"
 #include "UObject/StructOnScope.h"
 #include "UObject/StructScriptLoader.h"
 #include "UObject/PropertyHelper.h"
@@ -583,6 +584,27 @@ IMPLEMENT_CORE_INTRINSIC_CLASS(UField, UObject,
 	UStruct implementation.
 -----------------------------------------------------------------------------*/
 
+static bool PropertyTypeContainsStructOrEnum(UE::FPropertyTypeName Type)
+{
+	const FName Name = Type.GetName();
+	if ((Name == NAME_StructProperty) ||
+		(Name == NAME_EnumProperty) ||
+		(Name == NAME_ByteProperty && Type.GetParameterCount() > 0))
+	{
+		return true;
+	}
+	const int32 Count = Type.GetParameterCount();
+	for (int32 Index = 1; Index < Count; ++Index)
+	{
+		if (PropertyTypeContainsStructOrEnum(Type.GetParameter(Index)))
+		{
+			return true;
+		}
+	}
+	// Handle parameter 0 at the end to give the compiler freedom to make this a tail call.
+	return Count > 0 && PropertyTypeContainsStructOrEnum(Type.GetParameter(0));
+}
+
 #if WITH_EDITORONLY_DATA
 static int32 GetNextFieldPathSerialNumber()
 {
@@ -1489,22 +1511,9 @@ void UStruct::SerializeVersionedTaggedProperties(FStructuredArchive::FSlot Slot,
 		else
 #endif // WITH_TEXT_ARCHIVE_SUPPORT
 		{
-			auto TryFindPropertyBag = [PropertyBag = (FPropertyBag*)nullptr, bSearched = false, SerializeContext]() mutable -> FPropertyBag*
-			{
-				if (bSearched)
-				{
-					return PropertyBag;
-				}
-				bSearched = true;
-				if (SerializeContext && SerializeContext->bSerializeUnknownProperty)
-				{
-					if (UObject* Object = SerializeContext->SerializedObject)
-					{
-						PropertyBag = FPropertyBagRepository::Get().CreateOuterBag(Object);
-					}
-				}
-				return PropertyBag;
-			};
+			// Track whether the unknown property tree has been looked up.
+			FPropertyPathNameTree* UnknownPropertyTree = nullptr;
+			bool bSearchedForUnknownPropertyTree = false;
 
 			// Load tagged properties.
 			FStructuredArchive::FStream PropertiesStream = Slot.EnterStream();
@@ -1625,6 +1634,8 @@ void UStruct::SerializeVersionedTaggedProperties(FStructuredArchive::FSlot Slot,
 					SerializedPropertyPath.Emplace(SerializeContext, Segment, UE::ESerializedPropertyPathNotify::Yes);
 				}
 
+				bool bTryStoreUnknownPropertyPath = false;
+
 				if (Property)
 				{
 	#if WITH_EDITOR
@@ -1666,13 +1677,11 @@ void UStruct::SerializeVersionedTaggedProperties(FStructuredArchive::FSlot Slot,
 							}
 						}
 
-						bool bTryLoadIntoPropertyBag = false;
-
 						switch (Property->ConvertFromType(Tag, ValueSlot, Data, DefaultsStruct, Defaults))
 						{
 							case EConvertFromTypeResult::Converted:
 								bAdvanceProperty = true;
-								bTryLoadIntoPropertyBag = true;
+								bTryStoreUnknownPropertyPath = true;
 								break;
 
 							case EConvertFromTypeResult::Serialized:
@@ -1686,7 +1695,7 @@ void UStruct::SerializeVersionedTaggedProperties(FStructuredArchive::FSlot Slot,
 										*WriteToString<32>(Tag.Name), *WriteToString<32>(GetFName()),
 										*WriteToString<32>(Tag.Type), *WriteToString<32>(PropID),
 										*UnderlyingArchive.GetArchiveName());
-									bTryLoadIntoPropertyBag = true;
+									bTryStoreUnknownPropertyPath = true;
 								}
 								else
 								{
@@ -1700,45 +1709,68 @@ void UStruct::SerializeVersionedTaggedProperties(FStructuredArchive::FSlot Slot,
 								break;
 
 							case EConvertFromTypeResult::CannotConvert:
-								bTryLoadIntoPropertyBag = true;
+								bTryStoreUnknownPropertyPath = true;
 								break;
 
 							default:
 								checkNoEntry();
 								break;
 						}
+					}
+				}
+				else
+				{
+					bTryStoreUnknownPropertyPath = true;
+				}
 
-						if (bTryLoadIntoPropertyBag)
+				// Track the path for an unknown property and serialize it to track any unknown property within it.
+				if (UNLIKELY(bTryStoreUnknownPropertyPath))
+				{
+					if (!bSearchedForUnknownPropertyTree)
+					{
+						bSearchedForUnknownPropertyTree = true;
+						if (SerializeContext && SerializeContext->bSerializeUnknownProperty)
 						{
-							if (FPropertyBag* PropertyBag = TryFindPropertyBag())
+							if (UObject* Object = SerializeContext->SerializedObject)
 							{
-								Tag.SetProperty(nullptr);
-								UnderlyingArchive.Seek(StartOfProperty);
-								FStructuredArchive::FSlot ValueSlotCopy = PropertyRecord.EnterField(TEXT("Value"));
-								PropertyBag->LoadPropertyByTag(SerializeContext->SerializedPropertyPath, Tag, ValueSlotCopy, Property->ContainerPtrToValuePtrForDefaults<uint8>(DefaultsStruct, Defaults, Tag.ArrayIndex));
+								UnknownPropertyTree = FPropertyBagRepository::Get().CreateUnknownPropertyTree(Object);
 							}
 						}
 					}
-				}
-				else if (FPropertyBag* PropertyBag = TryFindPropertyBag(); PropertyBag && SerializeContext)
-				{
-					// TODO: Might we find defaults in a property bag for Defaults?
-					FStructuredArchive::FSlot ValueSlot = PropertyRecord.EnterField(TEXT("Value"));
-					PropertyBag->LoadPropertyByTag(SerializeContext->SerializedPropertyPath, Tag, ValueSlot);
-					
-					// we still want to set overrides for loose properties because they will be used by the IDO.
-					if (!UnderlyingArchive.IsTransacting())
+
+					if (UnknownPropertyTree)
 					{
-						if (FOverriddenPropertySet* OverriddenProperties = FOverridableSerializationLogic::GetOverriddenProperties())
+						UnknownPropertyTree->Add(SerializeContext->SerializedPropertyPath);
+
+						const bool bSerializeValue = PropertyTypeContainsStructOrEnum(Tag.GetType());
+						const bool bRestoreOverrideOperation = Tag.OverrideOperation != EOverriddenPropertyOperation:: None && !Property && !UnderlyingArchive.IsTransacting();
+
+						// Try to construct a field from the property tag for serialization and overrides.
+						if (bSerializeValue || bRestoreOverrideOperation)
 						{
-							// No need to restore none operations
-							if (Tag.OverrideOperation != EOverriddenPropertyOperation::None)
+							TUniquePtr<FField> TempField(FField::TryConstruct(Tag.Type, {}, Tag.Name, RF_NoFlags));
+							if (FProperty* TempProperty = CastField<FProperty>(TempField.Get()); TempProperty && TempProperty->LoadTypeName(Tag.GetType()))
 							{
-								// SetOverriddenPropertyOperation only needs to know the property name so we'll just construct a dummy property
-								FProperty* PropertyToOverride = CastField<FProperty>(FField::TryConstruct(Tag.Type, {}, Tag.Name, RF_NoFlags));
-								check(PropertyToOverride);
-								OverriddenProperties->SetOverriddenPropertyOperation(Tag.OverrideOperation, UnderlyingArchive.GetSerializedPropertyChain(), PropertyToOverride);
-								delete PropertyToOverride;
+								TempProperty->Link(UnderlyingArchive);
+
+								// Serialize the value if it contains a struct or enum because only those may contain an unknown property.
+								if (bSerializeValue)
+								{
+									UnderlyingArchive.Seek(StartOfProperty);
+									FStructuredArchive::FSlot ValueSlotCopy = PropertyRecord.EnterField(TEXT("Value"));
+									void* TempData = TempProperty->AllocateAndInitializeValue();
+									Tag.SerializeTaggedProperty(ValueSlotCopy, TempProperty, (uint8*)TempData, nullptr);
+									TempProperty->DestroyAndFreeValue(TempData);
+								}
+
+								// Restore overrides for unknown properties.
+								if (bRestoreOverrideOperation)
+								{
+									if (FOverriddenPropertySet* OverriddenProperties = FOverridableSerializationLogic::GetOverriddenProperties())
+									{
+										OverriddenProperties->SetOverriddenPropertyOperation(Tag.OverrideOperation, UnderlyingArchive.GetSerializedPropertyChain(), TempProperty);
+									}
+								}
 							}
 						}
 					}
diff --git a/Engine/Source/Runtime/CoreUObject/Private/UObject/InstanceDataObjectUtils.cpp b/Engine/Source/Runtime/CoreUObject/Private/UObject/InstanceDataObjectUtils.cpp
--- a/Engine/Source/Runtime/CoreUObject/Private/UObject/InstanceDataObjectUtils.cpp
+++ b/Engine/Source/Runtime/CoreUObject/Private/UObject/InstanceDataObjectUtils.cpp
@@ -3,6 +3,7 @@
 #include "UObject/InstanceDataObjectUtils.h"
 
 #include "HAL/IConsoleManager.h"
+#include "Misc/ReverseIterate.h"
 #include "UObject/Class.h"
 #include "UObject/EnumProperty.h"
 #include "UObject/Field.h"
@@ -10,6 +11,7 @@
 #include "UObject/PropertyHelper.h"
 #include "UObject/PropertyBagRepository.h"
 #include "UObject/PropertyOptional.h"
+#include "UObject/PropertyPathNameTree.h"
 #include "UObject/UnrealType.h"
 
 static const FName NAME_ValuesSetBySerialization(ANSITEXTVIEW("_ValuesSetBySerialization"));
@@ -311,13 +313,12 @@ namespace UE
 		return true;
 	}
 
-	static UStruct* CreateInstanceDataObjectStructRec(const UClass* StructClass, UStruct* TemplateStruct,
-		UObject* Outer, const TMap<FWildcardPropertyPathName, TMap<FName, const FProperty*>>& LooseProperties, FWildcardPropertyPathName& Path);
-	template <typename TStructType>
-	TStructType* CreateInstanceDataObjectStructRec(UStruct* TemplateStruct, UObject* Outer,
-		const TMap<FWildcardPropertyPathName, TMap<FName, const FProperty*>>& LooseProperties, FWildcardPropertyPathName& Path)
+	static UStruct* CreateInstanceDataObjectStructRec(const UClass* StructClass, UStruct* TemplateStruct, UObject* Outer, const FPropertyPathNameTree* PropertyTree);
+
+	template <typename StructType>
+	StructType* CreateInstanceDataObjectStructRec(UStruct* TemplateStruct, UObject* Outer, const FPropertyPathNameTree* PropertyTree)
 	{
-		return CastChecked<TStructType>(CreateInstanceDataObjectStructRec(TStructType::StaticClass(), TemplateStruct, Outer, LooseProperties, Path));
+		return CastChecked<StructType>(CreateInstanceDataObjectStructRec(StructType::StaticClass(), TemplateStruct, Outer, PropertyTree));
 	}
 
 	static FPropertyPathNameSegment CreateSegmentFromProperty(const FProperty* Inner, int32 Index = INDEX_NONE)
@@ -349,8 +350,7 @@ namespace UE
 	}
 
 	// recursively re-instances all structs contained by this property to include loose properties
-	static void ConvertToInstanceDataObjectProperty(FProperty* Property, FPropertyTypeName PropertyType, UObject* Outer,
-		const TMap<FWildcardPropertyPathName, TMap<FName, const FProperty*>>& LooseProperties, FWildcardPropertyPathName& Path)
+	static void ConvertToInstanceDataObjectProperty(FProperty* Property, FPropertyTypeName PropertyType, UObject* Outer, const FPropertyPathNameTree* PropertyTree)
 	{
 #if WITH_EDITORONLY_DATA
 		if (!Property->HasMetaData(NAME_DisplayName))
@@ -363,7 +363,8 @@ namespace UE
 			}
 		}
 #endif
-const auto TrySetContainsLooseProperties = [](FProperty* Property, const FFieldVariant& Inner)
+
+		const auto TrySetContainsLooseProperties = [](FProperty* Property, const FFieldVariant& Inner)
 		{
 #if WITH_EDITORONLY_DATA
 			if (Inner.HasMetaData(NAME_ContainsLoosePropertiesMetadata))
@@ -372,7 +373,7 @@ const auto TrySetContainsLooseProperties = [](FProperty* Property, const FFieldV
 			}
 #endif
 		};
-		
+
 		if (FStructProperty* AsStructProperty = CastField<FStructProperty>(Property))
 		{
 			if (!AsStructProperty->Struct->UseNativeSerialization())
@@ -400,7 +401,7 @@ const auto TrySetContainsLooseProperties = [](FProperty* Property, const FFieldV
 					OriginalName = WriteToString<256>(OriginalNameBuilder.Build()).ToView();
 				}
 #endif
-				UInstanceDataObjectStruct* Struct = CreateInstanceDataObjectStructRec<UInstanceDataObjectStruct>(AsStructProperty->Struct, Outer, LooseProperties, Path);
+				UInstanceDataObjectStruct* Struct = CreateInstanceDataObjectStructRec<UInstanceDataObjectStruct>(AsStructProperty->Struct, Outer, PropertyTree);
 				if (const FName StructGuidName = PropertyType.GetParameterName(1); !StructGuidName.IsNone())
 				{
 					FGuid::Parse(StructGuidName.ToString(), Struct->Guid);
@@ -417,32 +418,40 @@ const auto TrySetContainsLooseProperties = [](FProperty* Property, const FFieldV
 		}
 		else if (FArrayProperty* AsArrayProperty = CastField<FArrayProperty>(Property))
 		{
-			ConvertToInstanceDataObjectProperty(AsArrayProperty->Inner, PropertyType.GetParameter(0), Outer, LooseProperties, Path);
+			ConvertToInstanceDataObjectProperty(AsArrayProperty->Inner, PropertyType.GetParameter(0), Outer, PropertyTree);
 			TrySetContainsLooseProperties(AsArrayProperty, AsArrayProperty->Inner);
 		}
 		else if (FSetProperty* AsSetProperty = CastField<FSetProperty>(Property))
 		{
-			ConvertToInstanceDataObjectProperty(AsSetProperty->ElementProp, PropertyType.GetParameter(0), Outer, LooseProperties, Path);
+			ConvertToInstanceDataObjectProperty(AsSetProperty->ElementProp, PropertyType.GetParameter(0), Outer, PropertyTree);
 			TrySetContainsLooseProperties(AsSetProperty, AsSetProperty->ElementProp);
 		}
 		else if (FMapProperty* AsMapProperty = CastField<FMapProperty>(Property))
 		{
-			Path.Push({NAME_IDOMapKey});
-			ConvertToInstanceDataObjectProperty(AsMapProperty->KeyProp, PropertyType.GetParameter(0), Outer, LooseProperties, Path);
-			Path.Pop();
+			const FPropertyPathNameTree* KeyTree = nullptr;
+			const FPropertyPathNameTree* ValueTree = nullptr;
+			if (PropertyTree)
+			{
+				FPropertyPathName Path;
+				Path.Push({NAME_IDOMapKey});
+				PropertyTree->Find(&KeyTree, Path);
+				Path.Pop();
+				Path.Push({NAME_IDOMapValue});
+				PropertyTree->Find(&ValueTree, Path);
+				Path.Pop();
+			}
+
+			ConvertToInstanceDataObjectProperty(AsMapProperty->KeyProp, PropertyType.GetParameter(0), Outer, KeyTree);
 			TrySetContainsLooseProperties(AsMapProperty, AsMapProperty->KeyProp);
-			
-			Path.Push({NAME_IDOMapValue});
-			ConvertToInstanceDataObjectProperty(AsMapProperty->ValueProp, PropertyType.GetParameter(1), Outer, LooseProperties, Path);
-			Path.Pop();
+			ConvertToInstanceDataObjectProperty(AsMapProperty->ValueProp, PropertyType.GetParameter(1), Outer, ValueTree);
 			TrySetContainsLooseProperties(AsMapProperty, AsMapProperty->ValueProp);
 		}
 		else if (FOptionalProperty* AsOptionalProperty = CastField<FOptionalProperty>(Property))
 		{
-			ConvertToInstanceDataObjectProperty(AsOptionalProperty->GetValueProperty(), PropertyType.GetParameter(0), Outer, LooseProperties, Path);
+			ConvertToInstanceDataObjectProperty(AsOptionalProperty->GetValueProperty(), PropertyType.GetParameter(0), Outer, PropertyTree);
 			TrySetContainsLooseProperties(AsOptionalProperty, AsOptionalProperty->GetValueProperty());
 		}
-		
+
 #if WITH_EDITORONLY_DATA
 		if (Property->GetBoolMetaData(NAME_IsLooseMetadata) || Property->GetBoolMetaData(NAME_ContainsLoosePropertiesMetadata))
 		{
@@ -463,7 +472,7 @@ const auto TrySetContainsLooseProperties = [](FProperty* Property, const FFieldV
 			MarkPropertyAsLoose(AsArrayProperty->Inner);
 		}
 		else if (const FSetProperty* AsSetProperty = CastField<FSetProperty>(Property))
-		{
+{
 			MarkPropertyAsLoose(AsSetProperty->ElementProp);
 		}
 		else if (const FMapProperty* AsMapProperty = CastField<FMapProperty>(Property))
@@ -477,167 +486,92 @@ const auto TrySetContainsLooseProperties = [](FProperty* Property, const FFieldV
 		}
 	}
 
-	// copy template property then convert it into an InstanceDataObject property by adding loose properties
-	static FProperty* CreateInstanceDataObjectProperty(const FProperty* TemplateProperty, UObject* Outer,
-		const TMap<FWildcardPropertyPathName, TMap<FName, const FProperty*>>& LooseProperties, FWildcardPropertyPathName& Path, bool bIsLoose)
+	// constructs an InstanceDataObject struct by merging the properties in 
+	static UStruct* CreateInstanceDataObjectStructRec(const UClass* StructClass, UStruct* TemplateStruct, UObject* Outer, const FPropertyPathNameTree* PropertyTree)
 	{
-		FProperty* InstanceDataObjectProperty = CastFieldChecked<FProperty>(FField::Duplicate(TemplateProperty, Outer));
-#if WITH_EDITORONLY_DATA
-		FField::CopyMetaData(TemplateProperty, InstanceDataObjectProperty);
-#endif
-		if (bIsLoose)
-		{
-			MarkPropertyAsLoose(InstanceDataObjectProperty);
-		}
-		
-		ConvertToInstanceDataObjectProperty(InstanceDataObjectProperty, Path.GetSegment(Path.GetSegmentCount() - 1).Type, Outer, LooseProperties, Path);
-		return InstanceDataObjectProperty;
-	}
+		TSet<FPropertyPathName> SuperPropertyPathsFromTree;
 
-	// return a copy of Path with all the indices set to -1. This way all container elements will have the same wildcard path
-	static FWildcardPropertyPathName ConvertToWildcardPath(const FPropertyPathName& Path)
-	{
-		FWildcardPropertyPathName Result = Path;
-		// make path a wildcard path
-		for (int I = 0; I < Result.GetSegmentCount(); ++I)
-		{
-			FPropertyPathNameSegment Segment = Result.GetSegment(I);
-			Segment.Index = INDEX_NONE;
-			Result.SetSegment(I, Segment);
-		}
-		return Result;
-	}
+		// UClass is required to inherit from UObject
+		UStruct* Super = (StructClass == UClass::StaticClass()) ? UObject::StaticClass() : nullptr;
 
-	// recursively add all the wildcard paths of both Property and all it's sub-Properties to OutLooseProperties
-	static void AddWildcardedProperties(TMap<FWildcardPropertyPathName, TMap<FName, const FProperty*>>& OutProperties, FWildcardPropertyPathName& ParentPath, const FProperty* Property)
-	{
-		OutProperties.FindOrAdd(ParentPath).Add(Property->GetFName(), Property);
-		
-		ParentPath.Push(CreateSegmentFromProperty(Property));
-		if (const FStructProperty* AsStructProperty = CastField<FStructProperty>(Property))
+		if (TemplateStruct)
 		{
-			for (const FProperty* SubProperty : TFieldRange<FProperty>(AsStructProperty->Struct))
 			{
-				AddWildcardedProperties(OutProperties, ParentPath, SubProperty);
+				const FName SuperName(WriteToString<128>(TemplateStruct->GetName(), TEXTVIEW("_Super")));
+				UStruct* NewSuper = NewObject<UStruct>(Outer, StructClass, MakeUniqueObjectName(nullptr, StructClass, SuperName));
+				NewSuper->SetSuperStruct(Super);
+				Super = NewSuper;
 			}
-		}
-		else if (const FArrayProperty* AsArrayProperty = CastField<FArrayProperty>(Property))
-		{
-			AddWildcardedProperties(OutProperties, ParentPath, AsArrayProperty->Inner);
-		}
-		else if (const FSetProperty* AsSetProperty = CastField<FSetProperty>(Property))
-		{
-			AddWildcardedProperties(OutProperties, ParentPath, AsSetProperty->ElementProp);
-		}
-		else if (const FMapProperty* AsMapProperty = CastField<FMapProperty>(Property))
-		{
-			AddWildcardedProperties(OutProperties, ParentPath, AsMapProperty->KeyProp);
-			AddWildcardedProperties(OutProperties, ParentPath, AsMapProperty->ValueProp);
-		}
-		else if (const FOptionalProperty* AsOptionalProperty = CastField<FOptionalProperty>(Property))
-		{
-			AddWildcardedProperties(OutProperties, ParentPath, AsOptionalProperty->GetValueProperty());
-		}
-		ParentPath.Pop();
-	}
 
-	// construct a map that keys a parent struct by it's wildcard path and returns an array of all it's loose properties
-	static TMap<FWildcardPropertyPathName, TMap<FName, const FProperty*>> GetWildcardedLooseProperties(const FPropertyBag* PropertyBag)
-	{
-		
-		TMap<FWildcardPropertyPathName, TMap<FName, const FProperty*>> LooseProperties;
-		if (PropertyBag)
-		{
-			for (FPropertyBag::FConstIterator Itr = PropertyBag->CreateConstIterator(); Itr; ++Itr)
+			// Gather properties for Super Struct
+			TArray<FProperty*> SuperProperties;
+			for (const FProperty* TemplateProperty : TFieldRange<FProperty>(TemplateStruct))
 			{
-				FWildcardPropertyPathName ParentPath = ConvertToWildcardPath(Itr.GetPath());
-				ParentPath.Pop();
-				const FProperty* Property = Itr.GetProperty();
-				if (ensure(Property))
-				{
-					AddWildcardedProperties(LooseProperties, ParentPath, Property);
-				}
-			}
-		}
-		
-		return LooseProperties;
-	}
-	
-	// constructs an InstanceDataObject struct by merging the properties in 
-	static UStruct* CreateInstanceDataObjectStructRec(const UClass* StructClass, UStruct* TemplateStruct,
-		UObject* Outer, const TMap<FWildcardPropertyPathName, TMap<FName, const FProperty*>>& LooseProperties, FWildcardPropertyPathName& Path)
-	{
-		UStruct* Super = nullptr;
+				FProperty* SuperProperty = CastFieldChecked<FProperty>(FField::Duplicate(TemplateProperty, Super));
+				SuperProperties.Add(SuperProperty);
 
-		const TMap<FName, const FProperty*>* BagProperties = LooseProperties.Find(Path);
+			#if WITH_EDITORONLY_DATA
+				FField::CopyMetaData(TemplateProperty, SuperProperty);
+			#endif
 
-		auto MatchesBagProperty = [&BagProperties](const FProperty* Property)
-		{
-			if (BagProperties)
-			{
-				if (const FProperty* const* Found = BagProperties->Find(Property->GetFName()))
+				FPropertyTypeName Type;
 				{
-					return (*Found)->SameType(Property);
+					FPropertyTypeNameBuilder TypeBuilder;
+					TemplateProperty->SaveTypeName(TypeBuilder);
+					Type = TypeBuilder.Build();
 				}
-			}
-			return false;
-		};
-		
-		if (TemplateStruct)
-		{
-			const FName SuperName(TemplateStruct->GetName() + TEXT("_Super"));
-			Super = NewObject<UStruct>(Outer, StructClass, MakeUniqueObjectName(nullptr, StructClass, SuperName));
-			
-			// Gather properties for Super Struct
-			TArray<FProperty*> SuperProperties;
-			for (const FProperty* TemplateProperty : TFieldRange<FProperty>(TemplateStruct))
-			{
-				if (MatchesBagProperty(TemplateProperty))
+
+				// Find the sub-tree containing unknown properties for this template property.
+				const FPropertyPathNameTree* SubTree = nullptr;
+				if (PropertyTree)
 				{
-					// this property was determined to be loose despite it being in the template.
-					// this likely occurred due to an entire struct instance being loose and that instance becoming a template
-					continue;
+					FPropertyPathName Path;
+					Path.Push({TemplateProperty->GetFName(), Type});
+					if (PropertyTree->Find(&SubTree, Path))
+					{
+						SuperPropertyPathsFromTree.Add(MoveTemp(Path));
+					}
 				}
-				Path.Push(CreateSegmentFromProperty(TemplateProperty));
-				FProperty* SuperProperty = CreateInstanceDataObjectProperty(TemplateProperty, Super, LooseProperties, Path, false);
-				
-				Path.Pop();
-				SuperProperties.Add(SuperProperty);
-			}
 
-			if (StructClass == UClass::StaticClass())
-			{
-				// UClasses are required to inherit from a UObject class
-				Super->SetSuperStruct(UObject::StaticClass());
+				ConvertToInstanceDataObjectProperty(SuperProperty, Type, Super, SubTree);
 			}
-		    
+
 			// AddCppProperty expects reverse property order for StaticLink to work correctly
-			for (int32 I = SuperProperties.Num() - 1; I >= 0; --I)
+			for (FProperty* Property : ReverseIterate(SuperProperties))
 			{
-				Super->AddCppProperty(SuperProperties[I]);
+				Super->AddCppProperty(Property);
 			}
 			Super->Bind();
-			Super->StaticLink(/*RelinkExistingProperties*/true);
-		}
-		else if (StructClass == UClass::StaticClass())
-		{
-			// UClasses are required to inherit from a UObject class
-			Super = UObject::StaticClass();
+			Super->StaticLink(/*bRelinkExistingProperties*/true);
 		}
 
-		const FName InstanceDataObjectName = (TemplateStruct) ? FName(TemplateStruct->GetName() + TEXT("_InstanceDataObject")) : FName(TEXT("InstanceDataObject"));
+		const FName InstanceDataObjectName = (TemplateStruct) ? FName(WriteToString<128>(TemplateStruct->GetName(), TEXTVIEW("_InstanceDataObject"))) : FName(TEXTVIEW("InstanceDataObject"));
 		UStruct* Result = NewObject<UStruct>(Outer, StructClass, MakeUniqueObjectName(nullptr, StructClass, InstanceDataObjectName));
+		Result->SetSuperStruct(Super);
 
 		// Gather "loose" properties for child Struct
 		TArray<FProperty*> LooseInstanceDataObjectProperties;
-		if (BagProperties)
+		if (PropertyTree)
 		{
-			for (const TPair<FName, const FProperty*>& BagProperty : *BagProperties)
+			for (FPropertyPathNameTree::FConstIterator It = PropertyTree->CreateConstIterator(); It; ++It)
 			{
-				Path.Push(CreateSegmentFromProperty(BagProperty.Value));
-				FProperty* LooseProperty = CreateInstanceDataObjectProperty(BagProperty.Value, Result, LooseProperties, Path, true);
-				Path.Pop();
-				LooseInstanceDataObjectProperties.Add(LooseProperty);
+				FName Name = It.GetName();
+				FPropertyTypeName Type = It.GetType();
+				FPropertyPathName Path;
+				Path.Push({Name, Type});
+				if (!SuperPropertyPathsFromTree.Contains(Path))
+				{
+					// Construct a property from the type and try to use it to serialize the value.
+					FField* Field = FField::TryConstruct(Type.GetName(), Result, Name, RF_NoFlags);
+					if (FProperty* Property = CastField<FProperty>(Field); Property && Property->LoadTypeName(Type))
+					{
+						MarkPropertyAsLoose(Property);
+						ConvertToInstanceDataObjectProperty(Property, Type, Result, It.GetSubTree());
+						LooseInstanceDataObjectProperties.Add(Property);
+						continue;
+					}
+					delete Field;
+				}
 			}
 		}
 
@@ -650,15 +584,13 @@ const auto TrySetContainsLooseProperties = [](FProperty* Property, const FFieldV
 			Result->AddCppProperty(ValuesSetBySerializationProperty);
 		}
 
-		Result->SetSuperStruct(Super);
-		
 		// AddCppProperty expects reverse property order for StaticLink to work correctly
-		for (int32 I = LooseInstanceDataObjectProperties.Num() - 1; I >= 0; --I)
+		for (FProperty* Property : ReverseIterate(LooseInstanceDataObjectProperties))
 		{
-			Result->AddCppProperty(LooseInstanceDataObjectProperties[I]);
+			Result->AddCppProperty(Property);
 		}
 		Result->Bind();
-		Result->StaticLink(/*RelinkExistingProperties*/true);
+		Result->StaticLink(/*bRelinkExistingProperties*/true);
 		return Result;
 	}
 
@@ -697,11 +629,9 @@ const auto TrySetContainsLooseProperties = [](FProperty* Property, const FFieldV
 			CLASS_EditInlineNew | CLASS_CollapseCategories | CLASS_Const | CLASS_CompiledFromBlueprint | CLASS_HasInstancedReference);
 	}
 
-	UClass* CreateInstanceDataObjectClass(const FPropertyBag* PropertyBag, UClass* OwnerClass, UObject* Outer)
+	UClass* CreateInstanceDataObjectClass(const FPropertyPathNameTree* PropertyTree, UClass* OwnerClass, UObject* Outer)
 	{
-		const TMap<FWildcardPropertyPathName, TMap<FName, const FProperty*>> LooseProperties = GetWildcardedLooseProperties(PropertyBag);
-		FWildcardPropertyPathName ParentPath;
-		UClass* Result = CreateInstanceDataObjectStructRec<UClass>(OwnerClass, Outer, LooseProperties, ParentPath);
+		UClass* Result = CreateInstanceDataObjectStructRec<UClass>(OwnerClass, Outer, PropertyTree);
 #if WITH_EDITORONLY_DATA
 		if (const FString& DisplayName = OwnerClass->GetMetaData(NAME_DisplayName); !DisplayName.IsEmpty())
 		{
diff --git a/Engine/Source/Runtime/CoreUObject/Private/UObject/InstanceDataObjectUtils.h b/Engine/Source/Runtime/CoreUObject/Private/UObject/InstanceDataObjectUtils.h
--- a/Engine/Source/Runtime/CoreUObject/Private/UObject/InstanceDataObjectUtils.h
+++ b/Engine/Source/Runtime/CoreUObject/Private/UObject/InstanceDataObjectUtils.h
@@ -4,15 +4,14 @@
 
 #include "HAL/Platform.h"
 
+class FFieldVariant;
 class FProperty;
 class UClass;
 class UObject;
 class UStruct;
-class FProperty;
-class FFieldVariant;
 
-namespace UE { class FPropertyBag; }
 namespace UE { class FPropertyPathName; }
+namespace UE { class FPropertyPathNameTree; }
 
 namespace UE
 {
@@ -24,8 +23,8 @@ namespace UE
  */
 bool IsInstanceDataObjectSupportEnabled(UObject* Object = nullptr);
 
-/** Generate a UClass that contains the union of the properties of PropertyBag and OwnerClass. */
-UClass* CreateInstanceDataObjectClass(const FPropertyBag* PropertyBag, UClass* OwnerClass, UObject* Outer);
+/** Generate a UClass that contains the union of the properties of PropertyTree and OwnerClass. */
+UClass* CreateInstanceDataObjectClass(const FPropertyPathNameTree* PropertyTree, UClass* OwnerClass, UObject* Outer);
 
 /** Mark a property within the object as having been set during deserialization. */
 void MarkPropertySetBySerialization(UObject* Object, const FPropertyPathName& Path);
diff --git a/Engine/Source/Runtime/CoreUObject/Private/UObject/PropertyBagRepository.cpp b/Engine/Source/Runtime/CoreUObject/Private/UObject/PropertyBagRepository.cpp
--- a/Engine/Source/Runtime/CoreUObject/Private/UObject/PropertyBagRepository.cpp
+++ b/Engine/Source/Runtime/CoreUObject/Private/UObject/PropertyBagRepository.cpp
@@ -9,6 +9,7 @@
 #include "UObject/Object.h"
 #include "UObject/UObjectGlobals.h"
 #include "UObject/PropertyBag.h"
+#include "UObject/PropertyPathNameTree.h"
 #include "UObject/UObjectHash.h"
 #include "UObject/UObjectThreadContext.h"
 #include "UObject/LinkerLoad.h"
@@ -176,16 +177,36 @@ void FPropertyBagRepository::CleanupLevel(const UObject* Level)
 FPropertyBag* FPropertyBagRepository::CreateOuterBag(const UObject* Owner)
 {
 	FPropertyBagRepositoryLock LockRepo(this);
-	const FPropertyBagAssociationData* BagData = AssociatedData.Find(Owner);
-	if(!BagData)
+	FPropertyBagAssociationData* BagData = AssociatedData.Find(Owner);
+	if (!BagData)
 	{
 		FPropertyBagAssociationData NewBagData;
-		NewBagData.Bag = new FPropertyBag;
 		BagData = &AssociatedData.Emplace(Owner, NewBagData);
 	}
+	if (!BagData->Bag)
+	{
+		BagData->Bag = new FPropertyBag;
+	}
 	return BagData->Bag;
 }
 
+// TODO: Create these by class on construction?
+FPropertyPathNameTree* FPropertyBagRepository::CreateUnknownPropertyTree(const UObject* Owner)
+{
+	FPropertyBagRepositoryLock LockRepo(this);
+	FPropertyBagAssociationData* BagData = AssociatedData.Find(Owner);
+	if (!BagData)
+	{
+		FPropertyBagAssociationData NewBagData;
+		BagData = &AssociatedData.Emplace(Owner, NewBagData);
+	}
+	if (!BagData->Tree)
+	{
+		BagData->Tree = new FPropertyPathNameTree;
+	}
+	return BagData->Tree;
+}
+
 UObject* FPropertyBagRepository::CreateInstanceDataObject(UObject* Owner, FArchive* Archive)
 {
 	FPropertyBagRepositoryLock LockRepo(this);
@@ -206,8 +227,9 @@ void FPropertyBagRepository::DestroyOuterBag(const UObject* Owner)
 
 bool FPropertyBagRepository::RequiresFixup(const UObject* Object) const
 {
-	const FPropertyBag* PropertyBag = FindBag(Object);
-	return !PropertyBag || PropertyBag->IsEmpty();
+	FPropertyBagRepositoryLock LockRepo(this);
+	const FPropertyBagAssociationData* BagData = AssociatedData.Find(Object);
+	return BagData && BagData->Tree && !BagData->Tree->IsEmpty() && BagData->InstanceDataObject;
 }
 
 bool FPropertyBagRepository::RemoveAssociationUnsafe(const UObject* Owner)
@@ -304,10 +326,10 @@ FString FPropertyBagRepository::GetReferencerName() const
 void FPropertyBagRepository::CreateInstanceDataObjectUnsafe(UObject* Owner, FPropertyBagAssociationData& BagData, FArchive* Archive)
 {
 	check(!BagData.InstanceDataObject);	// No repeated calls
-	const FPropertyBag* PropertyBag = BagData.Bag;
+	const FPropertyPathNameTree* PropertyTree = BagData.Tree;
 	// construct InstanceDataObject class
 	// TODO: should we put the InstanceDataObject or it's class in a package?
-	const UClass* InstanceDataObjectClass = CreateInstanceDataObjectClass(PropertyBag, Owner->GetClass(), GetTransientPackage());
+	const UClass* InstanceDataObjectClass = CreateInstanceDataObjectClass(PropertyTree, Owner->GetClass(), GetTransientPackage());
 
 	TObjectPtr<UObject>* OuterPtr;
 	if (FPropertyBagAssociationData* OuterData = AssociatedData.Find(Owner->GetOuter()))
