SavePackage: DiffOnly Cooks and PropertyOverrides: when the DiffOnly cook calls SavePackage a second time to gather callstacks, it skips the call to PreSave on each object, but PreSave can write to PropertyOverrides. Save the PropertyOverrides from the first save and pass them into the second save.
#rnx

diff --git a/Engine/Source/Editor/UnrealEd/Private/CookOnTheFlyServer.cpp b/Engine/Source/Editor/UnrealEd/Private/CookOnTheFlyServer.cpp
--- a/Engine/Source/Editor/UnrealEd/Private/CookOnTheFlyServer.cpp
+++ b/Engine/Source/Editor/UnrealEd/Private/CookOnTheFlyServer.cpp
@@ -6118,6 +6118,7 @@ void UCookOnTheFlyServer::SaveCookedPackage(UE::Cook::FSaveCookedPackageContext&
 			UE_SCOPED_HIERARCHICAL_COOKTIMER(GEditorSavePackage);
 			UE_TRACK_REFERENCING_PLATFORM_SCOPED(TargetPlatform);
 
+			TMap<UObject*, FObjectSaveOverride> SaveOverrides;
 			FArchiveCookData CookData(*TargetPlatform, *Context.ArchiveCookContext);
 			FSavePackageArgs SaveArgs;
 			SaveArgs.TopLevelFlags = Context.FlagsToCook;
@@ -6127,6 +6128,7 @@ void UCookOnTheFlyServer::SaveCookedPackage(UE::Cook::FSaveCookedPackageContext&
 			SaveArgs.ArchiveCookData = &CookData;
 			SaveArgs.bSlowTask = false;
 			SaveArgs.SavePackageContext = Context.SavePackageContext;
+			SaveArgs.InOutSaveOverrides = &SaveOverrides;
 
 			Context.PackageWriter->UpdateSaveArguments(SaveArgs);
 			for(;;)
diff --git a/Engine/Source/Runtime/CoreUObject/Private/UObject/SavePackage/SaveContext.h b/Engine/Source/Runtime/CoreUObject/Private/UObject/SavePackage/SaveContext.h
--- a/Engine/Source/Runtime/CoreUObject/Private/UObject/SavePackage/SaveContext.h
+++ b/Engine/Source/Runtime/CoreUObject/Private/UObject/SavePackage/SaveContext.h
@@ -489,6 +489,10 @@ public:
 			ObjectSaveContext.CookType = SaveArgs.ArchiveCookData->CookContext.GetCookType();
 			ObjectSaveContext.CookingDLC = SaveArgs.ArchiveCookData->CookContext.GetCookingDLC();
 		}
+		if (SaveArgs.InOutSaveOverrides)
+		{
+			ObjectSaveContext.SaveOverrides = MoveTemp(*SaveArgs.InOutSaveOverrides);
+		}
 
 		// Setup the harvesting flags and generate the context for harvesting the package
 		SetupHarvestingRealms();
@@ -500,6 +504,13 @@ public:
 		{
 			UE::SavePackageUtilities::CallPostSaveRoot(Asset, ObjectSaveContext, bNeedPreSaveCleanup);
 		}
+
+		// Move the SaveOverrides that we copied and/or modified onto our ObjectSaveContext back to the
+		// InOut SaveOverrides parameter on the SaveArgs.
+		if (SaveArgs.InOutSaveOverrides)
+		{
+			*SaveArgs.InOutSaveOverrides = MoveTemp(ObjectSaveContext.SaveOverrides);
+		}
 	}
 
 	uint32 GetInitialPackageFlags() const
diff --git a/Engine/Source/Runtime/CoreUObject/Public/UObject/SavePackage.h b/Engine/Source/Runtime/CoreUObject/Public/UObject/SavePackage.h
--- a/Engine/Source/Runtime/CoreUObject/Public/UObject/SavePackage.h
+++ b/Engine/Source/Runtime/CoreUObject/Public/UObject/SavePackage.h
@@ -22,6 +22,7 @@
 #include "Templates/UniquePtr.h"
 #include "Templates/UnrealTemplate.h"
 #include "UObject/NameTypes.h"
+#include "UObject/ObjectSaveOverride.h"
 #include "UObject/Package.h"
 
 class ITargetPlatform;
@@ -92,6 +93,12 @@ struct FSavePackageArgs
 	UE_DEPRECATED(4.27, "UPackage::Guid has not been used by the engine for a long time and it will be removed.")
 	TOptional<FGuid> OutputPackageGuid;
 
+	/**
+	 * In/Out list of property overrides per object to apply to during save. This list can be extended by PreSave functions
+	 * during the save.
+	 */
+	TMap<UObject*, FObjectSaveOverride>* InOutSaveOverrides = nullptr;
+
 	PRAGMA_DISABLE_DEPRECATION_WARNINGS;
 	FSavePackageArgs() = default;
 	FSavePackageArgs(const FSavePackageArgs&) = default;
