OnCommit for Removal and Open and Abort during Add for Annotations

diff --git a/Engine/Source/Programs/AutoRTFMTests/Private/UObjectTests.cpp b/Engine/Source/Programs/AutoRTFMTests/Private/UObjectTests.cpp
--- a/Engine/Source/Programs/AutoRTFMTests/Private/UObjectTests.cpp
+++ b/Engine/Source/Programs/AutoRTFMTests/Private/UObjectTests.cpp
@@ -5,6 +5,7 @@
 #include <AutoRTFM/AutoRTFM.h>
 #include "UObject/GCObject.h"
 #include "UObject/ReachabilityAnalysis.h"
+#include "UObject/UObjectAnnotation.h"
 
 TEST_CASE("UObject.NewObject")
 {
@@ -135,3 +136,65 @@ TEST_CASE("UObject.MarkAsReachable")
 	// Reset it back just incase another test required the original time limit.
 	SetReachabilityAnalysisTimeLimit(Original);
 }
+
+
+namespace
+{
+	struct FTestAnnotation
+	{
+		// Default value, if it doesnt exist it will return a default value one
+		int TestAnnotationNumber = 32;
+
+		bool IsDefault() const
+		{
+			return false;
+		}
+	};
+	FUObjectAnnotationSparse<FTestAnnotation, true> GTestAnnotation;
+}
+
+TEST_CASE("UObject.TestAddAnnotation")
+{
+	SECTION("Create")
+	{
+		UMyAutoRTFMTestObject* Outer = NewObject<UMyAutoRTFMTestObject>();
+		UMyAutoRTFMTestObject* Object = nullptr;
+
+		AutoRTFM::Commit([&]
+			{
+				Object = NewObject<UMyAutoRTFMTestObject>(Outer);
+
+				FTestAnnotation Temp;
+				Temp.TestAnnotationNumber = 70;
+
+				GTestAnnotation.AddAnnotation(Object, Temp);
+			});
+
+		REQUIRE(nullptr != Object);
+		REQUIRE(42 == Object->Value);
+		REQUIRE(Object->IsInOuter(Outer));
+		REQUIRE(55 == Outer->Value);
+		REQUIRE(70 == GTestAnnotation.GetAnnotation(Object).TestAnnotationNumber);
+	}
+
+	SECTION("Abort")
+	{
+		UMyAutoRTFMTestObject* Outer = NewObject<UMyAutoRTFMTestObject>();
+		UMyAutoRTFMTestObject* Object = nullptr;
+
+		Object = NewObject<UMyAutoRTFMTestObject>(Outer);
+
+		REQUIRE(AutoRTFM::ETransactionResult::AbortedByRequest == AutoRTFM::Transact([&]
+			{
+
+				FTestAnnotation Temp;
+				Temp.TestAnnotationNumber = 70;
+
+				GTestAnnotation.AddAnnotation(Object, Temp);
+
+				AutoRTFM::AbortTransaction();
+			}));
+
+		REQUIRE(32 == GTestAnnotation.GetAnnotation(Object).TestAnnotationNumber);
+	}
+}
diff --git a/Engine/Source/Runtime/CoreUObject/Public/UObject/UObjectAnnotation.h b/Engine/Source/Runtime/CoreUObject/Public/UObject/UObjectAnnotation.h
--- a/Engine/Source/Runtime/CoreUObject/Public/UObject/UObjectAnnotation.h
+++ b/Engine/Source/Runtime/CoreUObject/Public/UObject/UObjectAnnotation.h
@@ -88,11 +88,21 @@ private:
 		{
 			bool bWasEmpty = false;
 			{
-				FScopeLock AnnotationMapLock(&AnnotationMapCritical);
-				AnnotationCacheKey = Object;
-				AnnotationCacheValue = MoveTemp(LocalAnnotation);
-				bWasEmpty = (AnnotationMap.Num() == 0);
-				AnnotationMap.Add(AnnotationCacheKey, AnnotationCacheValue);
+				// We hit this from many locations, for now lets open around adding this, and if we happen to be in a transcation
+				// and abort after we have added let us use the UObject to remove it from the annotation map
+				UE_AUTORTFM_OPEN(
+				{
+					FScopeLock AnnotationMapLock(&AnnotationMapCritical);
+					AnnotationCacheKey = Object;
+					AnnotationCacheValue = MoveTemp(LocalAnnotation);
+					bWasEmpty = (AnnotationMap.Num() == 0);
+					AnnotationMap.Add(AnnotationCacheKey, AnnotationCacheValue);
+				});
+
+				UE_AUTORTFM_ONABORT(
+				{
+					RemoveAnnotation(Object);
+				});
 			}
 
 			if (bWasEmpty)
@@ -102,7 +112,15 @@ private:
 				if (bAutoRemove)
 #endif
 				{
-					GUObjectArray.AddUObjectDeleteListener(this);
+					UE_AUTORTFM_OPEN(
+					{
+						GUObjectArray.AddUObjectDeleteListener(this);
+					});
+
+					UE_AUTORTFM_ONABORT(
+					{
+						GUObjectArray.RemoveUObjectDeleteListener(this);
+					});
 				}
 			}
 		}
@@ -165,28 +183,31 @@ public:
 	 */
 	void RemoveAnnotation(const UObjectBase *Object)
 	{
-		check(Object);
-		bool bHadElements = false;
-		bool bIsNowEmpty = false;
-		// Avoid holding the lock while we call GUObjectArray.RemoveUObjectDeleteListener as it could deadlock
+		AutoRTFM::OnCommit([this, Object]
 		{
-			FScopeLock AnnotationMapLock(&AnnotationMapCritical);
-			AnnotationCacheKey = Object;
-			AnnotationCacheValue = TAnnotation();
-			bHadElements = (AnnotationMap.Num() > 0);
-			AnnotationMap.Remove(AnnotationCacheKey);
-			bIsNowEmpty = (AnnotationMap.Num() == 0);
-		}
-		if (bHadElements && bIsNowEmpty)
-		{
-			// we are removing the last one, so if we are auto removing or verifying removal, unregister now
+			check(Object);
+			bool bHadElements = false;
+			bool bIsNowEmpty = false;
+			// Avoid holding the lock while we call GUObjectArray.RemoveUObjectDeleteListener as it could deadlock
+			{
+				FScopeLock AnnotationMapLock(&AnnotationMapCritical);
+				AnnotationCacheKey = Object;
+				AnnotationCacheValue = TAnnotation();
+				bHadElements = (AnnotationMap.Num() > 0);
+				AnnotationMap.Remove(AnnotationCacheKey);
+				bIsNowEmpty = (AnnotationMap.Num() == 0);
+			}
+			if (bHadElements && bIsNowEmpty)
+			{
+				// we are removing the last one, so if we are auto removing or verifying removal, unregister now
 #if (UE_BUILD_SHIPPING || UE_BUILD_TEST)
-			if (bAutoRemove)
+				if (bAutoRemove)
 #endif
-			{
-				GUObjectArray.RemoveUObjectDeleteListener(this);
+				{
+					GUObjectArray.RemoveUObjectDeleteListener(this);
+				}
 			}
-		}
+		});
 	}
 	/**
 	 * Removes all annotation from the annotation list. 
