Added verbose logging for localized package discovery and resolution

#rnx

diff --git a/Engine/Source/Runtime/CoreUObject/Private/Internationalization/PackageLocalizationCache.cpp b/Engine/Source/Runtime/CoreUObject/Private/Internationalization/PackageLocalizationCache.cpp
--- a/Engine/Source/Runtime/CoreUObject/Private/Internationalization/PackageLocalizationCache.cpp
+++ b/Engine/Source/Runtime/CoreUObject/Private/Internationalization/PackageLocalizationCache.cpp
@@ -7,7 +7,7 @@
 #include "Misc/PackageName.h"
 #include "Misc/ConfigCacheIni.h"
 
-DEFINE_LOG_CATEGORY_STATIC(LogPackageLocalizationCache, Log, All);
+DEFINE_LOG_CATEGORY(LogPackageLocalizationCache);
 
 FPackageLocalizationCultureCache::FPackageLocalizationCultureCache(FPackageLocalizationCache* InOwnerCache, const FString& InCultureName)
 	: OwnerCache(InOwnerCache)
@@ -49,6 +49,7 @@ void FPackageLocalizationCultureCache::ConditionalUpdateCache_NoLock()
 				LocalizedRootPaths.Add(LocalizedRootPath);
 				NewSourceToLocalizedPaths.FindOrAdd(SourceRootPath).Add(LocalizedRootPath);
 			}
+			UE_LOG(LogPackageLocalizationCache, Verbose, TEXT("Processing localized package path '%s'..."), *LocalizedRootPath);
 		}
 	}
 	OwnerCache->FindLocalizedPackages(NewSourceToLocalizedPaths, SourcePackagesToLocalizedPackages);
@@ -169,7 +170,11 @@ FName FPackageLocalizationCultureCache::FindLocalizedPackageName(const FName InS
 	ConditionalUpdateCache_NoLock();
 
 	const TArray<FName>* const FoundPrioritizedLocalizedPackageNames = SourcePackagesToLocalizedPackages.Find(InSourcePackageName);
-	return (FoundPrioritizedLocalizedPackageNames) ? (*FoundPrioritizedLocalizedPackageNames)[0] : NAME_None;
+	const FName LocalizedPackageName = (FoundPrioritizedLocalizedPackageNames) ? (*FoundPrioritizedLocalizedPackageNames)[0] : NAME_None;
+
+	UE_CLOG(!LocalizedPackageName.IsNone(), LogPackageLocalizationCache, Verbose, TEXT("Found localized package '%s' for source package '%s'"), *LocalizedPackageName.ToString(), *InSourcePackageName.ToString());
+
+	return LocalizedPackageName;
 }
 
 FPackageLocalizationCache::FPackageLocalizationCache()
diff --git a/Engine/Source/Runtime/CoreUObject/Public/Internationalization/PackageLocalizationCache.h b/Engine/Source/Runtime/CoreUObject/Public/Internationalization/PackageLocalizationCache.h
--- a/Engine/Source/Runtime/CoreUObject/Public/Internationalization/PackageLocalizationCache.h
+++ b/Engine/Source/Runtime/CoreUObject/Public/Internationalization/PackageLocalizationCache.h
@@ -15,6 +15,8 @@
 
 class FPackageLocalizationCache;
 
+COREUOBJECT_API DECLARE_LOG_CATEGORY_EXTERN(LogPackageLocalizationCache, Log, All);
+
 /** Package localization cache for a specific culture (may contain a chain of cultures internally based on their priority) */
 class FPackageLocalizationCultureCache
 {
diff --git a/Engine/Source/Runtime/Engine/Private/Internationalization/EnginePackageLocalizationCache.cpp b/Engine/Source/Runtime/Engine/Private/Internationalization/EnginePackageLocalizationCache.cpp
--- a/Engine/Source/Runtime/Engine/Private/Internationalization/EnginePackageLocalizationCache.cpp
+++ b/Engine/Source/Runtime/Engine/Private/Internationalization/EnginePackageLocalizationCache.cpp
@@ -82,6 +82,8 @@ void FEnginePackageLocalizationCache::FindLocalizedPackages(const TMap<FString,
 
 		TArray<FName>& PrioritizedLocalizedPackageNames = InOutSourcePackagesToLocalizedPackages.FindOrAdd(SourcePackageName);
 		PrioritizedLocalizedPackageNames.AddUnique(LocalizedAssetData.PackageName);
+
+		UE_LOG(LogPackageLocalizationCache, Verbose, TEXT("Discovered localized package '%s' for source package '%s'"), *LocalizedAssetData.PackageName.ToString(), *SourcePackageName.ToString());
 	}
 }
 
