Disable the overridable serialization in non editor builds

diff --git a/Engine/Source/Runtime/CoreUObject/Private/UObject/OverridableManager.cpp b/Engine/Source/Runtime/CoreUObject/Private/UObject/OverridableManager.cpp
--- a/Engine/Source/Runtime/CoreUObject/Private/UObject/OverridableManager.cpp
+++ b/Engine/Source/Runtime/CoreUObject/Private/UObject/OverridableManager.cpp
@@ -18,21 +18,30 @@ FOverridableManager& FOverridableManager::Get()
 
 bool FOverridableManager::IsEnabled(const UObject& Object)
 {
+#ifdef WITH_EDITORONLY_DATA
 	return OverriddenObjectAnnotations.IsEnabled(Object);
+#else
+	return false;
+#endif // WITH_EDITORONLY_DATA
 }
 
 void FOverridableManager::Enable(UObject& Object)
 {
+#ifdef WITH_EDITORONLY_DATA
 	OverriddenObjectAnnotations.FindOrAdd(Object);
+#endif // WITH_EDITORONLY_DATA
 }
 
 void FOverridableManager::Disable(UObject& Object)
 {
+#ifdef WITH_EDITORONLY_DATA
 	OverriddenObjectAnnotations.RemoveAnnotation(&Object);
+#endif // WITH_EDITORONLY_DATA
 }
 
 void FOverridableManager::InheritEnabledFrom(UObject& Object, const UObject* DefaultData)
 {
+#ifdef WITH_EDITORONLY_DATA
 	if (!OverriddenObjectAnnotations.IsEnabled(Object))
 	{
 		const UObject* Outer = Object.GetOuter();
@@ -41,37 +50,55 @@ void FOverridableManager::InheritEnabledFrom(UObject& Object, const UObject* Def
 			Enable(Object);
 		}
 	}
+#endif // WITH_EDITORONLY_DATA
 }
 
 bool FOverridableManager::NeedSubObjectTemplateInstantiation(const UObject& Object)
 {
+#ifdef WITH_EDITORONLY_DATA
 	if( const FOverriddenPropertySet* OverriddenProperties = OverriddenObjectAnnotations.Find(Object))
 	{
 		return OverriddenProperties->bNeedsSubobjectTemplateInstantiation;
 	}
+#endif // WITH_EDITORONLY_DATA
 	return false;
 }
 
 FOverriddenPropertySet* FOverridableManager::GetOverriddenProperties(UObject& Object)
 {
+#ifdef WITH_EDITORONLY_DATA
 	return OverriddenObjectAnnotations.Find(Object);
+#else
+	return nullptr;
+#endif // WITH_EDITORONLY_DATA
 }
 
 const FOverriddenPropertySet* FOverridableManager::GetOverriddenProperties(const UObject& Object)
 {
+#ifdef WITH_EDITORONLY_DATA
 	return OverriddenObjectAnnotations.Find(Object);
+#else
+	return nullptr;
+#endif // WITH_EDITORONLY_DATA
 }
 
 FOverriddenPropertySet& FOverridableManager::SetOverriddenProperties(UObject& Object, EOverriddenPropertyOperation Operation)
 {
+#ifdef WITH_EDITORONLY_DATA
 	FOverriddenPropertySet& ObjectOverriddenProperties = OverriddenObjectAnnotations.FindOrAdd(Object);
 	ObjectOverriddenProperties.Reset();
 	ObjectOverriddenProperties.SetOverriddenPropertyOperation(Operation, /*CurrentPropertyChain*/nullptr, /*Property*/nullptr);
 	return ObjectOverriddenProperties;
+#else
+	static FOverriddenPropertySet Dummy;
+	return Dummy;
+#endif // WITH_EDITORONLY_DATA
+
 }
 
 EOverriddenState FOverridableManager::GetOverriddenState(UObject& Object)
 {
+#ifdef WITH_EDITORONLY_DATA
 	if(const FOverriddenPropertySet* OverriddenProperties = GetOverriddenProperties(Object))
 	{
 		// Consider any object that its template is a CDO as added.
@@ -103,28 +130,34 @@ EOverriddenState FOverridableManager::GetOverriddenState(UObject& Object)
 			}
 		}
 	}
+#endif // WITH_EDITORONLY_DATA
 	return EOverriddenState::NoOverrides;
 }
 
 void FOverridableManager::OverrideObject(UObject& Object)
 {
+#ifdef WITH_EDITORONLY_DATA
 	if (FOverriddenPropertySet* ThisObjectOverriddenProperties = OverriddenObjectAnnotations.Find(Object))
 	{
 		// Passing no property node means we are overriding the object itself
 		ThisObjectOverriddenProperties->OverrideProperty(FPropertyChangedEvent(nullptr), /*PropertyNode*/nullptr, &Object);
 	}
+#endif // WITH_EDITORONLY_DATA
 }
 
 void FOverridableManager::OverrideInstancedSubObject(UObject& Object, UObject& InstancedSubObject)
 {
+#ifdef WITH_EDITORONLY_DATA
 	if (InstancedSubObject.IsIn(&Object))
 	{
 		OverrideObject(InstancedSubObject);
 	}
+#endif // WITH_EDITORONLY_DATA
 }
 
 void FOverridableManager::PropagateOverrideToInstancedSubObjects(UObject& Object)
 {
+#ifdef WITH_EDITORONLY_DATA
 	TSet<UObject*> InstancedSubObjects;
 	FFindInstancedReferenceSubobjectHelper::GetInstancedSubObjects(&Object, InstancedSubObjects);
 	for (UObject* InstancedSubObject : InstancedSubObjects)
@@ -132,77 +165,95 @@ void FOverridableManager::PropagateOverrideToInstancedSubObjects(UObject& Object
 		checkf(InstancedSubObject, TEXT("Expecting non null SubObjects"));
 		OverrideInstancedSubObject(Object, *InstancedSubObject);
 	}
+#endif // WITH_EDITORONLY_DATA
 }
 
 void FOverridableManager::OverrideProperty(UObject& Object, const FPropertyChangedEvent& PropertyEvent, const FEditPropertyChain& PropertyChain)
 {
+#ifdef WITH_EDITORONLY_DATA
 	if (FOverriddenPropertySet* ThisObjectOverriddenProperties = OverriddenObjectAnnotations.Find(Object))
 	{
 		ThisObjectOverriddenProperties->OverrideProperty(PropertyEvent, PropertyChain.GetActiveMemberNode() ? PropertyChain.GetActiveMemberNode() : PropertyChain.GetHead(), &Object);
 	}
+#endif // WITH_EDITORONLY_DATA
 }
 
 bool FOverridableManager::ClearOverriddenProperty(UObject& Object, const FPropertyChangedEvent& PropertyEvent, const FEditPropertyChain::TDoubleLinkedListNode* PropertyNode)
 {
+#ifdef WITH_EDITORONLY_DATA
 	if (FOverriddenPropertySet* ThisObjectOverriddenProperties = OverriddenObjectAnnotations.Find(Object))
 	{
 		return ThisObjectOverriddenProperties->ClearOverriddenProperty(PropertyEvent, PropertyNode);
 	}
 	return false;
+#endif // WITH_EDITORONLY_DATA
 }
 
 void FOverridableManager::PreOverrideProperty(UObject& Object, const FEditPropertyChain& PropertyChain)
 {
+#ifdef WITH_EDITORONLY_DATA
 	if (FOverriddenPropertySet* ThisObjectOverriddenProperties = OverriddenObjectAnnotations.Find(Object))
 	{
 		ThisObjectOverriddenProperties->NotifyPropertyChange(EPropertyNotificationType::PreEdit, FPropertyChangedEvent(nullptr), PropertyChain.GetActiveMemberNode() ? PropertyChain.GetActiveMemberNode() : PropertyChain.GetHead(), &Object);
 	}
+#endif // WITH_EDITORONLY_DATA
 }
 
 void FOverridableManager::PostOverrideProperty(UObject& Object, const FPropertyChangedEvent& PropertyEvent, const FEditPropertyChain& PropertyChain)
 {
+#ifdef WITH_EDITORONLY_DATA
 	if (FOverriddenPropertySet* ThisObjectOverriddenProperties = OverriddenObjectAnnotations.Find(Object))
 	{
 		ThisObjectOverriddenProperties->NotifyPropertyChange(EPropertyNotificationType::PostEdit, PropertyEvent, PropertyChain.GetActiveMemberNode() ? PropertyChain.GetActiveMemberNode() : PropertyChain.GetHead(), &Object);
 	}
+#endif // WITH_EDITORONLY_DATA
 }
 
 void FOverridableManager::NotifyPropertyChange(const EPropertyNotificationType Notification, UObject& Object, const FPropertyChangedEvent& PropertyEvent, const FEditPropertyChain::TDoubleLinkedListNode* PropertyNode)
 {
+#ifdef WITH_EDITORONLY_DATA
 	if (FOverriddenPropertySet* ThisObjectOverriddenProperties = OverriddenObjectAnnotations.Find(Object))
 	{
 		ThisObjectOverriddenProperties->NotifyPropertyChange(Notification, PropertyEvent, PropertyNode, &Object);
 	}
+#endif // WITH_EDITORONLY_DATA
 }
 
 EOverriddenPropertyOperation FOverridableManager::GetOverriddenPropertyOperation(UObject& Object, const FPropertyChangedEvent& PropertyEvent, const FEditPropertyChain::TDoubleLinkedListNode* PropertyNode, bool* bOutInheritedOperation)
 {
+#ifdef WITH_EDITORONLY_DATA
 	if (const FOverriddenPropertySet* ThisObjectOverriddenProperties = OverriddenObjectAnnotations.Find(Object))
 	{
 		return ThisObjectOverriddenProperties->GetOverriddenPropertyOperation(PropertyEvent, PropertyNode, bOutInheritedOperation);
 	}
 	return EOverriddenPropertyOperation::None;
+#endif // WITH_EDITORONLY_DATA
 }
 
 void FOverridableManager::ClearOverrides(UObject& Object)
 {
+#ifdef WITH_EDITORONLY_DATA
 	if(FOverriddenPropertySet* ThisObjectOverriddenProperties = OverriddenObjectAnnotations.Find(Object))
 	{
 		ThisObjectOverriddenProperties->Reset();
 	}
 	PropagateClearOverridesToInstancedSubObjects(Object);
+#endif // WITH_EDITORONLY_DATA
 }
 
 void FOverridableManager::ClearInstancedSubObjectOverrides(UObject& Object, UObject& InstancedSubObject)
 {
+#ifdef WITH_EDITORONLY_DATA
 	if (InstancedSubObject.IsIn(&Object))
 	{
 		ClearOverrides(InstancedSubObject);
 	}
+#endif // WITH_EDITORONLY_DATA
 }
 
 void FOverridableManager::PropagateClearOverridesToInstancedSubObjects(UObject& Object)
 {
+#ifdef WITH_EDITORONLY_DATA
 	TSet<UObject*> InstancedSubObjects;
 	FFindInstancedReferenceSubobjectHelper::GetInstancedSubObjects(&Object, InstancedSubObjects);
 	for (UObject* InstancedSubObject : InstancedSubObjects)
@@ -212,10 +263,12 @@ void FOverridableManager::PropagateClearOverridesToInstancedSubObjects(UObject&
 		// There are some cases where the property has information about that should be an instanced subobject, but it is not owned by us.
 		ClearInstancedSubObjectOverrides(Object, *InstancedSubObject);
 	}
+#endif // WITH_EDITORONLY_DATA
 }
 
 void FOverridableManager::SerializeOverriddenProperties(UObject& Object, FStructuredArchive::FRecord ObjectRecord)
 {
+#ifdef WITH_EDITORONLY_DATA
 	const FArchiveState& ArchiveState = ObjectRecord.GetArchiveState();
 	FOverriddenPropertySet* OverriddenProperties = ArchiveState.IsSaving() ? GetOverriddenProperties(Object) : nullptr;
 	TOptional<FStructuredArchiveSlot> OverridenPropertiesSlot = ObjectRecord.TryEnterField(TEXT("OverridenProperties"), OverriddenProperties != nullptr);
@@ -238,6 +291,7 @@ void FOverridableManager::SerializeOverriddenProperties(UObject& Object, FStruct
 	{
 		Disable(Object);
 	}
+#endif // WITH_EDITORONLY_DATA
 }
 
 FOverridableManager::FOverridableManager()
@@ -249,6 +303,7 @@ FOverridableManager::FOverridableManager()
 
 void FOverridableManager::HandleObjectsReInstantiated(const TMap<UObject*, UObject*>& OldToNewInstanceMap)
 {
+#ifdef WITH_EDITORONLY_DATA
 	const TMap<const UObjectBase *, FOverriddenPropertyAnnotation>& AnnotationMap = OverriddenObjectAnnotations.GetAnnotationMap();
 	for (const auto& Pair : AnnotationMap)
 	{
@@ -257,4 +312,5 @@ void FOverridableManager::HandleObjectsReInstantiated(const TMap<UObject*, UObje
 			OverridenProperties->HandleObjectsReInstantiated(OldToNewInstanceMap);
 		}
 	}
+#endif // WITH_EDITORONLY_DATA
 }
diff --git a/Engine/Source/Runtime/CoreUObject/Public/UObject/OverridableManager.h b/Engine/Source/Runtime/CoreUObject/Public/UObject/OverridableManager.h
--- a/Engine/Source/Runtime/CoreUObject/Public/UObject/OverridableManager.h
+++ b/Engine/Source/Runtime/CoreUObject/Public/UObject/OverridableManager.h
@@ -13,6 +13,8 @@
  *************************************************************************************
  */
 
+#ifdef WITH_EDITORONLY_DATA
+
 /**
  * Struct holding the shared ptr of the overridden properties
  */
@@ -79,6 +81,8 @@ private:
 	}
 };
 
+#endif // WITH_EDITORONLY_DATA
+
 enum class EOverriddenState : uint8
 {
 	NoOverrides, // no on this object and any of its instanced subobjects
@@ -268,6 +272,8 @@ public:
 protected:
 	FOverridableManager();
 
+#ifdef WITH_EDITORONLY_DATA
 	FOverriddenPropertyAnnotations OverriddenObjectAnnotations;
+#endif // WITH_EDITORONLY_DATA 
 };
 
