Enabled unknown property tracking for maps and sets

#rnx

diff --git a/Engine/Source/Runtime/CoreUObject/Private/UObject/PropertyMap.cpp b/Engine/Source/Runtime/CoreUObject/Private/UObject/PropertyMap.cpp
--- a/Engine/Source/Runtime/CoreUObject/Private/UObject/PropertyMap.cpp
+++ b/Engine/Source/Runtime/CoreUObject/Private/UObject/PropertyMap.cpp
@@ -715,13 +715,6 @@ void FMapProperty::SerializeItem(FStructuredArchive::FSlot Slot, void* Value, co
 				FMemory::Free(TempKeyValueStorage);
 			}
 
-			// Disable serialization of unknown properties until the TODO in the loop is addressed.
-			TOptional<TGuardValue<bool>> SerializeUnknownProperty;
-			if (Context)
-			{
-				SerializeUnknownProperty.Emplace(Context->bSerializeUnknownProperty, false);
-			}
-
 			int32 NumEntries = 0;
 			FStructuredArchive::FArray EntriesArray = Record.EnterArray(TEXT("Entries"), NumEntries);
 
@@ -785,13 +778,6 @@ void FMapProperty::SerializeItem(FStructuredArchive::FSlot Slot, void* Value, co
 				}
 			}
 
-			// Disable serialization of unknown properties until the TODO in the loop is addressed.
-			TOptional<TGuardValue<bool>> SerializeUnknownProperty;
-			if (Context)
-			{
-				SerializeUnknownProperty.Emplace(Context->bSerializeUnknownProperty, false);
-			}
-
 			int32 NumEntries = 0;
 			FStructuredArchive::FArray EntriesArray = Record.EnterArray(TEXT("Entries"), NumEntries);
 
@@ -805,8 +791,6 @@ void FMapProperty::SerializeItem(FStructuredArchive::FSlot Slot, void* Value, co
 			// Read remaining items into container
 			for (; NumEntries; --NumEntries)
 			{
-				// TODO: SetIndex on Context->SerializedPropertyPath and remove the element from the bag later if it existed.
-
 				FStructuredArchive::FRecord EntryRecord = EntriesArray.EnterElement().EnterRecord();
 
 				// Read key into temporary storage
diff --git a/Engine/Source/Runtime/CoreUObject/Private/UObject/PropertySet.cpp b/Engine/Source/Runtime/CoreUObject/Private/UObject/PropertySet.cpp
--- a/Engine/Source/Runtime/CoreUObject/Private/UObject/PropertySet.cpp
+++ b/Engine/Source/Runtime/CoreUObject/Private/UObject/PropertySet.cpp
@@ -342,19 +342,10 @@ void FSetProperty::SerializeItem(FStructuredArchive::FSlot Slot, void* Value, co
 				ElementProp->InitializeValue(TempElementStorage);
 			}
 
-			// Disable serialization of unknown properties until the TODO in the loop is addressed.
-			TOptional<TGuardValue<bool>> SerializeUnknownProperty;
-			if (Context)
-			{
-				SerializeUnknownProperty.Emplace(Context->bSerializeUnknownProperty, false);
-			}
-
 			FSerializedPropertyScope SerializedProperty(UnderlyingArchive, ElementProp, this);
 			// Read remaining items into container
 			for (; Num; --Num)
 			{
-				// TODO: SetIndex on Context->SerializedPropertyPath and remove the element from the bag later if it existed.
-
 				// Read key into temporary storage
 				ElementProp->SerializeItem(ElementsArray.EnterElement(), TempElementStorage);
 
