Added FUtf8StrProperty and FAnsiStrProperty.

diff --git a/Engine/Source/Runtime/Core/Public/UObject/UObjectHierarchyFwd.h b/Engine/Source/Runtime/Core/Public/UObject/UObjectHierarchyFwd.h
--- a/Engine/Source/Runtime/Core/Public/UObject/UObjectHierarchyFwd.h
+++ b/Engine/Source/Runtime/Core/Public/UObject/UObjectHierarchyFwd.h
@@ -52,3 +52,5 @@ class		FMulticastDelegateProperty;
 class		FMapProperty;
 class		FSetProperty;
 class		FEnumProperty;
+class		FUtf8StrProperty;
+class		FAnsiStrProperty;
diff --git a/Engine/Source/Runtime/Core/Public/UObject/UnrealNames.inl b/Engine/Source/Runtime/Core/Public/UObject/UnrealNames.inl
--- a/Engine/Source/Runtime/Core/Public/UObject/UnrealNames.inl
+++ b/Engine/Source/Runtime/Core/Public/UObject/UnrealNames.inl
@@ -51,6 +51,8 @@ REGISTER_NAME(33,CoreUObject)
 // More class properties
 REGISTER_NAME(34,EnumProperty)
 REGISTER_NAME(35,OptionalProperty)
+REGISTER_NAME(36,Utf8StrProperty)
+REGISTER_NAME(37,AnsiStrProperty)
 
 
 // Special types.
diff --git a/Engine/Source/Runtime/CoreUObject/Private/UObject/AnsiStrProperty.cpp b/Engine/Source/Runtime/CoreUObject/Private/UObject/AnsiStrProperty.cpp
new file mode 100644
--- /dev/null
+++ b/Engine/Source/Runtime/CoreUObject/Private/UObject/AnsiStrProperty.cpp
@@ -0,0 +1,16 @@
+// Copyright Epic Games, Inc. All Rights Reserved.
+
+#include "UObject/AnsiStrProperty.h"
+
+// Include StrProperty.cpp.inl's includes before defining the macros, in case the macros 'poison' other headers or there are re-entrant includes.
+#include "UObject/StrPropertyIncludes.cpp.inl"
+
+#define UE_STRPROPERTY_CLASS                FAnsiStrProperty
+#define UE_STRPROPERTY_STRINGTYPE           FAnsiString
+#define UE_STRPROPERTY_CASTCLASSFLAG        CASTCLASS_FAnsiStrProperty
+#define UE_STRPROPERTY_PROPERTYPARAMSSTRUCT FAnsiStrPropertyParams
+	#include "UObject/StrProperty.cpp.inl"
+#undef UE_STRPROPERTY_PROPERTYPARAMSSTRUCT
+#undef UE_STRPROPERTY_CASTCLASSFLAG
+#undef UE_STRPROPERTY_STRINGTYPE
+#undef UE_STRPROPERTY_CLASS
diff --git a/Engine/Source/Runtime/CoreUObject/Private/UObject/PropertyStr.cpp b/Engine/Source/Runtime/CoreUObject/Private/UObject/PropertyStr.cpp
deleted file mode 100644
diff --git a/Engine/Source/Runtime/CoreUObject/Private/UObject/StrProperty.cpp b/Engine/Source/Runtime/CoreUObject/Private/UObject/StrProperty.cpp
new file mode 100644
--- /dev/null
+++ b/Engine/Source/Runtime/CoreUObject/Private/UObject/StrProperty.cpp
@@ -0,0 +1,16 @@
+// Copyright Epic Games, Inc. All Rights Reserved.
+
+#include "UObject/StrProperty.h"
+
+// Include StrProperty.cpp.inl's includes before defining the macros, in case the macros 'poison' other headers or there are re-entrant includes.
+#include "UObject/StrPropertyIncludes.cpp.inl"
+
+#define UE_STRPROPERTY_CLASS                FStrProperty
+#define UE_STRPROPERTY_STRINGTYPE           FString
+#define UE_STRPROPERTY_CASTCLASSFLAG        CASTCLASS_FStrProperty
+#define UE_STRPROPERTY_PROPERTYPARAMSSTRUCT FStrPropertyParams
+	#include "UObject/StrProperty.cpp.inl"
+#undef UE_STRPROPERTY_PROPERTYPARAMSSTRUCT
+#undef UE_STRPROPERTY_CASTCLASSFLAG
+#undef UE_STRPROPERTY_STRINGTYPE
+#undef UE_STRPROPERTY_CLASS
diff --git a/Engine/Source/Runtime/CoreUObject/Private/UObject/StrProperty.cpp.inl b/Engine/Source/Runtime/CoreUObject/Private/UObject/StrProperty.cpp.inl
new file mode 100644
--- /dev/null
+++ b/Engine/Source/Runtime/CoreUObject/Private/UObject/StrProperty.cpp.inl
@@ -0,0 +1,153 @@
+// Copyright Epic Games, Inc. All Rights Reserved.
+
+/*********************************************************************************************************************
+ * NOTICE                                                                                                            *
+ *                                                                                                                   *
+ * This file is not intended to be included directly - it will be included by other .cpp files which have predefined *
+ * some macros to be expanded within this file.  As such, it does not have a #pragma once as it is intended to be    *
+ * included multiple times with different macro definitions.                                                         *
+ *                                                                                                                   *
+ * #includes needed to compile this file need to be specified in StringIncludes.cpp.inl file rather than here.       *
+ *********************************************************************************************************************/
+#define UE_INCLUDETOOL_IGNORE_INCONSISTENT_STATE
+
+#ifndef UE_STRPROPERTY_CLASS
+	#error "StrProperty.cpp.inl should only be included after defining UE_STRPROPERTY_CLASS"
+#endif
+#ifndef UE_STRPROPERTY_STRINGTYPE
+	#error "StrProperty.cpp.inl should only be included after defining UE_STRPROPERTY_STRINGTYPE"
+#endif
+#ifndef UE_STRPROPERTY_CASTCLASSFLAG
+	#error "StrProperty.cpp.inl should only be included after defining UE_STRPROPERTY_CASTCLASSFLAG"
+#endif
+#ifndef UE_STRPROPERTY_PROPERTYPARAMSSTRUCT
+	#error "StrProperty.cpp.inl should only be included after defining UE_STRPROPERTY_PROPERTYPARAMSSTRUCT"
+#endif
+
+IMPLEMENT_FIELD(UE_STRPROPERTY_CLASS)
+
+UE_STRPROPERTY_CLASS::UE_STRPROPERTY_CLASS(FFieldVariant InOwner, const UECodeGen_Private::UE_STRPROPERTY_PROPERTYPARAMSSTRUCT& Prop)
+	: Super(InOwner, (const UECodeGen_Private::FPropertyParamsBaseWithOffset&)Prop)
+{
+}
+
+EConvertFromTypeResult UE_STRPROPERTY_CLASS::ConvertFromType(const FPropertyTag& Tag, FStructuredArchive::FSlot Slot, uint8* Data, UStruct* DefaultsStruct, const uint8* Defaults)
+{
+	// Convert serialized text to string.
+	if (Tag.Type==NAME_TextProperty)
+	{ 
+		FText Text;
+		Slot << Text;
+		const UE_STRPROPERTY_STRINGTYPE String(FTextInspector::GetSourceString(Text) ? *FTextInspector::GetSourceString(Text) : TEXT(""));
+		SetPropertyValue_InContainer(Data, String, Tag.ArrayIndex);
+
+		return EConvertFromTypeResult::Converted;
+	}
+	else if (Tag.Type==NAME_StrProperty)
+	{
+		FString SavedStr;
+		Slot << SavedStr;
+		const UE_STRPROPERTY_STRINGTYPE ConvertedString(SavedStr);
+		SetPropertyValue_InContainer(Data, ConvertedString, Tag.ArrayIndex);
+
+		return EConvertFromTypeResult::Converted;
+	}
+	else if (Tag.Type==NAME_Utf8StrProperty)
+	{
+		FUtf8String SavedStr;
+		Slot << SavedStr;
+		const UE_STRPROPERTY_STRINGTYPE ConvertedString(SavedStr);
+		SetPropertyValue_InContainer(Data, ConvertedString, Tag.ArrayIndex);
+
+		return EConvertFromTypeResult::Converted;
+	}
+	else if (Tag.Type==NAME_AnsiStrProperty)
+	{
+		FAnsiString SavedStr;
+		Slot << SavedStr;
+		const UE_STRPROPERTY_STRINGTYPE ConvertedString(SavedStr);
+		SetPropertyValue_InContainer(Data, ConvertedString, Tag.ArrayIndex);
+
+		return EConvertFromTypeResult::Converted;
+	}
+
+	return EConvertFromTypeResult::UseSerializeItem;
+}
+
+FString UE_STRPROPERTY_CLASS::GetCPPTypeForwardDeclaration() const
+{
+	return FString();
+}
+
+void UE_STRPROPERTY_CLASS::ExportText_Internal( FString& ValueStr, const void* PropertyValueOrContainer, EPropertyPointerType PropertyPointerType, const void* DefaultValue, UObject* Parent, int32 PortFlags, UObject* ExportRootScope ) const
+{
+	UE_STRPROPERTY_STRINGTYPE StringValue;
+	if (PropertyPointerType == EPropertyPointerType::Container && HasGetter())
+	{
+		GetValue_InContainer(PropertyValueOrContainer, &StringValue);
+	}
+	else
+	{
+		StringValue = *(UE_STRPROPERTY_STRINGTYPE*)PointerToValuePtr(PropertyValueOrContainer, PropertyPointerType);
+	}
+
+	if (!(PortFlags & PPF_Delimited))
+	{
+		ValueStr += StringValue;
+	}
+	else if ( StringValue.Len() > 0 )
+	{
+		ValueStr += UE_STRPROPERTY_STRINGTYPE::Printf( CHARTEXT(UE_STRPROPERTY_STRINGTYPE::ElementType, "\"%s\""), *StringValue.ReplaceCharWithEscapedChar() );
+	}
+	else
+	{
+		ValueStr += CHARTEXT(UE_STRPROPERTY_STRINGTYPE::ElementType, "\"\"");
+	}
+}
+const TCHAR* UE_STRPROPERTY_CLASS::ImportText_Internal( const TCHAR* Buffer, void* ContainerOrPropertyPtr, EPropertyPointerType PropertyPointerType, UObject* Parent, int32 PortFlags, FOutputDevice* ErrorText ) const
+{
+	UE_STRPROPERTY_STRINGTYPE ImportedText;
+	if( !(PortFlags & PPF_Delimited) )
+	{
+		ImportedText = Buffer;
+
+		// in order to indicate that the value was successfully imported, advance the buffer past the last character that was imported
+		Buffer += FCString::Strlen(Buffer);
+	}
+	else
+	{
+		// require quoted string here
+		if (*Buffer != TCHAR('"'))
+		{
+			ErrorText->Logf(TEXT("Missing opening '\"' in string property value: %s"), Buffer);
+			return NULL;
+		}
+		const TCHAR* Start = Buffer;
+		FString Temp;
+		Buffer = FPropertyHelpers::ReadToken(Buffer, Temp);
+		if (Buffer == NULL)
+		{
+			return NULL;
+		}
+		if (Buffer > Start && Buffer[-1] != TCHAR('"'))
+		{
+			ErrorText->Logf(TEXT("Missing terminating '\"' in string property value: %s"), Start);
+			return NULL;
+		}
+		ImportedText = UE_STRPROPERTY_STRINGTYPE(MoveTemp(Temp));
+	}
+	if (PropertyPointerType == EPropertyPointerType::Container && HasSetter())
+	{
+		SetValue_InContainer(ContainerOrPropertyPtr, ImportedText);
+	}
+	else
+	{
+		*(UE_STRPROPERTY_STRINGTYPE*)PointerToValuePtr(ContainerOrPropertyPtr, PropertyPointerType) = MoveTemp(ImportedText);
+	}
+	return Buffer;
+}
+
+uint32 UE_STRPROPERTY_CLASS::GetValueTypeHashInternal(const void* Src) const
+{
+	return GetTypeHash(*(const UE_STRPROPERTY_STRINGTYPE*)Src);
+}
diff --git a/Engine/Source/Runtime/CoreUObject/Private/UObject/StrPropertyIncludes.cpp.inl b/Engine/Source/Runtime/CoreUObject/Private/UObject/StrPropertyIncludes.cpp.inl
new file mode 100644
--- /dev/null
+++ b/Engine/Source/Runtime/CoreUObject/Private/UObject/StrPropertyIncludes.cpp.inl
@@ -0,0 +1,26 @@
+// Copyright Epic Games, Inc. All Rights Reserved.
+
+#pragma once
+
+/*******************************************************************************************************
+ * NOTICE                                                                                              *
+ *                                                                                                     *
+ * This file is not intended to be included directly - it is only intended to contain the includes for *
+ * StrPropetryIncludes.cpp.inl.                                                                        *
+ *******************************************************************************************************/
+
+#ifdef UE_STRPROPERTY_CLASS
+	#error "StrPropertyIncludes.cpp.inl should not be included after defining UE_STRPROPERTY_CLASS"
+#endif
+#ifdef UE_STRPROPERTY_STRINGTYPE
+	#error "StrPropertyIncludes.cpp.inl should not be included after defining UE_STRPROPERTY_STRINGTYPE"
+#endif
+
+#include "CoreMinimal.h"
+#include "Containers/AnsiString.h"
+#include "Containers/UnrealString.h"
+#include "Containers/Utf8String.h"
+#include "UObject/ObjectMacros.h"
+#include "UObject/PropertyPortFlags.h"
+#include "Misc/StringBuilder.h"
+#include "Misc/AsciiSet.h"
diff --git a/Engine/Source/Runtime/CoreUObject/Private/UObject/UObjectGlobals.cpp b/Engine/Source/Runtime/CoreUObject/Private/UObject/UObjectGlobals.cpp
--- a/Engine/Source/Runtime/CoreUObject/Private/UObject/UObjectGlobals.cpp
+++ b/Engine/Source/Runtime/CoreUObject/Private/UObject/UObjectGlobals.cpp
@@ -60,6 +60,9 @@
 #include "Modules/ModuleManager.h"
 #include "UObject/EnumProperty.h"
 #include "UObject/TextProperty.h"
+#include "UObject/StrProperty.h"
+#include "UObject/AnsiStrProperty.h"
+#include "UObject/Utf8StrProperty.h"
 #include "UObject/FieldPathProperty.h"
 #include "UObject/MetaData.h"
 #include "HAL/LowLevelMemTracker.h"
@@ -6070,6 +6073,18 @@ namespace UECodeGen_Private
 			}
 			break;
 #endif
+
+			case EPropertyGenFlags::Utf8Str:
+			{
+				NewProp = NewFProperty<FUtf8StrProperty, FUtf8StrPropertyParams>(Outer, *PropBase);
+			}
+			break;
+
+			case EPropertyGenFlags::AnsiStr:
+			{
+				NewProp = NewFProperty<FAnsiStrProperty, FAnsiStrPropertyParams>(Outer, *PropBase);
+			}
+			break;
 		}
 
 		NewProp->ArrayDim = PropBase->ArrayDim;
diff --git a/Engine/Source/Runtime/CoreUObject/Private/UObject/Utf8StrProperty.cpp b/Engine/Source/Runtime/CoreUObject/Private/UObject/Utf8StrProperty.cpp
new file mode 100644
--- /dev/null
+++ b/Engine/Source/Runtime/CoreUObject/Private/UObject/Utf8StrProperty.cpp
@@ -0,0 +1,16 @@
+// Copyright Epic Games, Inc. All Rights Reserved.
+
+#include "UObject/Utf8StrProperty.h"
+
+// Include StrProperty.cpp.inl's includes before defining the macros, in case the macros 'poison' other headers or there are re-entrant includes.
+#include "UObject/StrPropertyIncludes.cpp.inl"
+
+#define UE_STRPROPERTY_CLASS                FUtf8StrProperty
+#define UE_STRPROPERTY_STRINGTYPE           FUtf8String
+#define UE_STRPROPERTY_CASTCLASSFLAG        CASTCLASS_FUtf8StrProperty
+#define UE_STRPROPERTY_PROPERTYPARAMSSTRUCT FUtf8StrPropertyParams
+	#include "UObject/StrProperty.cpp.inl"
+#undef UE_STRPROPERTY_PROPERTYPARAMSSTRUCT
+#undef UE_STRPROPERTY_CASTCLASSFLAG
+#undef UE_STRPROPERTY_STRINGTYPE
+#undef UE_STRPROPERTY_CLASS
diff --git a/Engine/Source/Runtime/CoreUObject/Public/Templates/Casts.h b/Engine/Source/Runtime/CoreUObject/Public/Templates/Casts.h
--- a/Engine/Source/Runtime/CoreUObject/Public/Templates/Casts.h
+++ b/Engine/Source/Runtime/CoreUObject/Public/Templates/Casts.h
@@ -445,6 +445,8 @@ DECLARE_CAST_BY_FLAG(FClassProperty)
 DECLARE_CAST_BY_FLAG(FInterfaceProperty)
 DECLARE_CAST_BY_FLAG(FNameProperty)
 DECLARE_CAST_BY_FLAG(FStrProperty)
+DECLARE_CAST_BY_FLAG(FUtf8StrProperty)
+DECLARE_CAST_BY_FLAG(FAnsiStrProperty)
 DECLARE_CAST_BY_FLAG(FTextProperty)
 DECLARE_CAST_BY_FLAG(FArrayProperty)
 DECLARE_CAST_BY_FLAG(FDelegateProperty)
diff --git a/Engine/Source/Runtime/CoreUObject/Public/UObject/AnsiStrProperty.h b/Engine/Source/Runtime/CoreUObject/Public/UObject/AnsiStrProperty.h
new file mode 100644
--- /dev/null
+++ b/Engine/Source/Runtime/CoreUObject/Public/UObject/AnsiStrProperty.h
@@ -0,0 +1,18 @@
+// Copyright Epic Games, Inc. All Rights Reserved.
+
+#pragma once
+
+// Include StrProperty.h.inl's includes before defining the macros, in case the macros 'poison' other headers or there are re-entrant includes.
+#include "UObject/StrPropertyIncludes.h.inl"
+
+#include "Containers/AnsiString.h"
+
+#define UE_STRPROPERTY_CLASS                FAnsiStrProperty
+#define UE_STRPROPERTY_STRINGTYPE           FAnsiString
+#define UE_STRPROPERTY_CASTCLASSFLAG        CASTCLASS_FAnsiStrProperty
+#define UE_STRPROPERTY_PROPERTYPARAMSSTRUCT FAnsiStrPropertyParams
+	#include "UObject/StrProperty.h.inl"
+#undef UE_STRPROPERTY_PROPERTYPARAMSSTRUCT
+#undef UE_STRPROPERTY_CASTCLASSFLAG
+#undef UE_STRPROPERTY_STRINGTYPE
+#undef UE_STRPROPERTY_CLASS
diff --git a/Engine/Source/Runtime/CoreUObject/Public/UObject/Field.h b/Engine/Source/Runtime/CoreUObject/Public/UObject/Field.h
--- a/Engine/Source/Runtime/CoreUObject/Public/UObject/Field.h
+++ b/Engine/Source/Runtime/CoreUObject/Public/UObject/Field.h
@@ -252,7 +252,7 @@ FField* TClass::Construct(const FFieldVariant& InOwner, const FName& InName, EOb
 } \
 FFieldClass* TClass::StaticClass() \
 { \
-	static FFieldClass StaticFieldClass(TEXT(#TClass), TClass::StaticClassCastFlagsPrivate(), TClass::StaticClassCastFlags(), TClass::Super::StaticClass(), &TClass::Construct); \
+	static FFieldClass StaticFieldClass(TEXT(PREPROCESSOR_TO_STRING(TClass)), TClass::StaticClassCastFlagsPrivate(), TClass::StaticClassCastFlags(), TClass::Super::StaticClass(), &TClass::Construct); \
 	return &StaticFieldClass; \
 } \
 
diff --git a/Engine/Source/Runtime/CoreUObject/Public/UObject/GeneratedCppIncludes.h b/Engine/Source/Runtime/CoreUObject/Public/UObject/GeneratedCppIncludes.h
--- a/Engine/Source/Runtime/CoreUObject/Public/UObject/GeneratedCppIncludes.h
+++ b/Engine/Source/Runtime/CoreUObject/Public/UObject/GeneratedCppIncludes.h
@@ -10,6 +10,9 @@
 #include "UObject/Package.h"
 #include "UObject/UnrealType.h"
 #include "UObject/EnumProperty.h"
+#include "UObject/StrProperty.h"
+#include "UObject/AnsiStrProperty.h"
+#include "UObject/Utf8StrProperty.h"
 #include "UObject/TextProperty.h"
 #include "UObject/FieldPathProperty.h"
 
diff --git a/Engine/Source/Runtime/CoreUObject/Public/UObject/ObjectMacros.h b/Engine/Source/Runtime/CoreUObject/Public/UObject/ObjectMacros.h
--- a/Engine/Source/Runtime/CoreUObject/Public/UObject/ObjectMacros.h
+++ b/Engine/Source/Runtime/CoreUObject/Public/UObject/ObjectMacros.h
@@ -377,6 +377,8 @@ enum EClassCastFlags : uint64
 	CASTCLASS_FOptionalProperty				= 0x0100000000000000,
 	CASTCLASS_FVValueProperty				= 0x0200000000000000,
 	CASTCLASS_FVRestValueProperty			= 0x0400000000000000,
+	CASTCLASS_FUtf8StrProperty				= 0x1000000000000000,
+	CASTCLASS_FAnsiStrProperty				= 0x2000000000000000,
 };
 
 #define CASTCLASS_AllFlags ((EClassCastFlags)0xFFFFFFFFFFFFFFFF)
diff --git a/Engine/Source/Runtime/CoreUObject/Public/UObject/StrProperty.h b/Engine/Source/Runtime/CoreUObject/Public/UObject/StrProperty.h
new file mode 100644
--- /dev/null
+++ b/Engine/Source/Runtime/CoreUObject/Public/UObject/StrProperty.h
@@ -0,0 +1,18 @@
+// Copyright Epic Games, Inc. All Rights Reserved.
+
+#pragma once
+
+// Include StrProperty.h.inl's includes before defining the macros, in case the macros 'poison' other headers or there are re-entrant includes.
+#include "UObject/StrPropertyIncludes.h.inl"
+
+#include "Containers/UnrealString.h"
+
+#define UE_STRPROPERTY_CLASS                FStrProperty
+#define UE_STRPROPERTY_STRINGTYPE           FString
+#define UE_STRPROPERTY_CASTCLASSFLAG        CASTCLASS_FStrProperty
+#define UE_STRPROPERTY_PROPERTYPARAMSSTRUCT FStrPropertyParams
+	#include "UObject/StrProperty.h.inl"
+#undef UE_STRPROPERTY_PROPERTYPARAMSSTRUCT
+#undef UE_STRPROPERTY_CASTCLASSFLAG
+#undef UE_STRPROPERTY_STRINGTYPE
+#undef UE_STRPROPERTY_CLASS
diff --git a/Engine/Source/Runtime/CoreUObject/Public/UObject/StrProperty.h.inl b/Engine/Source/Runtime/CoreUObject/Public/UObject/StrProperty.h.inl
new file mode 100644
--- /dev/null
+++ b/Engine/Source/Runtime/CoreUObject/Public/UObject/StrProperty.h.inl
@@ -0,0 +1,63 @@
+// Copyright Epic Games, Inc. All Rights Reserved.
+
+/*******************************************************************************************************************
+ * NOTICE                                                                                                          *
+ *                                                                                                                 *
+ * This file is not intended to be included directly - it will be included by other .h files which have predefined *
+ * some macros to be expanded within this file.  As such, it does not have a #pragma once as it is intended to be  *
+ * included multiple times with different macro definitions.                                                       *
+ *                                                                                                                 *
+ * #includes needed to compile this file need to be specified in StrPropertyIncludes.h.inl file rather than here.  *
+ *******************************************************************************************************************/
+#define UE_INCLUDETOOL_IGNORE_INCONSISTENT_STATE
+
+#ifndef UE_STRPROPERTY_CLASS
+	#error "StrProperty.h.inl should only be included after defining UE_STRPROPERTY_CLASS"
+#endif
+#ifndef UE_STRPROPERTY_STRINGTYPE
+	#error "StrProperty.h.inl should only be included after defining UE_STRPROPERTY_STRINGTYPE"
+#endif
+#ifndef UE_STRPROPERTY_CASTCLASSFLAG
+	#error "StrProperty.h.inl should only be included after defining UE_STRPROPERTY_CASTCLASSFLAG"
+#endif
+#ifndef UE_STRPROPERTY_PROPERTYPARAMSSTRUCT
+	#error "StrProperty.h.inl should only be included after defining UE_STRPROPERTY_PROPERTYPARAMSSTRUCT"
+#endif
+
+class COREUOBJECT_API UE_STRPROPERTY_CLASS : public TProperty_WithEqualityAndSerializer<UE_STRPROPERTY_STRINGTYPE, FProperty>
+{
+	DECLARE_FIELD(UE_STRPROPERTY_CLASS, (TProperty_WithEqualityAndSerializer<UE_STRPROPERTY_STRINGTYPE, FProperty>), UE_STRPROPERTY_CASTCLASSFLAG)
+public:
+	using TTypeFundamentals = Super::TTypeFundamentals;
+	using TCppType = TTypeFundamentals::TCppType;
+
+	UE_STRPROPERTY_CLASS(FFieldVariant InOwner, const FName& InName, EObjectFlags InObjectFlags)
+		: Super(InOwner, InName, InObjectFlags)
+	{
+	}
+
+	/**
+	 * Constructor used for constructing compiled in properties
+	 * @param InOwner Owner of the property
+	 * @param PropBase Pointer to the compiled in structure describing the property
+	 **/
+	UE_STRPROPERTY_CLASS(FFieldVariant InOwner, const UECodeGen_Private::UE_STRPROPERTY_PROPERTYPARAMSSTRUCT& Prop);
+
+#if WITH_EDITORONLY_DATA
+	explicit UE_STRPROPERTY_CLASS(UField* InField)
+		: Super(InField)
+	{
+	}
+#endif // WITH_EDITORONLY_DATA
+
+	// FProperty interface
+protected:
+	virtual void ExportText_Internal( FString& ValueStr, const void* PropertyValueOrContainer, EPropertyPointerType PropertyPointerType, const void* DefaultValue, UObject* Parent, int32 PortFlags, UObject* ExportRootScope ) const override;
+	virtual const TCHAR* ImportText_Internal(const TCHAR* Buffer, void* ContainerOrPropertyPtr, EPropertyPointerType PropertyPointerType, UObject* OwnerObject, int32 PortFlags, FOutputDevice* ErrorText) const override;
+public:
+	virtual EConvertFromTypeResult ConvertFromType(const FPropertyTag& Tag, FStructuredArchive::FSlot Slot, uint8* Data, UStruct* DefaultsStruct, const uint8* Defaults) override;
+	UE_DEPRECATED(5.4, "UnrealHeaderTool only API.  No replacement available.")
+	virtual FString GetCPPTypeForwardDeclaration() const override;
+	uint32 GetValueTypeHashInternal(const void* Src) const override;
+	// End of FProperty interface
+};
diff --git a/Engine/Source/Runtime/CoreUObject/Public/UObject/StrPropertyIncludes.h.inl b/Engine/Source/Runtime/CoreUObject/Public/UObject/StrPropertyIncludes.h.inl
new file mode 100644
--- /dev/null
+++ b/Engine/Source/Runtime/CoreUObject/Public/UObject/StrPropertyIncludes.h.inl
@@ -0,0 +1,17 @@
+// Copyright Epic Games, Inc. All Rights Reserved.
+
+/*=============================================================================
+	UnrealType.h: Unreal engine base type definitions.
+=============================================================================*/
+
+#pragma once
+
+#include "CoreGlobals.h"
+#include "Serialization/StructuredArchive.h"
+#include "UObject/Field.h"
+#include "UObject/NameTypes.h"
+#include "UObject/ObjectMacros.h"
+#include "UObject/PropertyPortFlags.h"
+#include "UObject/PropertyTag.h"
+#include "UObject/UObjectGlobals.h"
+#include "UObject/UnrealType.h"
diff --git a/Engine/Source/Runtime/CoreUObject/Public/UObject/UObjectGlobals.h b/Engine/Source/Runtime/CoreUObject/Public/UObject/UObjectGlobals.h
--- a/Engine/Source/Runtime/CoreUObject/Public/UObject/UObjectGlobals.h
+++ b/Engine/Source/Runtime/CoreUObject/Public/UObject/UObjectGlobals.h
@@ -3447,6 +3447,8 @@ namespace UECodeGen_Private
 		LargeWorldCoordinatesReal = 0x20,
 		Optional          = 0x21,
 		VValue            = 0x22,
+		Utf8Str           = 0x23,
+		AnsiStr           = 0x24,
 
 		// Property-specific flags
 		NativeBool        = 0x40,
@@ -3824,6 +3826,8 @@ namespace UECodeGen_Private
 	typedef FClassPropertyParams   FClassPtrPropertyParams;
 	typedef FObjectPropertyParams  FSoftObjectPropertyParams;
 	typedef FGenericPropertyParams FVerseValuePropertyParams;
+	typedef FGenericPropertyParams FUtf8StrPropertyParams;
+	typedef FGenericPropertyParams FAnsiStrPropertyParams;
 
 	struct FFunctionParams
 	{
diff --git a/Engine/Source/Runtime/CoreUObject/Public/UObject/UnrealType.h b/Engine/Source/Runtime/CoreUObject/Public/UObject/UnrealType.h
--- a/Engine/Source/Runtime/CoreUObject/Public/UObject/UnrealType.h
+++ b/Engine/Source/Runtime/CoreUObject/Public/UObject/UnrealType.h
@@ -3371,51 +3371,6 @@ public:
 	// End of FProperty interface
 };
 
-/*-----------------------------------------------------------------------------
-	FStrProperty.
------------------------------------------------------------------------------*/
-
-//
-// Describes a dynamic string variable.
-//
-class COREUOBJECT_API FStrProperty : public TProperty_WithEqualityAndSerializer<FString, FProperty>
-{
-	DECLARE_FIELD(FStrProperty, (TProperty_WithEqualityAndSerializer<FString, FProperty>), CASTCLASS_FStrProperty)
-public:
-	using TTypeFundamentals = Super::TTypeFundamentals;
-	using TCppType = TTypeFundamentals::TCppType;
-
-	FStrProperty(FFieldVariant InOwner, const FName& InName, EObjectFlags InObjectFlags)
-		: Super(InOwner, InName, InObjectFlags)
-	{
-	}
-
-	/**
-	 * Constructor used for constructing compiled in properties
-	 * @param InOwner Owner of the property
-	 * @param PropBase Pointer to the compiled in structure describing the property
-	 **/
-	FStrProperty(FFieldVariant InOwner, const UECodeGen_Private::FStrPropertyParams& Prop);
-
-#if WITH_EDITORONLY_DATA
-	explicit FStrProperty(UField* InField)
-		: Super(InField)
-	{
-	}
-#endif // WITH_EDITORONLY_DATA
-
-	// FProperty interface
-protected:
-	virtual void ExportText_Internal( FString& ValueStr, const void* PropertyValueOrContainer, EPropertyPointerType PropertyPointerType, const void* DefaultValue, UObject* Parent, int32 PortFlags, UObject* ExportRootScope ) const override;
-	virtual const TCHAR* ImportText_Internal(const TCHAR* Buffer, void* ContainerOrPropertyPtr, EPropertyPointerType PropertyPointerType, UObject* OwnerObject, int32 PortFlags, FOutputDevice* ErrorText) const override;
-public:
-	virtual EConvertFromTypeResult ConvertFromType(const FPropertyTag& Tag, FStructuredArchive::FSlot Slot, uint8* Data, UStruct* DefaultsStruct, const uint8* Defaults) override;
-	UE_DEPRECATED(5.4, "UnrealHeaderTool only API.  No replacement available.")
-	virtual FString GetCPPTypeForwardDeclaration() const override;
-	uint32 GetValueTypeHashInternal(const void* Src) const override;
-	// End of FProperty interface
-};
-
 /*-----------------------------------------------------------------------------
 	FArrayProperty.
 -----------------------------------------------------------------------------*/
@@ -7246,3 +7201,6 @@ class COREUOBJECT_API UMulticastInlineDelegatePropertyWrapper : public UMulticas
 #if UE_ENABLE_INCLUDE_ORDER_DEPRECATED_IN_5_5
 #include "Templates/IsTriviallyDestructible.h"
 #endif
+
+// This should really be included everywhere that FStrProperty is used, but that would require a lot of fixup
+#include "UObject/StrProperty.h"
diff --git a/Engine/Source/Runtime/CoreUObject/Public/UObject/Utf8StrProperty.h b/Engine/Source/Runtime/CoreUObject/Public/UObject/Utf8StrProperty.h
new file mode 100644
--- /dev/null
+++ b/Engine/Source/Runtime/CoreUObject/Public/UObject/Utf8StrProperty.h
@@ -0,0 +1,18 @@
+// Copyright Epic Games, Inc. All Rights Reserved.
+
+#pragma once
+
+// Include StrProperty.h.inl's includes before defining the macros, in case the macros 'poison' other headers or there are re-entrant includes.
+#include "UObject/StrPropertyIncludes.h.inl"
+
+#include "Containers/Utf8String.h"
+
+#define UE_STRPROPERTY_CLASS                FUtf8StrProperty
+#define UE_STRPROPERTY_STRINGTYPE           FUtf8String
+#define UE_STRPROPERTY_CASTCLASSFLAG        CASTCLASS_FUtf8StrProperty
+#define UE_STRPROPERTY_PROPERTYPARAMSSTRUCT FUtf8StrPropertyParams
+	#include "UObject/StrProperty.h.inl"
+#undef UE_STRPROPERTY_PROPERTYPARAMSSTRUCT
+#undef UE_STRPROPERTY_CASTCLASSFLAG
+#undef UE_STRPROPERTY_STRINGTYPE
+#undef UE_STRPROPERTY_CLASS
diff --git a/Engine/Source/Runtime/XmlParser/Private/XmlFile.cpp b/Engine/Source/Runtime/XmlParser/Private/XmlFile.cpp
--- a/Engine/Source/Runtime/XmlParser/Private/XmlFile.cpp
+++ b/Engine/Source/Runtime/XmlParser/Private/XmlFile.cpp
@@ -346,7 +346,7 @@ void FXmlFile::Tokenize(FStringView Input, TArray<FString>& Tokens)
 		}
 	};
 
-	enum TOKENTYPE { OPERATOR, STRING } Type = STRING;
+	enum TOKENTYPE { UNKNOWN, OPERATOR, STRING } Type = UNKNOWN;
 	bool bInQuote = false;
 
 	const TCHAR* PtrStart = GetData(Input);
@@ -367,88 +367,99 @@ void FXmlFile::Tokenize(FStringView Input, TArray<FString>& Tokens)
 			continue;
 		}
 
-		// Mark the start of a token
-		if(!WorkingToken.Len())
+		switch (Type)
 		{
-			WorkingToken += Ch;
-			if(CheckTagOperator(PtrStart, Ptr, PtrEnd))
+			case UNKNOWN:
 			{
-				// Add the working token if it's final (ie: ends with '>')
-				if (Ch == TCHAR('>'))
+				WorkingToken += Ch;
+				if(CheckTagOperator(PtrStart, Ptr, PtrEnd))
 				{
-					Tokens.Add(MoveTemp(WorkingToken));
-					checkSlow(WorkingToken.Len() == 0);
+					// Add the working token if it's final (ie: ends with '>')
+					if (Ch == TCHAR('>'))
+					{
+						Tokens.Add(MoveTemp(WorkingToken));
+						checkSlow(WorkingToken.Len() == 0);
+					}
+					else
+					{
+						Type = OPERATOR;
+					}
 				}
 				else
 				{
-					Type = OPERATOR;
+					Type = STRING;
 				}
 			}
-			else
-			{
-				Type = STRING;
-			}
+			break;
 
-			continue;
-		}
-
-		// Already in a token, so continue parsing
-		if(Type == OPERATOR)
-		{
-			// Still the tag, so add it to the working token
-			if(CheckTagOperator(PtrStart, Ptr, PtrEnd))
+			case OPERATOR:
 			{
-				WorkingToken += Ch;
+				// Still the tag, so add it to the working token
+				if(CheckTagOperator(PtrStart, Ptr, PtrEnd))
+				{
+					WorkingToken += Ch;
+
+					// Add the working token if it's final (ie: ends with '>')
+					if (Ch == TCHAR('>'))
+					{
+						Tokens.Add(MoveTemp(WorkingToken));
+						checkSlow(WorkingToken.Len() == 0);
+						Type = UNKNOWN;
+					}
+				}
 
-				// Add the working token if it's final (ie: ends with '>')
-				if (Ch == TCHAR('>'))
+				// Not a tag operator anymore, so add the old token and start a new one
+				else
 				{
 					Tokens.Add(MoveTemp(WorkingToken));
 					checkSlow(WorkingToken.Len() == 0);
+					WorkingToken += Ch;
+					Type = STRING;
 				}
 			}
+			break;
 
-			// Not a tag operator anymore, so add the old token and start a new one
-			else
-			{
-				Tokens.Add(MoveTemp(WorkingToken));
-				checkSlow(WorkingToken.Len() == 0);
-				WorkingToken += Ch;
-				Type = STRING;
-			}
-		}
-		else // STRING
-		{
-			if (IsQuote(Ch))
+			case STRING:
 			{
-				bInQuote = !bInQuote;
-			}
-
-			// Still a string. Allow '>' within a string
-			if(!CheckTagOperator(PtrStart, Ptr, PtrEnd) || (bInQuote && Ch == TCHAR('>')))
-			{
-				WorkingToken += Ch;
-			}
-
-			// Moving back to operator
-			else
-			{
-				Tokens.Add(MoveTemp(WorkingToken));
-				checkSlow(WorkingToken.Len() == 0);
-				WorkingToken += Ch;
-				bInQuote = false;
+				if (IsQuote(Ch))
+				{
+					bInQuote = !bInQuote;
+				}
 
-				// Add the working token if it's final (ie: ends with '>')
-				if (Ch == TCHAR('>'))
+				// Still a string. Allow '>' within a string
+				if(!CheckTagOperator(PtrStart, Ptr, PtrEnd) || (bInQuote && Ch == TCHAR('>')))
 				{
-					Tokens.Add(MoveTemp(WorkingToken));
-					checkSlow(WorkingToken.Len() == 0);
+					WorkingToken += Ch;
 				}
+
+				// Moving back to operator
 				else
 				{
-					Type = OPERATOR;
+					Tokens.Add(MoveTemp(WorkingToken));
+					checkSlow(WorkingToken.Len() == 0);
+					WorkingToken += Ch;
+					bInQuote = false;
+
+					// Add the working token if it's final (ie: ends with '>')
+					if (Ch == TCHAR('>'))
+					{
+						Tokens.Add(MoveTemp(WorkingToken));
+						checkSlow(WorkingToken.Len() == 0);
+						Type = UNKNOWN;
+					}
+					else
+					{
+						Type = OPERATOR;
+					}
 				}
 			}
+			break;
+
+			default:
+			{
+				checkSlow("Unhandled state");
+			}
+			break;
 		}
 	}
 }
