Removed null checks on the UObjectSerializeContext

#rnx

diff --git a/Engine/Source/Runtime/CoreUObject/Private/UObject/PropertyArray.cpp b/Engine/Source/Runtime/CoreUObject/Private/UObject/PropertyArray.cpp
--- a/Engine/Source/Runtime/CoreUObject/Private/UObject/PropertyArray.cpp
+++ b/Engine/Source/Runtime/CoreUObject/Private/UObject/PropertyArray.cpp
@@ -623,7 +623,7 @@ void FArrayProperty::SerializeItem(FStructuredArchive::FSlot Slot, void* Value,
 			}
 		}
 
-		if (Context && Context->bTrackSerializedPropertyPath)
+		if (Context->bTrackSerializedPropertyPath)
 		{
 			// Update the path with types from the inner tag if the outer tag is incomplete.
 			if (const int32 SegmentIndex = Context->SerializedPropertyPath.GetSegmentCount() - 1; SegmentIndex >= 0)
diff --git a/Engine/Source/Runtime/CoreUObject/Private/UObject/PropertyHelper.h b/Engine/Source/Runtime/CoreUObject/Private/UObject/PropertyHelper.h
--- a/Engine/Source/Runtime/CoreUObject/Private/UObject/PropertyHelper.h
+++ b/Engine/Source/Runtime/CoreUObject/Private/UObject/PropertyHelper.h
@@ -113,7 +113,7 @@ class FSerializedPropertyPathScope
 public:
 	[[nodiscard]] FSerializedPropertyPathScope(FUObjectSerializeContext* InContext, const FPropertyPathNameSegment& InSegment, ESerializedPropertyPathNotify InNotify)
 	{
-		if (InContext && InContext->bTrackSerializedPropertyPath)
+		if (InContext->bTrackSerializedPropertyPath)
 		{
 			Context = InContext;
 			Context->SerializedPropertyPath.Push(InSegment);
@@ -151,7 +151,7 @@ class FSerializedPropertyPathIndexScope
 public:
 	[[nodiscard]] FSerializedPropertyPathIndexScope(FUObjectSerializeContext* InContext, int32 InIndex, ESerializedPropertyPathNotify InNotify)
 	{
-		if (InContext && InContext->bTrackSerializedPropertyPath)
+		if (InContext->bTrackSerializedPropertyPath)
 		{
 			Context = InContext;
 			Context->SerializedPropertyPath.SetIndex(InIndex);
diff --git a/Engine/Source/Runtime/CoreUObject/Private/UObject/PropertyMap.cpp b/Engine/Source/Runtime/CoreUObject/Private/UObject/PropertyMap.cpp
--- a/Engine/Source/Runtime/CoreUObject/Private/UObject/PropertyMap.cpp
+++ b/Engine/Source/Runtime/CoreUObject/Private/UObject/PropertyMap.cpp
@@ -695,11 +695,7 @@ void FMapProperty::SerializeItem(FStructuredArchive::FSlot Slot, void* Value, co
 		{
 			if (NumKeysToRemove && !bReplaceMap)
 			{
-				TOptional<TGuardValue<bool>> SerializeUnknownProperty;
-				if (Context)
-				{
-					SerializeUnknownProperty.Emplace(Context->bSerializeUnknownProperty, false);
-				}
+				TGuardValue<bool> SerializeUnknownProperty(Context->bSerializeUnknownProperty, false);
 
 				// Load and discard keys to remove, map is empty
 				void* TempKeyValueStorage = FMemory::Malloc(MapLayout.SetLayout.Size);
@@ -753,11 +749,7 @@ void FMapProperty::SerializeItem(FStructuredArchive::FSlot Slot, void* Value, co
 
 			if (NumKeysToRemove)
 			{
-				TOptional<TGuardValue<bool>> SerializeUnknownProperty;
-				if (Context)
-				{
-					SerializeUnknownProperty.Emplace(Context->bSerializeUnknownProperty, false);
-				}
+				TGuardValue<bool> SerializeUnknownProperty(Context->bSerializeUnknownProperty, false);
 
 				TempKeyValueStorage = (uint8*)FMemory::Malloc(MapLayout.SetLayout.Size);
 				KeyProp->InitializeValue(TempKeyValueStorage);
diff --git a/Engine/Source/Runtime/CoreUObject/Private/UObject/PropertySet.cpp b/Engine/Source/Runtime/CoreUObject/Private/UObject/PropertySet.cpp
--- a/Engine/Source/Runtime/CoreUObject/Private/UObject/PropertySet.cpp
+++ b/Engine/Source/Runtime/CoreUObject/Private/UObject/PropertySet.cpp
@@ -262,11 +262,7 @@ void FSetProperty::SerializeItem(FStructuredArchive::FSlot Slot, void* Value, co
 		{
 			if (NumElementsToRemove)
 			{
-				TOptional<TGuardValue<bool>> SerializeUnknownProperty;
-				if (Context)
-				{
-					SerializeUnknownProperty.Emplace(Context->bSerializeUnknownProperty, false);
-				}
+				TGuardValue<bool> SerializeUnknownProperty(Context->bSerializeUnknownProperty, false);
 
 				// Load and discard elements to remove, set is empty
 				void* TempElementStorage = FMemory::Malloc(SetLayout.Size);
@@ -308,11 +304,7 @@ void FSetProperty::SerializeItem(FStructuredArchive::FSlot Slot, void* Value, co
 
 			if (NumElementsToRemove)
 			{
-				TOptional<TGuardValue<bool>> SerializeUnknownProperty;
-				if (Context)
-				{
-					SerializeUnknownProperty.Emplace(Context->bSerializeUnknownProperty, false);
-				}
+				TGuardValue<bool> SerializeUnknownProperty(Context->bSerializeUnknownProperty, false);
 
 				TempElementStorage = (uint8*)FMemory::Malloc(SetLayout.Size);
 				ElementProp->InitializeValue(TempElementStorage);
