Remove unused include in Class.cpp

#rnx

diff --git a/Engine/Source/Runtime/CoreUObject/Private/UObject/Class.cpp b/Engine/Source/Runtime/CoreUObject/Private/UObject/Class.cpp
--- a/Engine/Source/Runtime/CoreUObject/Private/UObject/Class.cpp
+++ b/Engine/Source/Runtime/CoreUObject/Private/UObject/Class.cpp
@@ -36,7 +36,6 @@
 #include "UObject/Interface.h"
 #include "UObject/LinkerPlaceholderClass.h"
 #include "UObject/LinkerPlaceholderFunction.h"
-#include "UObject/PropertyBag.h"
 #include "UObject/PropertyBagRepository.h"
 #include "UObject/PropertyOptional.h"
 #include "UObject/PropertyPathNameTree.h"
