Deleted FPropertyBag and remaining references to it

#rnx

diff --git a/Engine/Source/Runtime/CoreUObject/Internal/UObject/PropertyBag.h b/Engine/Source/Runtime/CoreUObject/Internal/UObject/PropertyBag.h
deleted file mode 100644
diff --git a/Engine/Source/Runtime/CoreUObject/Internal/UObject/PropertyBagRepository.h b/Engine/Source/Runtime/CoreUObject/Internal/UObject/PropertyBagRepository.h
--- a/Engine/Source/Runtime/CoreUObject/Internal/UObject/PropertyBagRepository.h
+++ b/Engine/Source/Runtime/CoreUObject/Internal/UObject/PropertyBagRepository.h
@@ -12,7 +12,6 @@ class UObject;
 namespace UE
 {
 
-class FPropertyBag;
 class FPropertyPathName;
 class FPropertyPathNameTree;
 
@@ -22,8 +21,7 @@ class FPropertyBagRepository : public FGCObject
 	struct FPropertyBagAssociationData
 	{
 		void Destroy();
-		
-		FPropertyBag* Bag = nullptr; // TODO: Ref bags via handle? Store as value?
+
 		FPropertyPathNameTree* Tree = nullptr;
 		TObjectPtr<UObject> InstanceDataObject = nullptr;
 	};
@@ -35,10 +33,9 @@ class FPropertyBagRepository : public FGCObject
 private:
 	friend class FPropertyBagRepositoryLock;
 	mutable FCriticalSection CriticalSection;
-	
+
 	// Lifetimes/ownership:
 	// Managed within UObjectBase and synced with object lifetime. The repo tracks pointers to bags, not the bags themselves.
-	// FPropertyBag destruction handles destruction of sub-bags automatically, via the FPropertyBagProperty tracking them.
 
 	/** Map of objects/subobjects to their top level property bag. */
 	// TODO: Currently will only exist in editor world, but could tracking per world make some sense for teardown in future? We're relying on object destruction to occur properly to free these up. 
@@ -63,10 +60,6 @@ public:
 
 	// Reclaim space - TODO: Hook up to GC.
 	void ShrinkMaps();
-	
-	// TODO: Restrict bag creation to actor creation and UStruct::SerializeVersionedTaggedProperties?
-	// Object owner is tracked internally
-	FPropertyBag* CreateOuterBag(const UObject* Owner);
 
 	/**
 	 * Finds or creates a property path name tree to collect unknown property paths within the owner.
@@ -106,18 +99,14 @@ public:
 	 * @return			- Does the object's InstanceDataObject contain any loose properties requiring user fixup before the object may be published?
 	 */
 	COREUOBJECT_API bool RequiresFixup(const UObject* Object) const;
-	
-	// Accessors
-	COREUOBJECT_API bool HasBag(const UObject* Owner) const;
-	COREUOBJECT_API FPropertyBag* FindBag(const UObject* Owner);
-	COREUOBJECT_API const FPropertyBag* FindBag(const UObject* Owner) const;
 
+	// Accessors
 	COREUOBJECT_API bool HasInstanceDataObject(const UObject* Owner) const;
 	COREUOBJECT_API UObject* FindInstanceDataObject(const UObject* Owner);
 	COREUOBJECT_API const UObject* FindInstanceDataObject(const UObject* Owner) const;
 
 	COREUOBJECT_API const UObject* FindInstanceForDataObject(const UObject* InstanceDataObject) const;
-	
+
 	// query whether a property in an object was set when the object was deserialized
 	COREUOBJECT_API static bool WasPropertySetBySerialization(UObject* Object, const FPropertyPathName& Path);
 	// query whether a property in Struct was set when the struct was deserialized
@@ -165,7 +154,7 @@ private:
 
 	// Delete owner reference and disassociate all data. Returns success.
 	bool RemoveAssociationUnsafe(const UObject* Owner);
-	
+
 	// Instantiate InstanceDataObject within BagData. Returns InstanceDataObject object. 
 	void CreateInstanceDataObjectUnsafe(UObject* Owner, FPropertyBagAssociationData& BagData, FArchive* Archive = nullptr);
 };
diff --git a/Engine/Source/Runtime/CoreUObject/Private/UObject/InstanceDataObjectUtils.cpp b/Engine/Source/Runtime/CoreUObject/Private/UObject/InstanceDataObjectUtils.cpp
--- a/Engine/Source/Runtime/CoreUObject/Private/UObject/InstanceDataObjectUtils.cpp
+++ b/Engine/Source/Runtime/CoreUObject/Private/UObject/InstanceDataObjectUtils.cpp
@@ -7,9 +7,8 @@
 #include "UObject/Class.h"
 #include "UObject/EnumProperty.h"
 #include "UObject/Field.h"
-#include "UObject/PropertyBag.h"
-#include "UObject/PropertyHelper.h"
 #include "UObject/PropertyBagRepository.h"
+#include "UObject/PropertyHelper.h"
 #include "UObject/PropertyOptional.h"
 #include "UObject/PropertyPathNameTree.h"
 #include "UObject/UnrealType.h"
diff --git a/Engine/Source/Runtime/CoreUObject/Private/UObject/PropertyBag.cpp b/Engine/Source/Runtime/CoreUObject/Private/UObject/PropertyBag.cpp
deleted file mode 100644
diff --git a/Engine/Source/Runtime/CoreUObject/Private/UObject/PropertyBagRepository.cpp b/Engine/Source/Runtime/CoreUObject/Private/UObject/PropertyBagRepository.cpp
--- a/Engine/Source/Runtime/CoreUObject/Private/UObject/PropertyBagRepository.cpp
+++ b/Engine/Source/Runtime/CoreUObject/Private/UObject/PropertyBagRepository.cpp
@@ -8,7 +8,6 @@
 #include "UObject/GarbageCollection.h"
 #include "UObject/Object.h"
 #include "UObject/UObjectGlobals.h"
-#include "UObject/PropertyBag.h"
 #include "UObject/PropertyPathNameTree.h"
 #include "UObject/UObjectHash.h"
 #include "UObject/UObjectThreadContext.h"
@@ -108,10 +107,10 @@ public:
 
 void FPropertyBagRepository::FPropertyBagAssociationData::Destroy()
 {
-	delete Bag;
-	Bag = nullptr;
-	
-	if(InstanceDataObject && InstanceDataObject->IsValidLowLevel())
+	delete Tree;
+	Tree = nullptr;
+
+	if (InstanceDataObject && InstanceDataObject->IsValidLowLevel())
 	{
 		InstanceDataObject = nullptr;
 	}
@@ -474,23 +473,6 @@ void FPropertyBagRepository::PostEditChangeChainProperty(const UObject* Object,
 #endif
 }
 
-// TODO: Create these by class on construction?
-FPropertyBag* FPropertyBagRepository::CreateOuterBag(const UObject* Owner)
-{
-	FPropertyBagRepositoryLock LockRepo(this);
-	FPropertyBagAssociationData* BagData = AssociatedData.Find(Owner);
-	if (!BagData)
-	{
-		FPropertyBagAssociationData NewBagData;
-		BagData = &AssociatedData.Emplace(Owner, NewBagData);
-	}
-	if (!BagData->Bag)
-	{
-		BagData->Bag = new FPropertyBag;
-	}
-	return BagData->Bag;
-}
-
 // TODO: Create these by class on construction?
 FPropertyPathNameTree* FPropertyBagRepository::CreateUnknownPropertyTree(const UObject* Owner)
 {
@@ -549,26 +531,6 @@ bool FPropertyBagRepository::RemoveAssociationUnsafe(const UObject* Owner)
 	return false;
 }
 
-bool FPropertyBagRepository::HasBag(const UObject* Object) const
-{
-	// TODO: Should be consistent across all objects of a given type, so handle via TStructOpsTypeTraits or similar?
-	FPropertyBagRepositoryLock LockRepo(this);
-	//return AssociatedData.Contains(Object);	// Better approach? Object data should guarantee existence of bag.
-	return FindBag(Object) != nullptr;
-}
-
-FPropertyBag* FPropertyBagRepository::FindBag(const UObject* Object)
-{
-	FPropertyBagRepositoryLock LockRepo(this);
-	const FPropertyBagAssociationData* BagData = AssociatedData.Find(Object);
-	return BagData ? BagData->Bag : nullptr;
-}
-
-const FPropertyBag* FPropertyBagRepository::FindBag(const UObject* Object) const
-{
-	return const_cast<FPropertyBagRepository*>(this)->FindBag(Object);
-}
-
 bool FPropertyBagRepository::HasInstanceDataObject(const UObject* Object) const
 {
 	FPropertyBagRepositoryLock LockRepo(this);
@@ -709,7 +671,9 @@ void FPropertyBagRepository::CreateInstanceDataObjectUnsafe(UObject* Owner, FPro
 	}
 	else
 	{
-		ensureMsgf(BagData.Bag == nullptr, TEXT("Linker missing when generating IDO for an object with loose properties. Loose properties will be lost"));
+		ensureMsgf(BagData.Tree == nullptr,
+			TEXT("Linker missing when generating IDO for an object with unknown properties. The unknown properties will be lost. Path: %s"),
+			*Owner->GetPathName());
 		// copy data from owner to IDO
 		CopyTaggedProperties(Owner, BagData.InstanceDataObject);
 	}
diff --git a/Engine/Source/Runtime/CoreUObject/Tests/PropertyBagTest.cpp b/Engine/Source/Runtime/CoreUObject/Tests/PropertyBagTest.cpp
deleted file mode 100644
