Add missing #pragma once.
#rnx

diff --git a/Engine/Source/Runtime/CoreUObject/Internal/UObject/ImportExportCollector.h b/Engine/Source/Runtime/CoreUObject/Internal/UObject/ImportExportCollector.h
--- a/Engine/Source/Runtime/CoreUObject/Internal/UObject/ImportExportCollector.h
+++ b/Engine/Source/Runtime/CoreUObject/Internal/UObject/ImportExportCollector.h
@@ -1,5 +1,7 @@
 // Copyright Epic Games, Inc. All Rights Reserved.
 
+#pragma once
+
 #include "CoreTypes.h"
 
 #if WITH_EDITORONLY_DATA
