Fixed remapping of lazy object ptr in PIE : now a default FGuid is remapped to default FGuid

diff --git a/Engine/Source/Runtime/CoreUObject/Private/UObject/LazyObjectPtr.cpp b/Engine/Source/Runtime/CoreUObject/Private/UObject/LazyObjectPtr.cpp
--- a/Engine/Source/Runtime/CoreUObject/Private/UObject/LazyObjectPtr.cpp
+++ b/Engine/Source/Runtime/CoreUObject/Private/UObject/LazyObjectPtr.cpp
@@ -19,8 +19,15 @@ static TMap<FGuid, FGuid> PIEGuidMap[MAX_PIE_INSTANCES];
 
 static FGuid RemapGuid(const FUniqueObjectGuid& Guid, int32 PIEInstanceID)
 {
+	// An invalid Guid should remain an invalid Guid :
+	if (Guid.IsDefault())
+	{
+		return FGuid();
+	}
+
 	check(PIEInstanceID != INDEX_NONE);
 	check(PIEInstanceID < MAX_PIE_INSTANCES);
+
 	FGuid& FoundGuid = PIEGuidMap[PIEInstanceID].FindOrAdd(Guid.GetGuid());
 
 	if (!FoundGuid.IsValid())
@@ -190,4 +197,4 @@ void FLazyObjectPtr::ResetPIEFixups()
 	check(PlayInEditorID != -1);
 	check(PlayInEditorID < MAX_PIE_INSTANCES);
 	PIEGuidMap[PlayInEditorID].Reset();
-}
\ No newline at end of file
+}
