EditorBulkData : make sure Payload.Size and PayloadSize stay in sync ; missed Reset when loading transaction with PayloadSize zero

diff --git a/Engine/Source/Runtime/CoreUObject/Private/Serialization/EditorBulkData.cpp b/Engine/Source/Runtime/CoreUObject/Private/Serialization/EditorBulkData.cpp
--- a/Engine/Source/Runtime/CoreUObject/Private/Serialization/EditorBulkData.cpp
+++ b/Engine/Source/Runtime/CoreUObject/Private/Serialization/EditorBulkData.cpp
@@ -949,6 +949,7 @@ void FEditorBulkData::Serialize(FArchive& Ar, UObject* Owner, bool bAllowRegiste
 				if ( PayloadSize == 0 )
 				{
 					// nothing to do
+					Payload.Reset();
 				}
 				else
 				{
@@ -959,7 +960,7 @@ void FEditorBulkData::Serialize(FArchive& Ar, UObject* Owner, bool bAllowRegiste
 					{
 						FCompressedBuffer CompressedPayload;
 						Ar << CompressedPayload;
-
+						
 						// there's no way to just set a CompressedBuffer for Payload, so we have to Decompress
 						//	it might be nice if Payload could be either FSharedBuffer or FCompressedBuffer
 						//	then we could just store it without decompressing
@@ -2099,7 +2100,7 @@ TFuture<FSharedBuffer> FEditorBulkData::GetPayload() const
 		// NOTE: DecompressedPayload is *NOT* cached in the Payload variable
 		//	so UnloadData will do nothing
 		//	and we will reload from disk if GetPayload() is called twice in a row
-
+		
 		// TODO: Not actually async yet!
 		Promise.SetValue(DecompressedPayload);
 	}
