Improve error reporting when making an invalid bulkdata stream request.

#rnx

- If a user tries to make a read request from FBulkDataBatchRequest that goes past the end of the bulkdata's payload we will eventually assert, but only at the point of reading which is often quite far away (context wise) from the code that initiated the request.
- Previously we were asserting if the user made a read request that exceeded the size of the bulkdata payload but this did not take into account reading from an offset.
- Fixed the assert to take offset into account so we now asset on invalid reads at the point that the read is first requested. In addition I have added more debug info to the assert making it easier to tell from the log file what was incorrect about the read request.
- This shouldn't raise any new issues as it's effectively just moving an assert that already exists earlier in execution.

diff --git a/Engine/Source/Runtime/CoreUObject/Private/Serialization/BulkDataStreaming.cpp b/Engine/Source/Runtime/CoreUObject/Private/Serialization/BulkDataStreaming.cpp
--- a/Engine/Source/Runtime/CoreUObject/Private/Serialization/BulkDataStreaming.cpp
+++ b/Engine/Source/Runtime/CoreUObject/Private/Serialization/BulkDataStreaming.cpp
@@ -1083,7 +1083,12 @@ FBulkDataBatchRequest::FBatchBuilder& FBulkDataBatchRequest::FBatchBuilder::Read
 	FIoBuffer& Dst, 
 	FBulkDataBatchReadRequest* OutRequest)
 {
-	check(Size == MAX_uint64 || Size <= uint64(BulkData.GetBulkDataSize()));
+	checkf(Size == MAX_uint64 || (Offset + Size) <= uint64(BulkData.GetBulkDataSize()),
+		TEXT("%s: Trying to read past the end of the payload, Offset: %llu, ReadSize: %llu, Payload Size: %lld"),
+		*BulkData.GetDebugName(),
+		Size,
+		Offset,
+		BulkData.GetBulkDataSize());
 
 	const uint64 ReadOffset = BulkData.GetBulkDataOffsetInFile() + Offset;
 	const uint64 ReadSize	= FMath::Min(uint64(BulkData.GetBulkDataSize()), Size);
