UObject::Rename will no longer call ResetLoaders as this frequently breaks recursive loads by resetting the loader currently in use. This change removes the functionality and a follow-up change will remove the define from usage.

Unless moving a CDO or moving to the transient package, an object moving out of it's package will have it's linker nulled but not reset.

diff --git a/Engine/Source/Runtime/CoreUObject/Private/UObject/Obj.cpp b/Engine/Source/Runtime/CoreUObject/Private/UObject/Obj.cpp
--- a/Engine/Source/Runtime/CoreUObject/Private/UObject/Obj.cpp
+++ b/Engine/Source/Runtime/CoreUObject/Private/UObject/Obj.cpp
@@ -223,19 +223,32 @@ bool UObject::Rename( const TCHAR* InName, UObject* NewOuter, ERenameFlags Flags
 	FMetaDataUtilities::FMoveMetadataHelperContext MoveMetaData(this, true);
 #endif //WITH_EDITOR
 
-	// Check that we are not renaming a within object into an Outer of the wrong type, unless we're renaming the CDO of a Blueprint.
-	// Moving objects to the transient package is commonly used halfway through destroying them so that is also fine
-	if( NewOuter && !NewOuter->IsA(GetClass()->ClassWithin) && !HasAnyFlags(RF_ClassDefaultObject|RF_ArchetypeObject) && NewOuter != GetTransientPackage())	
+	if(NewOuter)
 	{
-		if (Flags & REN_Test)
+		// Renaming the CDO of a Blueprint is a special case so we do not validate what would otherwise be incorrect use of Rename.
+		// Moving objects to the transient package is commonly used halfway through destroying them so that is also fine, otherwise
+		// proceed to check for incorrect use of Rename
+		if (!HasAnyFlags(RF_ClassDefaultObject | RF_ArchetypeObject) && NewOuter->GetPackage() != GetTransientPackage())
 		{
-			return false;
-		}
+			// Check that we are not renaming a within object into an Outer of the wrong type
+			if (!NewOuter->IsA(GetClass()->ClassWithin))
+			{
+				if (Flags & REN_Test)
+				{
+					return false;
+				}
 
-		UE_LOG(LogObj, Fatal, TEXT("Cannot rename %s into Outer %s as it is not of type %s"), 
-			*GetFullName(), 
-			*NewOuter->GetFullName(), 
-			*GetClass()->ClassWithin->GetName() );
+				UE_LOG(LogObj, Fatal, TEXT("Cannot rename %s into Outer %s as it is not of type %s"),
+					*GetFullName(),
+					*NewOuter->GetFullName(),
+					*GetClass()->ClassWithin->GetName());
+			}
+			// If moving the object to a new package, remove it's linker .
+			else if (GetLinker() && GetPackage() != NewOuter->GetPackage())
+			{
+				SetLinker(nullptr, INDEX_NONE);
+			}
+		}
 	}
 
 	// find an object with the same name and same class in the new outer
@@ -284,11 +297,6 @@ bool UObject::Rename( const TCHAR* InName, UObject* NewOuter, ERenameFlags Flags
 		return true;
 	}
 
-	if (!(Flags & REN_ForceNoResetLoaders))
-	{
-		ResetLoaders(GetPackage());
-	}
-
 	FName OldName = GetFName();
 	FName NewName;
 	bool bCreateRedirector = false;
diff --git a/Engine/Source/Runtime/CoreUObject/Public/UObject/ObjectMacros.h b/Engine/Source/Runtime/CoreUObject/Public/UObject/ObjectMacros.h
--- a/Engine/Source/Runtime/CoreUObject/Public/UObject/ObjectMacros.h
+++ b/Engine/Source/Runtime/CoreUObject/Public/UObject/ObjectMacros.h
@@ -2140,7 +2140,7 @@ typedef uint32 ERenameFlags;
 
 /** Default rename behavior */
 #define REN_None					(0x0000)
-/** Rename won't call ResetLoaders or flush async loading. You should pass this if you are renaming a deep subobject and do not need to reset loading for the outer package */
+/** UE_DEPRECATED(5.5, "Rename will no longer call ResetLoaders making this flag no longer needed.") -  Rename won't call ResetLoaders or flush async loading. You should pass this if you are renaming a deep subobject and do not need to reset loading for the outer package */
 #define REN_ForceNoResetLoaders		(0x0001) 
 /** Just test to make sure that the rename is guaranteed to succeed if an non test rename immediately follows */
 #define REN_Test					(0x0002) 
