Changes the zenloader interpretation of the flag `RF_WasLoaded` to match the -nozenloader meaning by setting the RF_WasLoaded flag on the linker root object once the package load is complete. Previously zenloader set this flag during creation of the package but before any loading had occurred (since 10482167).

This change resolves issues where codepaths would look for `RF_WasLoaded` to understand if a package had completed loading or not when in fact the package has not been loaded, but instead may only be partially loaded. Specifically, LinkerLoad::ResolveDeferredDependencies was changed to explicitly not load loaded or already loaded packages.

Any code guarding access to an partially loaded object will now be correct if guarding against objects missing RF_WasLoaded. However this change does means code looking for partially loaded objects will now need to look for RW_WillBeLoaded specifically. I imagine the amount of code that cares about partially loaded objects is quite low in practice compared to knowing if an object is fully loaded. Follow-up work aims to remove RF_WillBeLoaded completely in favour of another mechanism to check for in-flight work UE-212618

A test System.Engine.Loading.LoadBlueprintWithCircularDependency has been added that exhibits a hang if RF_WillBeLoaded is not checked for in ResolveDeferredDependencies before calling LoadPackage since will cause a flush to happen potentially as part of an already on-going flush which can lead to partial loads being unable to be unblocked and the loader softlocking
#tests System.Engine.Loading.LoadBlueprintWithCircularDependency

diff --git a/Engine/Source/Runtime/CoreUObject/Private/Blueprint/BlueprintSupport.cpp b/Engine/Source/Runtime/CoreUObject/Private/Blueprint/BlueprintSupport.cpp
--- a/Engine/Source/Runtime/CoreUObject/Private/Blueprint/BlueprintSupport.cpp
+++ b/Engine/Source/Runtime/CoreUObject/Private/Blueprint/BlueprintSupport.cpp
@@ -1575,7 +1575,7 @@ void FLinkerLoad::ResolveDeferredDependencies(UStruct* LoadStruct)
 			// this package may not have introduced any (possible) cyclic 
 			// dependencies, but it still could have been deferred (kept from
 			// fully loading... we need to make sure metadata gets loaded, etc.)
-			if ((SourcePackage != nullptr) && !SourcePackage->HasAnyFlags(RF_WasLoaded))
+			if ((SourcePackage != nullptr) && !SourcePackage->HasAnyFlags(RF_WasLoaded | RF_WillBeLoaded))
 			{
 				uint32 InternalLoadFlags = LoadFlags & (LOAD_NoVerify | LOAD_NoWarn | LOAD_Quiet | LOAD_RegenerateBulkDataGuids);
 				// make sure LoadAllObjects() is called for this package
diff --git a/Engine/Source/Runtime/CoreUObject/Private/Serialization/AsyncLoading2.cpp b/Engine/Source/Runtime/CoreUObject/Private/Serialization/AsyncLoading2.cpp
--- a/Engine/Source/Runtime/CoreUObject/Private/Serialization/AsyncLoading2.cpp
+++ b/Engine/Source/Runtime/CoreUObject/Private/Serialization/AsyncLoading2.cpp
@@ -7476,14 +7476,6 @@ EEventLoadNodeExecutionResult FAsyncPackage2::Event_PostLoadExportBundle(FAsyncL
 	if (++Package->ProcessedExportBundlesCount == Package->Data.TotalExportBundleCount)
 	{
 		Package->ProcessedExportBundlesCount = 0;
-		if (Package->LinkerRoot && !Package->bLoadHasFailed)
-		{
-			UE_ASYNC_PACKAGE_LOG(Verbose, Package->Desc, TEXT("AsyncThread: FullyLoaded"),
-				TEXT("Async loading of package is done, and UPackage is marked as fully loaded."));
-			// mimic old loader behavior for now, but this is more correctly also done in FinishUPackage
-			// called from ProcessLoadedPackagesFromGameThread just before completion callbacks
-			Package->LinkerRoot->MarkAsFullyLoaded();
-		}
 
 		FAsyncLoadingPostLoadGroup* DeferredPostLoadGroup = Package->DeferredPostLoadGroup;
 		check(DeferredPostLoadGroup);
@@ -9030,7 +9022,7 @@ void FAsyncPackage2::CreateUPackage()
 	}
 #endif
 
-	LinkerRoot->SetFlags(RF_Public | RF_WasLoaded);
+	LinkerRoot->SetFlags(RF_Public | RF_WillBeLoaded);
 	LinkerRoot->SetLoadedPath(Desc.PackagePathToLoad);
 	LinkerRoot->SetCanBeImportedFlag(Desc.bCanBeImported);
 	LinkerRoot->SetPackageId(Desc.UPackageId);
@@ -9150,6 +9142,8 @@ void FAsyncPackage2::FinishUPackage()
 	{
 		// Mark package as having been fully loaded and update load time.
 		LinkerRoot->MarkAsFullyLoaded();
+		LinkerRoot->SetFlags(RF_WasLoaded);
+		LinkerRoot->ClearFlags(RF_WillBeLoaded);
 		LinkerRoot->SetLoadTime((float)(FPlatformTime::Seconds() - LoadStartTime));
 	}
 }
diff --git a/Engine/Source/Runtime/Engine/Private/Tests/Loading/AsyncLoadingTests.cpp b/Engine/Source/Runtime/Engine/Private/Tests/Loading/AsyncLoadingTests.cpp
--- a/Engine/Source/Runtime/Engine/Private/Tests/Loading/AsyncLoadingTests.cpp
+++ b/Engine/Source/Runtime/Engine/Private/Tests/Loading/AsyncLoadingTests.cpp
@@ -56,5 +56,24 @@ bool FThreadSafeAsyncLoadingTest::RunTest(const FString& Parameters)
 	return true;
 }
 
+/**
+ * This test demonstrates that LoadPackage can load blueprints with circular dependencies who rely on dependencies with circular dependencies
+ */
+IMPLEMENT_SIMPLE_AUTOMATION_TEST(FCircularDependencyLoadingTest, TEXT(TEST_NAME_ROOT ".LoadBlueprintWithCircularDependencyTest"), EAutomationTestFlags::EditorContext | EAutomationTestFlags::EngineFilter)
+bool FCircularDependencyLoadingTest::RunTest(const FString& Parameters)
+{
+	static constexpr const TCHAR* ActorWithCircularReferences = TEXT("/Game/Tests/Core/AssetLoading/BlueprintActorWithCircularReferences");
+
+	UPackage* Package = LoadPackage(nullptr, ActorWithCircularReferences, LOAD_None);
+	TestTrue(TEXT("The object should have been properly loaded recursively"), Package != nullptr);
+
+	Package = FindPackage(nullptr, ActorWithCircularReferences);
+	TestTrue(TEXT("The object should have been properly loaded recursively"), Package != nullptr);
+
+	CollectGarbage(GARBAGE_COLLECTION_KEEPFLAGS);
+
+	return true;
+}
+
 #undef TEST_NAME_ROOT
 #endif // WITH_DEV_AUTOMATION_TESTS
\ No newline at end of file
