SOL-6284: initializing archetype proxy vars before PostInitProperties() fires

diff --git a/Engine/Source/Runtime/CoreUObject/Private/UObject/UObjectGlobals.cpp b/Engine/Source/Runtime/CoreUObject/Private/UObject/UObjectGlobals.cpp
--- a/Engine/Source/Runtime/CoreUObject/Private/UObject/UObjectGlobals.cpp
+++ b/Engine/Source/Runtime/CoreUObject/Private/UObject/UObjectGlobals.cpp
@@ -4045,6 +4045,14 @@ void FObjectInitializer::PostConstructInit()
 		Subobject->PostReinitProperties();
 	}
 
+	for (TFunction<void()>& Callback : PropertyPostInitCallbacks)
+	{
+		if (Callback)
+		{
+			Callback();
+		}
+	}
+
 	{
 		SCOPE_CYCLE_COUNTER(STAT_PostInitProperties);
 		Obj->PostInitProperties();
diff --git a/Engine/Source/Runtime/CoreUObject/Public/UObject/UObjectGlobals.h b/Engine/Source/Runtime/CoreUObject/Public/UObject/UObjectGlobals.h
--- a/Engine/Source/Runtime/CoreUObject/Public/UObject/UObjectGlobals.h
+++ b/Engine/Source/Runtime/CoreUObject/Public/UObject/UObjectGlobals.h
@@ -1241,6 +1241,11 @@ public:
 		return InstanceGraph;
 	}
 
+	FORCEINLINE void AddPropertyPostInitCallback(TFunction<void()>&& Callback)
+	{
+		PropertyPostInitCallbacks.Add(MoveTemp(Callback));
+	}
+
 	/**
 	* Return the class of the object that is being constructed
 	**/
@@ -1618,8 +1623,9 @@ private:
 	/**  Previously constructed object in the callstack */
 	UObject* LastConstructedObject = nullptr;
 
-	/** Callback for custom property initialization before PostInitProperties gets called */
+	/** Callbacks for custom property initialization before PostInitProperties gets called */
 	TFunction<void()> PropertyInitCallback;
+	TArray<TFunction<void()>> PropertyPostInitCallbacks;
 
 	friend struct FStaticConstructObjectParameters;
 
