[ZenLoader]
 - Detect when loader is hung and dump its state if it happens
 - Add support for full flush from inside the last loaded package's callback if it ever happens
 - Add unittest for that use case

diff --git a/Engine/Source/Runtime/CoreUObject/Private/Serialization/AsyncLoading2.cpp b/Engine/Source/Runtime/CoreUObject/Private/Serialization/AsyncLoading2.cpp
--- a/Engine/Source/Runtime/CoreUObject/Private/Serialization/AsyncLoading2.cpp
+++ b/Engine/Source/Runtime/CoreUObject/Private/Serialization/AsyncLoading2.cpp
@@ -418,6 +418,14 @@ static FAutoConsoleVariableRef CVarGOnlyProcessRequiredPackagesWhenSyncLoading(
 	ECVF_Default
 );
 
+static float GStallDetectorTimeout = 30.0f;
+static FAutoConsoleVariableRef CVarGStallDetectorTimeout(
+	TEXT("s.StallDetectorTimeout"),
+	GStallDetectorTimeout,
+	TEXT("Time in seconds after which we consider the loader stalled if no progress is being made"),
+	ECVF_Default
+);
+
 CSV_DECLARE_CATEGORY_MODULE_EXTERN(CORE_API, Basic);
 CSV_DECLARE_CATEGORY_MODULE_EXTERN(CORE_API, FileIO);
 
@@ -2191,8 +2199,13 @@ public:
 
 	uint64 GetSyncLoadContextId() const;
 
+	const FAsyncLoadEventSpec* GetSpec() const
+	{
+		return Spec;
+	}
 private:
 	friend class FAsyncLoadEventQueue2;
+	friend struct FAsyncPackage2;
 	friend TIoPriorityQueue<FEventLoadNode2>;
 
 	FEventLoadNode2()
@@ -2282,7 +2295,11 @@ public:
 	void Push(FAsyncLoadingThreadState2* ThreadState, FEventLoadNode2* Node);
 	bool ExecuteSyncLoadEvents(FAsyncLoadingThreadState2& ThreadState);
 	void UpdatePackagePriority(FAsyncPackage2* Package);
-
+	bool IsEmptyForDebug()
+	{
+		FScopeLock Lock(&ExternalCritical);
+		return LocalQueue.IsEmpty() && ExternalQueue.IsEmpty();
+	}
 private:
 	void PushLocal(FEventLoadNode2* Node);
 	void PushExternal(FEventLoadNode2* Node);
@@ -2328,6 +2345,7 @@ struct FAsyncLoadEventSpec
 	FAsyncLoadEventFunc Func = nullptr;
 	FAsyncLoadEventQueue2* EventQueue = nullptr;
 	bool bExecuteImmediately = false;
+	const TCHAR* Name = nullptr;
 };
 
 class FAsyncLoadingSyncLoadContext
@@ -3159,7 +3177,13 @@ private:
 	*/
 	void UpdateLoadPercentage();
 
+#if !UE_BUILD_SHIPPING
+	void DumpStateImpl(TSet<FAsyncPackage2*>& Set, int32 Indent = 0, TMultiMap<FEventLoadNode2*, FEventLoadNode2*>* MappedNodes = nullptr);
+#endif
 public:
+#if !UE_BUILD_SHIPPING
+	void DumpState();
+#endif
 
 	/** Serialization context for this package */
 	FUObjectSerializeContext* GetSerializeContext();
@@ -3589,6 +3613,15 @@ public:
 		return PendingRequests.Contains(RequestID);
 	}
 
+	/**
+	 * [ASYNC/GAME THREAD] Checks if there are any pending requests
+	 */
+	bool ContainsAnyPendingRequests()
+	{
+		FScopeLock Lock(&PendingRequestsCritical);
+		return !PendingRequests.IsEmpty();
+	}
+
 	/**
 	 * [ASYNC/GAME THREAD] Checks if a request ID already is added to the loading queue
 	 */
@@ -8151,20 +8184,19 @@ FAsyncLoadingThread2::FAsyncLoadingThread2(FIoDispatcher& InIoDispatcher, IAsync
 	RequestIdToPackageMap.Reserve(DefaultAsyncPackagesReserveCount);
 
 	EventSpecs.AddDefaulted(EEventLoadNode2::Package_NumPhases + EEventLoadNode2::ExportBundle_NumPhases);
-	EventSpecs[EEventLoadNode2::Package_ProcessSummary] = { &FAsyncPackage2::Event_ProcessPackageSummary, &EventQueue, false };
-	EventSpecs[EEventLoadNode2::Package_DependenciesReady] = { &FAsyncPackage2::Event_DependenciesReady, &EventQueue, false };
+	EventSpecs[EEventLoadNode2::Package_ProcessSummary] = { &FAsyncPackage2::Event_ProcessPackageSummary, &EventQueue, false, TEXT("ProcessSummary") };
+	EventSpecs[EEventLoadNode2::Package_DependenciesReady] = { &FAsyncPackage2::Event_DependenciesReady, &EventQueue, false, TEXT("DependenciesReady")};
 #if ALT2_ENABLE_LINKERLOAD_SUPPORT
-	EventSpecs[EEventLoadNode2::Package_CreateLinkerLoadExports] = { &FAsyncPackage2::Event_CreateLinkerLoadExports, &EventQueue, false };
-	EventSpecs[EEventLoadNode2::Package_ResolveLinkerLoadImports] = { &FAsyncPackage2::Event_ResolveLinkerLoadImports, &EventQueue, false };
-	EventSpecs[EEventLoadNode2::Package_PreloadLinkerLoadExports] = { &FAsyncPackage2::Event_PreloadLinkerLoadExports, &EventQueue, false };
+	EventSpecs[EEventLoadNode2::Package_CreateLinkerLoadExports] = { &FAsyncPackage2::Event_CreateLinkerLoadExports, &EventQueue, false, TEXT("CreateLinkerLoadExports") };
+	EventSpecs[EEventLoadNode2::Package_ResolveLinkerLoadImports] = { &FAsyncPackage2::Event_ResolveLinkerLoadImports, &EventQueue, false, TEXT("ResolveLinkerLoadImports") };
+	EventSpecs[EEventLoadNode2::Package_PreloadLinkerLoadExports] = { &FAsyncPackage2::Event_PreloadLinkerLoadExports, &EventQueue, false, TEXT("PreloadLinkerLoadExports") };
 #endif
-	EventSpecs[EEventLoadNode2::Package_ExportsSerialized] = { &FAsyncPackage2::Event_ExportsDone, &EventQueue, true };
+	EventSpecs[EEventLoadNode2::Package_ExportsSerialized] = { &FAsyncPackage2::Event_ExportsDone, &EventQueue, true, TEXT("ExportsSerialized") };
 
-	EventSpecs[EEventLoadNode2::Package_NumPhases + EEventLoadNode2::ExportBundle_Process] = { &FAsyncPackage2::Event_ProcessExportBundle, &EventQueue, false };
-	EventSpecs[EEventLoadNode2::Package_NumPhases + EEventLoadNode2::ExportBundle_PostLoad] = { &FAsyncPackage2::Event_PostLoadExportBundle, &EventQueue, false };
-	EventSpecs[EEventLoadNode2::Package_NumPhases + EEventLoadNode2::ExportBundle_DeferredPostLoad] = { &FAsyncPackage2::Event_DeferredPostLoadExportBundle, &MainThreadEventQueue, false };
+	EventSpecs[EEventLoadNode2::Package_NumPhases + EEventLoadNode2::ExportBundle_Process] = { &FAsyncPackage2::Event_ProcessExportBundle, &EventQueue, false, TEXT("ExportBundle_Process") };
+	EventSpecs[EEventLoadNode2::Package_NumPhases + EEventLoadNode2::ExportBundle_PostLoad] = { &FAsyncPackage2::Event_PostLoadExportBundle, &EventQueue, false, TEXT("ExportBundle_PostLoad") };
+	EventSpecs[EEventLoadNode2::Package_NumPhases + EEventLoadNode2::ExportBundle_DeferredPostLoad] = { &FAsyncPackage2::Event_DeferredPostLoadExportBundle, &MainThreadEventQueue, false, TEXT("ExportBundle_DeferredPostLoad") };
 
-	CancelLoadingEvent = FPlatformProcess::GetSynchEventFromPool();
 	ThreadSuspendedEvent = FPlatformProcess::GetSynchEventFromPool();
 	ThreadResumedEvent = FPlatformProcess::GetSynchEventFromPool();
 	AsyncLoadingTickCounter = 0;
@@ -9517,6 +9549,7 @@ void FAsyncLoadingThread2::FlushLoading(TConstArrayView<int32> RequestIDs)
 
 		// Flush async loaders by not using a time limit. Needed for e.g. garbage collection.
 		{
+			double LastActivity = 0.0;
 			while (IsAsyncLoadingPackages())
 			{
 				bool bDidSomething = false;
@@ -9526,15 +9559,67 @@ void FAsyncLoadingThread2::FlushLoading(TConstArrayView<int32> RequestIDs)
 					break;
 				}
 
-#if 0
 				if (!bDidSomething)
 				{
-					if (!AsyncLoadingThreadState && !PendingIoRequestsCounter)
+					// Early out in case we end up with a full flush with no more request to process
+					// (i.e. Could happen if a full flush is called from inside the last package completion callback).
+					if (RequestIDs.Num() == 0 && !ContainsAnyPendingRequests())
 					{
+						break;
+					}
+
+					if (LastActivity == 0.0)
+					{
+						LastActivity = FPlatformTime::Seconds();
+					}
+					else if (FPlatformTime::Seconds() - LastActivity > GStallDetectorTimeout)
+					{
+#if !UE_BUILD_SHIPPING
+						{
+							FScopeLock ScopeLock(&AsyncPackagesCritical);
+
+							if (GameThreadState->CurrentlyExecutingEventNodeStack.Num() > 0)
+							{
+								UE_LOG(LogStreaming, Warning, TEXT("============ Currently executing nodes on stack ============="));
+
+								int32 Index = 0;
+								for (FEventLoadNode2* Node : GameThreadState->CurrentlyExecutingEventNodeStack)
+								{
+									UE_LOG(LogStreaming, Warning, TEXT("#%d: Package %s executing node %s"), Index++, *Node->GetPackage()->Desc.UPackageName.ToString(), Node->GetSpec()->Name);
+								}
+
+								UE_LOG(LogStreaming, Warning, TEXT("============"));
+							}
+
+							if (SyncLoadContext)
+							{
+								for (FAsyncPackage2* Package : SyncLoadContext->RequestedPackages)
+								{
+									if (Package)
+									{
+										Package->DumpState();
+									}
+								}
+							}
+							else
+							{
+								for (auto& Pair : AsyncPackageLookup)
+								{
+									if (FAsyncPackage2* Package = Pair.Value)
+									{
+										Package->DumpState();
+									}
+								}
+							}
+						}
+#endif // !UE_BUILD_SHIPPING
 						UE_LOG(LogStreaming, Fatal, TEXT("Loading is stuck, flush will never finish"));
 					}
 				}
-#endif
+				else
+				{
+					LastActivity = 0.0;
+				}
 
 				if (IsMultithreaded())
 				{
@@ -9588,7 +9673,7 @@ void FAsyncLoadingThread2::FlushLoading(TConstArrayView<int32> RequestIDs)
 #endif
 
 		// If we asked to flush everything, we should no longer have anything in the pipeline
-		check(RequestIDs.Num() != 0 || !IsAsyncLoadingPackages());
+		check(RequestIDs.Num() != 0 || !IsAsyncLoadingPackages() || !ContainsAnyPendingRequests());
 	}
 }
 
@@ -9648,6 +9733,192 @@ void FGlobalImportStore::FlushDeferredDeletePackagesQueue()
 	AsyncLoadingThread.ProcessDeferredDeletePackagesQueue();
 }
 
+#if !UE_BUILD_SHIPPING
+
+void FAsyncPackage2::DumpState()
+{
+	UE_LOG(LogStreaming, Warning, TEXT("============ Dumping State of Package %s ============"), *Desc.UPackageName.ToString());
+
+	TSet<FAsyncPackage2*> Set;
+	DumpStateImpl(Set);
+
+	UE_LOG(LogStreaming, Warning, TEXT("============"));
+}
+
+void FAsyncPackage2::DumpStateImpl(TSet<FAsyncPackage2*>& Set, int32 Indent, TMultiMap<FEventLoadNode2*, FEventLoadNode2*>* MappedNodes)
+{
+	Set.Add(this);
+
+	auto GetPackageType = [](FAsyncPackage2* Package)
+	{
+#if ALT2_ENABLE_LINKERLOAD_SUPPORT
+		const bool bIsZenPackage = !Package->LinkerLoadState.IsSet();
+#else
+		const bool bIsZenPackage = true;
+#endif
+		return bIsZenPackage ? TEXT("ZEN") : TEXT("LINKERLOAD");
+	};
+
+	auto FormatPackage = [&GetPackageType](FAsyncPackage2* Package)
+	{
+		return FString::Printf(TEXT("%s package %s state %s"), GetPackageType(Package), *Package->Desc.UPackageName.ToString(), LexToString(Package->AsyncPackageLoadingState));
+	};
+
+	UE_LOG(LogStreaming, Warning, TEXT("%s%s:"), FCString::Spc(Indent), *FormatPackage(this));
+
+	auto VisitNodes = [](FEventLoadNode2& Node, TFunction<void(FEventLoadNode2*)> Visitor)
+	{
+		if (Node.DependenciesCount == 1)
+		{
+			Visitor(Node.SingleDependent);
+		}
+		else if (Node.DependenciesCount != 0)
+		{
+			FEventLoadNode2** Current = Node.MultipleDependents;
+			FEventLoadNode2** End = Node.MultipleDependents + Node.DependenciesCount;
+			for (; Current < End; ++Current)
+			{
+				Visitor(*Current);
+			}
+		}
+	};
+
+	auto DumpNode = [&FormatPackage, &VisitNodes](TMultiMap<FEventLoadNode2*, FEventLoadNode2*>& MappedNodes, const TCHAR* Header, FEventLoadNode2& Node, int32 Indent)
+	{
+		bool bIsEmpty = Node.Spec->EventQueue->IsEmptyForDebug();
+
+		if (Node.DependenciesCount == 0 && Node.BarrierCount == 0 && bIsEmpty)
+		{
+			return false;
+		}
+
+		if (Header)
+		{
+			UE_LOG(LogStreaming, Warning, TEXT("%s %s"), FCString::Spc(Indent), Header);
+		}
+		Indent++;
+
+		UE_LOG(LogStreaming, Warning, TEXT("%s Node %s (BarrierCount %d, EventQueue: %s)"), FCString::Spc(Indent), Node.GetSpec()->Name, Node.BarrierCount.load(), bIsEmpty ? TEXT("Empty") : TEXT("NonEmpty"));
+
+		TArray<FEventLoadNode2*> WaitingOn;
+		MappedNodes.MultiFind(&Node, WaitingOn);
+
+		for (FEventLoadNode2* WaitingNode : WaitingOn)
+		{
+			UE_LOG(LogStreaming, Warning, TEXT("%s Waiting on node %s from %s"), FCString::Spc(Indent + 1), WaitingNode->GetSpec()->Name, *FormatPackage(WaitingNode->GetPackage()));
+		}
+
+		VisitNodes(Node,
+			[Indent, &FormatPackage](FEventLoadNode2* Node)
+			{
+				UE_LOG(LogStreaming, Warning, TEXT("%s Will trigger node %s for %s"), FCString::Spc(Indent + 1), Node->GetSpec()->Name, *FormatPackage(Node->GetPackage()));
+			}
+		);
+
+		return true;
+	};
+
+	TMultiMap<FEventLoadNode2*, FEventLoadNode2*> LocalMappedNode;
+	if (MappedNodes == nullptr)
+	{
+		MappedNodes = &LocalMappedNode;
+
+		FScopeLock ScopeLock(&AsyncLoadingThread.AsyncPackagesCritical);
+		for (auto& Pair : AsyncLoadingThread.AsyncPackageLookup)
+		{
+			if (FAsyncPackage2* Package = Pair.Value)
+			{
+				for (uint8 PhaseIndex = 0; PhaseIndex < Package_NumPhases; ++PhaseIndex)
+				{
+					FEventLoadNode2& Node = Package->GetPackageNode((EEventLoadNode2)PhaseIndex);
+					VisitNodes(Node, [&](FEventLoadNode2* Dependent) { LocalMappedNode.AddUnique(Dependent, &Node); });
+				}
+
+				for (int32 ExportBundleIndex = 0; ExportBundleIndex < Package->Data.TotalExportBundleCount; ++ExportBundleIndex)
+				{
+					for (uint8 PhaseIndex = 0; PhaseIndex < ExportBundle_NumPhases; ++PhaseIndex)
+					{
+						EEventLoadNode2 EventNode = (EEventLoadNode2)PhaseIndex;
+						FEventLoadNode2& Node = Package->GetExportBundleNode(EventNode, ExportBundleIndex);
+						VisitNodes(Node, [&](FEventLoadNode2* Dependent) { LocalMappedNode.AddUnique(Dependent, &Node); });
+					}
+				}
+			}
+		}
+	}
+
+	{
+		FScopeLock ScopeLock(&AsyncLoadingThread.AsyncPackagesCritical);
+		for (auto& Pair : AsyncLoadingThread.AsyncPackageLookup)
+		{
+			if (FAsyncPackage2* Package = Pair.Value)
+			{
+				if (Package->HeaderData.ImportedPackageIds.Contains(this->Desc.UPackageId))
+				{
+					UE_LOG(LogStreaming, Warning, TEXT(" %s referenced by package %s"), FCString::Spc(Indent), *FormatPackage(Package));
+				}
+			}
+		}
+	}
+
+	for (uint8 PhaseIndex = 0; PhaseIndex < Package_NumPhases; ++PhaseIndex)
+	{
+		EEventLoadNode2 EventNode = (EEventLoadNode2)PhaseIndex;
+		DumpNode(*MappedNodes, nullptr, GetPackageNode(EventNode), Indent + 1);
+	}
+
+	for (int32 ExportBundleIndex = 0; ExportBundleIndex < Data.TotalExportBundleCount; ++ExportBundleIndex)
+	{
+		FString HeaderStr = FString::Printf(TEXT("ExportBundle %d"), ExportBundleIndex);
+		const TCHAR* Header = *HeaderStr;
+		for (uint8 PhaseIndex = 0; PhaseIndex < ExportBundle_NumPhases; ++PhaseIndex)
+		{
+			EEventLoadNode2 EventNode = (EEventLoadNode2)PhaseIndex;
+			if (DumpNode(*MappedNodes, Header, GetExportBundleNode(EventNode, ExportBundleIndex), Indent + 1))
+			{
+				Header = nullptr;
+			}
+		}
+	}
+
+	for (FAsyncPackage2* Import : Data.ImportedAsyncPackages)
+	{
+		if (Import && Import->AsyncPackageLoadingState < EAsyncPackageLoadingState2::DeferredDelete)
+		{
+			UE_LOG(LogStreaming, Warning, TEXT(" %s Imports %s package %s state %s"), FCString::Spc(Indent), GetPackageType(Import), *Import->Desc.UPackageName.ToString(), LexToString(Import->AsyncPackageLoadingState));
+		}
+	}
+
+	if (AllDependenciesSetupState.WaitingForPackage)
+	{
+		UE_LOG(LogStreaming, Warning, TEXT(" %s AllDependenciesSetupState is waiting on %s"), FCString::Spc(Indent), *FormatPackage(AllDependenciesSetupState.WaitingForPackage));
+	}
+#if ALT2_ENABLE_LINKERLOAD_SUPPORT
+	if (AllDependenciesImportState.WaitingForPackage)
+	{
+		UE_LOG(LogStreaming, Warning, TEXT(" %s AllDependenciesImportState is waiting on %s"), FCString::Spc(Indent), *FormatPackage(AllDependenciesImportState.WaitingForPackage));
+	}
+#endif
+	if (AllDependenciesFullyLoadedState.WaitingForPackage)
+	{
+		UE_LOG(LogStreaming, Warning, TEXT(" %s AllDependenciesFullyLoadedState is waiting on %s"), FCString::Spc(Indent), *FormatPackage(AllDependenciesFullyLoadedState.WaitingForPackage));
+	}
+
+	for (FAsyncPackage2* Import : Data.ImportedAsyncPackages)
+	{
+		if (Import && Import->AsyncPackageLoadingState < EAsyncPackageLoadingState2::DeferredDelete)
+		{
+			if (!Set.Contains(Import))
+			{
+				UE_LOG(LogStreaming, Warning, TEXT(""));
+				Import->DumpStateImpl(Set, 0, MappedNodes);
+			}
+		}
+	}
+}
+
+#endif // !UE_BUILD_SHIPPING
+
 IAsyncPackageLoader* MakeAsyncPackageLoader2(FIoDispatcher& InIoDispatcher, IAsyncPackageLoader* InUncookedPackageLoader)
 {
 	return new FAsyncLoadingThread2(InIoDispatcher, InUncookedPackageLoader);
diff --git a/Engine/Source/Runtime/Engine/Private/Tests/Loading/AsyncLoadingTests_Flushes.cpp b/Engine/Source/Runtime/Engine/Private/Tests/Loading/AsyncLoadingTests_Flushes.cpp
new file mode 100644
--- /dev/null
+++ b/Engine/Source/Runtime/Engine/Private/Tests/Loading/AsyncLoadingTests_Flushes.cpp
@@ -0,0 +1,67 @@
+// Copyright Epic Games, Inc. All Rights Reserved.
+
+#include "AsyncLoadingTests_Shared.h"
+#include "Misc/AutomationTest.h"
+#include "Misc/PackageName.h"
+#include "Async/ManualResetEvent.h"
+#include "UObject/CoreRedirects.h"
+#include "UObject/SavePackage.h"
+#include "Tasks/Task.h"
+
+#if WITH_DEV_AUTOMATION_TESTS
+
+// All Flush tests should run on zenloader only as the other loaders are not compliant.
+typedef FLoadingTests_ZenLoaderOnly_Base FLoadingTests_Flush_Base;
+
+/**
+ * This test validates loading an object synchronously during serialize.
+ */
+IMPLEMENT_CUSTOM_SIMPLE_AUTOMATION_TEST(
+	FLoadingTests_Flush_InvalidFromWorker,
+	FLoadingTests_Flush_Base,
+	TEXT("System.Engine.Loading.Flush.InvalidFromWorker"),
+	EAutomationTestFlags::EditorContext | EAutomationTestFlags::EngineFilter
+)
+bool FLoadingTests_Flush_InvalidFromWorker::RunTest(const FString& Parameters)
+{
+	AddExpectedError(TEXT("is unable to FlushAsyncLoading from the current thread"), EAutomationExpectedErrorFlags::Contains);
+	AddExpectedError(TEXT("[Callstack]"), EAutomationExpectedErrorFlags::Contains, 0 /* At least 1 occurrence */);
+	
+	FLoadingTestsScope LoadingTestScope(this);
+
+	UAsyncLoadingTests_Shared::OnSerialize.BindLambda(
+		[this](FArchive& Ar, UAsyncLoadingTests_Shared* Object)
+		{
+			if (Ar.IsLoading())
+			{
+				// Use event instead of waiting on the task to prevent retraction as we really want that task
+				// to execute on a worker thread instead of being retracted in the serialize thread.
+				UE::FManualResetEvent Event;
+				UE::Tasks::Launch(TEXT("FlushAsyncLoading"), [&Event]() { FlushAsyncLoading(); Event.Notify(); });
+				Event.Wait();
+			}
+		}
+	);
+
+	LoadingTestScope.LoadObjects();
+
+	return true;
+}
+
+IMPLEMENT_CUSTOM_SIMPLE_AUTOMATION_TEST(
+	FLoadingTests_Flush_ValidFromCallback,
+	FLoadingTests_Flush_Base,
+	TEXT("System.Engine.Loading.Flush.ValidFromCallback"),
+	EAutomationTestFlags::EditorContext | EAutomationTestFlags::EngineFilter
+)
+bool FLoadingTests_Flush_ValidFromCallback::RunTest(const FString& Parameters)
+{
+	FLoadingTestsScope LoadingTestScope(this);
+
+	LoadPackageAsync(FLoadingTestsScope::PackagePath1,
+		FLoadPackageAsyncDelegate::CreateLambda([](const FName&, UPackage*, EAsyncLoadingResult::Type) { FlushAsyncLoading(); }));
+
+	return true;
+}
+
+#endif // WITH_DEV_AUTOMATION_TESTS
diff --git a/Engine/Source/Runtime/Engine/Private/Tests/Loading/AsyncLoadingTests_InvalidFlush.cpp b/Engine/Source/Runtime/Engine/Private/Tests/Loading/AsyncLoadingTests_InvalidFlush.cpp
deleted file mode 100644
