ZenLoader: Move CreateExportBundleNodes from FinishInitializeAsyncPackage to InitializeAsyncPackageFromPackageStore.
It does not actually require any package summary information since it only relies on TotalExportBundleCount which is known already in InitializeAsyncPackageFromPackageStore.

Fix nullptr crash in FAsyncLoadingThread2::FindOrInsertPackage for the rare special case path where the following dependency is added:
ImportedByPackage->GetExportBundleNode(ExportBundle_Process, DependentExportBundleIndex).DependsOn(&Package->GetExportBundleNode(ExportBundle_DeferredPostLoad, DependsOnExportBundleIndex));

#tests Running Client, Server, Editor, CookedEditor

diff --git a/Engine/Source/Runtime/CoreUObject/Private/Serialization/AsyncLoading2.cpp b/Engine/Source/Runtime/CoreUObject/Private/Serialization/AsyncLoading2.cpp
--- a/Engine/Source/Runtime/CoreUObject/Private/Serialization/AsyncLoading2.cpp
+++ b/Engine/Source/Runtime/CoreUObject/Private/Serialization/AsyncLoading2.cpp
@@ -3821,17 +3821,22 @@ private:
 		const uint64 OptionalSegmentImportedPackageIdsMemSize = Align(sizeof(FPackageId) * OptionalSegmentImportedPackagesCount, 8);
 
 		const int32 TotalImportedPackagesCount = ImportedPackagesCount + OptionalSegmentImportedPackagesCount;
+		const int32 TotalExportBundleCount = bHasOptionalSegment ? 2 : 1;
 #else
 		const int32 TotalImportedPackagesCount = ImportedPackagesCount;
+		const int32 TotalExportBundleCount = 1;
 #endif
 		const int32 ShaderMapHashesCount = PackageStoreEntry.ShaderMapHashes.Num();
 
+		const int32 ExportBundleNodeCount = TotalExportBundleCount * EEventLoadNode2::ExportBundle_NumPhases;
+		const uint64 ExportBundleNodesMemSize = Align(sizeof(FEventLoadNode2) * ExportBundleNodeCount, 8);
 		const uint64 ImportedPackagesMemSize = Align(sizeof(FAsyncPackage2*) * TotalImportedPackagesCount, 8);
 		const uint64 ShaderMapHashesMemSize = Align(sizeof(FSHAHash) * ShaderMapHashesCount, 8);
 		const uint64 MemoryBufferSize =
 #if WITH_EDITOR
 			OptionalSegmentImportedPackageIdsMemSize +
 #endif
+			ExportBundleNodesMemSize +
 			ImportedPackageIdsMemSize +
 			ImportedPackagesMemSize +
 			ShaderMapHashesMemSize;
@@ -3847,6 +3852,11 @@ private:
 
 		uint8* DataPtr = Data.MemoryBuffer0;
 
+		Data.TotalExportBundleCount = TotalExportBundleCount;
+		Data.ExportBundleNodes = MakeArrayView(reinterpret_cast<FEventLoadNode2*>(DataPtr), ExportBundleNodeCount);
+		DataPtr += ExportBundleNodesMemSize;
+		AsyncPackage->CreateExportBundleNodes(EventSpecs.GetData());
+
 		Data.ShaderMapHashes = MakeArrayView(reinterpret_cast<const FSHAHash*>(DataPtr), ShaderMapHashesCount);
 		FMemory::Memcpy((void*)Data.ShaderMapHashes.GetData(), PackageStoreEntry.ShaderMapHashes.GetData(), sizeof(FSHAHash) * ShaderMapHashesCount);
 		DataPtr += ShaderMapHashesMemSize;
@@ -3859,7 +3869,6 @@ private:
 		FMemory::Memcpy((void*)HeaderData.ImportedPackageIds.GetData(), PackageStoreEntry.ImportedPackageIds.GetData(), sizeof(FPackageId) * ImportedPackagesCount);
 		DataPtr += ImportedPackageIdsMemSize;
 
-		Data.TotalExportBundleCount = 1;
 		HeaderData.ImportedAsyncPackagesView = Data.ImportedAsyncPackages;
 #if WITH_EDITOR
 		if (bHasOptionalSegment)
@@ -3870,7 +3879,6 @@ private:
 			FMemory::Memcpy((void*)OptionalSegmentHeaderData.ImportedPackageIds.GetData(), PackageStoreEntry.OptionalSegmentImportedPackageIds.GetData(), sizeof(FPackageId) * OptionalSegmentImportedPackagesCount);
 			DataPtr += OptionalSegmentImportedPackageIdsMemSize;
 
-			++Data.TotalExportBundleCount;
 			HeaderData.ImportedAsyncPackagesView = Data.ImportedAsyncPackages.Left(ImportedPackagesCount);
 			OptionalSegmentHeaderData.ImportedAsyncPackagesView = Data.ImportedAsyncPackages.Right(OptionalSegmentImportedPackagesCount);
 		}
@@ -3907,12 +3915,9 @@ private:
 			OptionalSegmentExportBundleEntriesCopyMemSize = Align(OptionalSegmentHeaderData->ExportBundleEntries.Num() * sizeof(FExportBundleEntry), 8);
 		}
 #endif
-		const int32 ExportBundleNodeCount = Data.TotalExportBundleCount * EEventLoadNode2::ExportBundle_NumPhases;
-		const uint64 ExportBundleNodesMemSize = Align(sizeof(FEventLoadNode2) * ExportBundleNodeCount, 8);
 		const uint64 ExportsMemSize = Align(sizeof(FExportObject) * TotalExportCount, 8);
 
 		const uint64 MemoryBufferSize =
-			ExportBundleNodesMemSize +
 			ExportsMemSize +
 #if WITH_EDITOR
 			OptionalSegmentExportBundleEntriesCopyMemSize +
@@ -3932,8 +3937,6 @@ private:
 
 		Data.Exports = MakeArrayView(reinterpret_cast<FExportObject*>(DataPtr), TotalExportCount);
 		DataPtr += ExportsMemSize;
-		Data.ExportBundleNodes = MakeArrayView(reinterpret_cast<FEventLoadNode2*>(DataPtr), ExportBundleNodeCount);
-		DataPtr += ExportBundleNodesMemSize;
 		HeaderData.ExportBundleEntriesCopyForPostLoad = MakeArrayView(reinterpret_cast<FExportBundleEntry*>(DataPtr), HeaderData.ExportBundleEntries.Num());
 		FMemory::Memcpy(DataPtr, HeaderData.ExportBundleEntries.GetData(), HeaderData.ExportBundleEntries.Num() * sizeof(FExportBundleEntry));
 		DataPtr += ExportBundleEntriesCopyMemSize;
@@ -3953,7 +3956,6 @@ private:
 #endif
 
 		check(DataPtr - Data.MemoryBuffer1 == MemoryBufferSize);
-		AsyncPackage->CreateExportBundleNodes(EventSpecs.GetData());
 
 		AsyncPackage->ConstructedObjects.Reserve(Data.Exports.Num() + 1); // +1 for UPackage, may grow dynamically beyond that
 		for (FExportObject& Export : Data.Exports)
