CookDependency: Fix typo at end of line that converted int compare into a bool value (uncovered by static analysis compiler).
#rnx

diff --git a/Engine/Source/Runtime/CoreUObject/Public/Cooker/CookDependency.h b/Engine/Source/Runtime/CoreUObject/Public/Cooker/CookDependency.h
--- a/Engine/Source/Runtime/CoreUObject/Public/Cooker/CookDependency.h
+++ b/Engine/Source/Runtime/CoreUObject/Public/Cooker/CookDependency.h
@@ -329,7 +329,7 @@ inline bool FCookDependency::operator<(const FCookDependency& Other) const
 	case ECookDependency::TransitiveBuild:
 	{
 		// FName.Compare is lexical and case-insensitive, which is what we want
-		int32 Compare = TransitiveBuildData.PackageName.Compare(Other.TransitiveBuildData.PackageName) < 0;
+		int32 Compare = TransitiveBuildData.PackageName.Compare(Other.TransitiveBuildData.PackageName);
 		if (Compare != 0)
 		{
 			return Compare < 0;
