OnCommit around the interlock compare for the delete, delay the delete until after the transaction

diff --git a/Engine/Source/Runtime/CoreUObject/Private/UObject/UObjectArray.cpp b/Engine/Source/Runtime/CoreUObject/Private/UObject/UObjectArray.cpp
--- a/Engine/Source/Runtime/CoreUObject/Private/UObject/UObjectArray.cpp
+++ b/Engine/Source/Runtime/CoreUObject/Private/UObject/UObjectArray.cpp
@@ -72,10 +72,14 @@ void FUObjectItem::CreateStatID() const
 	PROFILER_CHAR* StoragePtr = new PROFILER_CHAR[NumStorageChars];
 	FMemory::Memcpy(StoragePtr, ConversionData.Get(), NumStorageChars * sizeof(PROFILER_CHAR));
 
-	if (FPlatformAtomics::InterlockedCompareExchangePointer((void**)&StatIDStringStorage, StoragePtr, nullptr) != nullptr)
+	// delay the delete of the StatIDStringStorage swap until after the transaction
+	UE_AUTORTFM_ONCOMMIT(
 	{
-		delete[] StoragePtr;
-	}
+		if (FPlatformAtomics::InterlockedCompareExchangePointer((void**)&StatIDStringStorage, StoragePtr, nullptr) != nullptr)
+		{
+			delete[] StoragePtr;
+		}
+	});
 
 	StatID = TStatId(StatIDStringStorage);
 #endif
