Disable TObjectPtr placeholder type flag checks in editor runtime (not in use atm).

diff --git a/Engine/Source/Runtime/CoreUObject/Public/UObject/ObjectHandleDefines.h b/Engine/Source/Runtime/CoreUObject/Public/UObject/ObjectHandleDefines.h
--- a/Engine/Source/Runtime/CoreUObject/Public/UObject/ObjectHandleDefines.h
+++ b/Engine/Source/Runtime/CoreUObject/Public/UObject/ObjectHandleDefines.h
@@ -17,5 +17,5 @@
 
 
 #ifndef UE_WITH_OBJECT_HANDLE_TYPE_SAFETY
-	#define UE_WITH_OBJECT_HANDLE_TYPE_SAFETY WITH_EDITORONLY_DATA
+	#define UE_WITH_OBJECT_HANDLE_TYPE_SAFETY 0
 #endif
\ No newline at end of file
