Capture this here as lambda is calling a function that requries a class instance

diff --git a/Engine/Source/Runtime/CoreUObject/Private/UObject/LinkerLoad.cpp b/Engine/Source/Runtime/CoreUObject/Private/UObject/LinkerLoad.cpp
--- a/Engine/Source/Runtime/CoreUObject/Private/UObject/LinkerLoad.cpp
+++ b/Engine/Source/Runtime/CoreUObject/Private/UObject/LinkerLoad.cpp
@@ -6233,7 +6233,7 @@ FArchive& FLinkerLoad::operator<<(FObjectPtr& ObjectPtr)
 	Ar << Index;
 
 	// Wrapper that only allows pointers to exports with placeholder types when type safety features are enabled.
-	auto AsTypeSafeObjectPtr_Lambda = [](UObject* ResolvedObject)
+	auto AsTypeSafeObjectPtr_Lambda = [this](UObject* ResolvedObject)
 	{
 #if WITH_EDITOR && !UE_WITH_OBJECT_HANDLE_TYPE_SAFETY
 		// If type safety features are disabled, resolve unsafe references to placeholder-typed objects now to NULL.
