Implement intrusive optional states for FName and FText.
Fix FOptionalProperty flags to not set CPF_ZeroConstructor for inner types with intrusive unset states.
Make assignment of TOptional consistent between intrusive and non-intrusive types - always destroy and reconstruct the value. This allows TOptional<const T> to work correctly whether T is intrusive or not.

diff --git a/Engine/Source/Runtime/Core/Private/UObject/UnrealNames.cpp b/Engine/Source/Runtime/Core/Private/UObject/UnrealNames.cpp
--- a/Engine/Source/Runtime/Core/Private/UObject/UnrealNames.cpp
+++ b/Engine/Source/Runtime/Core/Private/UObject/UnrealNames.cpp
@@ -43,6 +43,8 @@ PRAGMA_DISABLE_UNSAFE_TYPECAST_WARNINGS
 
 DEFINE_LOG_CATEGORY_STATIC(LogUnrealNames, Log, All);
 
+static_assert(::HasIntrusiveUnsetOptionalState<FName>());
+
 // Console command declarations
 namespace UE::Name::Private
 {
diff --git a/Engine/Source/Runtime/Core/Public/Internationalization/Text.h b/Engine/Source/Runtime/Core/Public/Internationalization/Text.h
--- a/Engine/Source/Runtime/Core/Public/Internationalization/Text.h
+++ b/Engine/Source/Runtime/Core/Public/Internationalization/Text.h
@@ -768,6 +768,29 @@ public:
 	static CORE_API FText ChangeKey( const FTextKey& Namespace, const FTextKey& Key, const FText& Text );
 #endif
 
+	/////////////////////////////////////////////////////
+	// Start - intrusive TOptional<FText> state        //
+	/////////////////////////////////////////////////////
+	constexpr static bool bHasIntrusiveUnsetOptionalState = true;
+	using IntrusiveUnsetOptionalStateType = FText;
+
+	explicit FText(FIntrusiveUnsetOptionalState)
+	: TextData(nullptr)
+	, Flags(0)
+	{
+	}
+	void operator=(FIntrusiveUnsetOptionalState)
+	{
+		TextData = nullptr;
+		Flags = 0;	
+	}
+	bool operator==(FIntrusiveUnsetOptionalState) const
+	{
+		return TextData == nullptr;
+	}
+	/////////////////////////////////////////////////
+	// End - intrusive TOptional<FText> state      //
+	/////////////////////////////////////////////////
 private:
 	template <
 		typename HistoryType
diff --git a/Engine/Source/Runtime/Core/Public/Misc/Optional.h b/Engine/Source/Runtime/Core/Public/Misc/Optional.h
--- a/Engine/Source/Runtime/Core/Public/Misc/Optional.h
+++ b/Engine/Source/Runtime/Core/Public/Misc/Optional.h
@@ -6,6 +6,7 @@
 #include "Misc/AssertionMacros.h"
 #include "Misc/IntrusiveUnsetOptionalState.h"
 #include "Misc/OptionalFwd.h"
+#include "Templates/MemoryOps.h"
 #include "Templates/UnrealTemplate.h"
 #include "Serialization/Archive.h"
 
@@ -163,18 +164,13 @@ public:
 	{
 		if (&Other != this)
 		{
-			if constexpr (bUsingIntrusiveUnsetState)
+			if (Other.IsSet())
 			{
-				*(OptionalType*)&Value = *(const OptionalType*)&Other.Value;
+				Emplace(Other.GetValue());
 			}
 			else
 			{
 				Reset();
-				if (Other.Value.bIsSet)
-				{
-					::new((void*)&Value) OptionalType(*(const OptionalType*)&Other.Value);
-					Value.bIsSet = true;
-				}
 			}
 		}
 		return *this;
@@ -183,18 +179,13 @@ public:
 	{
 		if (&Other != this)
 		{
-			if constexpr (bUsingIntrusiveUnsetState)
+			if(Other.IsSet())
 			{
-				*(OptionalType*)&Value = MoveTempIfPossible(*(OptionalType*)&Other.Value);
+				Emplace(MoveTempIfPossible(Other.GetValue()));
 			}
 			else
 			{
 				Reset();
-				if (Other.Value.bIsSet)
-				{
-					::new((void*)&Value) OptionalType(MoveTempIfPossible(*(OptionalType*)&Other.Value));
-					Value.bIsSet = true;
-				}
 			}
 		}
 		return *this;
@@ -219,19 +210,16 @@ public:
 
 	void Reset()
 	{
-		if constexpr (bUsingIntrusiveUnsetState)
-		{
-			*(OptionalType*)&Value = FIntrusiveUnsetOptionalState{};
-		}
-		else
+		if (IsSet())
 		{
-			if (Value.bIsSet)
+			DestroyValue();
+			if constexpr (bUsingIntrusiveUnsetState)
+			{
+				::new((void*)&Value) OptionalType(FIntrusiveUnsetOptionalState{});
+			}
+			else
 			{
 				Value.bIsSet = false;
-
-				// We need a typedef here because VC won't compile the destructor call below if OptionalType itself has a member called OptionalType
-				typedef OptionalType OptionalDestructOptionalType;
-				((OptionalType*)&Value)->OptionalDestructOptionalType::~OptionalDestructOptionalType();
 			}
 		}
 	}
@@ -239,17 +227,10 @@ public:
 	template <typename... ArgsType>
 	OptionalType& Emplace(ArgsType&&... Args)
 	{
-		if constexpr (bUsingIntrusiveUnsetState)
+		// Destroy the member in-place before replacing it - a bit nasty, but it'll work since we don't support exceptions
+		if (IsSet())
 		{
-			// Destroy the member in-place before replacing it - a bit nasty, but it'll work since we don't support exceptions
-
-			// We need a typedef here because VC won't compile the destructor call below if OptionalType itself has a member called OptionalType
-			typedef OptionalType OptionalDestructOptionalType;
-			((OptionalType*)&Value)->OptionalDestructOptionalType::~OptionalDestructOptionalType();
-		}
-		else
-		{
-			Reset();
+			DestroyValue();
 		}
 
 		// If this fails to compile when trying to call Emplace with a non-public constructor,
@@ -406,6 +387,14 @@ public:
 	}
 
 private:
+	/** 
+	 * Destroys the value, must only be called if the value is set, and callers must then mark the value unset or construct a new value in place. 
+	 */
+	FORCEINLINE void DestroyValue()
+	{
+		DestructItem((OptionalType*)&Value);
+	}
+
 	using ValueStorageType = std::conditional_t<bUsingIntrusiveUnsetState, uint8[sizeof(OptionalType)], UE::Core::Private::TNonIntrusiveOptionalStorage<sizeof(OptionalType)>>;
 	alignas(OptionalType) ValueStorageType Value;
 };
diff --git a/Engine/Source/Runtime/Core/Public/UObject/NameTypes.h b/Engine/Source/Runtime/Core/Public/UObject/NameTypes.h
--- a/Engine/Source/Runtime/Core/Public/UObject/NameTypes.h
+++ b/Engine/Source/Runtime/Core/Public/UObject/NameTypes.h
@@ -14,6 +14,7 @@
 #include "UObject/UnrealNames.h"
 #include "Templates/Atomic.h"
 #include "Serialization/MemoryLayout.h"
+#include "Misc/Optional.h"
 #include "Misc/StringBuilder.h"
 #include "Trace/Trace.h"
 
@@ -50,9 +51,14 @@ class FName;
 /** Opaque id to a deduplicated name */
 struct FNameEntryId
 {
+	constexpr static bool bHasIntrusiveUnsetOptionalState = true;
+	using IntrusiveUnsetOptionalStateType = FNameEntryId;
+	
 	// Default initialize to be equal to NAME_None
 	FNameEntryId() : Value(0) {}
 	FNameEntryId(ENoInit) {}
+	explicit FNameEntryId(FIntrusiveUnsetOptionalState) : Value(~0u) {}
+	void operator=(FIntrusiveUnsetOptionalState) { Value = ~0u; }
 
 	bool IsNone() const
 	{
@@ -78,6 +84,12 @@ struct FNameEntryId
 	bool operator>(FNameEntryId Rhs) const { return Rhs.Value < Value; }
 	bool operator==(FNameEntryId Rhs) const { return Value == Rhs.Value; }
 	bool operator!=(FNameEntryId Rhs) const { return Value != Rhs.Value; }
+	
+	/** Comparison against special type for checking TOptional<FName>:IsSet */
+	bool operator==(FIntrusiveUnsetOptionalState) const
+	{
+		return Value == ~0u;
+	}
 
 	// Returns true if this FNameEntryId is not equivalent to NAME_None
 	explicit operator bool() const { return Value != 0; }
@@ -572,6 +584,9 @@ private:
 class FName
 {
 public:
+	constexpr static bool bHasIntrusiveUnsetOptionalState = true;
+	using IntrusiveUnsetOptionalStateType = FName;
+
 #if UE_FNAME_OUTLINE_NUMBER
 	CORE_API FNameEntryId GetComparisonIndex() const;
 	CORE_API FNameEntryId GetDisplayIndex() const;
@@ -706,6 +721,12 @@ public:
 		return !(*this == Other);
 	}
 
+	/** Special comparison operator for TOptional<FName>::IsSet */
+	inline bool operator==(FIntrusiveUnsetOptionalState I) const
+	{
+		return ComparisonIndex == I;
+	}
+
 	/** Fast non-alphabetical order that is only stable during this process' lifetime. */
 	FORCEINLINE bool FastLess(const FName& Other) const
 	{
@@ -921,6 +942,22 @@ public:
 #endif
 	{}
 
+	/** Special constructor used by TOptional<FName> */
+	explicit FName(FIntrusiveUnsetOptionalState I)
+	: FName()
+	{
+		ComparisonIndex = I;
+#if WITH_CASE_PRESERVING_NAME
+		DisplayIndex = I;
+#endif
+	}
+
+	/** Special assignment operator used by TOptional<FName> */
+	inline void operator=(FIntrusiveUnsetOptionalState I)
+	{
+		*this = FName(I);
+	}
+
 	FORCEINLINE explicit FName(FMinimalName InName);
 	FORCEINLINE explicit FName(FScriptName InName);
 	FORCEINLINE FName(FMemoryImageName InName);
diff --git a/Engine/Source/Runtime/CoreUObject/Private/UObject/PropertyOptional.cpp b/Engine/Source/Runtime/CoreUObject/Private/UObject/PropertyOptional.cpp
--- a/Engine/Source/Runtime/CoreUObject/Private/UObject/PropertyOptional.cpp
+++ b/Engine/Source/Runtime/CoreUObject/Private/UObject/PropertyOptional.cpp
@@ -158,8 +158,11 @@ void FOptionalProperty::LinkInternal(FArchive& Ar)
 	// After ValueProperty's size has been computed, compute the size of this property.
 	ElementSize = CalcSize();
 
-	// Optional properties can always be initialized by zeroing memory.
-	PropertyFlags |= CPF_ZeroConstructor;
+	// Standard non-intrusive optional properties can always be initialized by zeroing memory.
+	if (!ValueProperty->HasIntrusiveUnsetOptionalState())
+	{
+		PropertyFlags |= CPF_ZeroConstructor;
+	}
 
 	// Propagate CPF_NoDestructor, CPF_IsPlainOldData, and CPF_HasGetValueTypeHash from the value property.
 	PropertyFlags |= (ValueProperty->PropertyFlags & (CPF_NoDestructor|CPF_IsPlainOldData|CPF_HasGetValueTypeHash));
diff --git a/Engine/Source/Runtime/CoreUObject/Tests/OptionalPropertyTest.cpp b/Engine/Source/Runtime/CoreUObject/Tests/OptionalPropertyTest.cpp
--- a/Engine/Source/Runtime/CoreUObject/Tests/OptionalPropertyTest.cpp
+++ b/Engine/Source/Runtime/CoreUObject/Tests/OptionalPropertyTest.cpp
@@ -9,150 +9,255 @@
 #include "UObject/PropertyOptional.h"
 #include "UObject/UObjectGlobals.h"
 
+struct FOptionalTestObject
+{
+    UOptionalPropertyTestObject* Obj;
+    FOptionalProperty* StringProperty;
+    FOptionalProperty* TextProperty;
+    FOptionalProperty* NameProperty;
+    FOptionalProperty* IntProperty;
+
+    FOptionalTestObject()
+    {
+        const FName TestPackageName(TEXT("/Engine/TestPackage"));
+        UPackage* TestPackage = NewObject<UPackage>(nullptr, TestPackageName, RF_Transient);
+        Obj = NewObject<UOptionalPropertyTestObject>(TestPackage);
+
+        UClass* Class = Obj->GetClass();
+        StringProperty = CastField<FOptionalProperty>(Class->FindPropertyByName("OptionalString"));
+        REQUIRE(StringProperty != nullptr);
+        TextProperty = CastField<FOptionalProperty>(Class->FindPropertyByName("OptionalText"));
+        REQUIRE(TextProperty != nullptr);
+        NameProperty = CastField<FOptionalProperty>(Class->FindPropertyByName("OptionalName"));
+        REQUIRE(NameProperty != nullptr);
+        IntProperty = CastField<FOptionalProperty>(Class->FindPropertyByName("OptionalInt"));
+        REQUIRE(IntProperty != nullptr);
+    }
+
+    bool IsValid()
+    {
+        CHECK(StringProperty != nullptr);
+        CHECK(TextProperty != nullptr);
+        CHECK(NameProperty != nullptr);
+        CHECK(IntProperty != nullptr);
+        return true;
+    }
+};
+
 TEST_CASE_NAMED(FOptionalPropertyTestSize, "UE::CoreUObject::OptionalProperty::Size", "[Core][UObject][SmokeFilter]")
 {
-	const FName TestPackageName(TEXT("/Engine/TestPackage"));
-	UPackage* TestPackage = NewObject<UPackage>(nullptr, TestPackageName, RF_Transient);
-    UOptionalPropertyTestObject* Obj = NewObject<UOptionalPropertyTestObject>(TestPackage);
-
-    UClass* Class = Obj->GetClass();
-    FProperty* StringProperty = Class->FindPropertyByName("OptionalString");
-    REQUIRE(StringProperty != nullptr);
-    FProperty* TextProperty = Class->FindPropertyByName("OptionalText");
-    REQUIRE(TextProperty != nullptr);
-    FProperty* NameProperty = Class->FindPropertyByName("OptionalName");
-    REQUIRE(NameProperty != nullptr);
-    FProperty* IntProperty = Class->FindPropertyByName("OptionalInt");
-    REQUIRE(IntProperty != nullptr);
-
-    CHECK(StringProperty->GetSize() == sizeof(Obj->OptionalString));
-    CHECK(TextProperty->GetSize() == sizeof(Obj->OptionalText));
-    CHECK(NameProperty->GetSize() == sizeof(Obj->OptionalName));
-    CHECK(IntProperty->GetSize() == sizeof(Obj->OptionalInt));
+	FOptionalTestObject TestData = FOptionalTestObject();
+    REQUIRE(TestData.IsValid());
+
+    CHECK(TestData.StringProperty->GetSize() == sizeof(TestData.Obj->OptionalString));
+    CHECK(TestData.TextProperty->GetSize() == sizeof(TestData.Obj->OptionalText));
+    CHECK(TestData.NameProperty->GetSize() == sizeof(TestData.Obj->OptionalName));
+    CHECK(TestData.IntProperty->GetSize() == sizeof(TestData.Obj->OptionalInt));
+}
+
+TEST_CASE_NAMED(FOptionalPropertyTestInitializeValue, "UE::CoreUObject::OptionalProperty::InitializeValue", "[Core][UObject][SmokeFilter]")
+{
+	FOptionalTestObject TestData = FOptionalTestObject();
+    REQUIRE(TestData.IsValid());
+
+    uint8 OptionalStringStorage[sizeof(TOptional<FString>)];
+    uint8 OptionalTextStorage[sizeof(TOptional<FText>)];
+    uint8 OptionalNameStorage[sizeof(TOptional<FName>)];
+    uint8 OptionalIntStorage[sizeof(TOptional<int32>)];
+
+    TestData.StringProperty->InitializeValue(OptionalStringStorage);
+    TestData.TextProperty->InitializeValue(OptionalTextStorage);
+    TestData.NameProperty->InitializeValue(OptionalNameStorage);
+    TestData.IntProperty->InitializeValue(OptionalIntStorage);
+
+    CHECK_FALSE(reinterpret_cast<TOptional<FString>*>(OptionalStringStorage)->IsSet());
+    CHECK_FALSE(reinterpret_cast<TOptional<FText>*>(OptionalTextStorage)->IsSet());
+    CHECK_FALSE(reinterpret_cast<TOptional<FName>*>(OptionalNameStorage)->IsSet());
+    CHECK_FALSE(reinterpret_cast<TOptional<int32>*>(OptionalIntStorage)->IsSet());
 }
 
 TEST_CASE_NAMED(FOptionalPropertyTestClearValue, "UE::CoreUObject::OptionalProperty::ClearValue", "[Core][UObject][SmokeFilter]")
 {
-	const FName TestPackageName(TEXT("/Engine/TestPackage"));
-	UPackage* TestPackage = NewObject<UPackage>(nullptr, TestPackageName, RF_Transient);
-    UOptionalPropertyTestObject* Obj = NewObject<UOptionalPropertyTestObject>(TestPackage);
-
-    UClass* Class = Obj->GetClass();
-    FProperty* StringProperty = Class->FindPropertyByName("OptionalString");
-    REQUIRE(StringProperty != nullptr);
-    FProperty* TextProperty = Class->FindPropertyByName("OptionalText");
-    REQUIRE(TextProperty != nullptr);
-    FProperty* NameProperty = Class->FindPropertyByName("OptionalName");
-    REQUIRE(NameProperty != nullptr);
-    FProperty* IntProperty = Class->FindPropertyByName("OptionalInt");
-    REQUIRE(IntProperty != nullptr);
-
-    Obj->OptionalString.Emplace(TEXT("Optional"));
-    Obj->OptionalText.Emplace(FText::FromStringView(TEXTVIEW("Optional")));
-    Obj->OptionalName.Emplace(TEXT("Optional"));
-    Obj->OptionalInt.Emplace(42);
-
-    StringProperty->ClearValue(&Obj->OptionalString);
-    CHECK_FALSE(Obj->OptionalString.IsSet());
-    TextProperty->ClearValue(&Obj->OptionalText);
-    CHECK_FALSE(Obj->OptionalText.IsSet());
-    NameProperty->ClearValue(&Obj->OptionalName);
-    CHECK_FALSE(Obj->OptionalName.IsSet());
-    IntProperty->ClearValue(&Obj->OptionalInt);
-    CHECK_FALSE(Obj->OptionalInt.IsSet());
+	FOptionalTestObject TestData = FOptionalTestObject();
+    REQUIRE(TestData.IsValid());
+
+    TestData.Obj->OptionalString.Emplace(TEXT("Optional"));
+    TestData.Obj->OptionalText.Emplace(FText::FromStringView(TEXTVIEW("Optional")));
+    TestData.Obj->OptionalName.Emplace(TEXT("Optional"));
+    TestData.Obj->OptionalInt.Emplace(42);
+
+    TestData.StringProperty->ClearValue(&TestData.Obj->OptionalString);
+    CHECK_FALSE(TestData.Obj->OptionalString.IsSet());
+    TestData.TextProperty->ClearValue(&TestData.Obj->OptionalText);
+    CHECK_FALSE(TestData.Obj->OptionalText.IsSet());
+    TestData.NameProperty->ClearValue(&TestData.Obj->OptionalName);
+    CHECK_FALSE(TestData.Obj->OptionalName.IsSet());
+    TestData.IntProperty->ClearValue(&TestData.Obj->OptionalInt);
+    CHECK_FALSE(TestData.Obj->OptionalInt.IsSet());
 }
 
 TEST_CASE_NAMED(FOptionalPropertyTestCopyValueIn, "UE::CoreUObject::OptionalProperty::CopyValueIn", "[Core][UObject][SmokeFilter]")
 {
-	const FName TestPackageName(TEXT("/Engine/TestPackage"));
-	UPackage* TestPackage = NewObject<UPackage>(nullptr, TestPackageName, RF_Transient);
-    UOptionalPropertyTestObject* Obj = NewObject<UOptionalPropertyTestObject>(TestPackage);
-
-    UClass* Class = Obj->GetClass();
-    FProperty* StringProperty = Class->FindPropertyByName("OptionalString");
-    REQUIRE(StringProperty != nullptr);
-    FProperty* TextProperty = Class->FindPropertyByName("OptionalText");
-    REQUIRE(TextProperty != nullptr);
-    FProperty* NameProperty = Class->FindPropertyByName("OptionalName");
-    REQUIRE(NameProperty != nullptr);
-    FProperty* IntProperty = Class->FindPropertyByName("OptionalInt");
-    REQUIRE(IntProperty != nullptr);
+	FOptionalTestObject TestData = FOptionalTestObject();
+    REQUIRE(TestData.IsValid());
 
     TOptional<FString> OptString(FString(TEXT("Optional")));
     TOptional<FText> OptText(FText::FromStringView(TEXTVIEW("Optional")));
     TOptional<FName> OptName("Optional");
     TOptional<int32> OptInt(58);
 
-    StringProperty->CopySingleValue(&Obj->OptionalString, &OptString);
-    TextProperty->CopySingleValue(&Obj->OptionalText, &OptText);
-    NameProperty->CopySingleValue(&Obj->OptionalName, &OptName);
-    IntProperty->CopySingleValue(&Obj->OptionalInt, &OptInt);
-
-    CHECK(Obj->OptionalString.IsSet());
-    CHECK(Obj->OptionalString.Get(FString()) == OptString.GetValue());
-    CHECK(Obj->OptionalText.IsSet());
-    CHECK(Obj->OptionalText.Get(FText::GetEmpty()).EqualTo(OptText.GetValue()));
-    CHECK(Obj->OptionalName.IsSet());
-    CHECK(Obj->OptionalName.Get(FName()) == OptName.GetValue());
-    CHECK(Obj->OptionalInt.IsSet());
-    CHECK(Obj->OptionalInt.Get(0) == OptInt.GetValue());
+    TestData.StringProperty->CopySingleValue(&TestData.Obj->OptionalString, &OptString);
+    TestData.TextProperty->CopySingleValue(&TestData.Obj->OptionalText, &OptText);
+    TestData.NameProperty->CopySingleValue(&TestData.Obj->OptionalName, &OptName);
+    TestData.IntProperty->CopySingleValue(&TestData.Obj->OptionalInt, &OptInt);
+
+    CHECK(OptString.IsSet());
+    CHECK(TestData.Obj->OptionalString.IsSet());
+    CHECK(TestData.Obj->OptionalString.Get(FString()) == OptString.GetValue());
+    CHECK(OptText.IsSet());
+    CHECK(TestData.Obj->OptionalText.IsSet());
+    CHECK(TestData.Obj->OptionalText.Get(FText::GetEmpty()).EqualTo(OptText.GetValue()));
+    CHECK(OptName.IsSet());
+    CHECK(TestData.Obj->OptionalName.IsSet());
+    CHECK(TestData.Obj->OptionalName.Get(FName()) == OptName.GetValue());
+    CHECK(OptInt.IsSet());
+    CHECK(TestData.Obj->OptionalInt.IsSet());
+    CHECK(TestData.Obj->OptionalInt.Get(0) == OptInt.GetValue());
 }
 
-TEST_CASE_NAMED(FOptionalPropertyTestLayout, "UE::CoreUObject::OptionalProperty::OptionalPropertyLayout", "[Core][UObject][SmokeFilter]")
+TEST_CASE_NAMED(FOptionalPropertyTestCopyValueOut, "UE::CoreUObject::OptionalProperty::CopyValueOut", "[Core][UObject][SmokeFilter]")
+{
+	FOptionalTestObject TestData = FOptionalTestObject();
+    REQUIRE(TestData.IsValid());
+
+    TOptional<FString> OptString(FString(TEXT("Optional")));
+    TOptional<FText> OptText(FText::FromStringView(TEXTVIEW("Optional")));
+    TOptional<FName> OptName("Optional");
+    TOptional<int32> OptInt(58);
+
+    TestData.StringProperty->CopySingleValue(&TestData.Obj->OptionalString, &OptString);
+    TestData.TextProperty->CopySingleValue(&TestData.Obj->OptionalText, &OptText);
+    TestData.NameProperty->CopySingleValue(&TestData.Obj->OptionalName, &OptName);
+    TestData.IntProperty->CopySingleValue(&TestData.Obj->OptionalInt, &OptInt);
+
+    CHECK(OptString.IsSet());
+    CHECK(TestData.Obj->OptionalString.IsSet());
+    CHECK(TestData.Obj->OptionalString.Get(FString()) == OptString.GetValue());
+    CHECK(OptText.IsSet());
+    CHECK(TestData.Obj->OptionalText.IsSet());
+    CHECK(TestData.Obj->OptionalText.Get(FText::GetEmpty()).EqualTo(OptText.GetValue()));
+    CHECK(OptName.IsSet());
+    CHECK(TestData.Obj->OptionalName.IsSet());
+    CHECK(TestData.Obj->OptionalName.Get(FName()) == OptName.GetValue());
+    CHECK(OptInt.IsSet());
+    CHECK(TestData.Obj->OptionalInt.IsSet());
+    CHECK(TestData.Obj->OptionalInt.Get(0) == OptInt.GetValue());
+}
+
+TEST_CASE_NAMED(FOptionalPropertyTestIdentical, "UE::CoreUObject::OptionalProperty::Identical", "[Core][UObject][SmokeFilter]")
 {
-	const FName TestPackageName(TEXT("/Engine/TestPackage"));
-	UPackage* TestPackage = NewObject<UPackage>(nullptr, TestPackageName, RF_Transient);
-    UOptionalPropertyTestObject* Obj = NewObject<UOptionalPropertyTestObject>(TestPackage);
-
-    UClass* Class = Obj->GetClass();
-    FOptionalProperty* StringProperty = CastField<FOptionalProperty>(Class->FindPropertyByName("OptionalString"));
-    REQUIRE(StringProperty != nullptr);
-    FOptionalProperty* TextProperty = CastField<FOptionalProperty>(Class->FindPropertyByName("OptionalText"));
-    REQUIRE(TextProperty != nullptr);
-    FOptionalProperty* NameProperty = CastField<FOptionalProperty>(Class->FindPropertyByName("OptionalName"));
-    REQUIRE(NameProperty != nullptr);
-    FOptionalProperty* IntProperty = CastField<FOptionalProperty>(Class->FindPropertyByName("OptionalInt"));
-    REQUIRE(IntProperty != nullptr);
-
-    FOptionalPropertyLayout StringPropertyLayout(StringProperty->GetValueProperty());
-    FOptionalPropertyLayout TextPropertyLayout(TextProperty->GetValueProperty());
-    FOptionalPropertyLayout NamePropertyLayout(NameProperty->GetValueProperty());
-    FOptionalPropertyLayout IntPropertyLayout(IntProperty->GetValueProperty());
-
-    CHECK_FALSE(StringPropertyLayout.IsSet(&Obj->OptionalString));
-    CHECK_FALSE(TextPropertyLayout.IsSet(&Obj->OptionalText));
-    CHECK_FALSE(NamePropertyLayout.IsSet(&Obj->OptionalName));
-    CHECK_FALSE(IntPropertyLayout.IsSet(&Obj->OptionalInt));
-
-    FString* InnerString = (FString*)StringPropertyLayout.MarkSetAndGetInitializedValuePointerToReplace(&Obj->OptionalString);
-    CHECK(Obj->OptionalString.IsSet());
-    *InnerString = TEXT("Optional");
-    CHECK(Obj->OptionalString.GetValue() == TEXT("Optional"));
+	FOptionalTestObject TestData = FOptionalTestObject();
+    REQUIRE(TestData.IsValid());
+
+    TOptional<FString> UnsetOptString;
+    TOptional<FText> UnsetOptText;
+    TOptional<FName> UnsetOptName;
+    TOptional<int32> UnsetOptInt;
+
+    CHECK(TestData.StringProperty->Identical(&TestData.Obj->OptionalString, &UnsetOptString, PPF_None));
+    CHECK(TestData.TextProperty->Identical(&TestData.Obj->OptionalText, &UnsetOptText, PPF_None));
+    CHECK(TestData.NameProperty->Identical(&TestData.Obj->OptionalName, &UnsetOptName, PPF_None));
+    CHECK(TestData.IntProperty->Identical(&TestData.Obj->OptionalInt, &UnsetOptInt, PPF_None));
+
+    TOptional<FString> OptString(FString(TEXT("Optional")));
+    TOptional<FText> OptText(FText::FromStringView(TEXTVIEW("Optional")));
+    TOptional<FName> OptName("Optional");
+    TOptional<int32> OptInt(58);
+
+    CHECK_FALSE(TestData.StringProperty->Identical(&TestData.Obj->OptionalString, &OptString, PPF_None));
+    CHECK_FALSE(TestData.TextProperty->Identical(&TestData.Obj->OptionalText, &OptText, PPF_None));
+    CHECK_FALSE(TestData.NameProperty->Identical(&TestData.Obj->OptionalName, &OptName, PPF_None));
+    CHECK_FALSE(TestData.IntProperty->Identical(&TestData.Obj->OptionalInt, &OptInt, PPF_None));
+
+    TestData.Obj->OptionalString = OptString;
+    TestData.Obj->OptionalText = OptText;
+    TestData.Obj->OptionalName = OptName;
+    TestData.Obj->OptionalInt = OptInt;
+
+    CHECK(TestData.StringProperty->Identical(&TestData.Obj->OptionalString, &OptString, PPF_None));
+    CHECK(TestData.TextProperty->Identical(&TestData.Obj->OptionalText, &OptText, PPF_None));
+    CHECK(TestData.NameProperty->Identical(&TestData.Obj->OptionalName, &OptName, PPF_None));
+    CHECK(TestData.IntProperty->Identical(&TestData.Obj->OptionalInt, &OptInt, PPF_None));
     
-    FText* InnerText = (FText*)TextPropertyLayout.MarkSetAndGetInitializedValuePointerToReplace(&Obj->OptionalText);
-    CHECK(Obj->OptionalText.IsSet());
-    *InnerText = FText::FromStringView(TEXTVIEW("Optional"));
-    CHECK(Obj->OptionalText.GetValue().ToString() == TEXT("Optional"));
-
-    FName* InnerName = (FName*)NamePropertyLayout.MarkSetAndGetInitializedValuePointerToReplace(&Obj->OptionalName);
-    CHECK(Obj->OptionalName.IsSet());
-    *InnerName = FName("Optional");
-    CHECK(Obj->OptionalName.GetValue() == FName("Optional"));
-
-    int32* InnerInt = (int32*)IntPropertyLayout.MarkSetAndGetInitializedValuePointerToReplace(&Obj->OptionalInt);
-    CHECK(Obj->OptionalInt.IsSet());
-    *InnerInt = 79;
-    CHECK(Obj->OptionalInt.GetValue() == 79);
-
-    StringPropertyLayout.MarkUnset(&Obj->OptionalString);
-    TextPropertyLayout.MarkUnset(&Obj->OptionalText);
-    NamePropertyLayout.MarkUnset(&Obj->OptionalName);
-    IntPropertyLayout.MarkUnset(&Obj->OptionalInt);
-
-    CHECK_FALSE(StringPropertyLayout.IsSet(&Obj->OptionalString));
-    CHECK_FALSE(TextPropertyLayout.IsSet(&Obj->OptionalText));
-    CHECK_FALSE(NamePropertyLayout.IsSet(&Obj->OptionalName));
-    CHECK_FALSE(IntPropertyLayout.IsSet(&Obj->OptionalInt));
+    CHECK_FALSE(TestData.StringProperty->Identical(&TestData.Obj->OptionalString, &UnsetOptString, PPF_None));
+    CHECK_FALSE(TestData.TextProperty->Identical(&TestData.Obj->OptionalText, &UnsetOptText, PPF_None));
+    CHECK_FALSE(TestData.NameProperty->Identical(&TestData.Obj->OptionalName, &UnsetOptName, PPF_None));
+    CHECK_FALSE(TestData.IntProperty->Identical(&TestData.Obj->OptionalInt, &UnsetOptInt, PPF_None));
+}
+
+TEST_CASE_NAMED(FOptionalPropertyTestGetValueTypeHash, "UE::CoreUObject::OptionalProperty::GetValueTypeHash", "[Core][UObject][SmokeFilter]")
+{
+	FOptionalTestObject TestData = FOptionalTestObject();
+    REQUIRE(TestData.IsValid());
+
+    TOptional<FString> OptString(FString(TEXT("Optional")));
+    // GetTypeHash is undefined for FText 
+    // TOptional<FText> OptText(FText::FromStringView(TEXTVIEW("Optional")));
+    TOptional<FName> OptName("Optional");
+    TOptional<int32> OptInt(93);
+
+    CHECK(GetTypeHash(OptString) == TestData.StringProperty->GetValueTypeHash(&OptString));
+    // GetTypeHash is undefined for FText 
+    // CHECK(GetTypeHash(OptText) == TestData.TextProperty->GetValueTypeHash(&OptText));
+    CHECK_FALSE(TestData.TextProperty->HasAllPropertyFlags(CPF_HasGetValueTypeHash));
+    CHECK(GetTypeHash(OptName) == TestData.NameProperty->GetValueTypeHash(&OptName));
+    CHECK(GetTypeHash(OptInt) == TestData.IntProperty->GetValueTypeHash(&OptInt));
+}
+
+TEST_CASE_NAMED(FOptionalPropertyTestLayout, "UE::CoreUObject::OptionalProperty::OptionalPropertyLayout", "[Core][UObject][SmokeFilter]")
+{
+	FOptionalTestObject TestData = FOptionalTestObject();
+	REQUIRE(TestData.IsValid());
+
+	FOptionalPropertyLayout StringPropertyLayout(TestData.StringProperty->GetValueProperty());
+	FOptionalPropertyLayout TextPropertyLayout(TestData.TextProperty->GetValueProperty());
+	FOptionalPropertyLayout NamePropertyLayout(TestData.NameProperty->GetValueProperty());
+	FOptionalPropertyLayout IntPropertyLayout(TestData.IntProperty->GetValueProperty());
+
+	CHECK_FALSE(StringPropertyLayout.IsSet(&TestData.Obj->OptionalString));
+	CHECK_FALSE(TextPropertyLayout.IsSet(&TestData.Obj->OptionalText));
+	CHECK_FALSE(NamePropertyLayout.IsSet(&TestData.Obj->OptionalName));
+	CHECK_FALSE(IntPropertyLayout.IsSet(&TestData.Obj->OptionalInt));
+
+	FString* InnerString = (FString*)StringPropertyLayout.MarkSetAndGetInitializedValuePointerToReplace(&TestData.Obj->OptionalString);
+	CHECK(TestData.Obj->OptionalString.IsSet());
+	*InnerString = TEXT("Optional");
+	CHECK(TestData.Obj->OptionalString.GetValue() == TEXT("Optional"));
+
+	FText* InnerText = (FText*)TextPropertyLayout.MarkSetAndGetInitializedValuePointerToReplace(&TestData.Obj->OptionalText);
+	CHECK(TestData.Obj->OptionalText.IsSet());
+	*InnerText = FText::FromStringView(TEXTVIEW("Optional"));
+	CHECK(TestData.Obj->OptionalText.GetValue().ToString() == TEXT("Optional"));
+
+	FName* InnerName = (FName*)NamePropertyLayout.MarkSetAndGetInitializedValuePointerToReplace(&TestData.Obj->OptionalName);
+	CHECK(TestData.Obj->OptionalName.IsSet());
+	*InnerName = FName("Optional");
+	CHECK(TestData.Obj->OptionalName.GetValue() == FName("Optional"));
+
+	int32* InnerInt = (int32*)IntPropertyLayout.MarkSetAndGetInitializedValuePointerToReplace(&TestData.Obj->OptionalInt);
+	CHECK(TestData.Obj->OptionalInt.IsSet());
+	*InnerInt = 79;
+	CHECK(TestData.Obj->OptionalInt.GetValue() == 79);
+
+	StringPropertyLayout.MarkUnset(&TestData.Obj->OptionalString);
+	TextPropertyLayout.MarkUnset(&TestData.Obj->OptionalText);
+	NamePropertyLayout.MarkUnset(&TestData.Obj->OptionalName);
+	IntPropertyLayout.MarkUnset(&TestData.Obj->OptionalInt);
+
+	CHECK_FALSE(StringPropertyLayout.IsSet(&TestData.Obj->OptionalString));
+	CHECK_FALSE(TextPropertyLayout.IsSet(&TestData.Obj->OptionalText));
+	CHECK_FALSE(NamePropertyLayout.IsSet(&TestData.Obj->OptionalName));
+	CHECK_FALSE(IntPropertyLayout.IsSet(&TestData.Obj->OptionalInt));
 }
 
 #endif
\ No newline at end of file
