Don't return true ClassHasInstancesAsyncLoading if it's 'just' the CDO and/or some garbage objects. Since these won't go away until a GC is ran failure to indicate no pending loads can cause n^2 complexity for some object loading patterns

diff --git a/Engine/Source/Runtime/CoreUObject/Private/UObject/UObjectHash.cpp b/Engine/Source/Runtime/CoreUObject/Private/UObject/UObjectHash.cpp
--- a/Engine/Source/Runtime/CoreUObject/Private/UObject/UObjectHash.cpp
+++ b/Engine/Source/Runtime/CoreUObject/Private/UObject/UObjectHash.cpp
@@ -1574,7 +1574,13 @@ bool ClassHasInstancesAsyncLoading(const UClass* ClassToLookFor)
 			for (auto ObjectIt = List->CreateIterator(); ObjectIt; ++ObjectIt)
 			{
 				UObject *Object = static_cast<UObject*>(*ObjectIt);
-				if (Object->HasAnyInternalFlags(EInternalObjectFlags::AsyncLoading))
+				// If the object is async loading we'll want to indicate as such to the caller,
+				// excepting two cases - garbage objects and the CDO:
+				if (Object->HasAnyInternalFlags(EInternalObjectFlags::AsyncLoading) 
+					// garbage objects won't require that the class be kept alive:
+					&& !Object->HasAnyInternalFlags(EInternalObjectFlags::Garbage) 
+					// CDO is required and owned by the class - also doesn't need to keep the class alive:
+					&& !Object->HasAnyFlags(RF_ClassDefaultObject)) 
 				{
 					return true;
 				}
