Do not crash when waiting or cancelling a bulk data request that has not been started.

diff --git a/Engine/Source/Runtime/CoreUObject/Private/Serialization/BulkDataStreaming.cpp b/Engine/Source/Runtime/CoreUObject/Private/Serialization/BulkDataStreaming.cpp
--- a/Engine/Source/Runtime/CoreUObject/Private/Serialization/BulkDataStreaming.cpp
+++ b/Engine/Source/Runtime/CoreUObject/Private/Serialization/BulkDataStreaming.cpp
@@ -936,7 +936,7 @@ FBulkDataRequest& FBulkDataRequest::operator=(FBulkDataRequest&& Other)
 
 FBulkDataRequest::EStatus FBulkDataRequest::GetStatus() const
 {
-	return Handle != nullptr ? Handle->GetStatus() : EStatus::None;
+	return Handle.IsValid() ? Handle->GetStatus() : EStatus::None;
 }
 
 bool FBulkDataRequest::IsNone() const
@@ -962,8 +962,10 @@ bool FBulkDataRequest::IsCompleted() const
 
 void FBulkDataRequest::Cancel()
 {
-	check(Handle);
-	Handle->Cancel();
+	if (Handle.IsValid())
+	{
+		Handle->Cancel();
+	}
 }
 
 void FBulkDataRequest::Reset()
@@ -985,14 +987,20 @@ int32 FBulkDataRequest::GetRefCount() const
 
 void FBulkDataBatchRequest::Wait()
 {
-	check(Handle);
-	Handle->Wait(MAX_uint32);
+	if (Handle.IsValid())
+	{
+		Handle->Wait(MAX_uint32);
+	}
 }
 
 bool FBulkDataBatchRequest::WaitFor(uint32 Milliseconds)
 {
-	check(Handle);
-	return Handle->Wait(Milliseconds);
+	if (Handle.IsValid())
+	{
+		return Handle->Wait(Milliseconds);
+	}
+
+	return false;
 }
 
 bool FBulkDataBatchRequest::WaitFor(const FTimespan& WaitTime)
diff --git a/Engine/Source/Runtime/CoreUObject/Public/Serialization/BulkData.h b/Engine/Source/Runtime/CoreUObject/Public/Serialization/BulkData.h
--- a/Engine/Source/Runtime/CoreUObject/Public/Serialization/BulkData.h
+++ b/Engine/Source/Runtime/CoreUObject/Public/Serialization/BulkData.h
@@ -1378,7 +1378,7 @@ public:
 	/** Returns whether the request has been completed. */
 	COREUOBJECT_API bool IsCompleted() const;
 
-	/** Cancel the pending request. */
+	/** Cancel the pending request. Cancelling an invalid request will immediately return. */
 	COREUOBJECT_API void Cancel();
 
 	/** Reset the request handle to an invalid state. Will cancel and wait if the request is not completed. */
@@ -1433,13 +1433,19 @@ private:
 public:	
 	using FBulkDataRequest::FBulkDataRequest;
 
-	/** Blocks the calling thread until the request is completed. */
+	/** Blocks the calling thread until the request is completed. Waiting on an invalid request will immediately return. */
 	COREUOBJECT_API void Wait();
 
-	/** Waits the specified amount of time in milliseconds for the request to be completed. */
+	/**
+	 * Waits the specified amount of time in milliseconds for the request to be completed.
+	 * @return True if the request was completed before the wait time elapsed, otherwise false. Waiting on an invalid request returns false.
+	 */
 	COREUOBJECT_API bool WaitFor(uint32 Milliseconds);
 
-	/** Waits the specified amount of time for the request to be completed. */
+	/**
+	 * Waits the specified amount of time for the request to be completed.
+	 * @return True if the request was completed before the wait time elapsed, otherwise false. Waiting on an invalid request returns false.
+	 */
 	COREUOBJECT_API bool WaitFor(const FTimespan& WaitTime);
 
 	/** Issue one or more I/O request in a single batch. */
