Added UHT support for FIntRect

diff --git a/Engine/Source/Runtime/CoreUObject/Public/UObject/NoExportTypes.h b/Engine/Source/Runtime/CoreUObject/Public/UObject/NoExportTypes.h
--- a/Engine/Source/Runtime/CoreUObject/Public/UObject/NoExportTypes.h
+++ b/Engine/Source/Runtime/CoreUObject/Public/UObject/NoExportTypes.h
@@ -1064,6 +1064,74 @@ struct FUint64Point
 	int64 Y;
 };
 
+/**
+ * An integer rectangle in 2D space.
+ * @note The full C++ class is located here: Engine\Source\Runtime\Core\Public\Math\IntRect.h
+ */
+USTRUCT(immutable, noexport, BlueprintType, IsAlwaysAccessible, HasDefaults, HasNoOpConstructor, IsCoreType)
+struct FIntRect
+{
+	UPROPERTY(EditAnywhere, BlueprintReadWrite, Category=IntRect, SaveGame)
+	FIntPoint Min;
+
+	UPROPERTY(EditAnywhere, BlueprintReadWrite, Category=IntRect, SaveGame)
+	FIntPoint Max;
+};
+
+USTRUCT(immutable, noexport, IsAlwaysAccessible, HasDefaults, HasNoOpConstructor, IsCoreType)
+struct FInt32Rect
+{
+	UPROPERTY(EditAnywhere, Category=IntRect, SaveGame)
+	FInt32Point Min;
+
+	UPROPERTY(EditAnywhere, Category=IntRect, SaveGame)
+	FInt32Point Max;	
+};
+
+USTRUCT(immutable, noexport, IsAlwaysAccessible, HasDefaults, HasNoOpConstructor, IsCoreType)
+struct FInt64Rect
+{
+	UPROPERTY(EditAnywhere, Category=IntRect, SaveGame)
+	FInt64Point Min;
+
+	UPROPERTY(EditAnywhere, Category=IntRect, SaveGame)
+	FInt64Point Max;
+};
+
+/**
+ * An integer rectangle in 2D space.
+ * @note The full C++ class is located here: Engine\Source\Runtime\Core\Public\Math\IntRect.h
+ */
+USTRUCT(immutable, noexport, IsAlwaysAccessible, HasDefaults, HasNoOpConstructor, IsCoreType)
+struct FUintRect
+{
+	UPROPERTY(EditAnywhere, Category=IntRect, SaveGame)
+	FUintPoint Min;
+
+	UPROPERTY(EditAnywhere, Category=IntRect, SaveGame)
+	FUintPoint Max;
+};
+
+USTRUCT(immutable, noexport, IsAlwaysAccessible, HasDefaults, HasNoOpConstructor, IsCoreType)
+struct FUint32Rect
+{
+	UPROPERTY(EditAnywhere, Category=IntRect, SaveGame)
+	FUint32Point Min;
+
+	UPROPERTY(EditAnywhere, Category=IntRect, SaveGame)
+	FUint32Point Max;
+};
+
+USTRUCT(immutable, noexport, IsAlwaysAccessible, HasDefaults, HasNoOpConstructor, IsCoreType)
+struct FUint64Rect
+{
+	UPROPERTY(EditAnywhere, Category=IntRect, SaveGame)
+	FUint64Point Min;
+
+	UPROPERTY(EditAnywhere, Category=IntRect, SaveGame)
+	FUint64Point Max;
+};
+
 /**
  * An integer vector in 3D space.
  * @note The full C++ class is located here: Engine\Source\Runtime\Core\Public\Math\IntVector.h
