Explicitly wait for GC worker tasks (with retraction) at the end of ProcessAsync to avoid race condition with setup for next call to ProcessAsync.

diff --git a/Engine/Source/Runtime/CoreUObject/Private/UObject/GarbageCollection.cpp b/Engine/Source/Runtime/CoreUObject/Private/UObject/GarbageCollection.cpp
--- a/Engine/Source/Runtime/CoreUObject/Private/UObject/GarbageCollection.cpp
+++ b/Engine/Source/Runtime/CoreUObject/Private/UObject/GarbageCollection.cpp
@@ -6927,15 +6927,17 @@ void ProcessAsync(void (*ProcessSync)(void*, FWorkerContext&), void* Processor,
 	}
 
 	// Kick workers
+	TArray<Tasks::TTask<void>> AllWorkerTasks;
+	AllWorkerTasks.Reserve(GReachabilityState.GetNumWorkers() - 1);
 	for (int32 Idx = 1; Idx < GReachabilityState.GetNumWorkers(); ++Idx)
 	{
-		Tasks::Launch(TEXT("CollectReferences"), [=]() 
+		AllWorkerTasks.Add(Tasks::Launch(TEXT("CollectReferences"), [=]() 
 			{
 				if (FWorkerContext* Context = WorkCoordinator->TryStartWorking(Idx))
 				{
 					ProcessSync(Processor, *Context);
 				}
-			});
+			}));
 	}
 
 	// Start working ourselves
@@ -6984,6 +6986,10 @@ void ProcessAsync(void (*ProcessSync)(void*, FWorkerContext&), void* Processor,
 	{
 		ReleaseAsyncProcessingContexts(InContext, Contexts);		
 	}
+
+	// Wait for the tasks we launched so that they cannot survive until the next call to ProcessAsync without being started and race
+	// with work such as InitializeAsyncProcessingContexts.
+	Tasks::Wait(AllWorkerTasks);
 }
 
 //////////////////////////////////////////////////////////////////////////
