[BugFix] overrode FOptionalProperty::SameType to properly compare it's ValueProperty instead of just considering all optionals as the same type

diff --git a/Engine/Source/Runtime/CoreUObject/Private/UObject/PropertyOptional.cpp b/Engine/Source/Runtime/CoreUObject/Private/UObject/PropertyOptional.cpp
--- a/Engine/Source/Runtime/CoreUObject/Private/UObject/PropertyOptional.cpp
+++ b/Engine/Source/Runtime/CoreUObject/Private/UObject/PropertyOptional.cpp
@@ -70,6 +70,11 @@ bool FOptionalProperty::HasIntrusiveUnsetOptionalState() const
 	return false;
 }
 
+bool FOptionalProperty::SameType(const FProperty* Other) const
+{
+	return Super::SameType(Other) && ValueProperty && ValueProperty->SameType(CastFieldChecked<FOptionalProperty>(Other)->ValueProperty);
+}
+
 void FOptionalProperty::AddReferencedObjects(FReferenceCollector& Collector)
 {
 	Super::AddReferencedObjects(Collector);
diff --git a/Engine/Source/Runtime/CoreUObject/Public/UObject/PropertyOptional.h b/Engine/Source/Runtime/CoreUObject/Public/UObject/PropertyOptional.h
--- a/Engine/Source/Runtime/CoreUObject/Public/UObject/PropertyOptional.h
+++ b/Engine/Source/Runtime/CoreUObject/Public/UObject/PropertyOptional.h
@@ -237,6 +237,7 @@ public:
 	virtual void SaveTypeName(UE::FPropertyTypeNameBuilder& Type) const override;
 	virtual bool CanSerializeFromTypeName(UE::FPropertyTypeName Type) const override;
 	virtual bool HasIntrusiveUnsetOptionalState() const override;
+	virtual bool SameType(const FProperty* Other) const override;
 
 	virtual EPropertyVisitorControlFlow Visit(FPropertyVisitorPath& Path, void* Data, const TFunctionRef<EPropertyVisitorControlFlow(const FPropertyVisitorPath& /*Path*/, void* /*Data*/)> InFunc) const override;
 	// End of FProperty interface
