Fixed compilation error with AutoRTFM enabled

diff --git a/Engine/Source/Runtime/CoreUObject/Private/UObject/UObjectAllocator.cpp b/Engine/Source/Runtime/CoreUObject/Private/UObject/UObjectAllocator.cpp
--- a/Engine/Source/Runtime/CoreUObject/Private/UObject/UObjectAllocator.cpp
+++ b/Engine/Source/Runtime/CoreUObject/Private/UObject/UObjectAllocator.cpp
@@ -34,17 +34,22 @@ void FUObjectAllocator::BootMessage()
  */
 UObjectBase* FUObjectAllocator::AllocateUObject(int32 Size, int32 Alignment, bool bAllowPermanent)
 {
+	void* Result = nullptr;
 	// we want to perform this allocation uninstrumented so the GC can clean this up if the transaction is aborted
 	UE_AUTORTFM_OPEN(
 	{
 		if (bAllowPermanent)
 		{
 			// this allocation might go over the reserved memory amount and default to FMemory::Malloc, so we are moving it into the ARTFM scope
-			return (UObjectBase*)GetPersistentLinearAllocator().Allocate(Size, Alignment);
+			Result = GetPersistentLinearAllocator().Allocate(Size, Alignment);
+		}
+		else
+		{
+			Result = FMemory::Malloc(Size, Alignment);
 		}
-
-		return (UObjectBase*)FMemory::Malloc(Size, Alignment);
 	});
+
+	return (UObjectBase*)Result;
 }
 
 /**
