Reverted the new Persistent Linear Allocator

diff --git a/Engine/Config/BaseEngine.ini b/Engine/Config/BaseEngine.ini
--- a/Engine/Config/BaseEngine.ini
+++ b/Engine/Config/BaseEngine.ini
@@ -1527,6 +1527,7 @@ s.FlushStreamingOnExit=True
 
 [/Script/Engine.GarbageCollectionSettings]
 gc.MaxObjectsNotConsideredByGC=1
+gc.SizeOfPermanentObjectPool=0
 gc.FlushStreamingOnGC=0
 gc.NumRetriesBeforeForcingGC=10
 gc.AllowParallelGC=True
@@ -2199,7 +2200,6 @@ bRecompileOnLoad=true
 
 [MemoryPools]
 FLightPrimitiveInteractionInitialBlockSize=512
-PersistentAllocatorReserveSizeMB=1024
 
 [ProcessLimits]
 VirtualMemoryLimitInKB=0
diff --git a/Engine/Source/Runtime/Core/Core.Build.cs b/Engine/Source/Runtime/Core/Core.Build.cs
--- a/Engine/Source/Runtime/Core/Core.Build.cs
+++ b/Engine/Source/Runtime/Core/Core.Build.cs
@@ -3,8 +3,6 @@
 using UnrealBuildTool;
 using System;
 using System.IO;
-using EpicGames.Core;
-using Microsoft.Extensions.Logging;
 
 public class Core : ModuleRules
 {
@@ -348,8 +346,6 @@ public class Core : ModuleRules
 		IWYUSupport = IWYUSupport.KeepAsIs;
 
 		bAllowAutoRTFMInstrumentation = true;
-
-		PrivateDefinitions.Add("UE_PERSISTENT_ALLOCATOR_RESERVE_SIZE=" + GetPersistentAllocatorReserveSize().ToString() + "ULL");
 	}
 
 	protected virtual bool SupportsBinaryConfig(ReadOnlyTargetRules Target)
@@ -368,17 +364,4 @@ public class Core : ModuleRules
 			 Target.Platform == UnrealTargetPlatform.LinuxArm64 ||
 			 Target.Platform == UnrealTargetPlatform.Win64);
 	}
-
-	private ulong GetPersistentAllocatorReserveSize()
-	{
-		ConfigHierarchy Ini = ConfigCache.ReadHierarchy(ConfigHierarchyType.Engine, DirectoryReference.FromFile(Target.ProjectFile), Target.Platform);
-		int SizeOfPermanentObjectPool = 0;
-		if (Ini.GetInt32("/Script/Engine.GarbageCollectionSettings", "gc.SizeOfPermanentObjectPool", out SizeOfPermanentObjectPool))
-		{
-			//Target.Logger.LogWarning("/Script/Engine.GarbageCollectionSettings, gc.SizeOfPermanentObjectPool ini setting was deprecated in a favor of MemoryPools, PersistentAllocatorReserveSizeMB");
-		}
-		int PersistentAllocatorReserveSizeMB = 0;
-		Ini.GetInt32("MemoryPools", "PersistentAllocatorReserveSizeMB", out PersistentAllocatorReserveSizeMB);
-		return (ulong)PersistentAllocatorReserveSizeMB * 1024 * 1024;
-	}
 }
diff --git a/Engine/Source/Runtime/Core/Private/HAL/UnrealMemory.cpp b/Engine/Source/Runtime/Core/Private/HAL/UnrealMemory.cpp
--- a/Engine/Source/Runtime/Core/Private/HAL/UnrealMemory.cpp
+++ b/Engine/Source/Runtime/Core/Private/HAL/UnrealMemory.cpp
@@ -23,7 +23,6 @@
 #include "HAL/MallocDoubleFreeFinder.h"
 #include "HAL/MallocFrameProfiler.h"
 #include "HAL/MallocStomp2.h"
-#include "Memory/LinearVirtualMemoryAllocator.h"
 
 #if MALLOC_GT_HOOKS
 
@@ -672,24 +671,69 @@ void FUseSystemMallocForNew::operator delete[](void* Ptr)
 	FMemory::SystemFree(Ptr);
 }
 
+static bool GPersistentAuxiliaryEnabled = true;
+static uint8 * GPersistentAuxiliary = nullptr;
+static uint8 * GPersistentAuxiliaryEnd = nullptr;
+static TAtomic<SIZE_T> GPersistentAuxiliaryCurrentOffset;
+static SIZE_T GPersistentAuxiliarySize = 0;
+
+
+void FMemory::RegisterPersistentAuxiliary(void* InMemory, SIZE_T InSize)
+{
+	check(GPersistentAuxiliary == nullptr);
+	GPersistentAuxiliaryCurrentOffset = 0;
+	GPersistentAuxiliarySize = InSize;
+	GPersistentAuxiliary = (uint8 *)InMemory;
+	GPersistentAuxiliaryEnd = GPersistentAuxiliary + InSize;
+}
 void* FMemory::MallocPersistentAuxiliary(SIZE_T InSize, uint32 InAlignment)
 {
-	return GetPersistentLinearAllocator().Allocate(InSize, InAlignment);
+	if (GPersistentAuxiliary != nullptr && GPersistentAuxiliaryEnabled)
+	{
+		const uint32 Alignment = FMath::Max<uint32>(InAlignment, 16u);
+		const SIZE_T AlignedSize = Align(InSize, Alignment);
+		// 1st check if there is room, this is atomic but could still fail when actually incrementing the offset.
+		if (GPersistentAuxiliaryCurrentOffset + AlignedSize <= GPersistentAuxiliarySize)
+		{
+			SIZE_T OldOffset = GPersistentAuxiliaryCurrentOffset.AddExchange(AlignedSize);
+			if (OldOffset + AlignedSize <= GPersistentAuxiliarySize)
+			{
+				// we were able to increment the offset and it's still within the bounds of the aux memory.
+				return &GPersistentAuxiliary[OldOffset];
+			}
+			// we've gone over the end of the aux memory, this could waste some space, if it's a problem protect with a critical section.
+		}
+	}
+	return FMemory::Malloc(InSize, InAlignment);
 }
-
 void FMemory::FreePersistentAuxiliary(void* InPtr)
 {
-	GetPersistentLinearAllocator().TryDeallocate(InPtr, 0);
+	if (GPersistentAuxiliary != nullptr)
+	{
+		uint8* Ptr = (uint8*)InPtr;
+		if (Ptr >= GPersistentAuxiliary && Ptr < GPersistentAuxiliaryEnd)
+		{
+			// it is part of the GPersistentAuxiliary
+			return;
+		}
+	}
+	return FMemory::Free(InPtr);
 }
-
 bool FMemory::IsPersistentAuxiliaryActive()
 {
-	return GetPersistentLinearAllocator().IsInitialized();
+	return GPersistentAuxiliary != nullptr && GPersistentAuxiliaryEnabled;
+}
+void FMemory::DisablePersistentAuxiliary()
+{
+	GPersistentAuxiliaryEnabled = false;
+}
+void FMemory::EnablePersistentAuxiliary()
+{
+	GPersistentAuxiliaryEnabled = true;
 }
-
 SIZE_T FMemory::GetUsedPersistentAuxiliary()
 {
-	return GetPersistentLinearAllocator().GetAllocatedMemorySize();
+	return GPersistentAuxiliaryCurrentOffset;
 }
 
 
diff --git a/Engine/Source/Runtime/Core/Private/Memory/LinearVirtualMemoryAllocator.cpp b/Engine/Source/Runtime/Core/Private/Memory/LinearVirtualMemoryAllocator.cpp
deleted file mode 100644
diff --git a/Engine/Source/Runtime/Core/Private/UObject/UnrealNames.cpp b/Engine/Source/Runtime/Core/Private/UObject/UnrealNames.cpp
--- a/Engine/Source/Runtime/Core/Private/UObject/UnrealNames.cpp
+++ b/Engine/Source/Runtime/Core/Private/UObject/UnrealNames.cpp
@@ -10,7 +10,6 @@
 #include "Math/NumericLimits.h"
 #include "Math/UnrealMathUtility.h"
 #include "HAL/UnrealMemory.h"
-#include "Memory/LinearVirtualMemoryAllocator.h"
 #include "Templates/UnrealTemplate.h"
 #include "Misc/CString.h"
 #include "Misc/Crc.h"
@@ -454,9 +453,9 @@ public:
 
 	~FNameEntryAllocator()
 	{
-		for (int32 Index = CurrentBlock; Index >= 0; --Index)
+		for (uint32 Index = 0; Index <= CurrentBlock; ++Index)
 		{
-			GetPersistentLinearAllocator().TryDeallocate(Blocks[Index], BlockSizeBytes);
+			FMemory::Free(Blocks[Index]);
 		}
 	}
 
@@ -657,7 +656,7 @@ private:
 	static uint8* AllocBlock()
 	{
 		LLM_SCOPE(ELLMTag::FName);
-		return (uint8*)GetPersistentLinearAllocator().Allocate(BlockSizeBytes, alignof(FNameEntry));
+		return (uint8*)FMemory::MallocPersistentAuxiliary(BlockSizeBytes, alignof(FNameEntry));
 	}
 	
 	void AllocateNewBlock()
diff --git a/Engine/Source/Runtime/Core/Public/Containers/LockFreeFixedSizeAllocator.h b/Engine/Source/Runtime/Core/Public/Containers/LockFreeFixedSizeAllocator.h
--- a/Engine/Source/Runtime/Core/Public/Containers/LockFreeFixedSizeAllocator.h
+++ b/Engine/Source/Runtime/Core/Public/Containers/LockFreeFixedSizeAllocator.h
@@ -17,7 +17,7 @@
 #define USE_NAIVE_TLockFreeFixedSizeAllocator_TLSCacheBase (0) // this is useful for find who really leaked
 #endif
 
-template<int32 SIZE, typename TBundleRecycler, typename TTrackingCounter = FNoopCounter>
+template<int32 SIZE, typename TBundleRecycler, typename TTrackingCounter = FNoopCounter, bool AllowDisablingOfTrim = false>
 class TLockFreeFixedSizeAllocator_TLSCacheBase : public FNoncopyable
 {
 	enum
@@ -66,7 +66,7 @@ public:
 				TLS.PartialBundle = GlobalFreeListBundles.Pop();
 				if (!TLS.PartialBundle)
 				{
-					TLS.PartialBundle = (void**)FMemory::Malloc(SIZE_PER_BUNDLE);
+					TLS.PartialBundle = (void**)FMemory::MallocPersistentAuxiliary(SIZE_PER_BUNDLE);
 					void **Next = TLS.PartialBundle;
 					for (int32 Index = 0; Index < NUM_PER_BUNDLE - 1; Index++)
 					{
@@ -190,7 +190,7 @@ private:
  * never returns free space until program shutdown.
  * alignment isn't handled, assumes FMemory::Malloc will work
  */
-template<int32 SIZE, int TPaddingForCacheContention, typename TTrackingCounter = FNoopCounter>
+template<int32 SIZE, int TPaddingForCacheContention, typename TTrackingCounter = FNoopCounter, bool AllowDisablingOfTrim = false>
 class TLockFreeFixedSizeAllocator
 {
 public:
@@ -218,7 +218,7 @@ public:
 		}
 		else
 		{
-			Memory = FMemory::Malloc(SIZE);
+			Memory = FMemory::MallocPersistentAuxiliary(SIZE);
 		}
 		return Memory;
 	}
@@ -241,9 +241,16 @@ public:
 	*/
 	void Trim()
 	{
+		if (AllowDisablingOfTrim)
+		{
+			if (FMemory::IsPersistentAuxiliaryActive())
+			{
+				return;
+			}
+		}
 		while (void* Mem = FreeList.Pop())
 		{
-			FMemory::Free(Mem);
+			FMemory::FreePersistentAuxiliary(Mem);
 			NumFree.Decrement();
 		}
 	}
@@ -287,8 +294,8 @@ private:
  * never returns free space, even at shutdown
  * alignment isn't handled, assumes FMemory::Malloc will work
  */
-template<int32 SIZE, int TPaddingForCacheContention, typename TTrackingCounter = FNoopCounter>
-class TLockFreeFixedSizeAllocator_TLSCache : public TLockFreeFixedSizeAllocator_TLSCacheBase<SIZE, TLockFreePointerListUnordered<void*, TPaddingForCacheContention>, TTrackingCounter>
+template<int32 SIZE, int TPaddingForCacheContention, typename TTrackingCounter = FNoopCounter, bool AllowDisablingOfTrim = false>
+class TLockFreeFixedSizeAllocator_TLSCache : public TLockFreeFixedSizeAllocator_TLSCacheBase<SIZE, TLockFreePointerListUnordered<void*, TPaddingForCacheContention>, TTrackingCounter, AllowDisablingOfTrim>
 {
 };
 
@@ -297,8 +304,8 @@ class TLockFreeFixedSizeAllocator_TLSCache : public TLockFreeFixedSizeAllocator_
  *
  * Never returns free space until program shutdown.
  */
-template<class T, int TPaddingForCacheContention>
-class TLockFreeClassAllocator : private TLockFreeFixedSizeAllocator<sizeof(T), TPaddingForCacheContention, FNoopCounter>
+template<class T, int TPaddingForCacheContention, bool AllowDisablingOfTrim = false>
+class TLockFreeClassAllocator : private TLockFreeFixedSizeAllocator<sizeof(T), TPaddingForCacheContention, FNoopCounter, AllowDisablingOfTrim>
 {
 public:
 	/**
@@ -341,8 +348,8 @@ public:
  *
  * Never returns free space until program shutdown.
  */
-template<class T, int TPaddingForCacheContention>
-class TLockFreeClassAllocator_TLSCache : private TLockFreeFixedSizeAllocator_TLSCache<sizeof(T), TPaddingForCacheContention, FNoopCounter>
+template<class T, int TPaddingForCacheContention, bool AllowDisablingOfTrim = false>
+class TLockFreeClassAllocator_TLSCache : private TLockFreeFixedSizeAllocator_TLSCache<sizeof(T), TPaddingForCacheContention, FNoopCounter , AllowDisablingOfTrim>
 {
 public:
 	/**
diff --git a/Engine/Source/Runtime/Core/Public/HAL/UnrealMemory.h b/Engine/Source/Runtime/Core/Public/HAL/UnrealMemory.h
--- a/Engine/Source/Runtime/Core/Public/HAL/UnrealMemory.h
+++ b/Engine/Source/Runtime/Core/Public/HAL/UnrealMemory.h
@@ -277,25 +277,12 @@ struct FMemory
 	* Functions to handle special memory given to the title from the platform
 	* This memory is allocated like a stack, it's never really freed
 	*/
-	UE_DEPRECATED(5.4, "Persistent Auxiliary allocator is obsolete and is replaced by a GetPersistentLinearAllocator()")
-	static inline void RegisterPersistentAuxiliary(void* /*InMemory*/, SIZE_T /*InSize*/) {}
-
-	UE_DEPRECATED(5.4, "Persistent Auxiliary allocator is obsolete and is replaced by a GetPersistentLinearAllocator()")
+	static CORE_API void RegisterPersistentAuxiliary(void* InMemory, SIZE_T InSize);
 	static CORE_API void* MallocPersistentAuxiliary(SIZE_T InSize, uint32 InAlignment = 0);
-
-	UE_DEPRECATED(5.4, "Persistent Auxiliary allocator is obsolete and is replaced by a GetPersistentLinearAllocator()")
 	static CORE_API void FreePersistentAuxiliary(void* InPtr);
-
-	UE_DEPRECATED(5.4, "Persistent Auxiliary allocator is obsolete and is replaced by a GetPersistentLinearAllocator()")
 	static CORE_API bool IsPersistentAuxiliaryActive();
-
-	UE_DEPRECATED(5.4, "Persistent Auxiliary allocator is obsolete and is replaced by a GetPersistentLinearAllocator()")
-	static inline void DisablePersistentAuxiliary() {}
-
-	UE_DEPRECATED(5.4, "Persistent Auxiliary allocator is obsolete and is replaced by a GetPersistentLinearAllocator()")
-	static inline void EnablePersistentAuxiliary() {}
-
-	UE_DEPRECATED(5.4, "Persistent Auxiliary allocator is obsolete and is replaced by a GetPersistentLinearAllocator()")
+	static CORE_API void DisablePersistentAuxiliary();
+	static CORE_API void EnablePersistentAuxiliary();
 	static CORE_API SIZE_T GetUsedPersistentAuxiliary();
 private:
 	static CORE_API void GCreateMalloc();
diff --git a/Engine/Source/Runtime/Core/Public/Memory/LinearVirtualMemoryAllocator.h b/Engine/Source/Runtime/Core/Public/Memory/LinearVirtualMemoryAllocator.h
deleted file mode 100644
diff --git a/Engine/Source/Runtime/CoreUObject/Private/UObject/UObjectAllocator.cpp b/Engine/Source/Runtime/CoreUObject/Private/UObject/UObjectAllocator.cpp
--- a/Engine/Source/Runtime/CoreUObject/Private/UObject/UObjectAllocator.cpp
+++ b/Engine/Source/Runtime/CoreUObject/Private/UObject/UObjectAllocator.cpp
@@ -12,15 +12,32 @@ DEFINE_LOG_CATEGORY_STATIC(LogUObjectAllocator, Log, All);
 /** Global UObjectBase allocator							*/
 COREUOBJECT_API FUObjectAllocator GUObjectAllocator;
 
+/**
+ * Allocates and initializes the permanent object pool
+ *
+ * @param InPermanentObjectPoolSize size of permanent object pool
+ */
+void FUObjectAllocator::AllocatePermanentObjectPool(int32 InPermanentObjectPoolSize)
+{
+	PermanentObjectPoolSize	= InPermanentObjectPoolSize;
+	PermanentObjectPool		= (uint8*) FMemory::MallocPersistentAuxiliary( PermanentObjectPoolSize );
+	PermanentObjectPoolTail	= PermanentObjectPool;
+	PermanentObjectPoolExceededTail = PermanentObjectPoolTail;
+}
+
+
 /**
  * Prints a debugf message to allow tuning
  */
 void FUObjectAllocator::BootMessage()
 {
-	const SIZE_T ExceedingSize = GetPersistentLinearAllocator().GetExceedingSize();
-	if (ExceedingSize > 0)
+	if (PermanentObjectPoolSize && PermanentObjectPoolExceededTail - PermanentObjectPool > PermanentObjectPoolSize)
+	{
+		UE_LOG(LogUObjectAllocator, Warning, TEXT("%i Exceeds size of permanent object pool %i, please tune SizeOfPermanentObjectPool."), PermanentObjectPoolExceededTail - PermanentObjectPool, PermanentObjectPoolSize );
+	}
+	else
 	{
-		UE_LOG(LogUObjectAllocator, Warning, TEXT("Persistent memory pool exceeded by %u KB, please tune PersistentAllocatorReserveSizeMB setting in [MemoryPools] ini group."), ExceedingSize / 1024);
+		UE_LOG(LogUObjectAllocator, Log, TEXT("%i out of %i bytes used by permanent object pool."), PermanentObjectPoolExceededTail - PermanentObjectPool, PermanentObjectPoolSize );
 	}
 }
 
@@ -34,22 +51,48 @@ void FUObjectAllocator::BootMessage()
  */
 UObjectBase* FUObjectAllocator::AllocateUObject(int32 Size, int32 Alignment, bool bAllowPermanent)
 {
-	void* Result = nullptr;
-	// we want to perform this allocation uninstrumented so the GC can clean this up if the transaction is aborted
-	UE_AUTORTFM_OPEN(
+	// Force alignment to minimal of 16 bytes
+	Alignment = FMath::Max(16, Alignment);
+	int32 AlignedSize = Align( Size, Alignment );
+	UObjectBase* Result = nullptr;
+
+	bAllowPermanent &= PermanentObjectPool != nullptr;
+	const bool bPlaceInPerm = bAllowPermanent && (Align(PermanentObjectPoolTail,Alignment) + Size) <= (PermanentObjectPool + PermanentObjectPoolSize);
+	if (bAllowPermanent && !bPlaceInPerm)
+	{
+		// advance anyway so we can determine how much space we should set aside in the ini
+		uint8* AlignedPtr = Align( PermanentObjectPoolExceededTail, Alignment );
+		PermanentObjectPoolExceededTail = AlignedPtr + Size;
+	}
+	// Use object memory pool for objects disregarded by GC (initially loaded ones). This allows identifying their
+	// GC status by simply looking at their address.
+	if (bPlaceInPerm)
 	{
-		if (bAllowPermanent)
+		// Align current tail pointer and use it for object. 
+		uint8* AlignedPtr = Align( PermanentObjectPoolTail, Alignment );
+		// Update tail pointer.
+		PermanentObjectPoolTail = AlignedPtr + Size;
+		Result = (UObjectBase*)AlignedPtr;
+		if (PermanentObjectPoolExceededTail < PermanentObjectPoolTail)
 		{
-			// this allocation might go over the reserved memory amount and default to FMemory::Malloc, so we are moving it into the ARTFM scope
-			Result = GetPersistentLinearAllocator().Allocate(Size, Alignment);
+			PermanentObjectPoolExceededTail = PermanentObjectPoolTail;
 		}
-		else
+	}
+	else
+	{
+		// we want to perform this allocation uninstrumented so the GC can clean this up if the transaction is aborted
+		UE_AUTORTFM_OPEN(
 		{
-			Result = FMemory::Malloc(Size, Alignment);
-		}
-	});
+			// Allocate new memory of the appropriate size and alignment.
+			Result = (UObjectBase*)FMemory::Malloc( Size, Alignment );
+		});
+	}
+
+#if !UE_BUILD_TEST && !UE_BUILD_SHIPPING
+	checkf(IsAligned(Result, Alignment), TEXT("Allocated memory address does not match requirement of %d byte alignment for size %d"), Alignment, Size);
+#endif
 
-	return (UObjectBase*)Result;
+	return Result;
 }
 
 /**
@@ -61,7 +104,7 @@ void FUObjectAllocator::FreeUObject(UObjectBase *Object) const
 {
 	check(Object);
 	// Only free memory if it was allocated directly from allocator and not from permanent object pool.
-	if (FPermanentObjectPoolExtents().Contains(Object) == false)
+	if (FPermanentObjectPoolExtents(*this).Contains(Object) == false)
 	{
 		FMemory::Free(Object);
 	}
diff --git a/Engine/Source/Runtime/CoreUObject/Private/UObject/UObjectBase.cpp b/Engine/Source/Runtime/CoreUObject/Private/UObject/UObjectBase.cpp
--- a/Engine/Source/Runtime/CoreUObject/Private/UObject/UObjectBase.cpp
+++ b/Engine/Source/Runtime/CoreUObject/Private/UObject/UObjectBase.cpp
@@ -984,6 +984,7 @@ void UObjectBaseInit()
 
 	// Zero initialize and later on get value from .ini so it is overridable per game/ platform...
 	int32 MaxObjectsNotConsideredByGC = 0;
+	int32 SizeOfPermanentObjectPool = 0;
 	int32 MaxUObjects = 2 * 1024 * 1024; // Default to ~2M UObjects
 	bool bPreAllocateUObjectArray = false;	
 
@@ -1001,6 +1002,9 @@ void UObjectBaseInit()
 		else
 		{
 			GConfig->GetInt(TEXT("/Script/Engine.GarbageCollectionSettings"), TEXT("gc.MaxObjectsNotConsideredByGC"), MaxObjectsNotConsideredByGC, GEngineIni);
+
+			// Not used on PC as in-place creation inside bigger pool interacts with the exit purge and deleting UObject directly.
+			GConfig->GetInt(TEXT("/Script/Engine.GarbageCollectionSettings"), TEXT("gc.SizeOfPermanentObjectPool"), SizeOfPermanentObjectPool, GEngineIni);
 		}
 
 		// Maximum number of UObjects in cooked game
@@ -1021,10 +1025,23 @@ void UObjectBaseInit()
 #endif
 	}
 
+	if (MaxObjectsNotConsideredByGC <= 0 && SizeOfPermanentObjectPool > 0)
+	{
+		// If permanent object pool is enabled but disregard for GC is disabled, GC will mark permanent object pool objects
+		// as unreachable and may destroy them so disable permanent object pool too.
+		// An alternative would be to make GC not mark permanent object pool objects as unreachable but then they would have to
+		// be considered as root set objects because they could be referencing objects from outside of permanent object pool.
+		// This would be inconsistent and confusing and also counter productive (the more root set objects the more expensive MarkAsUnreachable phase is).
+		SizeOfPermanentObjectPool = 0;
+		UE_LOG(LogInit, Warning, TEXT("Disabling permanent object pool because disregard for GC is disabled (gc.MaxObjectsNotConsideredByGC=%d)."), MaxObjectsNotConsideredByGC);
+	}
+
 	// Log what we're doing to track down what really happens as log in LaunchEngineLoop doesn't report those settings in pristine form.
-	UE_LOG(LogInit, Log, TEXT("%s for max %d objects, including %i objects not considered by GC."), 
-		bPreAllocateUObjectArray ? TEXT("Pre-allocating") : TEXT("Presizing"), MaxUObjects, MaxObjectsNotConsideredByGC);
+	UE_LOG(LogInit, Log, TEXT("%s for max %d objects, including %i objects not considered by GC, pre-allocating %i bytes for permanent pool."), 
+		bPreAllocateUObjectArray ? TEXT("Pre-allocating") : TEXT("Presizing"),
+		MaxUObjects, MaxObjectsNotConsideredByGC, SizeOfPermanentObjectPool);
 
+	GUObjectAllocator.AllocatePermanentObjectPool(SizeOfPermanentObjectPool);
 	GUObjectArray.AllocateObjectPool(MaxUObjects, MaxObjectsNotConsideredByGC, bPreAllocateUObjectArray);
 #if UE_WITH_OBJECT_HANDLE_LATE_RESOLVE
 	UE::CoreUObject::Private::InitObjectHandles(GUObjectArray.GetObjectArrayCapacity());
diff --git a/Engine/Source/Runtime/CoreUObject/Public/UObject/UObjectAllocator.h b/Engine/Source/Runtime/CoreUObject/Public/UObject/UObjectAllocator.h
--- a/Engine/Source/Runtime/CoreUObject/Public/UObject/UObjectAllocator.h
+++ b/Engine/Source/Runtime/CoreUObject/Public/UObject/UObjectAllocator.h
@@ -8,26 +8,42 @@
 
 #include "CoreMinimal.h"
 #include "HAL/PlatformMath.h"
-#include "Memory/LinearVirtualMemoryAllocator.h"
 
 class UObjectBase;
 
 class FUObjectAllocator
 {
 public:
+
+	/**
+	 * Constructor, initializes to no permanent object pool
+	 */
+	FUObjectAllocator() :
+	  PermanentObjectPoolSize(0),
+	  PermanentObjectPool(NULL),
+	  PermanentObjectPoolTail(NULL),
+		PermanentObjectPoolExceededTail(NULL)
+	{
+	}
+
 	/**
 	 * Allocates and initializes the permanent object pool
 	 *
 	 * @param InPermanentObjectPoolSize size of permanent object pool
 	 */
-	UE_DEPRECATED(5.4, "Permanent Object Pool is handled by the global instance of FLinearVirtualMemoryAllocator now")
-	COREUOBJECT_API void AllocatePermanentObjectPool(int32 InPermanentObjectPoolSize) {}
+	COREUOBJECT_API void AllocatePermanentObjectPool(int32 InPermanentObjectPoolSize);
 
 	/**
 	 * Prints a debugf message to allow tuning
 	 */
 	COREUOBJECT_API void BootMessage();
 
+	UE_DEPRECATED(5.1, "Use the more efficient FPermanentObjectPoolExtents instead")
+	FORCEINLINE bool ResidesInPermanentPool(const UObjectBase *Object) const
+	{
+		return ((const uint8*)Object >= PermanentObjectPool) && ((const uint8*)Object < PermanentObjectPoolTail);
+	}
+
 	/**
 	 * Allocates a UObjectBase from the free store or the permanent object pool
 	 *
@@ -44,6 +60,18 @@ public:
 	 * @param Object object to free
 	 */
 	COREUOBJECT_API void FreeUObject(UObjectBase *Object) const;
+
+private:
+	friend class FPermanentObjectPoolExtents;
+
+	/** Size in bytes of pool for objects disregarded for GC.								*/
+	int32							PermanentObjectPoolSize;
+	/** Begin of pool for objects disregarded for GC.										*/
+	uint8*						PermanentObjectPool;
+	/** Current position in pool for objects disregarded for GC.							*/
+	uint8*						PermanentObjectPoolTail;
+	/** Tail that exceeded the size of the permanent object pool, >= PermanentObjectPoolTail.		*/
+	uint8*						PermanentObjectPoolExceededTail;
 };
 
 /** Global UObjectBase allocator							*/
@@ -53,11 +81,11 @@ extern COREUOBJECT_API FUObjectAllocator GUObjectAllocator;
 class FPermanentObjectPoolExtents
 {
 public:
-	FORCEINLINE FPermanentObjectPoolExtents(const FLinearVirtualMemoryAllocatorExtends& InAllocatorExtends = GPersistentLinearAllocatorExtends)
-		: Address(InAllocatorExtends.Address)
-		, Size(InAllocatorExtends.Size)
+	FORCEINLINE FPermanentObjectPoolExtents(const FUObjectAllocator& ObjectAllocator = GUObjectAllocator)
+		: Address(reinterpret_cast<uint64>(ObjectAllocator.PermanentObjectPool))
+		, Size(static_cast<uint64>(ObjectAllocator.PermanentObjectPoolSize))
 	{}
-
+	
 	FORCEINLINE bool Contains(const UObjectBase* Object) const
 	{
 		return reinterpret_cast<uint64>(Object) - Address < Size;
