Decoupled ValuesSetBySerialization from the index in the serialized property path

#rnx

diff --git a/Engine/Source/Runtime/CoreUObject/Internal/UObject/PropertyBagRepository.h b/Engine/Source/Runtime/CoreUObject/Internal/UObject/PropertyBagRepository.h
--- a/Engine/Source/Runtime/CoreUObject/Internal/UObject/PropertyBagRepository.h
+++ b/Engine/Source/Runtime/CoreUObject/Internal/UObject/PropertyBagRepository.h
@@ -12,7 +12,6 @@ class UObject;
 namespace UE
 {
 
-class FPropertyPathName;
 class FPropertyPathNameTree;
 
 // Singleton class tracking property bag association with objects
@@ -107,8 +106,6 @@ public:
 
 	COREUOBJECT_API const UObject* FindInstanceForDataObject(const UObject* InstanceDataObject) const;
 
-	// query whether a property in an object was set when the object was deserialized
-	COREUOBJECT_API static bool WasPropertySetBySerialization(UObject* Object, const FPropertyPathName& Path);
 	// query whether a property in Struct was set when the struct was deserialized
 	COREUOBJECT_API static bool WasPropertySetBySerialization(const UStruct* Struct, const void* StructData, const FProperty* Property, int32 ArrayIndex = 0);
 
diff --git a/Engine/Source/Runtime/CoreUObject/Private/UObject/Class.cpp b/Engine/Source/Runtime/CoreUObject/Private/UObject/Class.cpp
--- a/Engine/Source/Runtime/CoreUObject/Private/UObject/Class.cpp
+++ b/Engine/Source/Runtime/CoreUObject/Private/UObject/Class.cpp
@@ -33,6 +33,7 @@
 #include "Misc/PackageName.h"
 #include "UObject/ObjectResource.h"
 #include "UObject/LinkerSave.h"
+#include "UObject/InstanceDataObjectUtils.h"
 #include "UObject/Interface.h"
 #include "UObject/LinkerPlaceholderClass.h"
 #include "UObject/LinkerPlaceholderFunction.h"
@@ -1625,7 +1626,7 @@ void UStruct::SerializeVersionedTaggedProperties(FStructuredArchive::FSlot Slot,
 				}
 
 				TOptional<UE::FSerializedPropertyPathScope> SerializedPropertyPath;
-				if (SerializeContext && SerializeContext->bTrackSerializedPropertyPath)
+				if (SerializeContext->bTrackSerializedPropertyPath)
 				{
 					const FName Name = Property ? Property->GetFName() : Tag.Name;
 					const int32 Index = Tag.ArrayIndex > 0 || (Property && Property->ArrayDim > 1) ? Tag.ArrayIndex : INDEX_NONE;
@@ -1715,6 +1716,11 @@ void UStruct::SerializeVersionedTaggedProperties(FStructuredArchive::FSlot Slot,
 								checkNoEntry();
 								break;
 						}
+
+						if (bAdvanceProperty && SerializeContext->bImpersonateProperties)
+						{
+							UE::MarkPropertySetBySerialization(this, Data, Property, Tag.ArrayIndex);
+						}
 					}
 				}
 				else
@@ -1728,7 +1734,7 @@ void UStruct::SerializeVersionedTaggedProperties(FStructuredArchive::FSlot Slot,
 					if (!bSearchedForUnknownPropertyTree)
 					{
 						bSearchedForUnknownPropertyTree = true;
-						if (SerializeContext && SerializeContext->bSerializeUnknownProperty)
+						if (SerializeContext->bSerializeUnknownProperty)
 						{
 							if (UObject* Object = SerializeContext->SerializedObject)
 							{
diff --git a/Engine/Source/Runtime/CoreUObject/Private/UObject/InstanceDataObjectUtils.cpp b/Engine/Source/Runtime/CoreUObject/Private/UObject/InstanceDataObjectUtils.cpp
--- a/Engine/Source/Runtime/CoreUObject/Private/UObject/InstanceDataObjectUtils.cpp
+++ b/Engine/Source/Runtime/CoreUObject/Private/UObject/InstanceDataObjectUtils.cpp
@@ -105,9 +105,6 @@ uint32 UInstanceDataObjectStruct::GetStructTypeHash(const void* Src) const
 
 namespace UE
 {
-	// typedef to help make it clearer when a pathName has indices and when the indices are wildcarded away
-	using FWildcardPropertyPathName = FPropertyPathName;
-
 	static const FName NAME_DisplayName(ANSITEXTVIEW("DisplayName"));
 	static const FName NAME_PresentAsTypeMetadata(ANSITEXTVIEW("PresentAsType"));
 	static const FName NAME_IsLooseMetadata(ANSITEXTVIEW("IsLoose"));
@@ -116,16 +113,6 @@ namespace UE
 	static const FName NAME_IDOMapKey(ANSITEXTVIEW("Key"));
 	static const FName NAME_IDOMapValue(ANSITEXTVIEW("Value"));
 
-	struct ResolvePropertyPathNameHelperParams
-	{
-		void* Data = nullptr;
-		FProperty* ResultProperty = nullptr;
-		const UE::FPropertyPathName& Path;
-		int32 CurPathIndex = 0;
-		int32 EndPathIndex = INDEX_NONE; // INDEX_NONE has the same behavior as Path.GetSegmentCount()
-		bool bAddIfNeeded = false;
-	};
-
 	bool bEnableIDOSupport = false;
 	FAutoConsoleVariableRef EnableIDOSupportCVar(
 		TEXT("IDO.Enable"),
@@ -158,160 +145,6 @@ namespace UE
 		return bIsEnabled;
 	}
 
-	static FProperty* FindPropertyByType(UStruct* Struct, FName PropertyName, FPropertyTypeName PropertyType)
-	{
-		for (FProperty* Property : TFieldRange<FProperty>(Struct))
-		{
-			if (Property->GetFName() == PropertyName && Property->CanSerializeFromTypeName(PropertyType))
-			{
-				return Property;
-			}
-		}
-		return nullptr;
-	}
-
-	static bool ResolvePropertyPathName(
-		UObject* Object, const FPropertyPathName& Path,
-		void*& OutData, FProperty*& OutProperty,
-		void*& OutOwnerData, FStructProperty*& OutOwnerProperty)
-	{
-		OutData = Object;
-		OutProperty = nullptr;
-		OutOwnerData = nullptr;
-		OutOwnerProperty = nullptr;
-
-		if (!Object)
-		{
-			return false;
-		}
-
-		const int32 SegmentCount = Path.GetSegmentCount();
-		for (int32 SegmentIndex = 0; SegmentIndex < SegmentCount;)
-		{
-			// Assign the owner of the property in this segment.
-			OutOwnerProperty = CastField<FStructProperty>(OutProperty);
-			OutOwnerData = OutOwnerProperty ? OutData : nullptr;
-
-			// Find the struct to search within.
-			UStruct* Struct = nullptr;
-			if (SegmentIndex == 0)
-			{
-				Struct = Object->GetClass();
-			}
-			else if (OutOwnerProperty)
-			{
-				Struct = OutOwnerProperty->Struct;
-			}
-			else
-			{
-				return false;
-			}
-
-			// Find the named property within the struct/class.
-			const FPropertyPathNameSegment Segment = Path.GetSegment(SegmentIndex++);
-			FProperty* Property = FindPropertyByType(Struct, Segment.Name, Segment.Type);
-			if (!Property)
-			{
-				return false;
-			}
-
-			// Find the address of the property value.
-			const bool bIsIndexedProperty = Property->IsA<FArrayProperty>() || Property->IsA<FSetProperty>() || Property->IsA<FMapProperty>();
-			const int32 ArrayIndex = !bIsIndexedProperty && Segment.Index >= 0 ? Segment.Index : 0;
-			if (ArrayIndex >= Property->ArrayDim)
-			{
-				return false;
-			}
-			OutProperty = Property;
-			OutData = Property->ContainerPtrToValuePtr<void>(OutData, ArrayIndex);
-
-			if (bIsIndexedProperty)
-			{
-				if (Segment.Index == INDEX_NONE)
-				{
-					// A segment may only resolve directly to an indexed property if it is the last segment.
-					if (SegmentIndex < SegmentCount)
-					{
-						return false;
-					}
-				}
-				else if (FArrayProperty* ArrayProperty = CastField<FArrayProperty>(Property))
-				{
-					FScriptArrayHelper Array(ArrayProperty, OutData);
-					if (!Array.IsValidIndex(Segment.Index))
-					{
-						return false;
-					}
-					OutProperty = ArrayProperty->Inner;
-					OutData = Array.GetElementPtr(Segment.Index);
-				}
-				else if (FSetProperty* SetProperty = CastField<FSetProperty>(Property))
-				{
-					FScriptSetHelper Set(SetProperty, OutData);
-					if (!Set.IsValidIndex(Segment.Index))
-					{
-						return false;
-					}
-					OutProperty = SetProperty->ElementProp;
-					OutData = Set.GetElementPtr(Segment.Index);
-				}
-				else if (FMapProperty* MapProperty = CastField<FMapProperty>(Property))
-				{
-					FScriptMapHelper Map(MapProperty, OutData);
-					if (!Map.IsValidIndex(Segment.Index))
-					{
-						return false;
-					}
-					// Resolve to the pair if Key or Value are not in the path.
-					if (SegmentIndex == SegmentCount)
-					{
-						// Clear the property because there is no property for the pair, only the key and value.
-						OutProperty = nullptr;
-						OutData = Map.GetPairPtr(Segment.Index);
-						// Clear the owner because this logically resolves to the element, which is owned by the map.
-						OutOwnerProperty = nullptr;
-						OutOwnerData = nullptr;
-					}
-					else
-					{
-						const FPropertyPathNameSegment MapSegment = Path.GetSegment(SegmentIndex++);
-						if (MapSegment.Name == NAME_IDOMapKey)
-						{
-							OutProperty = MapProperty->KeyProp;
-							OutData = Map.GetKeyPtr(Segment.Index);
-						}
-						else if (MapSegment.Name == NAME_IDOMapValue)
-						{
-							OutProperty = MapProperty->ValueProp;
-							OutData = Map.GetValuePtr(Segment.Index);
-						}
-						else
-						{
-							return false;
-						}
-					}
-				}
-			}
-			else if (FOptionalProperty* OptionalProperty = CastField<FOptionalProperty>(Property))
-			{
-				FOptionalPropertyLayout Optional(OptionalProperty->GetValueProperty());
-				// Resolve to the value if it is set, otherwise resolve to the unset optional.
-				if (void* Data = Optional.GetValuePointerForReplaceIfSet(OutData))
-				{
-					OutProperty = Optional.GetValueProperty();
-					OutData = Data;
-				}
-				// A segment may only resolve directly to an unset optional if it is the last segment.
-				else if (SegmentIndex < SegmentCount)
-				{
-					return false;
-				}
-			}
-		}
-
-		return true;
-	}
-
 	static UStruct* CreateInstanceDataObjectStructRec(const UClass* StructClass, UStruct* TemplateStruct, UObject* Outer, const FPropertyPathNameTree* PropertyTree);
 
 	template <typename StructType>
@@ -320,18 +153,6 @@ namespace UE
 		return CastChecked<StructType>(CreateInstanceDataObjectStructRec(StructType::StaticClass(), TemplateStruct, Outer, PropertyTree));
 	}
 
-	static FPropertyPathNameSegment CreateSegmentFromProperty(const FProperty* Inner, int32 Index = INDEX_NONE)
-	{
-		FPropertyTypeNameBuilder TypeBuilder;
-		Inner->SaveTypeName(TypeBuilder);
-
-		FPropertyPathNameSegment Result;
-		Result.Index = INDEX_NONE; 
-		Result.Name = Inner->GetFName();
-		Result.Type = TypeBuilder.Build();
-		return Result;
-	}
-
 	static FString UnmanglePropertyName(const FName MaybeMangledName, bool& bOutNameWasMangled)
 	{
 		FString Result = MaybeMangledName.ToString();
@@ -649,11 +470,12 @@ namespace UE
 		return Result;
 	}
 
-	static void MarkPropertySetBySerialization(const UStruct* Struct, const void* StructData, const void* PropertyDataPtr)
+	void MarkPropertySetBySerialization(const UStruct* Struct, const void* StructData, const FProperty* Property, int32 ArrayIndex)
 	{
 		if (const FSetProperty* ValuesSetByPropertyBagProperty = CastField<FSetProperty>(Struct->FindPropertyByName(NAME_ValuesSetBySerialization)))
 		{
 			FScriptSetHelper ValuesSetByPropertyBag(ValuesSetByPropertyBagProperty, ValuesSetByPropertyBagProperty->ContainerPtrToValuePtr<void>(StructData));
+			const void* PropertyDataPtr = Property->ContainerPtrToValuePtr<void>(StructData, ArrayIndex);
 			const int64 ValueOffset = static_cast<const uint8*>(PropertyDataPtr) - static_cast<const uint8*>(StructData);
 			const int32 FoundIndex = ValuesSetByPropertyBag.FindElementIndex(&ValueOffset);
 			if (FoundIndex == INDEX_NONE)
@@ -662,66 +484,19 @@ namespace UE
 			}
 		}
 	}
-	
-	void MarkPropertySetBySerialization(const UStruct* Struct, const void* StructData, const FProperty* Property, int32 ArrayIndex)
-	{
-		return MarkPropertySetBySerialization(Struct, StructData, Property->ContainerPtrToValuePtr<void>(StructData, ArrayIndex));
-	}
-	
-	void MarkPropertySetBySerialization(UObject* Object, const FPropertyPathName& Path)
-	{
-		void* ResolvedData = nullptr;
-		void* ResolvedOwnerData = nullptr;
-		FProperty* ResolvedProperty = nullptr;
-		FStructProperty* ResolvedOwnerProperty = nullptr;
-		if (ensureMsgf(ResolvePropertyPathName(Object, Path, ResolvedData, ResolvedProperty, ResolvedOwnerData, ResolvedOwnerProperty),
-			TEXT("Failed to resolve property path name %s"), *WriteToString<256>(Path)))
-		{
-			if (Path.GetSegmentCount() == 1) // for paths of length == 1, the owner is the Object
-			{
-				MarkPropertySetBySerialization(Object->GetClass(), Object, ResolvedData);
-			}
-			else if (ResolvedOwnerProperty) // only mark properties set if they're in structs/classes
-			{
-				MarkPropertySetBySerialization(ResolvedOwnerProperty->Struct, ResolvedOwnerData, ResolvedData);
-			}
-		}
-	}
 
-	static bool WasPropertySetBySerialization(const UStruct* Struct, const void* StructData, const void* PropertyDataPtr)
+	bool WasPropertySetBySerialization(const UStruct* Struct, const void* StructData, const FProperty* Property, int32 ArrayIndex)
 	{
 		if (const FSetProperty* ValuesSetByPropertyBagProperty = CastField<FSetProperty>(Struct->FindPropertyByName(NAME_ValuesSetBySerialization)))
 		{
 			const FScriptSetHelper ValuesSetByPropertyBag(ValuesSetByPropertyBagProperty, ValuesSetByPropertyBagProperty->ContainerPtrToValuePtr<void>(StructData));
+			const void* PropertyDataPtr = Property->ContainerPtrToValuePtr<void>(StructData, ArrayIndex);
 			const int64 ValueOffset = static_cast<const uint8*>(PropertyDataPtr) - static_cast<const uint8*>(StructData);
 			return ValuesSetByPropertyBag.FindElementIndex(&ValueOffset) != INDEX_NONE;
 		}
 		return false;
 	}
 
-	bool WasPropertySetBySerialization(UObject* Object, const FPropertyPathName& Path)
-	{
-		void* ResolvedData = nullptr;
-		void* ResolvedOwnerData = nullptr;
-		FProperty* ResolvedProperty = nullptr;
-		FStructProperty* ResolvedOwnerProperty = nullptr;
-		if (ensureMsgf(ResolvePropertyPathName(Object, Path, ResolvedData, ResolvedProperty, ResolvedOwnerData, ResolvedOwnerProperty),
-			TEXT("Failed to resolve property path name %s"), *WriteToString<256>(Path)))
-		{
-			// only properties in structs/classes have been marked
-			if (ResolvedOwnerProperty)
-			{
-				return WasPropertySetBySerialization(ResolvedOwnerProperty->Struct, ResolvedOwnerData, ResolvedData);
-			}
-		}
-		return false;
-	}
-	
-	bool WasPropertySetBySerialization(const UStruct* Struct, const void* StructData, const FProperty* Property, int32 ArrayIndex)
-	{
-		return WasPropertySetBySerialization(Struct, StructData, Property->ContainerPtrToValuePtr<void>(StructData, ArrayIndex));
-	}
-
 	void CopyPropertySetBySerializationData(const FFieldVariant& OldField, void* OldDataPtr, const FFieldVariant& NewField, void* NewDataPtr)
 	{
 		if (const FStructProperty* OldAsStructProperty = OldField.Get<FStructProperty>())
diff --git a/Engine/Source/Runtime/CoreUObject/Private/UObject/InstanceDataObjectUtils.h b/Engine/Source/Runtime/CoreUObject/Private/UObject/InstanceDataObjectUtils.h
--- a/Engine/Source/Runtime/CoreUObject/Private/UObject/InstanceDataObjectUtils.h
+++ b/Engine/Source/Runtime/CoreUObject/Private/UObject/InstanceDataObjectUtils.h
@@ -10,7 +10,6 @@ class UClass;
 class UObject;
 class UStruct;
 
-namespace UE { class FPropertyPathName; }
 namespace UE { class FPropertyPathNameTree; }
 
 namespace UE
@@ -26,13 +25,8 @@ bool IsInstanceDataObjectSupportEnabled(UObject* Object = nullptr);
 /** Generate a UClass that contains the union of the properties of PropertyTree and OwnerClass. */
 UClass* CreateInstanceDataObjectClass(const FPropertyPathNameTree* PropertyTree, UClass* OwnerClass, UObject* Outer);
 
-/** Mark a property within the object as having been set during deserialization. */
-void MarkPropertySetBySerialization(UObject* Object, const FPropertyPathName& Path);
 /** Notify that a property in an struct was set when the struct was deserialized. */
 void MarkPropertySetBySerialization(const UStruct* Struct, const void* StructData, const FProperty* Property, int32 ArrayIndex = 0);
-
-/** Query whether a property within the object was set when the object was deserialized. */
-bool WasPropertySetBySerialization(UObject* Object, const FPropertyPathName& Path);
 /** Query whether a property in the struct was set when the struct was deserialized. */
 bool WasPropertySetBySerialization(const UStruct* Struct, const void* StructData, const FProperty* Property, int32 ArrayIndex = 0);
 /** Copy whether each property was set by serialization from one IDO to another. */
diff --git a/Engine/Source/Runtime/CoreUObject/Private/UObject/PropertyBagRepository.cpp b/Engine/Source/Runtime/CoreUObject/Private/UObject/PropertyBagRepository.cpp
--- a/Engine/Source/Runtime/CoreUObject/Private/UObject/PropertyBagRepository.cpp
+++ b/Engine/Source/Runtime/CoreUObject/Private/UObject/PropertyBagRepository.cpp
@@ -557,11 +557,6 @@ const UObject* FPropertyBagRepository::FindInstanceForDataObject(const UObject*
 	return Owner ? *Owner : nullptr;
 }
 
-bool FPropertyBagRepository::WasPropertySetBySerialization(UObject* Object, const FPropertyPathName& Path)
-{
-	return UE::WasPropertySetBySerialization(Object, Path);
-}
-
 bool FPropertyBagRepository::WasPropertySetBySerialization(const UStruct* Struct, const void* StructData, const FProperty* Property, int32 ArrayIndex)
 {
 	return UE::WasPropertySetBySerialization(Struct, StructData, Property, ArrayIndex);
@@ -646,26 +641,15 @@ void FPropertyBagRepository::CreateInstanceDataObjectUnsafe(UObject* Owner, FPro
 	}
 	else if (FLinkerLoad* Linker = Owner->GetLinker())
 	{
-		const FDelegateHandle OnTaggedPropertySerializeHandle = LoadContext->OnTaggedPropertySerialize.AddLambda(
-			[&BagData](const FUObjectSerializeContext& Context)
-			{
-				if (!Context.SerializedPropertyPath.IsEmpty())
-				{
-					MarkPropertySetBySerialization(BagData.InstanceDataObject, Context.SerializedPropertyPath);
-				}
-			}
-		);
-
 		// TODO: @jordan.hoffmann - this is very inefficient! We should remove this call to Preload. To do so, we'd need to change MarkPropertySetBySerialization
 		// to cache the serialized property list in the property bag instead of the structs. We'd also need to copy the property bag values to the IDO
 		Owner->SetFlags(RF_NeedLoad);
 		{
 			TGuardValue<bool> ScopedSkipKnownProperties(Linker->bSkipKnownProperties, true);
 			FGuardValue_Bitfield(Linker->ArMergeOverrides, true);
-            Linker->Preload(Owner);
+			Linker->Preload(Owner);
 		}
-		LoadContext->OnTaggedPropertySerialize.Remove(OnTaggedPropertySerializeHandle);
-		
+
 		// copy data from owner to IDO
 		CopyTaggedProperties(Owner, BagData.InstanceDataObject);
 	}
diff --git a/Engine/Source/Runtime/CoreUObject/Private/UObject/PropertyHelper.h b/Engine/Source/Runtime/CoreUObject/Private/UObject/PropertyHelper.h
--- a/Engine/Source/Runtime/CoreUObject/Private/UObject/PropertyHelper.h
+++ b/Engine/Source/Runtime/CoreUObject/Private/UObject/PropertyHelper.h
@@ -117,7 +117,6 @@ public:
 		{
 			Context = InContext;
 			Context->SerializedPropertyPath.Push(InSegment);
-			Notify = InNotify;
 		}
 	}
 
@@ -125,10 +124,6 @@ public:
 	{
 		if (Context)
 		{
-			if (Notify == ESerializedPropertyPathNotify::Yes)
-			{
-				Context->OnTaggedPropertySerialize.Broadcast(*Context);
-			}
 			Context->SerializedPropertyPath.Pop();
 		}
 	}
@@ -138,7 +133,6 @@ public:
 
 private:
 	FUObjectSerializeContext* Context = nullptr;
-	ESerializedPropertyPathNotify Notify = ESerializedPropertyPathNotify::No;
 };
 
 /**
@@ -155,7 +149,6 @@ public:
 		{
 			Context = InContext;
 			Context->SerializedPropertyPath.SetIndex(InIndex);
-			Notify = InNotify;
 		}
 	}
 
@@ -163,10 +156,6 @@ public:
 	{
 		if (Context)
 		{
-			if (Notify == ESerializedPropertyPathNotify::Yes)
-			{
-				Context->OnTaggedPropertySerialize.Broadcast(*Context);
-			}
 			Context->SerializedPropertyPath.SetIndex(INDEX_NONE);
 		}
 	}
@@ -176,7 +165,6 @@ public:
 
 private:
 	FUObjectSerializeContext* Context = nullptr;
-	ESerializedPropertyPathNotify Notify = ESerializedPropertyPathNotify::No;
 };
 
 } // UE
diff --git a/Engine/Source/Runtime/CoreUObject/Public/UObject/UObjectThreadContext.h b/Engine/Source/Runtime/CoreUObject/Public/UObject/UObjectThreadContext.h
--- a/Engine/Source/Runtime/CoreUObject/Public/UObject/UObjectThreadContext.h
+++ b/Engine/Source/Runtime/CoreUObject/Public/UObject/UObjectThreadContext.h
@@ -177,10 +177,6 @@ public:
 	/** True when the SerializedObject properties are being impersonated. */
 	UE_INTERNAL bool bImpersonateProperties;
 
-	/** event called after each tagged property is deserialized */
-	DECLARE_MULTICAST_DELEGATE_OneParam(FOnTaggedPropertySerialized, const FUObjectSerializeContext&)
-	UE_INTERNAL FOnTaggedPropertySerialized OnTaggedPropertySerialize;
-
 	/** Adds a new loaded object */
 	COREUOBJECT_API void AddLoadedObject(UObject* InObject);
 	COREUOBJECT_API void AddUniqueLoadedObjects(const TArray<UObject*>& InObjects);
