Fix for AutomaticOptionalInclusion of a BP that depends on a native engine class resulting in a fatal cook error. We now explicitly check if the missing import is to a native type and suppress the error "Attempt to map an object during save that was not listed as a dependency".

diff --git a/Engine/Source/Runtime/CoreUObject/Private/UObject/LinkerSave.cpp b/Engine/Source/Runtime/CoreUObject/Private/UObject/LinkerSave.cpp
--- a/Engine/Source/Runtime/CoreUObject/Private/UObject/LinkerSave.cpp
+++ b/Engine/Source/Runtime/CoreUObject/Private/UObject/LinkerSave.cpp
@@ -5,6 +5,7 @@
 #include "Misc/CommandLine.h"
 #include "Misc/ConfigCacheIni.h"
 #include "Misc/Optional.h"
+#include "Misc/PackageName.h"
 #include "Serialization/BulkData.h"
 #include "Serialization/LargeMemoryWriter.h"
 #include "UObject/Package.h"
@@ -243,6 +244,7 @@ FPackageIndex FLinkerSave::MapObject(TObjectPtr<const UObject> Object) const
 	if (Object)
 	{
 		const FPackageIndex *Found = ObjectIndicesMap.Find(Object);
+
 		if (Found)
 		{
 			if (IsCooking() && CurrentlySavingExport.IsExport() &&
@@ -280,10 +282,15 @@ FPackageIndex FLinkerSave::MapObject(TObjectPtr<const UObject> Object) const
 				}
 				if (!bFoundDep)
 				{
-					UE_LOG(LogLinker, Fatal, TEXT("Attempt to map an object during save that was not listed as a dependency. Saving Export %d %s in %s. Missing Dep on %s %s."),
-						CurrentlySavingExport.ForDebugging(), *SavingExport.ObjectName.ToString(), *GetArchiveName(),
-						Found->IsExport() ? TEXT("Export") : TEXT("Import"), *ImpExp(*Found).ObjectName.ToString()
-						);
+					const FString ImpExpObjectNameString = ImpExp(*Found).ObjectName.ToString();
+					const bool IsNativeDep = FPackageName::IsScriptPackage(ImpExpObjectNameString);
+					if (!IsNativeDep)
+					{
+						UE_LOG(LogLinker, Fatal, TEXT("Attempt to map an object during save that was not listed as a dependency. Saving Export %d %s in %s. Missing Dep on %s %s."),
+							CurrentlySavingExport.ForDebugging(), *SavingExport.ObjectName.ToString(), *GetArchiveName(),
+							Found->IsExport() ? TEXT("Export") : TEXT("Import"), *ImpExpObjectNameString
+							);
+					}
 				}
 			}
 
