[UObject]
  - Add thread-safe ref-counting mecanism to UObject
  - Make TStrongObjectPtr more light-weight and usable on any thread by using ref-count instead of FGCObject
  - Add pinning API to WeakObjectPtr so they can be converted safely to StrongObjectPtr from any thread
  - Make delegate broadcast thread-safe when used with UObjects by pinning during broadcast for non game-thread.

diff --git a/Engine/Source/Programs/Shared/EpicGames.Core/UnrealEngineTypes.cs b/Engine/Source/Programs/Shared/EpicGames.Core/UnrealEngineTypes.cs
--- a/Engine/Source/Programs/Shared/EpicGames.Core/UnrealEngineTypes.cs
+++ b/Engine/Source/Programs/Shared/EpicGames.Core/UnrealEngineTypes.cs
@@ -712,6 +712,11 @@ namespace EpicGames.Core
 		/// </summary>
 		Unreachable = 1 << 28,
 
+		/// <summary>
+		/// Object currently has ref-counts associated with it.
+		/// </summary>
+		RefCounted = 1 << 29,
+
 		/// <summary>
 		/// Object will not be garbage collected, even if unreferenced.
 		/// </summary>
@@ -722,10 +727,10 @@ namespace EpicGames.Core
 		/// </summary>
 		PendingConstruction = 1 << 31,
 
-		GarbageCollectionKeepFlags = Native | Async | AsyncLoading | LoaderImport,
+		GarbageCollectionKeepFlags = Native | Async | AsyncLoading | LoaderImport | RefCounted,
 
 		//~ Make sure this is up to date!
-		AllFlags = ReachabilityFlag0 | ReachabilityFlag1 | MaybeUnreachable | LoaderImport | Garbage | ReachableInCluster | ClusterRoot | Native | Async | AsyncLoading | Unreachable | RootSet | PendingConstruction
+		AllFlags = ReachabilityFlag0 | ReachabilityFlag1 | MaybeUnreachable | LoaderImport | Garbage | ReachableInCluster | ClusterRoot | Native | Async | AsyncLoading | Unreachable | RootSet | PendingConstruction | RefCounted
 	};
 
 	/// <summary>
diff --git a/Engine/Source/Runtime/Core/Public/Delegates/DelegateInstancesImpl.h b/Engine/Source/Runtime/Core/Public/Delegates/DelegateInstancesImpl.h
--- a/Engine/Source/Runtime/Core/Public/Delegates/DelegateInstancesImpl.h
+++ b/Engine/Source/Runtime/Core/Public/Delegates/DelegateInstancesImpl.h
@@ -156,6 +156,7 @@ public:
 
 	bool ExecuteIfSafe(ParamTypes... Params) const final
 	{
+		TStrongObjectPtr<UserClass> PinnedObject = IsInGameThread() ? nullptr : this->UserObjectPtr.Pin();
 		if (UserClass* ActualUserObject = this->UserObjectPtr.Get())
 		{
 			using FParmsWithPayload = TPayload<RetValType(typename TDecay<ParamTypes>::Type..., typename TDecay<VarTypes> ::Type...)>;
@@ -651,6 +652,7 @@ public:
 
 	bool ExecuteIfSafe(ParamTypes... Params) const final
 	{
+		TStrongObjectPtr<UserClass> PinnedObject = IsInGameThread() ? nullptr : this->UserObject.Pin();
 		if (UserClass* ActualUserObject = this->UserObject.Get())
 		{
 			using MutableUserClass = std::remove_const_t<UserClass>;
@@ -963,6 +965,7 @@ public:
 
 	bool ExecuteIfSafe(ParamTypes... Params) const final
 	{
+		TStrongObjectPtr<UserClass> PinnedObject = IsInGameThread() ? nullptr : ContextObject.Pin();
 		if (ContextObject.IsValid())
 		{
 			(void)this->Payload.ApplyAfter(Functor, Forward<ParamTypes>(Params)...);
diff --git a/Engine/Source/Runtime/Core/Public/UObject/StrongObjectPtrTemplates.h b/Engine/Source/Runtime/Core/Public/UObject/StrongObjectPtrTemplates.h
new file mode 100644
--- /dev/null
+++ b/Engine/Source/Runtime/Core/Public/UObject/StrongObjectPtrTemplates.h
@@ -0,0 +1,254 @@
+// Copyright Epic Games, Inc. All Rights Reserved.
+
+#pragma once
+
+#include "Templates/EnableIf.h"
+#include "Templates/PointerIsConvertibleFromTo.h"
+#include "Templates/UniquePtr.h"
+#include "UObject/StrongObjectPtrTemplatesFwd.h"
+
+class UObject;
+
+namespace UEStrongObjectPtr_Private
+{
+	struct FInternalReferenceCollectorReferencerNameProvider
+	{
+	};
+
+	COREUOBJECT_API void ReleaseUObject(const UObject*);
+}
+
+/**
+ * Take a ref-count on a UObject to prevent it from being GC'd while this guard is in scope.
+ */
+template <typename ObjectType, typename ReferencerNameProvider>
+class TStrongObjectPtr
+{
+public:
+	using ElementType = ObjectType;
+
+	TStrongObjectPtr(TStrongObjectPtr&& InOther)
+	{
+		Reset();
+		Object = InOther.Object;
+		InOther.Object = nullptr;
+	}
+
+	TStrongObjectPtr& operator=(TStrongObjectPtr&& InOther)
+	{
+		if (this != &InOther)
+		{
+			Reset();
+			Object = InOther.Object;
+			InOther.Object = nullptr;
+		}
+		return *this;
+	}
+
+	FORCEINLINE_DEBUGGABLE ~TStrongObjectPtr()
+	{
+		Reset();
+	}
+
+	FORCEINLINE_DEBUGGABLE TStrongObjectPtr(TYPE_OF_NULLPTR = nullptr)
+	{
+		static_assert(TPointerIsConvertibleFromTo<ObjectType, const volatile UObject>::Value, "TStrongObjectPtr can only be constructed with UObject types");
+	}
+
+	FORCEINLINE_DEBUGGABLE explicit TStrongObjectPtr(ObjectType* InObject)
+	{
+		static_assert(TPointerIsConvertibleFromTo<ObjectType, const volatile UObject>::Value, "TStrongObjectPtr can only be constructed with UObject types");
+		Reset(InObject);
+	}
+
+	FORCEINLINE_DEBUGGABLE TStrongObjectPtr(const TStrongObjectPtr& InOther)
+	{
+		Reset(InOther.Get());
+	}
+
+	template <
+		typename OtherObjectType,
+		typename OtherReferencerNameProvider
+		UE_REQUIRES(std::is_convertible_v<OtherObjectType*, ObjectType*>)
+	>
+	FORCEINLINE_DEBUGGABLE TStrongObjectPtr(const TStrongObjectPtr<OtherObjectType, OtherReferencerNameProvider>& InOther)
+	{
+		Reset(InOther.Get());
+	}
+
+	FORCEINLINE_DEBUGGABLE TStrongObjectPtr& operator=(const TStrongObjectPtr& InOther)
+	{
+		Reset(InOther.Get());
+		return *this;
+	}
+
+	template <
+		typename OtherObjectType,
+		typename OtherReferencerNameProvider
+		UE_REQUIRES(std::is_convertible_v<OtherObjectType*, ObjectType*>)
+	>
+	FORCEINLINE_DEBUGGABLE TStrongObjectPtr& operator=(const TStrongObjectPtr<OtherObjectType, OtherReferencerNameProvider>& InOther)
+	{
+		Reset(InOther.Get());
+		return *this;
+	}
+
+	FORCEINLINE_DEBUGGABLE ObjectType& operator*() const
+	{
+		check(IsValid());
+		return *(ObjectType*)Get();
+	}
+
+	FORCEINLINE_DEBUGGABLE ObjectType* operator->() const
+	{
+		check(IsValid());
+		return (ObjectType*)Get();
+	}
+
+	FORCEINLINE_DEBUGGABLE bool IsValid() const
+	{
+		return Object != nullptr;
+	}
+
+	FORCEINLINE_DEBUGGABLE explicit operator bool() const
+	{
+		return IsValid();
+	}
+
+	FORCEINLINE_DEBUGGABLE ObjectType* Get() const
+	{
+		return (ObjectType*)Object;
+	}
+
+	FORCEINLINE_DEBUGGABLE void Reset()
+	{
+		if (Object)
+		{
+			// UObject type is forward declared, ReleaseRef() is not known.
+			// So move the implementation to the cpp file instead.
+			UEStrongObjectPtr_Private::ReleaseUObject(Object);
+			Object = nullptr;
+		}
+	}
+	
+	// Attach an object without incrementing its ref-count.
+	FORCEINLINE_DEBUGGABLE void Attach(ObjectType* InNewObject)
+	{
+		Reset();
+		Object = InNewObject;
+	}
+
+	// Detach the current owned object without decrementing its ref-count.
+	FORCEINLINE_DEBUGGABLE ObjectType* Detach()
+	{
+		ObjectType* DetachedObject = Get();
+		Object = nullptr;
+		return DetachedObject;
+	}
+
+	FORCEINLINE_DEBUGGABLE void Reset(ObjectType* InNewObject)
+	{
+		if (InNewObject)
+		{
+			if (Object == InNewObject)
+			{
+				return;
+			}
+
+			if (Object)
+			{
+				// UObject type is forward declared, ReleaseRef() is not known.
+				// So move the implementation to the cpp file instead.
+				UEStrongObjectPtr_Private::ReleaseUObject(Object);
+			}
+			InNewObject->AddRef();
+			Object = InNewObject;
+		}
+		else
+		{
+			Reset();
+		}
+	}
+
+	[[nodiscard]] FORCEINLINE_DEBUGGABLE friend uint32 GetTypeHash(const TStrongObjectPtr& InStrongObjectPtr)
+	{
+		return GetTypeHash(InStrongObjectPtr.Get());
+	}
+
+private:
+	// Store as UObject to allow forward declarations without having to fully resolve ObjectType before construction.
+	// This is required because the destructor calls Reset, which need to be fully resolved at declaration.
+	const UObject* Object{ nullptr };
+
+	[[nodiscard]] friend FORCEINLINE bool operator==(const TStrongObjectPtr& InLHS, const TStrongObjectPtr& InRHS)
+	{
+		return InLHS.Get() == InRHS.Get();
+	}
+
+	[[nodiscard]] friend FORCEINLINE bool operator==(const TStrongObjectPtr& InLHS, TYPE_OF_NULLPTR)
+	{
+		return !InLHS.IsValid();
+	}
+
+	[[nodiscard]] friend FORCEINLINE bool operator==(TYPE_OF_NULLPTR, const TStrongObjectPtr& InRHS)
+	{
+		return !InRHS.IsValid();
+	}
+
+	template <
+		typename RHSObjectType,
+		typename RHSReferencerNameProvider
+		UE_REQUIRES(UE_REQUIRES_EXPR(std::declval<ObjectType*>() == std::declval<RHSObjectType*>()))
+	>
+	[[nodiscard]] friend FORCEINLINE bool operator==(const TStrongObjectPtr& InLHS, const TStrongObjectPtr<RHSObjectType, RHSReferencerNameProvider>& InRHS)
+	{
+		return InLHS.Get() == InRHS.Get();
+	}
+
+	template <
+		typename LHSObjectType,
+		typename LHSReferencerNameProvider
+		UE_REQUIRES(UE_REQUIRES_EXPR(std::declval<LHSObjectType*>() == std::declval<ObjectType*>()))
+	>
+	[[nodiscard]] friend FORCEINLINE bool operator==(const TStrongObjectPtr<LHSObjectType, LHSReferencerNameProvider>& InLHS, const TStrongObjectPtr& InRHS)
+	{
+		return InLHS.Get() == InRHS.Get();
+	}
+
+#if !PLATFORM_COMPILER_HAS_GENERATED_COMPARISON_OPERATORS
+	[[nodiscard]] friend FORCEINLINE bool operator!=(const TStrongObjectPtr& InLHS, const TStrongObjectPtr& InRHS)
+	{
+		return InLHS.Get() != InRHS.Get();
+	}
+
+	[[nodiscard]] friend FORCEINLINE bool operator!=(const TStrongObjectPtr& InLHS, TYPE_OF_NULLPTR)
+	{
+		return InLHS.IsValid();
+	}
+
+	[[nodiscard]] friend FORCEINLINE bool operator!=(TYPE_OF_NULLPTR, const TStrongObjectPtr& InRHS)
+	{
+		return InRHS.IsValid();
+	}
+
+	template <
+		typename RHSObjectType,
+		typename RHSReferencerNameProvider
+		UE_REQUIRES(UE_REQUIRES_EXPR(std::declval<ObjectType*>() == std::declval<RHSObjectType*>()))
+	>
+	[[nodiscard]] friend FORCEINLINE bool operator!=(const TStrongObjectPtr& InLHS, const TStrongObjectPtr<RHSObjectType, RHSReferencerNameProvider>& InRHS)
+	{
+		return InLHS.Get() != InRHS.Get();
+	}
+
+	template <
+		typename LHSObjectType,
+		typename LHSReferencerNameProvider
+		UE_REQUIRES(UE_REQUIRES_EXPR(std::declval<LHSObjectType*>() == std::declval<ObjectType*>()))
+	>
+	[[nodiscard]] friend FORCEINLINE bool operator!=(const TStrongObjectPtr<LHSObjectType, LHSReferencerNameProvider>& InLHS, const TStrongObjectPtr& InRHS)
+	{
+		return InLHS.Get() != InRHS.Get();
+	}
+#endif
+};
diff --git a/Engine/Source/Runtime/Core/Public/UObject/StrongObjectPtrTemplatesFwd.h b/Engine/Source/Runtime/Core/Public/UObject/StrongObjectPtrTemplatesFwd.h
new file mode 100644
--- /dev/null
+++ b/Engine/Source/Runtime/Core/Public/UObject/StrongObjectPtrTemplatesFwd.h
@@ -0,0 +1,21 @@
+// Copyright Epic Games, Inc. All Rights Reserved.
+
+#pragma once
+
+class UObject;
+
+namespace UEStrongObjectPtr_Private
+{
+	struct FInternalReferenceCollectorReferencerNameProvider;
+}
+
+/**
+ * TStrongObjectPtr is a strong pointer to a UObject.
+ * It can return nullptr if it has not been initialized or has been constructed from a weak ptr that is already garbage collected.
+ * It prevents an object from being garbage collected.
+ * It can't be directly used across a network.
+ *
+ * Most often it is used when you explicitly want to prevent something from being garbage collected.
+ */
+template <typename ObjectType, typename ReferencerNameProvider = UEStrongObjectPtr_Private::FInternalReferenceCollectorReferencerNameProvider>
+class TStrongObjectPtr;
diff --git a/Engine/Source/Runtime/Core/Public/UObject/WeakObjectPtrTemplates.h b/Engine/Source/Runtime/Core/Public/UObject/WeakObjectPtrTemplates.h
--- a/Engine/Source/Runtime/Core/Public/UObject/WeakObjectPtrTemplates.h
+++ b/Engine/Source/Runtime/Core/Public/UObject/WeakObjectPtrTemplates.h
@@ -7,6 +7,7 @@
 #include "Templates/Requires.h"
 #include "Containers/Map.h"
 #include "UObject/WeakObjectPtrTemplatesFwd.h"
+#include "UObject/StrongObjectPtrTemplatesFwd.h"
 
 #include <type_traits>
 
@@ -126,6 +127,28 @@ public:
 		return (T*)TWeakObjectPtrBase::Get();
 	}
 
+	/**
+	 * Pin the weak pointer and get a strongptr.
+	 * @param bEvenIfPendingKill if this is true, pendingkill objects are considered valid
+	 * @return nullptr if this object is gone or the weak pointer is explicitly null, otherwise a valid uobject pointer
+	 */
+	FORCEINLINE TStrongObjectPtr<T> Pin(bool bEvenIfPendingKill) const
+	{
+		TStrongObjectPtr<T> StrongPtr;
+		StrongPtr.Attach((T*)TWeakObjectPtrBase::Pin(bEvenIfPendingKill).Detach());
+		return StrongPtr;
+	}
+
+	/**
+	 * Pin the weak pointer as a strong ptr. This is an optimized version implying bEvenIfPendingKill=false.
+	 */
+	FORCEINLINE TStrongObjectPtr<T> Pin(/*bool bEvenIfPendingKill = false*/) const
+	{
+		TStrongObjectPtr<T> StrongPtr;
+		StrongPtr.Attach((T*)TWeakObjectPtrBase::Pin().Detach());
+		return StrongPtr;
+	}
+
 	/** Deferences the weak pointer even if its marked RF_Unreachable. This is needed to resolve weak pointers during GC (such as ::AddReferenceObjects) */
 	FORCEINLINE T* GetEvenIfUnreachable() const
 	{
diff --git a/Engine/Source/Runtime/CoreUObject/Private/UObject/GarbageCollection.cpp b/Engine/Source/Runtime/CoreUObject/Private/UObject/GarbageCollection.cpp
--- a/Engine/Source/Runtime/CoreUObject/Private/UObject/GarbageCollection.cpp
+++ b/Engine/Source/Runtime/CoreUObject/Private/UObject/GarbageCollection.cpp
@@ -624,6 +624,14 @@ bool FUObjectItem::ClearRootFlags(EInternalObjectFlags FlagsToClear)
 	constexpr int32 RootFlags = (int32)EInternalObjectFlags_RootFlags;
 	FScopeLock RootsLock(&GRootsCritical);
 	const int32 OldFlags = GetFlagsInternal();
+
+	// Now that we have the lock, we prevent the removal if refcount is not 0
+	// This avoids a race condition between ReleaseRef and AddRef.
+	if ((FlagsToClear & EInternalObjectFlags::RefCounted) == EInternalObjectFlags::RefCounted && GetRefCountInternal())
+	{
+		FlagsToClear &= ~EInternalObjectFlags::RefCounted;
+	}
+
 	if ((OldFlags & RootFlags) != 0 && ((OldFlags & ~(int32)FlagsToClear) & RootFlags) == 0)
 	{
 		GRoots.Remove(GUObjectArray.ObjectToIndex(Object));
@@ -6225,6 +6233,18 @@ void UObjectBase::MarkAsReachable() const
 	});
 }
 
+void UObjectBase::AddRef() const
+{
+	FUObjectItem* ObjectItem = GUObjectArray.ObjectToObjectItem(this);
+	ObjectItem->AddRef();
+}
+
+void UObjectBase::ReleaseRef() const
+{
+	FUObjectItem* ObjectItem = GUObjectArray.ObjectToObjectItem(this);
+	ObjectItem->ReleaseRef();
+}
+
 #if WITH_VERSE_VM || defined(__INTELLISENSE__)
 void UObject::VerseMarkAsReachable() const
 {
diff --git a/Engine/Source/Runtime/CoreUObject/Private/UObject/ReferenceChainSearch.cpp b/Engine/Source/Runtime/CoreUObject/Private/UObject/ReferenceChainSearch.cpp
--- a/Engine/Source/Runtime/CoreUObject/Private/UObject/ReferenceChainSearch.cpp
+++ b/Engine/Source/Runtime/CoreUObject/Private/UObject/ReferenceChainSearch.cpp
@@ -1289,6 +1289,11 @@ FString FReferenceChainSearch::GetObjectFlags(const FGCObjectInfo& InObject)
 		Flags += FString::Printf(TEXT("(Unreachable<%d>) "), (int32)UE::GC::GUnreachableObjectFlag);
 	}
 
+	if (InObject.HasAnyInternalFlags(EInternalObjectFlags::RefCounted))
+	{
+		Flags += FString::Printf(TEXT("(refcounted<%d>) "), InObject.GetRefCount());
+	}
+
 	if (InObject.IsRooted())
 	{
 		Flags += TEXT("(root) ");
diff --git a/Engine/Source/Runtime/CoreUObject/Private/UObject/StrongObjectPtr.cpp b/Engine/Source/Runtime/CoreUObject/Private/UObject/StrongObjectPtr.cpp
new file mode 100644
--- /dev/null
+++ b/Engine/Source/Runtime/CoreUObject/Private/UObject/StrongObjectPtr.cpp
@@ -0,0 +1,17 @@
+// Copyright Epic Games, Inc. All Rights Reserved.
+
+
+#include "UObject/StrongObjectPtr.h"
+#include "UObject/GarbageCollection.h"
+#include "UObject/Object.h"
+
+namespace UEStrongObjectPtr_Private
+{
+	void ReleaseUObject(const UObject* InObject)
+	{
+		if (!GExitPurge)
+		{
+			InObject->ReleaseRef();
+		}
+	}
+}
diff --git a/Engine/Source/Runtime/CoreUObject/Private/UObject/UObjectArray.cpp b/Engine/Source/Runtime/CoreUObject/Private/UObject/UObjectArray.cpp
--- a/Engine/Source/Runtime/CoreUObject/Private/UObject/UObjectArray.cpp
+++ b/Engine/Source/Runtime/CoreUObject/Private/UObject/UObjectArray.cpp
@@ -269,6 +269,7 @@ void FUObjectArray::AllocateUObjectIndex(UObjectBase* Object, EInternalObjectFla
 	{
 		ObjectItem->Flags |= (int32)UE::GC::GReachableObjectFlag;
 	}
+	ObjectItem->RefCount = 0;
 	ObjectItem->ClusterRootIndex = 0;
 	ObjectItem->SerialNumber = SerialNumber;
 	Object->InternalIndex = Index;
@@ -338,6 +339,9 @@ void FUObjectArray::FreeUObjectIndex(UObjectBase* Object)
 	FUObjectItem* ObjectItem = IndexToObject(Index);
 	UE_CLOG(ObjectItem->Object != Object, LogUObjectArray, Fatal, TEXT("Removing object (0x%016llx) at index %d but the index points to a different object (0x%016llx)!"), (int64)(PTRINT)Object, Index, (int64)(PTRINT)ObjectItem->Object);
 
+	// Can't destroy a refcounted object
+	check((ObjectItem->RefCount == 0 && (ObjectItem->GetFlags() & EInternalObjectFlags::RefCounted) == EInternalObjectFlags::None) || GExitPurge);
+
 	// Clear root flags to remove this object's index from UE::GC::Private::GRoots array 
 	if ((ObjectItem->Flags & (int32)EInternalObjectFlags_RootFlags) != 0)
 	{
@@ -346,6 +350,7 @@ void FUObjectArray::FreeUObjectIndex(UObjectBase* Object)
 
 	ObjectItem->Object = nullptr;
 	ObjectItem->Flags = 0;
+	ObjectItem->RefCount = 0;
 	ObjectItem->ClusterRootIndex = 0;
 	ObjectItem->SerialNumber = 0;
 
diff --git a/Engine/Source/Runtime/CoreUObject/Private/UObject/WeakObjectPtr.cpp b/Engine/Source/Runtime/CoreUObject/Private/UObject/WeakObjectPtr.cpp
--- a/Engine/Source/Runtime/CoreUObject/Private/UObject/WeakObjectPtr.cpp
+++ b/Engine/Source/Runtime/CoreUObject/Private/UObject/WeakObjectPtr.cpp
@@ -5,6 +5,7 @@
 =============================================================================*/
 
 #include "UObject/WeakObjectPtr.h"
+#include "UObject/GarbageCollection.h"
 #include "UObject/Object.h"
 
 DEFINE_LOG_CATEGORY_STATIC(LogWeakObjectPtr, Log, All);
@@ -104,6 +105,36 @@ UObject* FWeakObjectPtr::GetEvenIfUnreachable() const
 	return Result;
 }
 
+TStrongObjectPtr<UObject> FWeakObjectPtr::Pin(/*bool bEvenIfGarbage = false*/) const
+{
+	// Using a literal here allows the optimizer to remove branches later down the chain.
+	return Internal_Pin(false);
+}
+
+TStrongObjectPtr<UObject> FWeakObjectPtr::Pin(bool bEvenIfGarbage) const
+{
+	return Internal_Pin(bEvenIfGarbage);
+}
+
+TStrongObjectPtr<UObject> FWeakObjectPtr::PinEvenIfUnreachable() const
+{
+	FGCScopeGuard GCScopeGuard;
+	UObject* Result = nullptr;
+	if (Internal_IsValid(true, true))
+	{
+		FUObjectItem* ObjectItem = GUObjectArray.IndexToObject(GetObjectIndex_Private(), true);
+		Result = static_cast<UObject*>(ObjectItem->Object);
+	}
+	return TStrongObjectPtr<UObject>(Result);
+}
+
+TStrongObjectPtr<UObject> FWeakObjectPtr::Internal_Pin(bool bEvenIfGarbage) const
+{
+	FGCScopeGuard GCScopeGuard;
+	FUObjectItem* const ObjectItem = Internal_GetObjectItem();
+	return TStrongObjectPtr<UObject>(((ObjectItem != nullptr) && GUObjectArray.IsValid(ObjectItem, bEvenIfGarbage)) ? (UObject*)ObjectItem->Object : nullptr);
+}
+
 void FWeakObjectPtr::Serialize(FArchive& Ar)
 {
 	FArchiveUObject::SerializeWeakObjectPtr(Ar, *this);
diff --git a/Engine/Source/Runtime/CoreUObject/Public/UObject/GCObjectInfo.h b/Engine/Source/Runtime/CoreUObject/Public/UObject/GCObjectInfo.h
--- a/Engine/Source/Runtime/CoreUObject/Public/UObject/GCObjectInfo.h
+++ b/Engine/Source/Runtime/CoreUObject/Public/UObject/GCObjectInfo.h
@@ -35,6 +35,7 @@ public:
 		, InternalFlags(Obj->GetInternalFlags())
 		, ClusterRootIndex(GUObjectArray.ObjectToObjectItem(Obj)->ClusterRootIndex)
 		, bDisregardForGC(GUObjectArray.IsDisregardForGC(Obj))
+		, RefCount(GUObjectArray.ObjectToObjectItem(Obj)->GetRefCount())
 	{
 
 	}
@@ -55,6 +56,8 @@ private:
 	int32 ClusterRootIndex = -1;
 	/** True if the object was inside of the disregard for GC set */
 	bool bDisregardForGC = false;
+	/** Current ref-count of the UObject. */
+	int32 RefCount = 0;
 
 public:
 
@@ -122,6 +125,11 @@ public:
 		return Class->Name.ToString();
 	}
 
+	int32 GetRefCount() const
+	{
+		return RefCount;
+	}
+
 	bool IsValid() const
 	{
 		return !HasAnyInternalFlags(EInternalObjectFlags::Garbage);
diff --git a/Engine/Source/Runtime/CoreUObject/Public/UObject/ObjectMacros.h b/Engine/Source/Runtime/CoreUObject/Public/UObject/ObjectMacros.h
--- a/Engine/Source/Runtime/CoreUObject/Public/UObject/ObjectMacros.h
+++ b/Engine/Source/Runtime/CoreUObject/Public/UObject/ObjectMacros.h
@@ -622,7 +622,7 @@ enum class EInternalObjectFlags : int32
 	Async = 1 << 26, ///< Object exists only on a different thread than the game thread.
 	AsyncLoading = 1 << 27, ///< Object is being asynchronously loaded.
 	Unreachable UE_DEPRECATED(5.4, "Unreachable flag should no longer be used. Use UObjectBaseUtility::IsUnreachable() or UE::GC::GUnreachableObjectFlag flag instead.") = 1 << 28, ///< Object is not reachable on the object graph.
-	// Unused = 1 << 29,
+	RefCounted = 1 << 29, ///< Object currently has ref-counts associated with it.
 	RootSet = 1 << 30, ///< Object will not be garbage collected, even if unreferenced.
 	PendingConstruction = 1 << 31, ///< Object didn't have its class constructor called yet (only the UObjectBase one to initialize its most basic members)
 
@@ -630,19 +630,19 @@ enum class EInternalObjectFlags : int32
 	// DO NOT ADD new compound flags to EInternalObjectFlags. The below flags are deprecated so that one day we can remove them.
 	// !!!!!!!!!!!!!!!!!!!!!!!!!!!!!!!!!!!!!!!!!!!!!!!!!!!!!!!!!!!!!!!!!!!!!!!!!!!!!!!!!!!!!!!!!!!!!!!!!!!!!!!!!!!!!!!!!!!!!!!!
 
-	GarbageCollectionKeepFlags UE_DEPRECATED(5.4, "GarbageCollectionKeepFlags should no longer be used. Use EInternalObjectFlags_GarbageCollectionKeepFlags instead.") = Native | Async | AsyncLoading | LoaderImport,
+	GarbageCollectionKeepFlags UE_DEPRECATED(5.4, "GarbageCollectionKeepFlags should no longer be used. Use EInternalObjectFlags_GarbageCollectionKeepFlags instead.") = Native | Async | AsyncLoading | LoaderImport | RefCounted,
 	MirroredFlags UE_DEPRECATED(5.4, "MirroredFlags should no longer be used. Use Garbage instead.") = Garbage,
 
 	//~ Make sure this is up to date!
 	PRAGMA_DISABLE_DEPRECATION_WARNINGS
-	AllFlags UE_DEPRECATED(5.4, "AllFlags should no longer be used. Use EInternalObjectFlags_AllFlags instead.") = ReachabilityFlag0 | ReachabilityFlag1 | MaybeUnreachable | LoaderImport | Garbage | ReachableInCluster | ClusterRoot | Native | Async | AsyncLoading | Unreachable | RootSet | PendingConstruction
+	AllFlags UE_DEPRECATED(5.4, "AllFlags should no longer be used. Use EInternalObjectFlags_AllFlags instead.") = ReachabilityFlag0 | ReachabilityFlag1 | MaybeUnreachable | LoaderImport | Garbage | ReachableInCluster | ClusterRoot | Native | Async | AsyncLoading | Unreachable | RootSet | PendingConstruction | RefCounted
 	PRAGMA_ENABLE_DEPRECATION_WARNINGS
 };
 ENUM_CLASS_FLAGS(EInternalObjectFlags);
 
 //~ Make sure these macros are up to date!
-#define EInternalObjectFlags_GarbageCollectionKeepFlags (EInternalObjectFlags::Native | EInternalObjectFlags::Async | EInternalObjectFlags::AsyncLoading | EInternalObjectFlags::LoaderImport)
-#define EInternalObjectFlags_AllFlags (EInternalObjectFlags::ReachabilityFlag0 | EInternalObjectFlags::ReachabilityFlag1 | EInternalObjectFlags::ReachabilityFlag2 | EInternalObjectFlags::LoaderImport | EInternalObjectFlags::Garbage | EInternalObjectFlags::ReachableInCluster | EInternalObjectFlags::ClusterRoot | EInternalObjectFlags::Native | EInternalObjectFlags::Async | EInternalObjectFlags::AsyncLoading | EInternalObjectFlags::RootSet | EInternalObjectFlags::PendingConstruction | (EInternalObjectFlags)(1 << 19 /*MaybeUnreachable*/) | (EInternalObjectFlags)(1 << 28 /*Unreachable*/))
+#define EInternalObjectFlags_GarbageCollectionKeepFlags (EInternalObjectFlags::Native | EInternalObjectFlags::Async | EInternalObjectFlags::AsyncLoading | EInternalObjectFlags::LoaderImport | EInternalObjectFlags::RefCounted)
+#define EInternalObjectFlags_AllFlags (EInternalObjectFlags::ReachabilityFlag0 | EInternalObjectFlags::ReachabilityFlag1 | EInternalObjectFlags::ReachabilityFlag2 | EInternalObjectFlags::LoaderImport | EInternalObjectFlags::Garbage | EInternalObjectFlags::ReachableInCluster | EInternalObjectFlags::ClusterRoot | EInternalObjectFlags::Native | EInternalObjectFlags::RefCounted | EInternalObjectFlags::Async | EInternalObjectFlags::AsyncLoading | EInternalObjectFlags::RootSet | EInternalObjectFlags::PendingConstruction | (EInternalObjectFlags)(1 << 19 /*MaybeUnreachable*/) | (EInternalObjectFlags)(1 << 28 /*Unreachable*/))
 #define EInternalObjectFlags_RootFlags (EInternalObjectFlags::RootSet | EInternalObjectFlags_GarbageCollectionKeepFlags)
 
 /**
diff --git a/Engine/Source/Runtime/CoreUObject/Public/UObject/StrongObjectPtr.h b/Engine/Source/Runtime/CoreUObject/Public/UObject/StrongObjectPtr.h
--- a/Engine/Source/Runtime/CoreUObject/Public/UObject/StrongObjectPtr.h
+++ b/Engine/Source/Runtime/CoreUObject/Public/UObject/StrongObjectPtr.h
@@ -7,249 +7,4 @@
 #include "Templates/PointerIsConvertibleFromTo.h"
 #include "Templates/UniquePtr.h"
 
-namespace UEStrongObjectPtr_Private
-{
-	struct FInternalReferenceCollectorReferencerNameProvider
-	{
-		static FString GetReferencerName()
-		{
-			return TEXT("UEStrongObjectPtr_Private::TInternalReferenceCollector");
-		}
-	};
-
-	template <typename ReferencerNameProvider = FInternalReferenceCollectorReferencerNameProvider>
-	class TInternalReferenceCollector : public FGCObject
-	{
-	public:
-		explicit TInternalReferenceCollector(const UObject* InObject)
-			: Object(InObject)
-		{
-			checkf(IsInGameThread(), TEXT("TStrongObjectPtr can only be created on the game thread otherwise it may introduce threading issues with Grbage Collector"));
-		}
-
-		virtual ~TInternalReferenceCollector()
-		{
-			check(IsInGameThread() || IsInGarbageCollectorThread());
-		}
-
-		bool IsValid() const
-		{
-			return Object != nullptr;
-		}
-
-		template <typename UObjectType>
-		FORCEINLINE UObjectType* GetAs() const
-		{
-			return (UObjectType*)Object;
-		}
-
-		FORCEINLINE void Set(const UObject* InObject)
-		{
-			Object = InObject;
-		}
-
-		//~ FGCObject interface
-		virtual void AddReferencedObjects(FReferenceCollector& Collector) override
-		{
-			Collector.AddReferencedObject(Object);
-		}
-
-		virtual FString GetReferencerName() const override
-		{
-			return ReferencerNameProvider::GetReferencerName();
-		}
-
-	private:
-		TObjectPtr<const UObject> Object;
-	};
-}
-
-/**
- * Specific implementation of FGCObject that prevents a single UObject-based pointer from being GC'd while this guard is in scope.
- * @note This is the "full-fat" version of FGCObjectScopeGuard which uses a heap-allocated FGCObject so *can* safely be used with containers that treat types as trivially relocatable.
- */
-template <typename ObjectType, typename ReferencerNameProvider = UEStrongObjectPtr_Private::FInternalReferenceCollectorReferencerNameProvider>
-class TStrongObjectPtr
-{
-public:
-	using ElementType = ObjectType;
-	
-	TStrongObjectPtr(TStrongObjectPtr&& InOther) = default;
-	TStrongObjectPtr& operator=(TStrongObjectPtr&& InOther) = default;
-	~TStrongObjectPtr() = default;
-
-	FORCEINLINE_DEBUGGABLE TStrongObjectPtr(TYPE_OF_NULLPTR = nullptr)
-	{
-		static_assert(TPointerIsConvertibleFromTo<ObjectType, const volatile UObject>::Value, "TStrongObjectPtr can only be constructed with UObject types");
-	}
-
-	FORCEINLINE_DEBUGGABLE explicit TStrongObjectPtr(ObjectType* InObject)
-	{
-		static_assert(TPointerIsConvertibleFromTo<ObjectType, const volatile UObject>::Value, "TStrongObjectPtr can only be constructed with UObject types");
-		Reset(InObject);
-	}
-
-	FORCEINLINE_DEBUGGABLE TStrongObjectPtr(const TStrongObjectPtr& InOther)
-	{
-		Reset(InOther.Get());
-	}
-
-	template <
-		typename OtherObjectType,
-		typename OtherReferencerNameProvider
-		UE_REQUIRES(std::is_convertible_v<OtherObjectType*, ObjectType*>)
-	>
-	FORCEINLINE_DEBUGGABLE TStrongObjectPtr(const TStrongObjectPtr<OtherObjectType, OtherReferencerNameProvider>& InOther)
-	{
-		Reset(InOther.Get());
-	}
-
-	FORCEINLINE_DEBUGGABLE TStrongObjectPtr& operator=(const TStrongObjectPtr& InOther)
-	{
-		Reset(InOther.Get());
-		return *this;
-	}
-
-	template <
-		typename OtherObjectType,
-		typename OtherReferencerNameProvider
-		UE_REQUIRES(std::is_convertible_v<OtherObjectType*, ObjectType*>)
-	>
-	FORCEINLINE_DEBUGGABLE TStrongObjectPtr& operator=(const TStrongObjectPtr<OtherObjectType, OtherReferencerNameProvider>& InOther)
-	{
-		Reset(InOther.Get());
-		return *this;
-	}
-
-	FORCEINLINE_DEBUGGABLE ObjectType& operator*() const
-	{
-		check(IsValid());
-		return *Get();
-	}
-
-	FORCEINLINE_DEBUGGABLE ObjectType* operator->() const
-	{
-		check(IsValid());
-		return Get();
-	}
-
-	FORCEINLINE_DEBUGGABLE bool IsValid() const
-	{
-		return ReferenceCollector && ReferenceCollector->IsValid();
-	}
-
-	FORCEINLINE_DEBUGGABLE explicit operator bool() const
-	{
-		return IsValid();
-	}
-
-	FORCEINLINE_DEBUGGABLE ObjectType* Get() const
-	{
-		return ReferenceCollector
-			? ReferenceCollector->template GetAs<ObjectType>()
-			: nullptr;
-	}
-
-	FORCEINLINE_DEBUGGABLE void Reset(ObjectType* InNewObject = nullptr)
-	{
-		if (InNewObject)
-		{
-			if (ReferenceCollector)
-			{
-				// Update the referenced object
-				ReferenceCollector->Set(InNewObject);
-			}
-			else
-			{
-				// Lazily create the ReferenceCollector to allow TStrongObjectPtr to be used during static initialization
-				ReferenceCollector = MakeUnique< UEStrongObjectPtr_Private::TInternalReferenceCollector<ReferencerNameProvider> >(InNewObject);
-			}
-		}
-		else
-		{
-			// Destroy the ReferenceCollector immediately to allow TStrongObjectPtr to be manually cleared prior to 
-			// static deinitialization, as not all platforms use the main thread as their game thread
-			ReferenceCollector.Reset();
-		}
-	}
-
-	[[nodiscard]] FORCEINLINE_DEBUGGABLE friend uint32 GetTypeHash(const TStrongObjectPtr& InStrongObjectPtr)
-	{
-		return GetTypeHash(InStrongObjectPtr.Get());
-	}
-
-private:
-	TUniquePtr< UEStrongObjectPtr_Private::TInternalReferenceCollector<ReferencerNameProvider> > ReferenceCollector;
-
-	[[nodiscard]] friend FORCEINLINE bool operator==(const TStrongObjectPtr& InLHS, const TStrongObjectPtr& InRHS)
-	{
-		return InLHS.Get() == InRHS.Get();
-	}
-
-	[[nodiscard]] friend FORCEINLINE bool operator==(const TStrongObjectPtr& InLHS, TYPE_OF_NULLPTR)
-	{
-		return !InLHS.IsValid();
-	}
-
-	[[nodiscard]] friend FORCEINLINE bool operator==(TYPE_OF_NULLPTR, const TStrongObjectPtr& InRHS)
-	{
-		return !InRHS.IsValid();
-	}
-
-	template <
-		typename RHSObjectType,
-		typename RHSReferencerNameProvider
-		UE_REQUIRES(UE_REQUIRES_EXPR(std::declval<ObjectType*>() == std::declval<RHSObjectType*>()))
-	>
-	[[nodiscard]] friend FORCEINLINE bool operator==(const TStrongObjectPtr& InLHS, const TStrongObjectPtr<RHSObjectType, RHSReferencerNameProvider>& InRHS)
-	{
-		return InLHS.Get() == InRHS.Get();
-	}
-
-	template <
-		typename LHSObjectType,
-		typename LHSReferencerNameProvider
-		UE_REQUIRES(UE_REQUIRES_EXPR(std::declval<LHSObjectType*>() == std::declval<ObjectType*>()))
-	>
-	[[nodiscard]] friend FORCEINLINE bool operator==(const TStrongObjectPtr<LHSObjectType, LHSReferencerNameProvider>& InLHS, const TStrongObjectPtr& InRHS)
-	{
-		return InLHS.Get() == InRHS.Get();
-	}
-
-#if !PLATFORM_COMPILER_HAS_GENERATED_COMPARISON_OPERATORS
-	[[nodiscard]] friend FORCEINLINE bool operator!=(const TStrongObjectPtr& InLHS, const TStrongObjectPtr& InRHS)
-	{
-		return InLHS.Get() != InRHS.Get();
-	}
-
-	[[nodiscard]] friend FORCEINLINE bool operator!=(const TStrongObjectPtr& InLHS, TYPE_OF_NULLPTR)
-	{
-		return InLHS.IsValid();
-	}
-
-	[[nodiscard]] friend FORCEINLINE bool operator!=(TYPE_OF_NULLPTR, const TStrongObjectPtr& InRHS)
-	{
-		return InRHS.IsValid();
-	}
-
-	template <
-		typename RHSObjectType,
-		typename RHSReferencerNameProvider
-		UE_REQUIRES(UE_REQUIRES_EXPR(std::declval<ObjectType*>() == std::declval<RHSObjectType*>()))
-	>
-	[[nodiscard]] friend FORCEINLINE bool operator!=(const TStrongObjectPtr& InLHS, const TStrongObjectPtr<RHSObjectType, RHSReferencerNameProvider>& InRHS)
-	{
-		return InLHS.Get() != InRHS.Get();
-	}
-
-	template <
-		typename LHSObjectType,
-		typename LHSReferencerNameProvider
-		UE_REQUIRES(UE_REQUIRES_EXPR(std::declval<LHSObjectType*>() == std::declval<ObjectType*>()))
-	>
-	[[nodiscard]] friend FORCEINLINE bool operator!=(const TStrongObjectPtr<LHSObjectType, LHSReferencerNameProvider>& InLHS, const TStrongObjectPtr& InRHS)
-	{
-		return InLHS.Get() != InRHS.Get();
-	}
-#endif
-};
+#include "UObject/StrongObjectPtrTemplates.h"
\ No newline at end of file
diff --git a/Engine/Source/Runtime/CoreUObject/Public/UObject/UObjectArray.h b/Engine/Source/Runtime/CoreUObject/Public/UObject/UObjectArray.h
--- a/Engine/Source/Runtime/CoreUObject/Public/UObject/UObjectArray.h
+++ b/Engine/Source/Runtime/CoreUObject/Public/UObject/UObjectArray.h
@@ -48,6 +48,8 @@ public:
 	int32 ClusterRootIndex;	
 	// Weak Object Pointer Serial number associated with the object
 	int32 SerialNumber;
+	// RefCount associated with the object preventing its destruction.
+	int32 RefCount;
 
 #if STATS || ENABLE_STATNAMEDEVENTS_UOBJECT
 	/** Stat id of this object, 0 if nobody asked for it yet */
@@ -63,6 +65,7 @@ public:
 		, Flags(0)
 		, ClusterRootIndex(0)
 		, SerialNumber(0)
+		, RefCount(0)
 #if ENABLE_STATNAMEDEVENTS_UOBJECT
 		, StatIDStringStorage(nullptr)
 #endif
@@ -163,7 +166,8 @@ public:
 		bool Result = false;
 		UE_AUTORTFM_OPEN(
 		{
-			FlagToClear &= ~UE::GC::GReachableObjectFlag; // reachability bit can only be cleared by GC through *_ForGC functions
+			FlagToClear &= ~UE::GC::GReachableObjectFlag;     // reachability bit can only be cleared by GC through *_ForGC functions
+			FlagToClear &= ~EInternalObjectFlags::RefCounted; // refcounted flag is internal and must only be cleared internally by AddRef/ReleaseRef.
 			if (!!(FlagToClear & EInternalObjectFlags_RootFlags))
 			{
 				Result = ClearRootFlags(FlagToClear);
@@ -213,6 +217,7 @@ public:
 		bool Result = false;
 		UE_AUTORTFM_OPEN(
 		{
+			FlagToSet &= ~EInternalObjectFlags::RefCounted; // refcounted flag is internal and must only be set by AddRef/ReleaseRef.
 			if (!!(FlagToSet & EInternalObjectFlags_RootFlags))
 			{
 				Result = SetRootFlags(FlagToSet);
@@ -308,6 +313,38 @@ public:
 		return !!(GetFlagsInternal() & int32(EInternalObjectFlags::RootSet));
 	}
 
+	FORCEINLINE int32 GetRefCount() const
+	{
+		return RefCount;
+	}
+
+	void AddRef()
+	{
+		UE_AUTORTFM_OPEN(
+		{
+			FPlatformAtomics::InterlockedIncrement(&RefCount);
+			if ((GetFlags() & EInternalObjectFlags::RefCounted) != EInternalObjectFlags::RefCounted)
+			{
+				SetRootFlags(EInternalObjectFlags::RefCounted);
+			}
+		});
+	}
+
+	void ReleaseRef()
+	{
+		UE_AUTORTFM_OPEN(
+		{
+			// This alone is not thread-safe as we may race with AddRef and in that case we don't want ClearRootFlags to apply.
+			// We fix this by validating that the refcount is still 0 while inside the root locks in ClearRootFlags.
+			const int32 NewRefCount = FPlatformAtomics::InterlockedDecrement(&RefCount);
+			check(NewRefCount >= 0);
+			if (NewRefCount == 0)
+			{
+				ClearRootFlags(EInternalObjectFlags::RefCounted);
+			}
+		});
+	}
+
 #if STATS || ENABLE_STATNAMEDEVENTS_UOBJECT
 	COREUOBJECT_API void CreateStatID() const;
 #endif
@@ -355,6 +392,12 @@ private:
 	{
 		return FPlatformAtomics::AtomicRead_Relaxed((int32*)&Flags);
 	}
+	
+	FORCEINLINE int32 GetRefCountInternal() const
+	{
+		return FPlatformAtomics::AtomicRead_Relaxed((int32*)&RefCount);
+	}
+
 	COREUOBJECT_API bool SetRootFlags(EInternalObjectFlags FlagsToSet);
 	COREUOBJECT_API bool ClearRootFlags(EInternalObjectFlags FlagsToClear);
 };
diff --git a/Engine/Source/Runtime/CoreUObject/Public/UObject/UObjectBase.h b/Engine/Source/Runtime/CoreUObject/Public/UObject/UObjectBase.h
--- a/Engine/Source/Runtime/CoreUObject/Public/UObject/UObjectBase.h
+++ b/Engine/Source/Runtime/CoreUObject/Public/UObject/UObjectBase.h
@@ -186,6 +186,12 @@ public:
 	*/
 	COREUOBJECT_API void MarkAsReachable() const;
 
+	/** Increments ref-count on the object. Should always be paired with RAII (i.e. TStrongObjectPtr). */
+	COREUOBJECT_API void AddRef() const;
+
+	/** Release ref-count on the object. */
+	COREUOBJECT_API void ReleaseRef() const;
+
 protected:
 	/**
 	 * Set the object flags directly
diff --git a/Engine/Source/Runtime/CoreUObject/Public/UObject/WeakObjectPtr.h b/Engine/Source/Runtime/CoreUObject/Public/UObject/WeakObjectPtr.h
--- a/Engine/Source/Runtime/CoreUObject/Public/UObject/WeakObjectPtr.h
+++ b/Engine/Source/Runtime/CoreUObject/Public/UObject/WeakObjectPtr.h
@@ -11,6 +11,7 @@
 #include "Templates/UnrealTemplate.h"
 #include "UObject/ScriptDelegates.h"
 #include "UObject/UObjectArray.h"
+#include "UObject/StrongObjectPtr.h"
 #include "UObject/WeakObjectPtrFwd.h"
 
 #include <type_traits>
@@ -163,6 +164,25 @@ public:
 	/** Dereference the weak pointer even if it is marked as Garbage or Unreachable */
 	COREUOBJECT_API class UObject* GetEvenIfUnreachable() const;
 
+	/**
+	 * Get a strong object ptr to the weak pointer.
+	 * @param bEvenIfGarbage if this is true, Garbage objects are considered valid
+	 * @return TStrongObjectPtr will be invalid if this object is gone or the weak pointer is explicitly null, otherwise a valid TStrongObjectPtr
+	 */
+	COREUOBJECT_API class TStrongObjectPtr<UObject> Pin(bool bEvenIfGarbage) const;
+
+	/**
+	 * Get a strong object ptr to the weak pointer. This is an optimized version implying bEvenIfGarbage=false.
+	 * @return TStrongObjectPtr will be invalid if this object is gone or the weak pointer is explicitly null, otherwise a valid TStrongObjectPtr
+	 */
+	COREUOBJECT_API class TStrongObjectPtr<UObject> Pin(/*bool bEvenIfGarbage = false*/) const;
+
+	/*
+	 * Get a strong object ptr even if it is marked as Garbage or Unreachable 
+	 * @return TStrongObjectPtr will be invalid if this object is gone or the weak pointer is explicitly null, otherwise a valid TStrongObjectPtr
+	 */
+	COREUOBJECT_API class TStrongObjectPtr<UObject> PinEvenIfUnreachable() const;
+
 	// This is explicitly not added to avoid resolving weak pointers too often - use Get() once in a function.
 	explicit operator bool() const = delete;
 
@@ -314,6 +334,8 @@ private:
 		return ((ObjectItem != nullptr) && GUObjectArray.IsValid(ObjectItem, bEvenIfGarbage)) ? (UObject*)ObjectItem->Object : nullptr;
 	}
 
+	COREUOBJECT_API TStrongObjectPtr<UObject> Internal_Pin(bool bEvenIfGarbage) const;
+	
 #if UE_WEAKOBJECTPTR_ZEROINIT_FIX
 	int32		ObjectIndex = UE::Core::Private::InvalidWeakObjectIndex;
 	int32		ObjectSerialNumber = 0;
diff --git a/Engine/Source/Runtime/CoreUObject/Tests/ObjectPinningTest.cpp b/Engine/Source/Runtime/CoreUObject/Tests/ObjectPinningTest.cpp
new file mode 100644
--- /dev/null
+++ b/Engine/Source/Runtime/CoreUObject/Tests/ObjectPinningTest.cpp
@@ -0,0 +1,77 @@
+// Copyright Epic Games, Inc. All Rights Reserved.
+
+#if WITH_LOW_LEVEL_TESTS
+
+#include "TestHarness.h"
+#include "TestMacros/Assertions.h"
+#include "Tests/Benchmark.h"
+#include "Tasks/Task.h"
+
+namespace UE::CoreObject::Private::Tests
+{
+
+void ObjectPinningStressTest(int32 Count, int32 NumTasks)
+{
+	TArray<UE::Tasks::FTask> Tasks;
+	Tasks.Reserve(NumTasks);
+
+	std::atomic<int64> Success{ 0 };
+	std::atomic<int64> Failure{ 0 };
+
+	for (int32 Index = 0; Index < Count; ++Index)
+	{
+		UObject* Obj = NewObject<UClass>(GetTransientPackage());
+
+		UE::Tasks::FTaskEvent Trigger(TEXT("Trigger"));
+		FWeakObjectPtr WeakPtr = Obj;
+		auto Lambda =
+			[WeakPtr, &Success, &Failure]()
+			{
+				FPlatformProcess::YieldCycles(FMath::RandRange(0, 10000));
+
+				if (TStrongObjectPtr<UObject> StrongPtr = WeakPtr.Pin())
+				{
+					Success++;
+					for (int i = 0; i < 1000; ++i)
+					{
+						check(WeakPtr.IsValid());
+					}
+				}
+				else
+				{
+					Failure++;
+				}
+			};
+
+		Tasks.Reset();
+
+		for (int32 TaskIndex = 0; TaskIndex < NumTasks; ++TaskIndex)
+		{
+			Tasks.Emplace(UE::Tasks::Launch(TEXT("Pin"), [Lambda]() { Lambda(); }, Trigger));
+		}
+
+		Trigger.Trigger();
+
+		CollectGarbage(GARBAGE_COLLECTION_KEEPFLAGS);
+		UE::Tasks::Wait(Tasks);
+		if (WeakPtr.IsValid())
+		{
+			CollectGarbage(GARBAGE_COLLECTION_KEEPFLAGS);
+			check(!WeakPtr.IsValid());
+		}
+	}
+}
+
+TEST_CASE("CoreUObject::ObjectPinning::Stress Tests", "[CoreUObject][ObjectPinning]")
+{
+	// We need to be sure we've done the static GC initialization before we start doing a garbage collection.
+	FGCObject::StaticInit();
+
+	UE_BENCHMARK(5, [] { ObjectPinningStressTest(1000, 1); } );
+	UE_BENCHMARK(5, [] { ObjectPinningStressTest(1000, 2); });
+	UE_BENCHMARK(5, [] { ObjectPinningStressTest(1000, 4); });
+}
+
+} // namespace UE::CoreObject::Private::Tests
+
+#endif // WITH_LOW_LEVEL_TESTS
diff --git a/Engine/Source/Runtime/Engine/Private/UnrealEngine.cpp b/Engine/Source/Runtime/Engine/Private/UnrealEngine.cpp
--- a/Engine/Source/Runtime/Engine/Private/UnrealEngine.cpp
+++ b/Engine/Source/Runtime/Engine/Private/UnrealEngine.cpp
@@ -8981,6 +8981,7 @@ bool UEngine::HandleObjCommand( const TCHAR* Cmd, FOutputDevice& Ar )
 			const bool bOnlyListGCObjects = FParse::Param(Cmd, TEXT("GCONLY"));
 			const bool bOnlyListGCObjectsNoClusters = FParse::Param(Cmd, TEXT("GCNOCLUSTERS"));
 			const bool bOnlyListRootObjects = FParse::Param(Cmd, TEXT("ROOTONLY"));
+			const bool bOnlyListRefCountedObjects = FParse::Param(Cmd, TEXT("REFCOUNTEDONLY"));
 			const bool bShouldIncludeDefaultObjects = FParse::Param(Cmd, TEXT("INCLUDEDEFAULTS"));
 			const bool bOnlyListDefaultObjects = FParse::Param(Cmd, TEXT("DEFAULTSONLY"));
 			const bool bShowDetailedObjectInfo = FParse::Param(Cmd, TEXT("NODETAILEDINFO")) == false && bTrackDetailedObjectInfo;
@@ -9032,6 +9033,11 @@ bool UEngine::HandleObjCommand( const TCHAR* Cmd, FOutputDevice& Ar )
 					continue;
 				}
 
+				if (bOnlyListRefCountedObjects && !It->HasAnyInternalFlags(EInternalObjectFlags::RefCounted))
+				{
+					continue;
+				}
+
 				if ( CheckType && !It->IsA(CheckType) )
 				{
 					continue;
