Fixed TWeakInterfacePtr<const T> and TScriptInterface<const T>.

diff --git a/Engine/Source/Runtime/CoreUObject/Public/UObject/ScriptInterface.h b/Engine/Source/Runtime/CoreUObject/Public/UObject/ScriptInterface.h
--- a/Engine/Source/Runtime/CoreUObject/Public/UObject/ScriptInterface.h
+++ b/Engine/Source/Runtime/CoreUObject/Public/UObject/ScriptInterface.h
@@ -139,7 +139,8 @@ class TScriptInterface : public FScriptInterface
 {
 public:
 	using InterfaceType = InInterfaceType;
-	
+	using UObjectType = typename TCopyQualifiersFromTo<InterfaceType, UObject>::Type;
+
 	/**
 	 * Default constructor
 	 */
@@ -155,12 +156,12 @@ public:
 	 */
 	template <
 		typename U
-		UE_REQUIRES(std::is_convertible_v<U, UObject*>)
+		UE_REQUIRES(std::is_convertible_v<U, UObjectType*>)
 	>
 	FORCEINLINE TScriptInterface(U&& Source)
 	{
 		// Always set the object
-		UObject* SourceObject = ImplicitConv<UObject*>(Source);
+		UObjectType* SourceObject = ImplicitConv<UObjectType*>(Source);
 		SetObject(SourceObject);
 
 		if constexpr (std::is_base_of<InInterfaceType, std::remove_pointer_t<std::remove_reference_t<U>>>::value)
@@ -235,7 +236,7 @@ public:
 	 */
 	template <
 		typename U
-		UE_REQUIRES(std::is_convertible_v<U, UObject*>)
+		UE_REQUIRES(std::is_convertible_v<U, UObjectType*>)
 	>
 	TScriptInterface& operator=(U&& Source)
 	{
@@ -344,6 +345,30 @@ public:
 		FScriptInterface::SetInterface((void*)InInterfacePointer);
 	}
 
+	/**
+	 * Returns the ObjectPointer contained by this TScriptInterface
+	 */
+	FORCEINLINE UObjectType* GetObject() const
+	{
+		return FScriptInterface::GetObject();
+	}
+
+	/**
+	 * Returns the ObjectPointer contained by this TScriptInterface
+	 */
+	FORCEINLINE TObjectPtr<UObjectType>& GetObjectRef()
+	{
+		return *(TObjectPtr<UObjectType>*)&FScriptInterface::GetObjectRef();
+	}
+
+	/**
+	 * Sets the value of the ObjectPointer for this TScriptInterface
+	 */
+	FORCEINLINE void SetObject( UObjectType* InObjectPointer )
+	{
+		FScriptInterface::SetObject(const_cast<UObject*>(InObjectPointer));
+	}
+
 	/**
 	 * Boolean operator, returns true if this object natively implements InterfaceType.
 	 * This will return false for objects that only implement the interface in blueprint classes.
diff --git a/Engine/Source/Runtime/CoreUObject/Public/UObject/WeakInterfacePtr.h b/Engine/Source/Runtime/CoreUObject/Public/UObject/WeakInterfacePtr.h
--- a/Engine/Source/Runtime/CoreUObject/Public/UObject/WeakInterfacePtr.h
+++ b/Engine/Source/Runtime/CoreUObject/Public/UObject/WeakInterfacePtr.h
@@ -17,7 +17,8 @@ template<class T>
 struct TWeakInterfacePtr
 {
 	using ElementType = T;
-	
+	using UObjectType = typename TCopyQualifiersFromTo<T, UObject>::Type;
+
 	FORCEINLINE TWeakInterfacePtr() = default;
 	FORCEINLINE TWeakInterfacePtr(const TWeakInterfacePtr& Other) = default;
 	FORCEINLINE TWeakInterfacePtr(TWeakInterfacePtr&& Other) = default;
@@ -99,7 +100,7 @@ struct TWeakInterfacePtr
 	/**
 	 * Dereference the weak pointer into a UObject pointer.
 	 */
-	FORCEINLINE UObject* GetObject() const
+	FORCEINLINE UObjectType* GetObject() const
 	{
 		return ObjectInstance.Get();
 	}
@@ -153,7 +154,7 @@ struct TWeakInterfacePtr
 
 	FORCEINLINE TScriptInterface<T> ToScriptInterface() const
 	{
-		UObject* Object = ObjectInstance.Get();
+		UObjectType* Object = ObjectInstance.Get();
 		if (Object)
 		{
 			return TScriptInterface<T>(Object);
@@ -162,12 +163,12 @@ struct TWeakInterfacePtr
 		return TScriptInterface<T>();
 	}
 
-	FORCEINLINE TWeakObjectPtr<UObject> GetWeakObjectPtr() const
+	FORCEINLINE TWeakObjectPtr<UObjectType> GetWeakObjectPtr() const
 	{
 		return ObjectInstance;
 	}
 
 private:
-	TWeakObjectPtr<UObject> ObjectInstance;
+	TWeakObjectPtr<UObjectType> ObjectInstance;
 	T* InterfaceInstance = nullptr;
 };
