autortfm: fixes for FScopeCycleCounterUObject and FScopeCycleCounter, moving the trace start/stop to be in the Open

diff --git a/Engine/Source/Runtime/Core/Private/AutoRTFM/API.cpp b/Engine/Source/Runtime/Core/Private/AutoRTFM/API.cpp
--- a/Engine/Source/Runtime/Core/Private/AutoRTFM/API.cpp
+++ b/Engine/Source/Runtime/Core/Private/AutoRTFM/API.cpp
@@ -522,8 +522,7 @@ extern "C" UE_AUTORTFM_NOAUTORTFM void RTFM_PopOnAbortHandler(const void* Key)
 {
 	FContext* Context = FContext::Get();
     ASSERT(Context->GetStatus() == EContextStatus::OnTrack);
-    const bool bResult = Context->GetCurrentTransaction()->PopDeferUntilAbortHandler(Key);
-	ASSERT(bResult);
+    Context->GetCurrentTransaction()->PopDeferUntilAbortHandler(Key);
 }
 
 extern "C" UE_AUTORTFM_NOAUTORTFM void RTFM_autortfm_on_commit(void (*Work)(void*), void* Arg)
diff --git a/Engine/Source/Runtime/Core/Public/AutoRTFM/AutoRTFM.h b/Engine/Source/Runtime/Core/Public/AutoRTFM/AutoRTFM.h
--- a/Engine/Source/Runtime/Core/Public/AutoRTFM/AutoRTFM.h
+++ b/Engine/Source/Runtime/Core/Public/AutoRTFM/AutoRTFM.h
@@ -42,6 +42,7 @@
 #define UE_AUTORTFM_FORCEINLINE inline
 #define UE_AUTORTFM_MEMCPY ::memcpy
 #else
+#include <Templates/Function.h>
 #include <HAL/Platform.h>
 #include <HAL/PlatformMemory.h>
 #define UE_AUTORTFM_API CORE_API
@@ -49,10 +50,6 @@
 #define UE_AUTORTFM_MEMCPY FPlatformMemory::Memcpy
 #endif
 
-#if UE_AUTORTFM
-template <typename FuncType> class TFunction;
-#endif
-
 #define UE_AUTORTFM_UNUSED(UNUSEDVAR) (void)UNUSEDVAR
 
 #ifdef __cplusplus
diff --git a/Engine/Source/Runtime/Core/Public/Stats/Stats.h b/Engine/Source/Runtime/Core/Public/Stats/Stats.h
--- a/Engine/Source/Runtime/Core/Public/Stats/Stats.h
+++ b/Engine/Source/Runtime/Core/Public/Stats/Stats.h
@@ -16,6 +16,7 @@
 #include "StatsCommon.h"
 #include "ProfilingDebugging/UMemoryDefines.h"
 #include "Stats/Stats2.h"
+#include "AutoRTFM/AutoRTFM.h"
 
 /** 
  *	Learn about the Stats System at docs.unrealengine.com
@@ -66,19 +67,30 @@ public:
 	 */
 	FORCEINLINE_STATS FScopeCycleCounter( TStatId StatId, EStatFlags StatFlags, bool bAlways = false)
 	{
-		Start( StatId, StatFlags, bAlways);
+		AutoRTFM::Open([&]
+		{
+			Start( StatId, StatFlags, bAlways);
+		});
+
+		AutoRTFM::PushOnAbortHandler(this, [this](){ this->Stop(); });
 	}
 
 	FORCEINLINE_STATS FScopeCycleCounter(TStatId StatId, bool bAlways = false)
 		: FScopeCycleCounter(StatId, EStatFlags::None, bAlways)
-	{}
+	{
+		AutoRTFM::PushOnAbortHandler(this, [this](){ this->Stop(); });
+	}
 
 	/**
 	 * Updates the stat with the time spent
 	 */
 	FORCEINLINE_STATS ~FScopeCycleCounter()
 	{
-		Stop();
+		AutoRTFM::PopOnAbortHandler(this);
+		AutoRTFM::Open([&]
+		{
+			Stop();
+		});
 	}
 
 };
@@ -186,8 +198,12 @@ public:
 	{
 		if (GCycleStatsShouldEmitNamedEvents && InStatId.IsValidStat())
 		{
-			bPop = true;
-			FPlatformMisc::BeginNamedEvent(FColor(0), InStatId.StatString);
+			AutoRTFM::Open([&]
+			{
+				bPop = true;
+				FPlatformMisc::BeginNamedEvent(FColor(0), InStatId.StatString);
+			});
+			AutoRTFM::PushOnAbortHandler(this, [](){ FPlatformMisc::EndNamedEvent(); });
 		}
 	}
 
@@ -198,10 +214,15 @@ public:
 
 	FORCEINLINE ~FScopeCycleCounter()
 	{
-		if (bPop)
+		AutoRTFM::PopOnAbortHandler(this);
+
+		AutoRTFM::Open([&]
 		{
-			FPlatformMisc::EndNamedEvent();
-		}
+			if (bPop)
+			{
+				FPlatformMisc::EndNamedEvent();
+			}
+		});
 	}
 private:
 #if USE_LIGHTWEIGHT_STATS_FOR_HITCH_DETECTION && USE_HITCH_DETECTION
diff --git a/Engine/Source/Runtime/CoreUObject/Public/UObject/UObjectBaseUtility.h b/Engine/Source/Runtime/CoreUObject/Public/UObject/UObjectBaseUtility.h
--- a/Engine/Source/Runtime/CoreUObject/Public/UObject/UObjectBaseUtility.h
+++ b/Engine/Source/Runtime/CoreUObject/Public/UObject/UObjectBaseUtility.h
@@ -26,6 +26,7 @@
 #include "UObject/UObjectBase.h"
 #include "UObject/UObjectMarks.h"
 #include "UObject/ObjectFwd.h"
+#include "AutoRTFM/AutoRTFM.h"
 
 class UClass;
 class UObject;
@@ -771,29 +772,37 @@ public:
 	  */
 	FORCEINLINE TStatId GetStatID(bool bForDeferredUse = false) const
 	{
+		TStatId Result{};
+		AutoRTFM::Open([&]
+		{
 #if STATS
-		const TStatId& StatID = GUObjectArray.IndexToObject(InternalIndex)->StatID;
+			const TStatId& StatID = GUObjectArray.IndexToObject(InternalIndex)->StatID;
 
-		// this is done to avoid even registering stats for a disabled group (unless we plan on using it later)
-		if (bForDeferredUse || FThreadStats::IsCollectingData(GET_STATID(STAT_UObjectsStatGroupTester)))
-		{
-			if (!StatID.IsValidStat())
+			// this is done to avoid even registering stats for a disabled group (unless we plan on using it later)
+			if (bForDeferredUse || FThreadStats::IsCollectingData(GET_STATID(STAT_UObjectsStatGroupTester)))
 			{
-				CreateStatID();
+				if (!StatID.IsValidStat())
+				{
+					CreateStatID();
+				}
+				Result = StatID;
+			}
+			else
+			{
+				Result = TStatId(); // not doing stats at the moment, or ever
 			}
-			return StatID;
-		}
-		return TStatId(); // not doing stats at the moment, or ever
 #elif ENABLE_STATNAMEDEVENTS_UOBJECT
-		const TStatId& StatID = GUObjectArray.IndexToObject(InternalIndex)->StatID;
-		if (!StatID.IsValidStat() && (bForDeferredUse || GCycleStatsShouldEmitNamedEvents))
-		{
-			CreateStatID();
-		}
-		return StatID;
+			const TStatId& StatID = GUObjectArray.IndexToObject(InternalIndex)->StatID;
+			if (!StatID.IsValidStat() && (bForDeferredUse || GCycleStatsShouldEmitNamedEvents))
+			{
+				CreateStatID();
+			}
+			Result = StatID;
 #else
-		return TStatId(); // not doing stats at the moment, or ever
+			Result = TStatId(); // not doing stats at the moment, or ever
 #endif // STATS
+		});
+		return Result;
 	}
 
 private:
@@ -930,49 +939,59 @@ public:
 	 */
 	FORCEINLINE_STATS FScopeCycleCounterUObject(const UObjectBaseUtility *Object)
 	{
-		if (Object)
+		AutoRTFM::Open([&]
 		{
-			bool bStarted = false;
-			TStatId ObjectStatId = Object->GetStatID();
-			if (FThreadStats::IsCollectingData(ObjectStatId))
+			if (Object)
 			{
-				Start(ObjectStatId);
-				bStarted = true;
-			}
+				bool bStarted = false;
+				TStatId ObjectStatId = Object->GetStatID();
+				if (FThreadStats::IsCollectingData(ObjectStatId))
+				{
+					Start(ObjectStatId);
+					bStarted = true;
+				}
 
 #if CPUPROFILERTRACE_ENABLED
-			if (!bStarted && UE_TRACE_CHANNELEXPR_IS_ENABLED(CpuChannel))
-			{
-				StartObjectTrace(Object);
-			}
+				if (!bStarted && UE_TRACE_CHANNELEXPR_IS_ENABLED(CpuChannel))
+				{
+					StartObjectTrace(Object);
+				}
 #endif
-		}
+			}
+		});
+
+		AutoRTFM::PushOnAbortHandler(this, [this](){ this->Stop(); });
 	}
 	/**
 	 * Constructor, starts timing with an alternate enable stat to use high performance disable for only SOME UObject stats
 	 */
 	FORCEINLINE_STATS FScopeCycleCounterUObject(const UObjectBaseUtility *Object, TStatId OtherStat)
 	{
-		if (Object)
+		AutoRTFM::Open([&]
 		{
-			bool bStarted = false;
-			if (FThreadStats::IsCollectingData(OtherStat))
+			if (Object)
 			{
-				TStatId ObjectStatId = Object->GetStatID();
-				if (!ObjectStatId.IsNone())
+				bool bStarted = false;
+				if (FThreadStats::IsCollectingData(OtherStat))
 				{
-					Start(ObjectStatId);
-					bStarted = true;
+					TStatId ObjectStatId = Object->GetStatID();
+					if (!ObjectStatId.IsNone())
+					{
+						Start(ObjectStatId);
+						bStarted = true;
+					}
 				}
-			}
 
 #if CPUPROFILERTRACE_ENABLED
-			if (!bStarted && UE_TRACE_CHANNELEXPR_IS_ENABLED(CpuChannel))
-			{
-				StartObjectTrace(Object);
-			}
+				if (!bStarted && UE_TRACE_CHANNELEXPR_IS_ENABLED(CpuChannel))
+				{
+					StartObjectTrace(Object);
+				}
 #endif
-		}
+			}
+		});
+		
+		AutoRTFM::PushOnAbortHandler(this, [this](){ this->Stop(); });
 	}
 
 #if CPUPROFILERTRACE_ENABLED
@@ -984,7 +1003,11 @@ public:
 	 */
 	FORCEINLINE_STATS ~FScopeCycleCounterUObject()
 	{
-		Stop();
+		AutoRTFM::PopOnAbortHandler(this);
+		AutoRTFM::Open([&]
+		{
+			Stop();
+		});
 	}
 };
 
@@ -1009,14 +1032,22 @@ public:
 #endif
 	{
 #if ENABLE_STATNAMEDEVENTS_UOBJECT && CPUPROFILERTRACE_ENABLED
-		if (GCycleStatsShouldEmitNamedEvents && UE_TRACE_CHANNELEXPR_IS_ENABLED(CpuChannel) && Object)
+		AutoRTFM::Open([&]
 		{
-			const TStatId ObjectStatId = Object->GetStatID();
-			if (ObjectStatId.IsValidStat())
+			if (GCycleStatsShouldEmitNamedEvents && UE_TRACE_CHANNELEXPR_IS_ENABLED(CpuChannel) && Object)
 			{
-				bPop = true;
-				FCpuProfilerTrace::OutputBeginDynamicEvent(ObjectStatId.StatString);
+				const TStatId ObjectStatId = Object->GetStatID();
+				if (ObjectStatId.IsValidStat())
+				{
+					bPop = true;
+					FCpuProfilerTrace::OutputBeginDynamicEvent(ObjectStatId.StatString);
+				}
 			}
+		});
+
+		if(bPop)
+		{
+			AutoRTFM::PushOnAbortHandler(this, [](){ FCpuProfilerTrace::OutputEndEvent(); });
 		}
 #endif
 	}
@@ -1029,10 +1060,15 @@ public:
 	FORCEINLINE_STATS ~FScopeCycleCounterUObject()
 	{
 #if ENABLE_STATNAMEDEVENTS_UOBJECT && CPUPROFILERTRACE_ENABLED
-		if (bPop)
+		AutoRTFM::PopOnAbortHandler(this);
+
+		AutoRTFM::Open([&]
 		{
-			FCpuProfilerTrace::OutputEndEvent();
-		}
+			if (bPop)
+			{
+				FCpuProfilerTrace::OutputEndEvent();
+			}
+		});
 #endif
 	}
 };
@@ -1056,7 +1092,10 @@ public:
 	{
 		if (GHitchDetected &&  StatObject)
 		{
-			ReportHitch();
+			AutoRTFM::Open([&]
+			{
+				ReportHitch();
+			});
 		}
 	}
 
