[Backout] - CL33377480
Original CL Desc
-----------------------------------------------------------------
autortfm: fixes for FScopeCycleCounterUObject and FScopeCycleCounter, moving the trace start/stop to be in the Open

diff --git a/Engine/Source/Runtime/Core/Private/AutoRTFM/API.cpp b/Engine/Source/Runtime/Core/Private/AutoRTFM/API.cpp
--- a/Engine/Source/Runtime/Core/Private/AutoRTFM/API.cpp
+++ b/Engine/Source/Runtime/Core/Private/AutoRTFM/API.cpp
@@ -522,7 +522,8 @@ extern "C" UE_AUTORTFM_NOAUTORTFM void RTFM_PopOnAbortHandler(const void* Key)
 {
 	FContext* Context = FContext::Get();
     ASSERT(Context->GetStatus() == EContextStatus::OnTrack);
-    Context->GetCurrentTransaction()->PopDeferUntilAbortHandler(Key);
+    const bool bResult = Context->GetCurrentTransaction()->PopDeferUntilAbortHandler(Key);
+	ASSERT(bResult);
 }
 
 extern "C" UE_AUTORTFM_NOAUTORTFM void RTFM_autortfm_on_commit(void (*Work)(void*), void* Arg)
diff --git a/Engine/Source/Runtime/Core/Public/AutoRTFM/AutoRTFM.h b/Engine/Source/Runtime/Core/Public/AutoRTFM/AutoRTFM.h
--- a/Engine/Source/Runtime/Core/Public/AutoRTFM/AutoRTFM.h
+++ b/Engine/Source/Runtime/Core/Public/AutoRTFM/AutoRTFM.h
@@ -42,7 +42,6 @@
 #define UE_AUTORTFM_FORCEINLINE inline
 #define UE_AUTORTFM_MEMCPY ::memcpy
 #else
-#include <Templates/Function.h>
 #include <HAL/Platform.h>
 #include <HAL/PlatformMemory.h>
 #define UE_AUTORTFM_API CORE_API
@@ -50,6 +49,10 @@
 #define UE_AUTORTFM_MEMCPY FPlatformMemory::Memcpy
 #endif
 
+#if UE_AUTORTFM
+template <typename FuncType> class TFunction;
+#endif
+
 #define UE_AUTORTFM_UNUSED(UNUSEDVAR) (void)UNUSEDVAR
 
 #ifdef __cplusplus
diff --git a/Engine/Source/Runtime/Core/Public/Stats/Stats.h b/Engine/Source/Runtime/Core/Public/Stats/Stats.h
--- a/Engine/Source/Runtime/Core/Public/Stats/Stats.h
+++ b/Engine/Source/Runtime/Core/Public/Stats/Stats.h
@@ -16,7 +16,6 @@
 #include "StatsCommon.h"
 #include "ProfilingDebugging/UMemoryDefines.h"
 #include "Stats/Stats2.h"
-#include "AutoRTFM/AutoRTFM.h"
 
 /** 
  *	Learn about the Stats System at docs.unrealengine.com
@@ -67,30 +66,19 @@ public:
 	 */
 	FORCEINLINE_STATS FScopeCycleCounter( TStatId StatId, EStatFlags StatFlags, bool bAlways = false)
 	{
-		AutoRTFM::Open([&]
-		{
-			Start( StatId, StatFlags, bAlways);
-		});
-
-		AutoRTFM::PushOnAbortHandler(this, [this](){ this->Stop(); });
+		Start( StatId, StatFlags, bAlways);
 	}
 
 	FORCEINLINE_STATS FScopeCycleCounter(TStatId StatId, bool bAlways = false)
 		: FScopeCycleCounter(StatId, EStatFlags::None, bAlways)
-	{
-		AutoRTFM::PushOnAbortHandler(this, [this](){ this->Stop(); });
-	}
+	{}
 
 	/**
 	 * Updates the stat with the time spent
 	 */
 	FORCEINLINE_STATS ~FScopeCycleCounter()
 	{
-		AutoRTFM::PopOnAbortHandler(this);
-		AutoRTFM::Open([&]
-		{
-			Stop();
-		});
+		Stop();
 	}
 
 };
@@ -198,12 +186,8 @@ public:
 	{
 		if (GCycleStatsShouldEmitNamedEvents && InStatId.IsValidStat())
 		{
-			AutoRTFM::Open([&]
-			{
-				bPop = true;
-				FPlatformMisc::BeginNamedEvent(FColor(0), InStatId.StatString);
-			});
-			AutoRTFM::PushOnAbortHandler(this, [](){ FPlatformMisc::EndNamedEvent(); });
+			bPop = true;
+			FPlatformMisc::BeginNamedEvent(FColor(0), InStatId.StatString);
 		}
 	}
 
@@ -214,15 +198,10 @@ public:
 
 	FORCEINLINE ~FScopeCycleCounter()
 	{
-		AutoRTFM::PopOnAbortHandler(this);
-
-		AutoRTFM::Open([&]
+		if (bPop)
 		{
-			if (bPop)
-			{
-				FPlatformMisc::EndNamedEvent();
-			}
-		});
+			FPlatformMisc::EndNamedEvent();
+		}
 	}
 private:
 #if USE_LIGHTWEIGHT_STATS_FOR_HITCH_DETECTION && USE_HITCH_DETECTION
diff --git a/Engine/Source/Runtime/CoreUObject/Public/UObject/UObjectBaseUtility.h b/Engine/Source/Runtime/CoreUObject/Public/UObject/UObjectBaseUtility.h
--- a/Engine/Source/Runtime/CoreUObject/Public/UObject/UObjectBaseUtility.h
+++ b/Engine/Source/Runtime/CoreUObject/Public/UObject/UObjectBaseUtility.h
@@ -26,7 +26,6 @@
 #include "UObject/UObjectBase.h"
 #include "UObject/UObjectMarks.h"
 #include "UObject/ObjectFwd.h"
-#include "AutoRTFM/AutoRTFM.h"
 
 class UClass;
 class UObject;
@@ -772,37 +771,29 @@ public:
 	  */
 	FORCEINLINE TStatId GetStatID(bool bForDeferredUse = false) const
 	{
-		TStatId Result{};
-		AutoRTFM::Open([&]
-		{
 #if STATS
-			const TStatId& StatID = GUObjectArray.IndexToObject(InternalIndex)->StatID;
+		const TStatId& StatID = GUObjectArray.IndexToObject(InternalIndex)->StatID;
 
-			// this is done to avoid even registering stats for a disabled group (unless we plan on using it later)
-			if (bForDeferredUse || FThreadStats::IsCollectingData(GET_STATID(STAT_UObjectsStatGroupTester)))
-			{
-				if (!StatID.IsValidStat())
-				{
-					CreateStatID();
-				}
-				Result = StatID;
-			}
-			else
-			{
-				Result = TStatId(); // not doing stats at the moment, or ever
-			}
-#elif ENABLE_STATNAMEDEVENTS_UOBJECT
-			const TStatId& StatID = GUObjectArray.IndexToObject(InternalIndex)->StatID;
-			if (!StatID.IsValidStat() && (bForDeferredUse || GCycleStatsShouldEmitNamedEvents))
+		// this is done to avoid even registering stats for a disabled group (unless we plan on using it later)
+		if (bForDeferredUse || FThreadStats::IsCollectingData(GET_STATID(STAT_UObjectsStatGroupTester)))
+		{
+			if (!StatID.IsValidStat())
 			{
 				CreateStatID();
 			}
-			Result = StatID;
+			return StatID;
+		}
+		return TStatId(); // not doing stats at the moment, or ever
+#elif ENABLE_STATNAMEDEVENTS_UOBJECT
+		const TStatId& StatID = GUObjectArray.IndexToObject(InternalIndex)->StatID;
+		if (!StatID.IsValidStat() && (bForDeferredUse || GCycleStatsShouldEmitNamedEvents))
+		{
+			CreateStatID();
+		}
+		return StatID;
 #else
-			Result = TStatId(); // not doing stats at the moment, or ever
+		return TStatId(); // not doing stats at the moment, or ever
 #endif // STATS
-		});
-		return Result;
 	}
 
 private:
@@ -939,59 +930,49 @@ public:
 	 */
 	FORCEINLINE_STATS FScopeCycleCounterUObject(const UObjectBaseUtility *Object)
 	{
-		AutoRTFM::Open([&]
+		if (Object)
 		{
-			if (Object)
+			bool bStarted = false;
+			TStatId ObjectStatId = Object->GetStatID();
+			if (FThreadStats::IsCollectingData(ObjectStatId))
 			{
-				bool bStarted = false;
-				TStatId ObjectStatId = Object->GetStatID();
-				if (FThreadStats::IsCollectingData(ObjectStatId))
-				{
-					Start(ObjectStatId);
-					bStarted = true;
-				}
+				Start(ObjectStatId);
+				bStarted = true;
+			}
 
 #if CPUPROFILERTRACE_ENABLED
-				if (!bStarted && UE_TRACE_CHANNELEXPR_IS_ENABLED(CpuChannel))
-				{
-					StartObjectTrace(Object);
-				}
-#endif
+			if (!bStarted && UE_TRACE_CHANNELEXPR_IS_ENABLED(CpuChannel))
+			{
+				StartObjectTrace(Object);
 			}
-		});
-
-		AutoRTFM::PushOnAbortHandler(this, [this](){ this->Stop(); });
+#endif
+		}
 	}
 	/**
 	 * Constructor, starts timing with an alternate enable stat to use high performance disable for only SOME UObject stats
 	 */
 	FORCEINLINE_STATS FScopeCycleCounterUObject(const UObjectBaseUtility *Object, TStatId OtherStat)
 	{
-		AutoRTFM::Open([&]
+		if (Object)
 		{
-			if (Object)
+			bool bStarted = false;
+			if (FThreadStats::IsCollectingData(OtherStat))
 			{
-				bool bStarted = false;
-				if (FThreadStats::IsCollectingData(OtherStat))
+				TStatId ObjectStatId = Object->GetStatID();
+				if (!ObjectStatId.IsNone())
 				{
-					TStatId ObjectStatId = Object->GetStatID();
-					if (!ObjectStatId.IsNone())
-					{
-						Start(ObjectStatId);
-						bStarted = true;
-					}
+					Start(ObjectStatId);
+					bStarted = true;
 				}
+			}
 
 #if CPUPROFILERTRACE_ENABLED
-				if (!bStarted && UE_TRACE_CHANNELEXPR_IS_ENABLED(CpuChannel))
-				{
-					StartObjectTrace(Object);
-				}
-#endif
+			if (!bStarted && UE_TRACE_CHANNELEXPR_IS_ENABLED(CpuChannel))
+			{
+				StartObjectTrace(Object);
 			}
-		});
-		
-		AutoRTFM::PushOnAbortHandler(this, [this](){ this->Stop(); });
+#endif
+		}
 	}
 
 #if CPUPROFILERTRACE_ENABLED
@@ -1003,11 +984,7 @@ public:
 	 */
 	FORCEINLINE_STATS ~FScopeCycleCounterUObject()
 	{
-		AutoRTFM::PopOnAbortHandler(this);
-		AutoRTFM::Open([&]
-		{
-			Stop();
-		});
+		Stop();
 	}
 };
 
@@ -1032,22 +1009,14 @@ public:
 #endif
 	{
 #if ENABLE_STATNAMEDEVENTS_UOBJECT && CPUPROFILERTRACE_ENABLED
-		AutoRTFM::Open([&]
+		if (GCycleStatsShouldEmitNamedEvents && UE_TRACE_CHANNELEXPR_IS_ENABLED(CpuChannel) && Object)
 		{
-			if (GCycleStatsShouldEmitNamedEvents && UE_TRACE_CHANNELEXPR_IS_ENABLED(CpuChannel) && Object)
+			const TStatId ObjectStatId = Object->GetStatID();
+			if (ObjectStatId.IsValidStat())
 			{
-				const TStatId ObjectStatId = Object->GetStatID();
-				if (ObjectStatId.IsValidStat())
-				{
-					bPop = true;
-					FCpuProfilerTrace::OutputBeginDynamicEvent(ObjectStatId.StatString);
-				}
+				bPop = true;
+				FCpuProfilerTrace::OutputBeginDynamicEvent(ObjectStatId.StatString);
 			}
-		});
-
-		if(bPop)
-		{
-			AutoRTFM::PushOnAbortHandler(this, [](){ FCpuProfilerTrace::OutputEndEvent(); });
 		}
 #endif
 	}
@@ -1060,15 +1029,10 @@ public:
 	FORCEINLINE_STATS ~FScopeCycleCounterUObject()
 	{
 #if ENABLE_STATNAMEDEVENTS_UOBJECT && CPUPROFILERTRACE_ENABLED
-		AutoRTFM::PopOnAbortHandler(this);
-
-		AutoRTFM::Open([&]
+		if (bPop)
 		{
-			if (bPop)
-			{
-				FCpuProfilerTrace::OutputEndEvent();
-			}
-		});
+			FCpuProfilerTrace::OutputEndEvent();
+		}
 #endif
 	}
 };
@@ -1092,10 +1056,7 @@ public:
 	{
 		if (GHitchDetected &&  StatObject)
 		{
-			AutoRTFM::Open([&]
-			{
-				ReportHitch();
-			});
+			ReportHitch();
 		}
 	}
 
