Revise TObjectPtr to use pointer semantics for type safety checks (editor only).

Notes:
- This replaces the previous solution that relied on indirection.
- This solution relies on packed ref handles which requires late resolve features to also be enabled.

Change summary:
- Modified FPropertyBagRepository::IsPropertyBagPlaceholderType()/IsPropertyBagPlaceholderObject() to make the parameter read-only.
- Re-enabled UE_WITH_OBJECT_HANDLE_TYPE_SAFETY for the editor context. This will unblock testing of IDO w/ placeholder type features.
- Modified UE::CoreUObject::Private::Pack() to include a new "type" trait bit in the packed object handle reference. If set, it will signal that the reference is also a placeholder-typed object.
- Modified UE::CoreUObject::Private::MakePackedObjectRef() to include a check for whether or not the given object reference is also a placeholder-typed object, and allow this case to pass.
- Added UE::CoreUObject::Private::IsObjectHandleTypeSafeNoReadNoCheck(). This allows the implementation of IsObjectHandleTypeSafe() to rely on packed reference semantics, which are implemented as internal inline definitions.
- Modified FLinkerLoad serialization of FObjectPtr values to register a packed object reference mapping for placeholder-typed objects. Note: This requires late resolve features to also be enabled.
- Modified the TObjectPtr equality test to handle the case when only the RHS side is NULL. This can occur if the RHS is a non-UObject-type bound ptr to a placeholder object while the LHS is not.
- Revised all object handle/ptr unit tests for type safety. Ensured that TObjectPtr stress tests are still in line both with/without type safety enabled.
- Added a GObjectIndexToPackedObjectRefDebug global ptr to assist with debugger visualization of resolved placeholder-typed object references to allow for removal of RF_HasPlaceholderType.
- Updated VS natvis file to revise queries for IsPlaceholderObject() on non-UObject-typed TObjectPtr values (removed ObjectFlags check and replaced w/ GObjectIndexToPackedObjectRefDebug).
- Deprecated/removed all usage of RF_HasPlaceholderType (no longer required).
- Removed FPropertyBagRepository::PropertyBagPlaceholderTypeRegistry and decoupled the current (static) API from the repo singleton instance.
- (Minor) Fixed up some handle tracking unit tests in ObjectPropertyTest.cpp that did not compile with UE_WITH_OBJECT_HANDLE_LATE_RESOLVE disabled.

diff --git a/Engine/Source/Runtime/Core/Private/Misc/CoreGlobals.cpp b/Engine/Source/Runtime/Core/Private/Misc/CoreGlobals.cpp
--- a/Engine/Source/Runtime/Core/Private/Misc/CoreGlobals.cpp
+++ b/Engine/Source/Runtime/Core/Private/Misc/CoreGlobals.cpp
@@ -253,6 +253,7 @@ float				GNearClippingPlane_RenderThread		= 10.0f;				/* Near clipping plane (Re
 
 bool					GExitPurge						= false;
 
+const UPTRINT* GCoreObjectIndexToPackedObjectRefDebug = nullptr;
 FChunkedFixedUObjectArray* GCoreObjectArrayForDebugVisualizers = nullptr;
 
 namespace UE::CoreUObject::Private
@@ -264,6 +265,7 @@ UE::CoreUObject::Private::FStoredObjectPathDebug* GCoreComplexObjectPathDebug =
 UE::CoreUObject::Private::FObjectHandlePackageDebugData* GCoreObjectHandlePackageDebug = nullptr;
 #if PLATFORM_UNIX
 uint8** CORE_API GNameBlocksDebug = FNameDebugVisualizer::GetBlocks();
+const UPTRINT*& CORE_API GObjectIndexToPackedObjectRefDebug = GCoreObjectIndexToPackedObjectRefDebug;
 FChunkedFixedUObjectArray*& CORE_API GObjectArrayForDebugVisualizers = GCoreObjectArrayForDebugVisualizers;
 UE::CoreUObject::Private::FStoredObjectPathDebug*& GComplexObjectPathDebug = GCoreComplexObjectPathDebug;
 UE::CoreUObject::Private::FObjectHandlePackageDebugData*& CORE_API GObjectHandlePackageDebug = GCoreObjectHandlePackageDebug;
diff --git a/Engine/Source/Runtime/Core/Public/CoreGlobals.h b/Engine/Source/Runtime/Core/Public/CoreGlobals.h
--- a/Engine/Source/Runtime/Core/Public/CoreGlobals.h
+++ b/Engine/Source/Runtime/Core/Public/CoreGlobals.h
@@ -808,6 +808,9 @@ struct FScopedLoadingState
 bool CORE_API GetEmitDrawEvents();
 void CORE_API SetEmitDrawEvents(bool EmitDrawEvents);
 
+/** Array to help visualize packed object references */
+extern CORE_API const UPTRINT* GCoreObjectIndexToPackedObjectRefDebug;
+
 /** Array to help visualize weak pointers in the debugger */
 class FChunkedFixedUObjectArray;
 
diff --git a/Engine/Source/Runtime/Core/Public/Modules/Boilerplate/ModuleBoilerplate.h b/Engine/Source/Runtime/Core/Public/Modules/Boilerplate/ModuleBoilerplate.h
--- a/Engine/Source/Runtime/Core/Public/Modules/Boilerplate/ModuleBoilerplate.h
+++ b/Engine/Source/Runtime/Core/Public/Modules/Boilerplate/ModuleBoilerplate.h
@@ -83,6 +83,7 @@ class FChunkedFixedUObjectArray;
 #else
 	#define UE4_VISUALIZERS_HELPERS \
 		uint8** GNameBlocksDebug = FNameDebugVisualizer::GetBlocks(); \
+		const UPTRINT*& GObjectIndexToPackedObjectRefDebug = GCoreObjectIndexToPackedObjectRefDebug; \
 		FChunkedFixedUObjectArray*& GObjectArrayForDebugVisualizers = GCoreObjectArrayForDebugVisualizers; \
 		UE::CoreUObject::Private::FStoredObjectPathDebug*& GComplexObjectPathDebug = GCoreComplexObjectPathDebug; \
 		UE::CoreUObject::Private::FObjectHandlePackageDebugData*& GObjectHandlePackageDebug = GCoreObjectHandlePackageDebug;
diff --git a/Engine/Source/Runtime/CoreUObject/Internal/UObject/PropertyBagRepository.h b/Engine/Source/Runtime/CoreUObject/Internal/UObject/PropertyBagRepository.h
--- a/Engine/Source/Runtime/CoreUObject/Internal/UObject/PropertyBagRepository.h
+++ b/Engine/Source/Runtime/CoreUObject/Internal/UObject/PropertyBagRepository.h
@@ -45,10 +45,7 @@ private:
 	// used to make sure IDOs don't have name overlap
 	TMap<const UObject*, TObjectPtr<UObject>> Namespaces;
 
-	/** Internal registry that tracks the current set of types for property bag container objects instanced as placeholders for package exports that have invalid or missing class imports on load. */
-	TUniquePtr<class FPropertyBagPlaceholderTypeRegistry> PropertyBagPlaceholderTypeRegistry;
-
-	FPropertyBagRepository();
+	FPropertyBagRepository() = default;
 
 public:
 	FPropertyBagRepository(const FPropertyBagRepository &) = delete;
@@ -115,9 +112,9 @@ public:
 	// End FGCObject interface
 
 	// query for whether or not the given struct/class is a placeholder type
-	static COREUOBJECT_API bool IsPropertyBagPlaceholderType(UStruct* Type);
+	static COREUOBJECT_API bool IsPropertyBagPlaceholderType(const UStruct* Type);
 	// query for whether or not the given object was created as a placeholder type
-	static COREUOBJECT_API bool IsPropertyBagPlaceholderObject(UObject* Object);
+	static COREUOBJECT_API bool IsPropertyBagPlaceholderObject(const UObject* Object);
 	// query for whether or not creating property bag placeholder objects should be allowed
 	static COREUOBJECT_API bool IsPropertyBagPlaceholderObjectSupportEnabled();
 	// query whether an object supports IDO generation
diff --git a/Engine/Source/Runtime/CoreUObject/Private/UObject/LinkerLoad.cpp b/Engine/Source/Runtime/CoreUObject/Private/UObject/LinkerLoad.cpp
--- a/Engine/Source/Runtime/CoreUObject/Private/UObject/LinkerLoad.cpp
+++ b/Engine/Source/Runtime/CoreUObject/Private/UObject/LinkerLoad.cpp
@@ -661,7 +661,7 @@ void FLinkerLoad::PRIVATE_PatchNewObjectIntoExport(UObject* OldObject, UObject*
 		FObjectExport& ObjExport = OldObjectLinker->ExportMap[CachedLinkerIndex];
 		
 		// Since we don't copy the internal flags, the mirrored flags need can't be set on the new object as well
-		const EObjectFlags OldObjectFlags = OldObject->GetFlags() & ~(RF_MirroredGarbage | RF_HasPlaceholderType);
+		const EObjectFlags OldObjectFlags = OldObject->GetFlags() & ~(RF_MirroredGarbage);
 
 		// Detach the old object to make room for the new
 		OldObject->ClearFlags(RF_NeedLoad|RF_NeedPostLoad|RF_NeedPostLoadSubobjects);
@@ -5105,11 +5105,6 @@ UObject* FLinkerLoad::CreateExport( int32 Index )
 		{
 			return nullptr;
 		}
-		else if (LoadClass && UE::FPropertyBagRepository::IsPropertyBagPlaceholderType(LoadClass))
-		{
-			// Modify the export's object flags for instancing to indicate that it has a placeholder type.
-			Export.ObjectFlags |= RF_HasPlaceholderType;
-		}
 #endif
 		if( !LoadClass )
 		{
@@ -6237,17 +6232,27 @@ FArchive& FLinkerLoad::operator<<(FObjectPtr& ObjectPtr)
 	// Wrapper that only allows pointers to exports with placeholder types when type safety features are enabled.
 	auto AsTypeSafeObjectPtr_Lambda = [this](UObject* ResolvedObject)
 	{
-#if WITH_EDITOR && !UE_WITH_OBJECT_HANDLE_TYPE_SAFETY
-		// If type safety features are disabled, resolve unsafe references to placeholder-typed objects now to NULL.
-		// Note: Similar to hard references above, this means we won't find it for replacement at reinstancing time.
+#if WITH_EDITOR
+		// Note: References to placeholder objects cannot resolve to it if the underlying pointer type is unsafe.
 		if (ResolvedObject && UE::FPropertyBagRepository::IsPropertyBagPlaceholderObject(ResolvedObject))
 		{
+#if UE_WITH_OBJECT_HANDLE_LATE_RESOLVE && UE_WITH_OBJECT_HANDLE_TYPE_SAFETY
+			// If type safety features are enabled, create a packed reference mapping for the placeholder-typed object.
+			// This resolves to the placeholder object for UObject-typed pointers. For other types it resolves to NULL.
+			// However, the underlying value when serialized will always resolve to the object reference (e.g. for GC).
+			// Note: We could return an FObjectPtr that wraps the packed reference result, but that implies a lazy load.
+			// We need to register the packed object reference here, but we don't need to also defer the pointer resolve.
+			UE::CoreUObject::Private::MakePackedObjectRef(ResolvedObject);
+#else
+			// If type safety features are disabled, resolve unsafe references to placeholder-typed objects now to NULL.
+			// Note: Similar to hard references above, this means we won't find it for replacement at reinstancing time.
 			const FObjectProperty* ObjectProperty = CastField<FObjectProperty>(GetSerializedProperty());
 			if (!ObjectProperty || !ResolvedObject->GetClass()->IsChildOf(ObjectProperty->PropertyClass))
 			{
 				UE_LOG(LogLinker, Warning, TEXT("Serializing reference to \"%s\" as NULL to ensure type safety."), *ResolvedObject->GetPathName());
 				ResolvedObject = nullptr;
 			}
+#endif
 		}
 #endif
 		return FObjectPtr(ResolvedObject);
diff --git a/Engine/Source/Runtime/CoreUObject/Private/UObject/ObjectHandleIndex.cpp b/Engine/Source/Runtime/CoreUObject/Private/UObject/ObjectHandleIndex.cpp
--- a/Engine/Source/Runtime/CoreUObject/Private/UObject/ObjectHandleIndex.cpp
+++ b/Engine/Source/Runtime/CoreUObject/Private/UObject/ObjectHandleIndex.cpp
@@ -15,11 +15,10 @@
 #include "UObject/Package.h"
 #include "UObject/UObjectHash.h"
 #include "UObject/ObjectPathId.h"
+#include "UObject/PropertyBagRepository.h"
 
 #if UE_WITH_OBJECT_HANDLE_LATE_RESOLVE
 
-
-
 static inline FName GetNameOrNone(UObject* Object)
 {
 	return Object ? Object->GetFName() : NAME_None;
@@ -131,6 +130,9 @@ namespace UE::CoreUObject::Private
 	void InitObjectHandles(int32 MaxObjects)
 	{
 		GObjectHandleIndex.ObjectIndexToPackedObjectRef.SetNumZeroed(MaxObjects);
+
+		// This allows the debug visualizer to evaluate a resolved handle to see if it has a packed ref stored in the global table.
+		GCoreObjectIndexToPackedObjectRefDebug = reinterpret_cast<const UPTRINT*>(GObjectHandleIndex.ObjectIndexToPackedObjectRef.GetData());
 	}
 
 	static FAutoConsoleCommand CmdPrintUnresolvedObjects(
@@ -178,11 +180,12 @@ namespace UE::CoreUObject::Private
 			})
 	);
 
-	static inline FPackedObjectRef Pack(FPackageId PackageId, FObjectId ObjectId)
+	static inline FPackedObjectRef Pack(FPackageId PackageId, FObjectId ObjectId, bool bHasPlaceholderType = false)
 	{
-		checkf(PackageId.ToIndex() <= 0x7FFFFFFF, TEXT("Package count exceeded the space permitted within packed object references.  This implies over 2 billion packages are in use."));
+		checkf(PackageId.ToIndex() <= PackageIdMask, TEXT("Package count exceeded the space permitted within packed object references.  This implies over %d packages are in use."), PackageIdMask);
 		return { static_cast<UPTRINT>(PackageId.ToIndex()) << PackageIdShift |
-				static_cast<UPTRINT>(ObjectId.ToPackedIndex() << ObjectIdShift) | 1 };
+				static_cast<UPTRINT>(ObjectId.ToPackedIndex() << ObjectIdShift) |
+				static_cast<UPTRINT>(bHasPlaceholderType << TypeIdShift) | 1 };
 	}
 
 	static inline void Unpack(FPackedObjectRef PackedObjectRef, FPackageId& OutPackageId, FObjectId& OutObjectId)
@@ -293,12 +296,12 @@ namespace UE::CoreUObject::Private
 		return;
 	}
 
-	static inline FPackedObjectRef MakePackedObjectRef(FName PackageName, FName ClassPackageName, FName ClassName, FObjectPathId ObjectPath)
+	static inline FPackedObjectRef MakePackedObjectRef(FName PackageName, FName ClassPackageName, FName ClassName, FObjectPathId ObjectPath, bool bHasPlaceholderType = false)
 	{
 		FPackageId PackageId;
 		FObjectId ObjectId;
 		MakeReferenceIds(PackageName, ClassPackageName, ClassName, ObjectPath, PackageId, ObjectId);
-		return Pack(PackageId, ObjectId);
+		return Pack(PackageId, ObjectId, bHasPlaceholderType);
 	}
 
 	static void GetObjectDataFromId(FPackageId PackageId, FObjectId ObjectId, FMinimalName& OutPackageName, FObjectPathId& OutPathId, FMinimalName& OutClassPackageName, FMinimalName& OutClassName)
@@ -526,7 +529,8 @@ namespace UE::CoreUObject::Private
 			return { 0 };
 		}
 
-		return MakePackedObjectRef(ObjectRef.PackageName, ObjectRef.ClassPackageName, ObjectRef.ClassName, ObjectRef.GetObjectPath());
+		const bool bHasPlaceholderType = UE::FPropertyBagRepository::IsPropertyBagPlaceholderType(ObjectRef.ResolveObjectRefClass());
+		return MakePackedObjectRef(ObjectRef.PackageName, ObjectRef.ClassPackageName, ObjectRef.ClassName, ObjectRef.GetObjectPath(), bHasPlaceholderType);
 	}
 
 
@@ -564,7 +568,8 @@ namespace UE::CoreUObject::Private
 			return { 0 };
 		}
 
-		if (UE::LinkerLoad::FindLoadBehavior(*Object->GetClass()) == UE::LinkerLoad::EImportBehavior::Eager)
+		const bool bHasPlaceholderType = UE::FPropertyBagRepository::IsPropertyBagPlaceholderObject(Object);
+		if (!bHasPlaceholderType && UE::LinkerLoad::FindLoadBehavior(*Object->GetClass()) == UE::LinkerLoad::EImportBehavior::Eager)
 		{
 			return { 0 };
 		}
@@ -580,7 +585,7 @@ namespace UE::CoreUObject::Private
 
 		UObject* Class = Object->GetClass();
 		FName ClassPackageName = GetNameOrNone(Class->GetOutermost());
-		PackedObjectRef = UE::CoreUObject::Private::MakePackedObjectRef(PackageName, ClassPackageName, GetNameOrNone(Class), FObjectPathId(Object));
+		PackedObjectRef = UE::CoreUObject::Private::MakePackedObjectRef(PackageName, ClassPackageName, GetNameOrNone(Class), FObjectPathId(Object), bHasPlaceholderType);
 		return PackedObjectRef;
 	}
 
diff --git a/Engine/Source/Runtime/CoreUObject/Private/UObject/PropertyBagRepository.cpp b/Engine/Source/Runtime/CoreUObject/Private/UObject/PropertyBagRepository.cpp
--- a/Engine/Source/Runtime/CoreUObject/Private/UObject/PropertyBagRepository.cpp
+++ b/Engine/Source/Runtime/CoreUObject/Private/UObject/PropertyBagRepository.cpp
@@ -27,9 +27,16 @@ DEFINE_LOG_CATEGORY_STATIC(LogPropertyBagRepository, Log, All);
 namespace UE
 {
 
+/** Internal registry that tracks the current set of types for property bag container objects instanced as placeholders for package exports that have invalid or missing class imports on load. */
 class FPropertyBagPlaceholderTypeRegistry
 {
 public:
+	static FPropertyBagPlaceholderTypeRegistry& Get()
+	{
+		static FPropertyBagPlaceholderTypeRegistry Instance;
+		return Instance;
+	}
+
 	void AddReferencedObjects(FReferenceCollector& Collector)
 	{
 		ConsumePendingPlaceholderTypes();
@@ -41,13 +48,17 @@ public:
 		PendingPlaceholderTypes.Enqueue(Type);
 	}
 
-	bool Contains(UStruct* Type)
+	bool Contains(const UStruct* Type)
 	{
 		ConsumePendingPlaceholderTypes();
 		return PlaceholderTypes.Contains(Type);
 	}
 
 protected:
+	FPropertyBagPlaceholderTypeRegistry() = default;
+	FPropertyBagPlaceholderTypeRegistry(const FPropertyBagPlaceholderTypeRegistry&) = delete;
+	FPropertyBagPlaceholderTypeRegistry& operator=(const FPropertyBagPlaceholderTypeRegistry&) = delete;
+
 	void ConsumePendingPlaceholderTypes()
 	{
 		if (!PendingPlaceholderTypes.IsEmpty())
@@ -122,11 +133,6 @@ FPropertyBagRepository& FPropertyBagRepository::Get()
 	return Repo;
 }
 
-FPropertyBagRepository::FPropertyBagRepository()
-{
-	PropertyBagPlaceholderTypeRegistry = MakeUnique<FPropertyBagPlaceholderTypeRegistry>();
-}
-
 void FPropertyBagRepository::ReassociateObjects(const TMap<UObject*, UObject*>& ReplacedObjects)
 {
 	FPropertyBagRepositoryLock LockRepo(this);
@@ -573,7 +579,7 @@ void FPropertyBagRepository::AddReferencedObjects(FReferenceCollector& Collector
 		Collector.AddReferencedObject(Element.Value);
 	}
 
-	PropertyBagPlaceholderTypeRegistry->AddReferencedObjects(Collector);
+	FPropertyBagPlaceholderTypeRegistry::Get().AddReferencedObjects(Collector);
 }
 
 FString FPropertyBagRepository::GetReferencerName() const
@@ -671,25 +677,24 @@ void FPropertyBagRepository::ShrinkMaps()
 	InstanceDataObjectToOwner.Compact();
 }
 
-bool FPropertyBagRepository::IsPropertyBagPlaceholderType(UStruct* Type)
+bool FPropertyBagRepository::IsPropertyBagPlaceholderType(const UStruct* Type)
 {
 	if (!Type)
 	{
 		return false;
 	}
 
-	return FPropertyBagRepository::Get().PropertyBagPlaceholderTypeRegistry->Contains(Type);
+	return FPropertyBagPlaceholderTypeRegistry::Get().Contains(Type);
 }
 
-bool FPropertyBagRepository::IsPropertyBagPlaceholderObject(UObject* Object)
+bool FPropertyBagRepository::IsPropertyBagPlaceholderObject(const UObject* Object)
 {
 	if (!Object)
 	{
 		return false;
 	}
 
-	return Object->HasAnyFlags(RF_HasPlaceholderType|RF_ClassDefaultObject)
-		&& IsPropertyBagPlaceholderType(Object->GetClass());
+	return IsPropertyBagPlaceholderType(Object->GetClass());
 }
 
 namespace Private
@@ -751,7 +756,7 @@ UStruct* FPropertyBagRepository::CreatePropertyBagPlaceholderType(UObject* Outer
 
 	// Use the property bag repository for now to manage property bag placeholder types (e.g. object lifetime).
 	// Note: The object lifetime of instances of this type will rely on existing references that are serialized.
-	FPropertyBagRepository::Get().PropertyBagPlaceholderTypeRegistry->Add(PlaceholderType);
+	FPropertyBagPlaceholderTypeRegistry::Get().Add(PlaceholderType);
 
 	return PlaceholderType;
 }
diff --git a/Engine/Source/Runtime/CoreUObject/Public/UObject/ObjectHandle.h b/Engine/Source/Runtime/CoreUObject/Public/UObject/ObjectHandle.h
--- a/Engine/Source/Runtime/CoreUObject/Public/UObject/ObjectHandle.h
+++ b/Engine/Source/Runtime/CoreUObject/Public/UObject/ObjectHandle.h
@@ -88,6 +88,11 @@ namespace UE::CoreUObject::Private
 
 	/** Read the handle as a pointer without checking if it is resolved. Invalid to call for unresolved handles. */
 	inline UObject* ReadObjectHandlePointerNoCheck(FObjectHandle Handle);
+
+#if UE_WITH_OBJECT_HANDLE_TYPE_SAFETY
+	/** Determines if the handle represents a reference that is considered to be safe for object subtype pointers. */
+	inline bool IsObjectHandleTypeSafeNoReadNoCheck(FObjectHandle Handle);
+#endif
 }
 
 /////////////////////////////////////////////////////////////////////////////////////////////////////////////////////////
@@ -116,12 +121,12 @@ inline bool IsObjectHandleResolved(FObjectHandle Handle)
 }
 
 /* return true if a handle is type safe.
- * null and unresolved handles are considered safe
+ * null handles are considered safe, otherwise does a bit test
  */ 
 inline bool IsObjectHandleTypeSafe(FObjectHandle Handle)
 {
 #if UE_WITH_OBJECT_HANDLE_TYPE_SAFETY
-	return IsObjectHandleNull(Handle) || !IsObjectHandleResolved(Handle) || !UE::CoreUObject::Private::HasAnyFlags(UE::CoreUObject::Private::ReadObjectHandlePointerNoCheck(Handle), static_cast<int32>(RF_HasPlaceholderType));
+	return UE::CoreUObject::Private::IsObjectHandleTypeSafeNoReadNoCheck(Handle);
 #else
 	return true;
 #endif
@@ -350,9 +355,22 @@ namespace UE::CoreUObject::Private
 		};
 	};
 
-	inline constexpr uint32 ObjectIdShift = 1;
-	inline constexpr uint32 PackageIdShift = 33;
-	inline constexpr uint32 PackageIdMask = 0x7FFF'FFFF;
+	inline constexpr uint32 TypeIdShift = 1;
+	inline constexpr uint32 ObjectIdShift = 2;
+	inline constexpr uint32 PackageIdShift = 34;
+	inline constexpr uint32 PackageIdMask = 0x3FFF'FFFF;
+
+#if UE_WITH_OBJECT_HANDLE_TYPE_SAFETY
+	inline bool IsObjectHandleTypeSafeNoReadNoCheck(FObjectHandle Handle)
+	{
+#if UE_WITH_OBJECT_HANDLE_LATE_RESOLVE
+		constexpr uint32 UnsafeTypeMask = (1 << TypeIdShift);
+		return IsObjectHandleNull(Handle) || ((Handle.PointerOrRef & (UnsafeTypeMask | 1)) == 1) || (IsObjectHandleResolved(Handle) && !(FindExistingPackedObjectRef(ReadObjectHandlePointerNoCheck(Handle)).EncodedRef & UnsafeTypeMask));
+#else
+		return true;
+#endif
+	}
+#endif
 
 #if UE_WITH_OBJECT_HANDLE_LATE_RESOLVE
 	//forward declarations
diff --git a/Engine/Source/Runtime/CoreUObject/Public/UObject/ObjectHandleDefines.h b/Engine/Source/Runtime/CoreUObject/Public/UObject/ObjectHandleDefines.h
--- a/Engine/Source/Runtime/CoreUObject/Public/UObject/ObjectHandleDefines.h
+++ b/Engine/Source/Runtime/CoreUObject/Public/UObject/ObjectHandleDefines.h
@@ -17,5 +17,5 @@
 
 
 #ifndef UE_WITH_OBJECT_HANDLE_TYPE_SAFETY
-	#define UE_WITH_OBJECT_HANDLE_TYPE_SAFETY 0
+	#define UE_WITH_OBJECT_HANDLE_TYPE_SAFETY WITH_EDITORONLY_DATA
 #endif
\ No newline at end of file
diff --git a/Engine/Source/Runtime/CoreUObject/Public/UObject/ObjectMacros.h b/Engine/Source/Runtime/CoreUObject/Public/UObject/ObjectMacros.h
--- a/Engine/Source/Runtime/CoreUObject/Public/UObject/ObjectMacros.h
+++ b/Engine/Source/Runtime/CoreUObject/Public/UObject/ObjectMacros.h
@@ -571,7 +571,7 @@ enum EObjectFlags
 	// RF_Dynamic				=0x04000000,	///< Was removed along with bp nativization
 	RF_WillBeLoaded				=0x08000000,	///< This object was constructed during load and will be loaded shortly
 	RF_HasExternalPackage		=0x10000000,	///< This object has an external package assigned and should look it up when getting the outermost package
-	RF_HasPlaceholderType		=0x20000000,	///< This object was instanced from a placeholder type (e.g. on load). References to it are serialized but externally resolve to NULL from a logical point of view (for type safety).
+	// RF_Unused				=0x20000000,
 
 	// RF_MirroredGarbage is mirrored in EInternalObjectFlags::Garbage because checking the internal flags is much faster for the Garbage Collector
 	// while checking the object flags is much faster outside of it where the Object pointer is already available and most likely cached.
@@ -585,6 +585,8 @@ UE_DEPRECATED(5.4, "RF_InternalGarbage should no longer be used. Use IsValid(Obj
 inline constexpr EObjectFlags RF_InternalGarbage = RF_MirroredGarbage;
 UE_DEPRECATED(5.4, "RF_InternalMirroredFlags should no longer be used. Use IsValid(Object) instead.")
 inline constexpr EObjectFlags RF_InternalMirroredFlags = RF_MirroredGarbage;
+UE_DEPRECATED(5.5, "RF_HasPlaceholderType is no longer in use.")
+inline constexpr EObjectFlags RF_HasPlaceholderType = EObjectFlags(0x20000000);
 
 /** Mask for all object flags */
 #define RF_AllFlags				(EObjectFlags)0xffffffff	///< All flags, used mainly for error checking
diff --git a/Engine/Source/Runtime/CoreUObject/Public/UObject/ObjectPtr.h b/Engine/Source/Runtime/CoreUObject/Public/UObject/ObjectPtr.h
--- a/Engine/Source/Runtime/CoreUObject/Public/UObject/ObjectPtr.h
+++ b/Engine/Source/Runtime/CoreUObject/Public/UObject/ObjectPtr.h
@@ -603,10 +603,16 @@ public:
 #if UE_WITH_OBJECT_HANDLE_TYPE_SAFETY
 		// Do a NULL test first before comparing the underlying handles, in case either side is
 		// a non-NULL, but unsafe type pointer (which would equate to NULL when Get() is called).
-		return ObjectPtr_Private::IsObjectPtrNull<T>(ObjectPtr) ? ObjectPtr_Private::IsObjectPtrNull<U>(Other.ObjectPtr) : ObjectPtr == Other.ObjectPtr;
-#else
-		return ObjectPtr == Other.ObjectPtr;
+		if (ObjectPtr_Private::IsObjectPtrNull<T>(ObjectPtr))
+		{
+			return ObjectPtr_Private::IsObjectPtrNull<U>(Other.ObjectPtr);
+		}
+		else if (ObjectPtr_Private::IsObjectPtrNull<U>(Other.ObjectPtr))
+		{
+			return false;
+		}
 #endif
+		return ObjectPtr == Other.ObjectPtr;
 	}
 
 	// Equality/Inequality comparisons against nullptr
diff --git a/Engine/Source/Runtime/CoreUObject/Tests/ObjectHandleTest.cpp b/Engine/Source/Runtime/CoreUObject/Tests/ObjectHandleTest.cpp
--- a/Engine/Source/Runtime/CoreUObject/Tests/ObjectHandleTest.cpp
+++ b/Engine/Source/Runtime/CoreUObject/Tests/ObjectHandleTest.cpp
@@ -14,6 +14,7 @@
 #include "TestHarness.h"
 #include "UObject/ObjectRef.h"
 #include "UObject/ObjectPathId.h"
+#include "UObject/PropertyBagRepository.h"
 
 static_assert(sizeof(FObjectHandle) == sizeof(void*), "FObjectHandle type must always compile to something equivalent to a pointer size.");
 
@@ -382,65 +383,99 @@ TEST_CASE_METHOD(FObjectHandleTestBase, "CoreUObject::FObjectHandle::Type Safety
 		TestPackage->RemoveFromRoot();
 	};
 
-	// simulate an unsafe class type
-	UClass* TestClass = NewObject<UClass>(TestPackage, TEXT("TestClass"), RF_Transient);
-	TestClass->SetSuperStruct(UObject::StaticClass());
-	TestClass->Bind();
-	TestClass->StaticLink(/*bRelinkExistingProperties =*/ true);
-	UObject* TestClassDefaults = TestClass->GetDefaultObject();
-	TestClass->PostLoadDefaultObject(TestClassDefaults);
-
-	// validate helper method(s)
-	CHECK_FALSE(UE::CoreUObject::Private::HasAnyFlags(TestClassDefaults, RF_NoFlags));
-	CHECK(UE::CoreUObject::Private::HasAnyFlags(TestClassDefaults, RF_ClassDefaultObject));
+	// construct an unsafe class type
+	UClass* TestUnsafeClass = UE::FPropertyBagRepository::CreatePropertyBagPlaceholderClass(TestPackage, UClass::StaticClass(), TEXT("TestUnsafeClass"));
 
 	// construct objects for testing
 	UObjectPtrTestClass* TestSafeObject = NewObject<UObjectPtrTestClass>(TestPackage, TEXT("TestSafeObject"), RF_Transient);
-	UObject* TestUnsafeObject = NewObject<UObject>(TestPackage, TestClass, TEXT("TestUnsafeObject"), RF_Transient | RF_HasPlaceholderType);
+	UObject* TestUnsafeObject = NewObject<UObject>(TestPackage, TestUnsafeClass, TEXT("TestUnsafeObject"), RF_Transient);
 
 	// construct object handles for testing
 	FObjectHandle NullObjectHandle = UE::CoreUObject::Private::MakeObjectHandle(nullptr);
 	FObjectHandle TestSafeObjectHandle = UE::CoreUObject::Private::MakeObjectHandle(TestSafeObject);
 	FObjectHandle TestUnsafeObjectHandle = UE::CoreUObject::Private::MakeObjectHandle(TestUnsafeObject);
 #if UE_WITH_OBJECT_HANDLE_LATE_RESOLVE
-	// note that unresolved object handles are type safe by definition (since it implies the underlying type was not a placeholder)
-	FObjectRef TestSafeObjectRef(TestSafeObject);
-	UE::CoreUObject::Private::FPackedObjectRef PackedSafeObjectRef = UE::CoreUObject::Private::MakePackedObjectRef(TestSafeObjectRef);
-	FObjectHandle TestUnresolvedSafeObjectHandle = { PackedSafeObjectRef.EncodedRef };
+	FObjectHandle TestLateResolveSafeObjectHandle = { UE::CoreUObject::Private::MakePackedObjectRef(TestSafeObject).EncodedRef };
+	FObjectHandle TestLateResolveUnsafeObjectHandle = { UE::CoreUObject::Private::MakePackedObjectRef(TestUnsafeObject).EncodedRef };
 #endif
 
-	// NULL/type-safe objects should report as being safe
+	// NULL/type-safe object handles should report as being safe
 	CHECK(IsObjectHandleTypeSafe(NullObjectHandle));
 	CHECK(IsObjectHandleTypeSafe(TestSafeObjectHandle));
 #if UE_WITH_OBJECT_HANDLE_LATE_RESOLVE
-	CHECK(IsObjectHandleTypeSafe(TestUnresolvedSafeObjectHandle));
-	CHECK(!IsObjectHandleResolved(TestUnresolvedSafeObjectHandle));	// the call above should not resolve the handle
+	CHECK(IsObjectHandleTypeSafe(TestLateResolveSafeObjectHandle));
 #endif
 
-	// unsafe type object handles should report as being unsafe
+	// unsafe type object handles should report as being unsafe, but only when late resolve is enabled (required for support)
+#if UE_WITH_OBJECT_HANDLE_LATE_RESOLVE
 	CHECK_FALSE(IsObjectHandleTypeSafe(TestUnsafeObjectHandle));
+	CHECK_FALSE(IsObjectHandleTypeSafe(TestLateResolveUnsafeObjectHandle));
+#else
+	CHECK(IsObjectHandleTypeSafe(TestUnsafeObjectHandle));
+#endif
 
-	// object handles should resolve the class to the unsafe type
-	CHECK(UE::CoreUObject::Private::ResolveObjectHandleClass(TestUnsafeObjectHandle) == TestClass);
-
-	// object handles should resolve/evaluate to the original type object
-	CHECK(UE::CoreUObject::Private::ResolveObjectHandle(TestUnsafeObjectHandle) == TestUnsafeObject);
+	// unsafe type object handles should resolve the class to the unsafe type
+	CHECK(UE::CoreUObject::Private::ResolveObjectHandleClass(TestUnsafeObjectHandle) == TestUnsafeClass);
+#if UE_WITH_OBJECT_HANDLE_LATE_RESOLVE
+	CHECK(UE::CoreUObject::Private::ResolveObjectHandleClass(TestLateResolveUnsafeObjectHandle) == TestUnsafeClass);
+#endif
 
-	// an unsafe type object handle should not equate to other unsafe type object handles except for itself (including NULL)
-	CHECK(NullObjectHandle != TestUnsafeObjectHandle);			// note: this intentionally differs from object *pointers* (see below)
+	// an unsafe type object handle should not equate to other unsafe type object handles (including NULL), except for itself
+	CHECK(NullObjectHandle != TestUnsafeObjectHandle);			// note: this behavior intentionally differs from object *pointers* (see below)
 	CHECK(TestUnsafeObjectHandle != NullObjectHandle);			// see note directly above
 	CHECK(TestSafeObjectHandle != TestUnsafeObjectHandle);
 	CHECK(TestUnsafeObjectHandle != TestSafeObjectHandle);
 	CHECK(TestUnsafeObjectHandle == TestUnsafeObjectHandle);
+#if UE_WITH_OBJECT_HANDLE_LATE_RESOLVE
+	CHECK(NullObjectHandle != TestLateResolveUnsafeObjectHandle);
+	CHECK(TestLateResolveUnsafeObjectHandle != NullObjectHandle);
+	CHECK(TestSafeObjectHandle != TestLateResolveUnsafeObjectHandle);
+	CHECK(TestLateResolveUnsafeObjectHandle != TestSafeObjectHandle);
+	CHECK(TestUnsafeObjectHandle == TestLateResolveUnsafeObjectHandle);
+	CHECK(TestLateResolveUnsafeObjectHandle == TestUnsafeObjectHandle);
+	CHECK(TestLateResolveSafeObjectHandle != TestLateResolveUnsafeObjectHandle);
+	CHECK(TestLateResolveUnsafeObjectHandle != TestLateResolveSafeObjectHandle);
+	CHECK(TestLateResolveUnsafeObjectHandle == TestLateResolveUnsafeObjectHandle);
+#endif
 
-	// simulate a placeholder object for type-safe pointer tests
-	TestSafeObject->SetFlags(RF_HasPlaceholderType);
+	// the type safety and class queries above should not have resolved an object handle that's using late resolve
+#if UE_WITH_OBJECT_HANDLE_LATE_RESOLVE
+	CHECK_FALSE(IsObjectHandleResolved(TestLateResolveSafeObjectHandle));
+	CHECK_FALSE(IsObjectHandleResolved(TestLateResolveUnsafeObjectHandle));
+#endif
 
-	// construct object pointers for testing
+	// unsafe type object handles should resolve/evaluate to the original type object
+	CHECK(UE::CoreUObject::Private::ResolveObjectHandle(TestSafeObjectHandle) == TestSafeObject);
+	CHECK(UE::CoreUObject::Private::ResolveObjectHandle(TestUnsafeObjectHandle) == TestUnsafeObject);
+#if UE_WITH_OBJECT_HANDLE_LATE_RESOLVE
+	CHECK(UE::CoreUObject::Private::ResolveObjectHandle(TestLateResolveSafeObjectHandle) == TestSafeObject);
+	CHECK(UE::CoreUObject::Private::ResolveObjectHandle(TestLateResolveUnsafeObjectHandle) == TestUnsafeObject);
+#endif
+
+	// all handles should now report as being resolved
+	CHECK(IsObjectHandleResolved(NullObjectHandle));
+	CHECK(IsObjectHandleResolved(TestSafeObjectHandle));
+	CHECK(IsObjectHandleResolved(TestUnsafeObjectHandle));
+#if UE_WITH_OBJECT_HANDLE_LATE_RESOLVE
+	CHECK(IsObjectHandleResolved(TestLateResolveSafeObjectHandle));
+	CHECK(IsObjectHandleResolved(TestLateResolveUnsafeObjectHandle));
+#endif
+
+	// construct object pointers for testing intentionally different behaviors of UObject-type vs. non-UObject-type bindings
 	TObjectPtr<UObject> NullObjectPtr(nullptr);
-	TObjectPtr<UObject> TestSafeObjectPtr(TestSafeObject);									// type safe pointer to placeholder (UObject type)
+	TObjectPtr<UObject> TestSafeObjectPtr(TestUnsafeObject);								// type safe pointer to placeholder (bound to UObject type)
 	FObjectPtr TestUnsafeObjectPtr_Untyped(TestUnsafeObject);
-	TObjectPtr<UObjectPtrTestClass> TestUnsafeObjectPtr(TestUnsafeObjectPtr_Untyped);		// unsafe pointer to placeholder object (non-UObject type)
+	TObjectPtr<UObjectPtrTestClass> TestUnsafeObjectPtr(TestUnsafeObjectPtr_Untyped);		// unsafe pointer to placeholder object (bound to non-UObject type)
+#if UE_WITH_OBJECT_HANDLE_LATE_RESOLVE
+	// note: "safe" in this context means the pointer should be type safe because it's bound to the UObject base type, but both reference the same "unsafe" object
+	TObjectPtr<UObject> TestLateResolveSafeObjectPtr(FObjectPtr({ UE::CoreUObject::Private::MakePackedObjectRef(TestUnsafeObject).EncodedRef }));
+	TObjectPtr<UObjectPtrTestClass> TestLateResolveUnsafeObjectPtr(FObjectPtr({ UE::CoreUObject::Private::MakePackedObjectRef(TestUnsafeObject).EncodedRef }));
+#endif
+
+	// an unsafe packed object reference should now exist with the unsafe type bit set for an unsafe object reference
+#if UE_WITH_OBJECT_HANDLE_LATE_RESOLVE
+	CHECK((UE::CoreUObject::Private::FindExistingPackedObjectRef(TestUnsafeObject).EncodedRef & (1 << UE::CoreUObject::Private::TypeIdShift)) != 0);
+#endif
 
 	// type safe object pointers should evaluate to true/non-NULL
 	CHECK(TestSafeObjectPtr);
@@ -449,21 +484,57 @@ TEST_CASE_METHOD(FObjectHandleTestBase, "CoreUObject::FObjectHandle::Type Safety
 	CHECK(TestSafeObjectPtr != NULL);
 	CHECK(nullptr != TestSafeObjectPtr);
 	CHECK(TestSafeObjectPtr != nullptr);
+#if UE_WITH_OBJECT_HANDLE_LATE_RESOLVE
+	CHECK(TestLateResolveSafeObjectPtr);
+	CHECK(!!TestLateResolveSafeObjectPtr);
+	CHECK(NULL != TestLateResolveSafeObjectPtr);
+	CHECK(TestLateResolveSafeObjectPtr != NULL);
+	CHECK(nullptr != TestLateResolveSafeObjectPtr);
+	CHECK(TestLateResolveSafeObjectPtr != nullptr);
+#endif
 
-	// unsafe type object pointers should evaluate to NULL/false (for type safety)
-	CHECK(!TestUnsafeObjectPtr);
+	// unsafe type object pointers should evaluate to NULL/false (for type safety), but only when late resolve is enabled (required for support)
+#if UE_WITH_OBJECT_HANDLE_LATE_RESOLVE
+	CHECK_FALSE(TestUnsafeObjectPtr);
 	CHECK_FALSE(!!TestUnsafeObjectPtr);
 	CHECK(NULL == TestUnsafeObjectPtr);
 	CHECK(TestUnsafeObjectPtr == NULL);
 	CHECK(nullptr == TestUnsafeObjectPtr);
 	CHECK(TestUnsafeObjectPtr == nullptr);
+	CHECK_FALSE(TestLateResolveUnsafeObjectPtr);
+	CHECK_FALSE(!!TestLateResolveUnsafeObjectPtr);
+	CHECK(NULL == TestLateResolveUnsafeObjectPtr);
+	CHECK(TestLateResolveUnsafeObjectPtr == NULL);
+	CHECK(nullptr == TestLateResolveUnsafeObjectPtr);
+	CHECK(TestLateResolveUnsafeObjectPtr == nullptr);
+#else
+	CHECK(TestUnsafeObjectPtr);
+	CHECK(!!TestUnsafeObjectPtr);
+	CHECK(NULL != TestUnsafeObjectPtr);
+	CHECK(TestUnsafeObjectPtr != NULL);
+	CHECK(nullptr != TestUnsafeObjectPtr);
+	CHECK(TestUnsafeObjectPtr != nullptr);
+#endif
 
-	// an unsafe type object pointer should not equate to other pointers except for NULL and itself
-	CHECK(NullObjectPtr == TestUnsafeObjectPtr);				// note: this intentionally differs from object *handles* (see above)
+	// an unsafe type object pointer should not equate to other unsafe type object pointers, excluding NULL and itself
+#if UE_WITH_OBJECT_HANDLE_LATE_RESOLVE
+	CHECK(NullObjectPtr == TestUnsafeObjectPtr);				// note: this intentionally differs from object *handles* (see above), and depends on late resolve
 	CHECK(TestUnsafeObjectPtr == NullObjectPtr);				// see note directly above
-	CHECK(TestSafeObjectPtr != TestUnsafeObjectPtr);
-	CHECK(TestUnsafeObjectPtr != TestSafeObjectPtr);
-	CHECK(TestUnsafeObjectPtr == TestUnsafeObjectPtr);
+	CHECK(NullObjectPtr != TestLateResolveSafeObjectPtr);
+	CHECK(TestLateResolveSafeObjectPtr != NullObjectPtr);
+	CHECK(NullObjectPtr == TestLateResolveUnsafeObjectPtr);
+	CHECK(TestLateResolveUnsafeObjectPtr == NullObjectPtr);
+	CHECK(TestSafeObjectPtr != TestLateResolveUnsafeObjectPtr);
+	CHECK(TestLateResolveUnsafeObjectPtr != TestSafeObjectPtr);
+	CHECK(TestUnsafeObjectPtr == TestLateResolveUnsafeObjectPtr);
+	CHECK(TestLateResolveUnsafeObjectPtr == TestUnsafeObjectPtr);
+	CHECK(TestLateResolveSafeObjectPtr != TestLateResolveUnsafeObjectPtr);
+	CHECK(TestLateResolveUnsafeObjectPtr != TestLateResolveSafeObjectPtr);
+	CHECK(TestLateResolveUnsafeObjectPtr == TestLateResolveUnsafeObjectPtr);
+#else
+	CHECK(NullObjectPtr != TestUnsafeObjectPtr);				// note: if late resolve is disabled, these will no longer be equivalent
+	CHECK(TestUnsafeObjectPtr != NullObjectPtr);				// see note directly above
+#endif
 
 	// an unsafe type object should evaluate the object's attributes correctly
 	CHECK(TestUnsafeObjectPtr.GetName() == TestUnsafeObject->GetName());
@@ -473,12 +544,43 @@ TEST_CASE_METHOD(FObjectHandleTestBase, "CoreUObject::FObjectHandle::Type Safety
 	CHECK(TestUnsafeObjectPtr.GetOuter() == TestUnsafeObject->GetOuter());
 	CHECK(TestUnsafeObjectPtr.GetClass() == TestUnsafeObject->GetClass());
 	CHECK(TestUnsafeObjectPtr.GetPackage() == TestUnsafeObject->GetPackage());
+#if UE_WITH_OBJECT_HANDLE_LATE_RESOLVE
+	CHECK(TestLateResolveUnsafeObjectPtr.GetName() == TestUnsafeObject->GetName());
+	CHECK(TestLateResolveUnsafeObjectPtr.GetFName() == TestUnsafeObject->GetFName());
+	CHECK(TestLateResolveUnsafeObjectPtr.GetPathName() == TestUnsafeObject->GetPathName());
+	CHECK(TestLateResolveUnsafeObjectPtr.GetFullName() == TestUnsafeObject->GetFullName());
+	CHECK(TestLateResolveUnsafeObjectPtr.GetOuter() == TestUnsafeObject->GetOuter());
+	CHECK(TestLateResolveUnsafeObjectPtr.GetClass() == TestUnsafeObject->GetClass());
+	CHECK(TestLateResolveUnsafeObjectPtr.GetPackage() == TestUnsafeObject->GetPackage());
+#endif
+
+	// the type safety and queries above should not have resolved an object pointer that's using late resolve
+#if UE_WITH_OBJECT_HANDLE_LATE_RESOLVE
+	CHECK_FALSE(TestLateResolveSafeObjectPtr.IsResolved());
+	CHECK_FALSE(TestLateResolveUnsafeObjectPtr.IsResolved());
+#endif
 
 	// a type safe object pointer should resolve to a non-NULL value when dereferenced
-	CHECK(TestSafeObjectPtr.Get() == TestSafeObject);
+	CHECK(TestSafeObjectPtr.Get() == TestUnsafeObject);
+#if UE_WITH_OBJECT_HANDLE_LATE_RESOLVE
+	CHECK(TestLateResolveSafeObjectPtr.Get() == TestUnsafeObject);
+#endif
 
-	// an unsafe type object pointer should resolve to NULL when dereferenced (for type safety)
+	// an unsafe type object pointer should resolve to NULL when dereferenced (for type safety), but only when late resolve is enabled (required for support)
+#if UE_WITH_OBJECT_HANDLE_LATE_RESOLVE
 	CHECK(TestUnsafeObjectPtr.Get() == nullptr);
+	CHECK(TestLateResolveUnsafeObjectPtr.Get() == nullptr);
+#else
+	CHECK(TestUnsafeObjectPtr.Get() == TestUnsafeObject);
+#endif
+
+	// all pointers should now be resolved
+	CHECK(TestSafeObjectPtr.IsResolved());
+	CHECK(TestUnsafeObjectPtr.IsResolved());
+#if UE_WITH_OBJECT_HANDLE_LATE_RESOLVE
+	CHECK(TestLateResolveSafeObjectPtr.IsResolved());
+	CHECK(TestLateResolveUnsafeObjectPtr.IsResolved());
+#endif
 }
 #endif
 
diff --git a/Engine/Source/Runtime/CoreUObject/Tests/ObjectPropertyTest.cpp b/Engine/Source/Runtime/CoreUObject/Tests/ObjectPropertyTest.cpp
--- a/Engine/Source/Runtime/CoreUObject/Tests/ObjectPropertyTest.cpp
+++ b/Engine/Source/Runtime/CoreUObject/Tests/ObjectPropertyTest.cpp
@@ -41,7 +41,7 @@ TEST_CASE("UE::CoreUObject::FObjectProperty::CheckValidAddress")
 	UObjectPtrTestClassWithRef* Obj = NewObject<UObjectPtrTestClassWithRef>(TestPackage, TEXT("Object"));
 	UObjectPtrTestClass* Other = NewObject<UObjectPtrTestClass>(Obj, TEXT("Other"));
 
-#if UE_WITH_OBJECT_HANDLE_TRACKING
+#if UE_WITH_OBJECT_HANDLE_TRACKING && UE_WITH_OBJECT_HANDLE_LATE_RESOLVE
 	FObjectHandle Handle = MakeUnresolvedHandle(Other);
 	TObjectPtr<UObjectPtrTestClass> ObjectPtr = *reinterpret_cast<TObjectPtr<UObjectPtrTestClass>*>(&Handle);
 #else
@@ -104,7 +104,7 @@ TEST_CASE("UE::CoreUObject::FObjectProperty::CheckValidAddressNonNullable")
 	UObjectPtrTestClassWithRef* Obj = NewObject<UObjectPtrTestClassWithRef>(TestPackage, TEXT("Object"));
 	UObjectPtrTestClass* Other = NewObject<UObjectPtrTestClass>(Obj, TEXT("Other"));
 
-#if UE_WITH_OBJECT_HANDLE_TRACKING
+#if UE_WITH_OBJECT_HANDLE_TRACKING && UE_WITH_OBJECT_HANDLE_LATE_RESOLVE
 	FObjectHandle Handle = MakeUnresolvedHandle(Other);
 	TObjectPtr<UObjectPtrTestClass> ObjectPtr = *reinterpret_cast<TObjectPtr<UObjectPtrTestClass>*>(&Handle);
 #else
@@ -426,7 +426,7 @@ TEST_CASE("UE::FObjectProperty::Identical::ObjectPtr")
 	TObjectPtr<UObject> Obj1 = NewObject<UObjectPtrTestClass>(TestPackage, TEXT("UObjectPtrTestClass"));
 	TObjectPtr<UObject> Obj2 = NewObject<UObjectPtrTestClass>(TestPackage2, TEXT("UObjectPtrTestClass"));
 
-#if UE_WITH_OBJECT_HANDLE_TRACKING
+#if UE_WITH_OBJECT_HANDLE_TRACKING && UE_WITH_OBJECT_HANDLE_LATE_RESOLVE
 
 	FObjectHandle Handle1 = MakeUnresolvedHandle(ObjWithRef.Get());
 	FObjectHandle Handle2 = MakeUnresolvedHandle(Obj1.Get());
@@ -447,7 +447,7 @@ TEST_CASE("UE::FObjectProperty::Identical::ObjectPtr")
 
 	CHECK(ResolveCount == 0);
 	CHECK(Property->Identical(&Obj1, &Obj2, PPF_DeepComparison));
-#if UE_WITH_OBJECT_HANDLE_TRACKING
+#if UE_WITH_OBJECT_HANDLE_TRACKING && UE_WITH_OBJECT_HANDLE_LATE_RESOLVE
 	CHECK(ResolveCount == 2);
 #endif
 	
@@ -506,7 +506,7 @@ TEST_CASE("UE::FObjectProperty::CopySingleValue")
 	UObject* RawPtr = Obj1;
 	TObjectPtr<UObject> PtrObj2 = Obj2;
 
-#if UE_WITH_OBJECT_HANDLE_TRACKING
+#if UE_WITH_OBJECT_HANDLE_TRACKING && UE_WITH_OBJECT_HANDLE_LATE_RESOLVE
 	FObjectHandle Handle = MakeUnresolvedHandle(Obj2);
 	PtrObj2 = TObjectPtr<UObject>(FObjectPtr(Handle) );
 #endif
