PropertyVisitor updates

* Added FPropertyVisitorInfo::ParentStructType to provide the information about the actual UStruct type of the instance we're visiting the property data for.
    * This isn't the same as Property->GetOwnerStruct() as it's based on the instance data and not the type data.
* Updated FPropertyVisitorPath::GetPropertyDataPtr to stop iterating at a layer once it already found a match.
* Added tests for FPropertyVisitorPath::GetPropertyDataPtr.

#rnx

diff --git a/Engine/Source/Runtime/CoreUObject/Private/UObject/Class.cpp b/Engine/Source/Runtime/CoreUObject/Private/UObject/Class.cpp
--- a/Engine/Source/Runtime/CoreUObject/Private/UObject/Class.cpp
+++ b/Engine/Source/Runtime/CoreUObject/Private/UObject/Class.cpp
@@ -775,6 +775,7 @@ EPropertyVisitorControlFlow UStruct::Visit(FPropertyVisitorPath& Path, void* Dat
 	for (const FProperty* Property = PropertyLink; Property; Property = Property->PropertyLinkNext)
 	{
 		FPropertyVisitorScope Scope(Path, FPropertyVisitorInfo(Property));
+		Path.Top().ParentStructType = this;
 
 		// ArrayDim of one means it is just a single property, not a static array
 		if(Property->ArrayDim == 1)
diff --git a/Engine/Source/Runtime/CoreUObject/Private/UObject/PropertyVisitor.cpp b/Engine/Source/Runtime/CoreUObject/Private/UObject/PropertyVisitor.cpp
--- a/Engine/Source/Runtime/CoreUObject/Private/UObject/PropertyVisitor.cpp
+++ b/Engine/Source/Runtime/CoreUObject/Private/UObject/PropertyVisitor.cpp
@@ -111,11 +111,18 @@ void* FPropertyVisitorPath::GetPropertyDataPtr(UObject* Object) const
 	void* DataPtr = nullptr;
 	checkf(Object, TEXT("Expecting an valid object"));
 
-	Object->GetClass()->Visit(Object, [this, &DataPtr](const FPropertyVisitorPath& InPath, void* Data)
+	int32 MatchedPathDepth = 0;
+	Object->GetClass()->Visit(Object, [this, &DataPtr, &MatchedPathDepth](const FPropertyVisitorPath& InPath, void* Data)
 	{
+		if (InPath.Num() <= MatchedPathDepth)
+		{
+			// We've returned a level that we previously found a match in; we can stop now
+			return EPropertyVisitorControlFlow::Stop;
+		}
 		bool bIsEqual = false;
 		if (InPath.Contained(*this, &bIsEqual))
 		{
+			MatchedPathDepth = InPath.Num();
 			if(bIsEqual)
 			{
 				DataPtr = Data;
diff --git a/Engine/Source/Runtime/CoreUObject/Public/UObject/PropertyVisitor.h b/Engine/Source/Runtime/CoreUObject/Public/UObject/PropertyVisitor.h
--- a/Engine/Source/Runtime/CoreUObject/Public/UObject/PropertyVisitor.h
+++ b/Engine/Source/Runtime/CoreUObject/Public/UObject/PropertyVisitor.h
@@ -31,8 +31,12 @@ struct FPropertyVisitorInfo
 		, PropertyInfo(InPropertyInfo)
 		, bContainsInnerProperties(false)
 	{}
-	FPropertyVisitorInfo(const FPropertyVisitorInfo& Other) = default;
-	FPropertyVisitorInfo(FPropertyVisitorInfo&& Other) = default;
+
+	FPropertyVisitorInfo(const FPropertyVisitorInfo&) = default;
+	FPropertyVisitorInfo(FPropertyVisitorInfo&&) = default;
+
+	FPropertyVisitorInfo& operator=(const FPropertyVisitorInfo&) = default;
+	FPropertyVisitorInfo& operator=(FPropertyVisitorInfo&&) = default;
 
 	void SetIndex(int32 InIndex, EPropertyVisitorInfoType InPropertyInfo)
 	{
@@ -43,6 +47,13 @@ struct FPropertyVisitorInfo
 	/** The property currently being visited */
 	const FProperty* Property;
 
+	/**
+	 * The parent struct that provided the property being iterated, if iterating a sub-property within a struct.
+	 * @note This is slighty different than Property->GetOwnerStruct() as you might be iterating a FDerived instance 
+	 *       but processing a FBase struct property. In this case this will be set to FDerived rather than FBase.
+	 */
+	const UStruct* ParentStructType = nullptr;
+
 	/**
      * Index of the element being visited in the container, otherwise INDEX_NONE.
      * For maps and sets it indicates the logical index. */
@@ -59,8 +70,13 @@ struct FPropertyVisitorPath
 {
 public:
 	FPropertyVisitorPath() = default;
-	FPropertyVisitorPath(const FPropertyVisitorPath& Other) = default;
-	FPropertyVisitorPath(FPropertyVisitorPath&& Other) = default;
+	
+	FPropertyVisitorPath(const FPropertyVisitorPath&) = default;
+	FPropertyVisitorPath(FPropertyVisitorPath&&) = default;
+
+	FPropertyVisitorPath& operator=(const FPropertyVisitorPath&) = default;
+	FPropertyVisitorPath& operator=(FPropertyVisitorPath&&) = default;
+	
 	explicit FPropertyVisitorPath(const FPropertyVisitorInfo& Info)
 	{
 		Push(Info);
