Remove an unnecessary/unused check for placeholder object types.

This eliminates a few warnings emitted by low-level unit tests that are fuzzing late resolve object handle APIs using dummy package names.

diff --git a/Engine/Source/Runtime/CoreUObject/Private/UObject/ObjectHandleIndex.cpp b/Engine/Source/Runtime/CoreUObject/Private/UObject/ObjectHandleIndex.cpp
--- a/Engine/Source/Runtime/CoreUObject/Private/UObject/ObjectHandleIndex.cpp
+++ b/Engine/Source/Runtime/CoreUObject/Private/UObject/ObjectHandleIndex.cpp
@@ -529,8 +529,7 @@ namespace UE::CoreUObject::Private
 			return { 0 };
 		}
 
-		const bool bHasPlaceholderType = UE::FPropertyBagRepository::IsPropertyBagPlaceholderType(ObjectRef.ResolveObjectRefClass());
-		return MakePackedObjectRef(ObjectRef.PackageName, ObjectRef.ClassPackageName, ObjectRef.ClassName, ObjectRef.GetObjectPath(), bHasPlaceholderType);
+		return MakePackedObjectRef(ObjectRef.PackageName, ObjectRef.ClassPackageName, ObjectRef.ClassName, ObjectRef.GetObjectPath());
 	}
 
 
