Added index-in-owner to FProperty to support efficient storage of per-property flags

#rnx

diff --git a/Engine/Source/Runtime/CoreUObject/Private/UObject/Class.cpp b/Engine/Source/Runtime/CoreUObject/Private/UObject/Class.cpp
--- a/Engine/Source/Runtime/CoreUObject/Private/UObject/Class.cpp
+++ b/Engine/Source/Runtime/CoreUObject/Private/UObject/Class.cpp
@@ -849,11 +849,17 @@ void UStruct::Link(FArchive& Ar, bool bRelinkExistingProperties)
 		int32 LoopNum = 1;
 		for (int32 LoopIter = 0; LoopIter < LoopNum; LoopIter++)
 		{
+		#if WITH_EDITORONLY_DATA
+			TotalFieldCount = 0;
+		#endif
 			PropertiesSize = 0;
 			MinAlignment = 1;
 
 			if (InheritanceSuper)
 			{
+			#if WITH_EDITORONLY_DATA
+				TotalFieldCount = InheritanceSuper->TotalFieldCount;
+			#endif
 				PropertiesSize = InheritanceSuper->GetPropertiesSize();
 				MinAlignment = InheritanceSuper->GetMinAlignment();
 			}
diff --git a/Engine/Source/Runtime/CoreUObject/Private/UObject/Property.cpp b/Engine/Source/Runtime/CoreUObject/Private/UObject/Property.cpp
--- a/Engine/Source/Runtime/CoreUObject/Private/UObject/Property.cpp
+++ b/Engine/Source/Runtime/CoreUObject/Private/UObject/Property.cpp
@@ -1357,10 +1357,17 @@ int32 FProperty::SetupOffset()
 	{
 		UStruct* OwnerStruct = (UStruct*)OwnerUObject;
 		Offset_Internal = Align(OwnerStruct->GetPropertiesSize(), GetMinAlignment());
+	#if WITH_EDITORONLY_DATA
+		IndexInOwner = IntCastChecked<uint16>(OwnerStruct->TotalFieldCount);
+		OwnerStruct->TotalFieldCount += ArrayDim;
+	#endif
 	}
 	else
 	{
 		Offset_Internal = Align(0, GetMinAlignment());
+	#if WITH_EDITORONLY_DATA
+		IndexInOwner = 0;
+	#endif
 	}
 
 	uint32 UnsignedTotal = (uint32)Offset_Internal + (uint32)GetSize();
diff --git a/Engine/Source/Runtime/CoreUObject/Public/UObject/Class.h b/Engine/Source/Runtime/CoreUObject/Public/UObject/Class.h
--- a/Engine/Source/Runtime/CoreUObject/Public/UObject/Class.h
+++ b/Engine/Source/Runtime/CoreUObject/Public/UObject/Class.h
@@ -430,6 +430,9 @@ public:
 	TArray<TObjectPtr<UPropertyWrapper>> PropertyWrappers;
 	/** Unique id incremented each time this class properties get destroyed */
 	int32 FieldPathSerialNumber;
+
+	/** Number of fields, inclusive of base fields and fixed-size array elements. Generated during Link(). */
+	UE_INTERNAL int32 TotalFieldCount = 0;
 #endif
 
 	/** Cached schema for optimized unversioned and filtereditoronly property serialization, owned by this. */
diff --git a/Engine/Source/Runtime/CoreUObject/Public/UObject/UnrealType.h b/Engine/Source/Runtime/CoreUObject/Public/UObject/UnrealType.h
--- a/Engine/Source/Runtime/CoreUObject/Public/UObject/UnrealType.h
+++ b/Engine/Source/Runtime/CoreUObject/Public/UObject/UnrealType.h
@@ -186,6 +186,11 @@ class FProperty : public FField
 private:
 	TEnumAsByte<ELifetimeCondition> BlueprintReplicationCondition;
 
+#if WITH_EDITORONLY_DATA
+	/** Index of the property within its owner, inclusive of base properties. Generated during Link(). */
+	uint16		IndexInOwner = 0;
+#endif
+
 	// In memory variables (generated during Link()).
 	int32		Offset_Internal;
 
@@ -356,6 +361,14 @@ protected:
 	COREUOBJECT_API void Init();
 
 public:
+#if WITH_EDITORONLY_DATA
+	/** Return the index of the property in its owner. Only valid after Link(). */
+	UE_INTERNAL FORCEINLINE int32 GetIndexInOwner() const
+	{
+		return IndexInOwner;
+	}
+#endif
+
 	/** Return offset of property from container base. */
 	FORCEINLINE int32 GetOffset_ForDebug() const
 	{
