Optimize UUserDefinedStruct::InitializeStruct in non editor builds, as of 18672462 all cooked assets include XBeforeY dependency information to safely order loads. GEventDrivenLoaderEnabled is also no longer set, so we can't rely on that anymore

diff --git a/Engine/Source/Runtime/CoreUObject/Private/Blueprint/BlueprintSupport.cpp b/Engine/Source/Runtime/CoreUObject/Private/Blueprint/BlueprintSupport.cpp
--- a/Engine/Source/Runtime/CoreUObject/Private/Blueprint/BlueprintSupport.cpp
+++ b/Engine/Source/Runtime/CoreUObject/Private/Blueprint/BlueprintSupport.cpp
@@ -163,14 +163,10 @@ bool FBlueprintSupport::IsDeferredDependencyPlaceholder(const UObject* LoadedObj
 		LoadedObj->IsA<ULinkerPlaceholderExportObject>() );
 }
 
+#if WITH_EDITOR
 void FBlueprintSupport::RegisterDeferredDependenciesInStruct(const UStruct* Struct, void* StructData)
 {
 #if USE_CIRCULAR_DEPENDENCY_LOAD_DEFERRING
-	if (GEventDrivenLoaderEnabled)
-	{
-		return;
-	}
-
 	for (TPropertyValueIterator<const FObjectProperty> It(Struct, StructData); It; ++It)
 	{
 		const FObjectProperty* Property = It.Key();
@@ -216,6 +212,7 @@ void FBlueprintSupport::RegisterDeferredDependenciesInStruct(const UStruct* Stru
 	}
 #endif // USE_CIRCULAR_DEPENDENCY_LOAD_DEFERRING
 }
+#endif
 
 void FBlueprintSupport::RepairDeferredDependenciesInObject(UObject* Object)
 {
diff --git a/Engine/Source/Runtime/CoreUObject/Public/Blueprint/BlueprintSupport.h b/Engine/Source/Runtime/CoreUObject/Public/Blueprint/BlueprintSupport.h
--- a/Engine/Source/Runtime/CoreUObject/Public/Blueprint/BlueprintSupport.h
+++ b/Engine/Source/Runtime/CoreUObject/Public/Blueprint/BlueprintSupport.h
@@ -109,8 +109,10 @@ struct FBlueprintSupport
 	/** Tells if the specified object is one of the many flavors of FLinkerPlaceholderBase that we have. */
 	COREUOBJECT_API static bool IsDeferredDependencyPlaceholder(const UObject* LoadedObj);
 
+#if WITH_EDITOR
 	/** Registers any object properties in this struct with the deferred dependency system */
 	COREUOBJECT_API static void RegisterDeferredDependenciesInStruct(const UStruct* Struct, void* StructData);
+#endif
 
 	/** Repair function to scan for any deferred dependency placeholders and attempt to resolve them in-place with real assets */
 	COREUOBJECT_API static void RepairDeferredDependenciesInObject(UObject* Object);
diff --git a/Engine/Source/Runtime/Engine/Private/UserDefinedStruct.cpp b/Engine/Source/Runtime/Engine/Private/UserDefinedStruct.cpp
--- a/Engine/Source/Runtime/Engine/Private/UserDefinedStruct.cpp
+++ b/Engine/Source/Runtime/Engine/Private/UserDefinedStruct.cpp
@@ -313,9 +313,11 @@ void UUserDefinedStruct::InitializeStruct(void* Dest, int32 ArrayDim) const
 			void* DestStruct = (uint8*)Dest + (Stride * ArrayIndex);
 			CopyScriptStruct(DestStruct, DefaultInstance);
 
+#if WITH_EDITOR
 			// When copying into another struct we need to register this raw struct pointer so any deferred dependencies will get fixed later
 			FScopedPlaceholderRawContainerTracker TrackStruct(DestStruct);
 			FBlueprintSupport::RegisterDeferredDependenciesInStruct(this, DestStruct);
+#endif
 		}	
 	}
 }
