[AsyncLoading2]
 - Prevent stall detection to trigger during debug sessions

diff --git a/Engine/Source/Runtime/CoreUObject/Private/Serialization/AsyncLoading2.cpp b/Engine/Source/Runtime/CoreUObject/Private/Serialization/AsyncLoading2.cpp
--- a/Engine/Source/Runtime/CoreUObject/Private/Serialization/AsyncLoading2.cpp
+++ b/Engine/Source/Runtime/CoreUObject/Private/Serialization/AsyncLoading2.cpp
@@ -426,6 +426,14 @@ static FAutoConsoleVariableRef CVarGStallDetectorTimeout(
 	ECVF_Default
 );
 
+static int32 GStallDetectorIdleLoops = 50;
+static FAutoConsoleVariableRef CVarGStallDetectorIdleLoops(
+	TEXT("s.StallDetectorIdleLoops"),
+	GStallDetectorIdleLoops,
+	TEXT("The minimum amount of idle loops before considering the loader stalled if no progress is being made"),
+	ECVF_Default
+);
+
 CSV_DECLARE_CATEGORY_MODULE_EXTERN(CORE_API, Basic);
 CSV_DECLARE_CATEGORY_MODULE_EXTERN(CORE_API, FileIO);
 
@@ -9553,6 +9561,7 @@ void FAsyncLoadingThread2::FlushLoading(TConstArrayView<int32> RequestIDs)
 		// Flush async loaders by not using a time limit. Needed for e.g. garbage collection.
 		{
 			double LastActivity = 0.0;
+			int32  IdleLoopCount = 0;
 			while (IsAsyncLoadingPackages())
 			{
 				bool bDidSomething = false;
@@ -9574,8 +9583,11 @@ void FAsyncLoadingThread2::FlushLoading(TConstArrayView<int32> RequestIDs)
 					if (LastActivity == 0.0)
 					{
 						LastActivity = FPlatformTime::Seconds();
+
+						// The loop count offers additional protection against timeout that could arise during debugging.
+						IdleLoopCount = 0;
 					}
-					else if (FPlatformTime::Seconds() - LastActivity > GStallDetectorTimeout)
+					else if (FPlatformTime::Seconds() - LastActivity > GStallDetectorTimeout && ++IdleLoopCount > GStallDetectorIdleLoops)
 					{
 #if !UE_BUILD_SHIPPING
 						{
