Converted ValuesSetBySerialization from a set to a bit array

#rnx

diff --git a/Engine/Source/Runtime/CoreUObject/Private/UObject/InstanceDataObjectUtils.cpp b/Engine/Source/Runtime/CoreUObject/Private/UObject/InstanceDataObjectUtils.cpp
--- a/Engine/Source/Runtime/CoreUObject/Private/UObject/InstanceDataObjectUtils.cpp
+++ b/Engine/Source/Runtime/CoreUObject/Private/UObject/InstanceDataObjectUtils.cpp
@@ -395,11 +395,9 @@ namespace UE
 			}
 		}
 
-		// add a hidden set property used to record whether this struct's properties were set serialization.
+		// Add a hidden byte array property to record whether its sibling properties were set by serialization.
+		FByteProperty* ValuesSetBySerializationProperty = CastFieldChecked<FByteProperty>(FByteProperty::Construct(Result, NAME_ValuesSetBySerialization, RF_Transient | RF_MarkAsNative));
 		{
-			FSetProperty* ValuesSetBySerializationProperty = CastFieldChecked<FSetProperty>(FSetProperty::Construct(Result, NAME_ValuesSetBySerialization, RF_Transient | RF_MarkAsNative));
-			static FName Name_PropertyName(TEXT("PropertyName"));
-			ValuesSetBySerializationProperty->ElementProp = CastFieldChecked<FProperty>(FInt64Property::Construct(ValuesSetBySerializationProperty, Name_PropertyName, RF_Transient));
 			ValuesSetBySerializationProperty->SetPropertyFlags(CPF_Transient | CPF_EditorOnly | CPF_NativeAccessSpecifierPrivate);
 			Result->AddCppProperty(ValuesSetBySerializationProperty);
 		}
@@ -409,6 +407,15 @@ namespace UE
 		{
 			Result->AddCppProperty(Property);
 		}
+
+		// Count properties and set the size of the array of flags.
+		int32 PropertyCount = -1; // Start at -1 to exclude ValuesSetBySerialization.
+		for (TFieldIterator<FProperty> It(Result); It; ++It)
+		{
+			PropertyCount += It->ArrayDim;
+		}
+		ValuesSetBySerializationProperty->ArrayDim = FMath::DivideAndRoundUp(PropertyCount, 8);
+
 		Result->Bind();
 		Result->StaticLink(/*bRelinkExistingProperties*/true);
 		return Result;
@@ -470,30 +477,38 @@ namespace UE
 		return Result;
 	}
 
-	void MarkPropertySetBySerialization(const UStruct* Struct, const void* StructData, const FProperty* Property, int32 ArrayIndex)
+	void MarkPropertySetBySerialization(const UStruct* Struct, void* StructData, const FProperty* Property, int32 ArrayIndex)
 	{
-		if (const FSetProperty* ValuesSetByPropertyBagProperty = CastField<FSetProperty>(Struct->FindPropertyByName(NAME_ValuesSetBySerialization)))
+	#if WITH_EDITORONLY_DATA
+		if (const FByteProperty* ValuesSetBySerializationProperty = CastField<FByteProperty>(Struct->FindPropertyByName(NAME_ValuesSetBySerialization)))
 		{
-			FScriptSetHelper ValuesSetByPropertyBag(ValuesSetByPropertyBagProperty, ValuesSetByPropertyBagProperty->ContainerPtrToValuePtr<void>(StructData));
-			const void* PropertyDataPtr = Property->ContainerPtrToValuePtr<void>(StructData, ArrayIndex);
-			const int64 ValueOffset = static_cast<const uint8*>(PropertyDataPtr) - static_cast<const uint8*>(StructData);
-			const int32 FoundIndex = ValuesSetByPropertyBag.FindElementIndex(&ValueOffset);
-			if (FoundIndex == INDEX_NONE)
+			const int32 PropertyIndex = Property->GetIndexInOwner() + ArrayIndex;
+			const int32 ByteIndex = PropertyIndex / 8;
+			const int32 BitOffset = PropertyIndex % 8;
+			if (ByteIndex < ValuesSetBySerializationProperty->ArrayDim)
 			{
-				ValuesSetByPropertyBag.AddElement(&ValueOffset);
+				uint8* PropertyDataPtr = ValuesSetBySerializationProperty->ContainerPtrToValuePtr<uint8>(StructData, ByteIndex);
+				*PropertyDataPtr |= (1 << BitOffset);
 			}
 		}
+	#endif
 	}
 
 	bool WasPropertySetBySerialization(const UStruct* Struct, const void* StructData, const FProperty* Property, int32 ArrayIndex)
 	{
-		if (const FSetProperty* ValuesSetByPropertyBagProperty = CastField<FSetProperty>(Struct->FindPropertyByName(NAME_ValuesSetBySerialization)))
+	#if WITH_EDITORONLY_DATA
+		if (const FByteProperty* ValuesSetBySerializationProperty = CastField<FByteProperty>(Struct->FindPropertyByName(NAME_ValuesSetBySerialization)))
 		{
-			const FScriptSetHelper ValuesSetByPropertyBag(ValuesSetByPropertyBagProperty, ValuesSetByPropertyBagProperty->ContainerPtrToValuePtr<void>(StructData));
-			const void* PropertyDataPtr = Property->ContainerPtrToValuePtr<void>(StructData, ArrayIndex);
-			const int64 ValueOffset = static_cast<const uint8*>(PropertyDataPtr) - static_cast<const uint8*>(StructData);
-			return ValuesSetByPropertyBag.FindElementIndex(&ValueOffset) != INDEX_NONE;
+			const int32 PropertyIndex = Property->GetIndexInOwner() + ArrayIndex;
+			const int32 ByteIndex = PropertyIndex / 8;
+			const int32 BitOffset = PropertyIndex % 8;
+			if (ByteIndex < ValuesSetBySerializationProperty->ArrayDim)
+			{
+				const uint8* PropertyDataPtr = ValuesSetBySerializationProperty->ContainerPtrToValuePtr<uint8>(StructData, ByteIndex);
+				return (*PropertyDataPtr & (1 << BitOffset)) != 0;
+			}
 		}
+	#endif
 		return false;
 	}
 
diff --git a/Engine/Source/Runtime/CoreUObject/Private/UObject/InstanceDataObjectUtils.h b/Engine/Source/Runtime/CoreUObject/Private/UObject/InstanceDataObjectUtils.h
--- a/Engine/Source/Runtime/CoreUObject/Private/UObject/InstanceDataObjectUtils.h
+++ b/Engine/Source/Runtime/CoreUObject/Private/UObject/InstanceDataObjectUtils.h
@@ -26,7 +26,7 @@ bool IsInstanceDataObjectSupportEnabled(UObject* Object = nullptr);
 UClass* CreateInstanceDataObjectClass(const FPropertyPathNameTree* PropertyTree, UClass* OwnerClass, UObject* Outer);
 
 /** Notify that a property in an struct was set when the struct was deserialized. */
-void MarkPropertySetBySerialization(const UStruct* Struct, const void* StructData, const FProperty* Property, int32 ArrayIndex = 0);
+void MarkPropertySetBySerialization(const UStruct* Struct, void* StructData, const FProperty* Property, int32 ArrayIndex = 0);
 /** Query whether a property in the struct was set when the struct was deserialized. */
 bool WasPropertySetBySerialization(const UStruct* Struct, const void* StructData, const FProperty* Property, int32 ArrayIndex = 0);
 /** Copy whether each property was set by serialization from one IDO to another. */
