Removed index tracking in SerializedPropertyPath

#rnx

diff --git a/Engine/Source/Runtime/CoreUObject/Private/UObject/Class.cpp b/Engine/Source/Runtime/CoreUObject/Private/UObject/Class.cpp
--- a/Engine/Source/Runtime/CoreUObject/Private/UObject/Class.cpp
+++ b/Engine/Source/Runtime/CoreUObject/Private/UObject/Class.cpp
@@ -1638,7 +1638,7 @@ void UStruct::SerializeVersionedTaggedProperties(FStructuredArchive::FSlot Slot,
 					const FName Name = Property ? Property->GetFName() : Tag.Name;
 					const int32 Index = Tag.ArrayIndex > 0 || (Property && Property->ArrayDim > 1) ? Tag.ArrayIndex : INDEX_NONE;
 					const UE::FPropertyPathNameSegment Segment{Name, Tag.GetType(), Index};
-					SerializedPropertyPath.Emplace(SerializeContext, Segment, UE::ESerializedPropertyPathNotify::Yes);
+					SerializedPropertyPath.Emplace(SerializeContext, Segment);
 				}
 
 				bool bTryStoreUnknownPropertyPath = false;
diff --git a/Engine/Source/Runtime/CoreUObject/Private/UObject/PropertyArray.cpp b/Engine/Source/Runtime/CoreUObject/Private/UObject/PropertyArray.cpp
--- a/Engine/Source/Runtime/CoreUObject/Private/UObject/PropertyArray.cpp
+++ b/Engine/Source/Runtime/CoreUObject/Private/UObject/PropertyArray.cpp
@@ -684,7 +684,6 @@ void FArrayProperty::SerializeItem(FStructuredArchive::FSlot Slot, void* Value,
 
 				// Serialize the item at this array index
 				i = PropertyNode->ArrayIndex;
-				UE::FSerializedPropertyPathIndexScope SerializedPropertyPathIndex(Context, i, UE::ESerializedPropertyPathNotify::Yes);
 				SerializeContainerItem(Array.EnterElement(), ArrayHelper.GetRawPtr(i));
 				PropertyNode = PropertyNode->PropertyListNext;
 
@@ -709,7 +708,6 @@ void FArrayProperty::SerializeItem(FStructuredArchive::FSlot Slot, void* Value,
 			NAME_UArraySerializeCount.SetNumber(i);
 			FArchive::FScopeAddDebugData P(UnderlyingArchive, NAME_UArraySerializeCount);
 #endif
-			UE::FSerializedPropertyPathIndexScope SerializedPropertyPathIndex(Context, i, UE::ESerializedPropertyPathNotify::Yes);
 			SerializeContainerItem(Array.EnterElement(), ArrayHelper.GetRawPtr(i++));
 		}
 
@@ -1236,16 +1234,11 @@ EConvertFromTypeResult FArrayProperty::ConvertFromType(const FPropertyTag& Tag,
 	FUObjectSerializeContext* Context = FUObjectThreadContext::Get().GetSerializeContext();
 	FStructuredArchive::FStream ValueStream = Slot.EnterStream();
 
-	EConvertFromTypeResult ConvertResult;
-	{
-		UE::FSerializedPropertyPathIndexScope SerializedPropertyPathIndex(Context, 0, UE::ESerializedPropertyPathNotify::Yes);
-		ConvertResult = Inner->ConvertFromType(InnerPropertyTag, ValueStream.EnterElement(), ScriptArrayHelper.GetRawPtr(0), DefaultsStruct, nullptr);
-	}
+	EConvertFromTypeResult ConvertResult = Inner->ConvertFromType(InnerPropertyTag, ValueStream.EnterElement(), ScriptArrayHelper.GetRawPtr(0), DefaultsStruct, nullptr);
 	if (ConvertResult == EConvertFromTypeResult::Converted || ConvertResult == EConvertFromTypeResult::Serialized)
 	{
 		for (int32 ElementIndex = 1; ElementIndex < ElementCount; ++ElementIndex)
 		{
-			UE::FSerializedPropertyPathIndexScope SerializedPropertyPathIndex(Context, ElementIndex, UE::ESerializedPropertyPathNotify::Yes);
 			ConvertResult = Inner->ConvertFromType(InnerPropertyTag, ValueStream.EnterElement(), ScriptArrayHelper.GetRawPtr(ElementIndex), DefaultsStruct, nullptr);
 			check(ConvertResult == EConvertFromTypeResult::Converted || ConvertResult == EConvertFromTypeResult::Serialized);
 		}
diff --git a/Engine/Source/Runtime/CoreUObject/Private/UObject/PropertyHelper.h b/Engine/Source/Runtime/CoreUObject/Private/UObject/PropertyHelper.h
--- a/Engine/Source/Runtime/CoreUObject/Private/UObject/PropertyHelper.h
+++ b/Engine/Source/Runtime/CoreUObject/Private/UObject/PropertyHelper.h
@@ -98,20 +98,13 @@ FPropertyTypeName FindOriginalType(const FStructProperty* Struct);
  */
 FPropertyTypeName ApplyRedirectsToPropertyType(FPropertyTypeName OldTypeName, const FProperty* Property = nullptr);
 
-/** Whether to notify that a property has been serialized when terminating the property path scopes below. */
-enum class ESerializedPropertyPathNotify
-{
-	No,
-	Yes,
-};
-
 /**
  * Pushes a segment to SerializedPropertyPath for the lifetime of this object if path tracking is active.
  */
 class FSerializedPropertyPathScope
 {
 public:
-	[[nodiscard]] FSerializedPropertyPathScope(FUObjectSerializeContext* InContext, const FPropertyPathNameSegment& InSegment, ESerializedPropertyPathNotify InNotify)
+	[[nodiscard]] FSerializedPropertyPathScope(FUObjectSerializeContext* InContext, const FPropertyPathNameSegment& InSegment)
 	{
 		if (InContext->bTrackSerializedPropertyPath)
 		{
@@ -135,36 +128,4 @@ private:
 	FUObjectSerializeContext* Context = nullptr;
 };
 
-/**
- * Sets the index of the last segment of SerializedPropertyPath for the lifetime of this object if path tracking is active.
- *
- * Resets the index to INDEX_NONE when destructed.
- */
-class FSerializedPropertyPathIndexScope
-{
-public:
-	[[nodiscard]] FSerializedPropertyPathIndexScope(FUObjectSerializeContext* InContext, int32 InIndex, ESerializedPropertyPathNotify InNotify)
-	{
-		if (InContext->bTrackSerializedPropertyPath)
-		{
-			Context = InContext;
-			Context->SerializedPropertyPath.SetIndex(InIndex);
-		}
-	}
-
-	~FSerializedPropertyPathIndexScope()
-	{
-		if (Context)
-		{
-			Context->SerializedPropertyPath.SetIndex(INDEX_NONE);
-		}
-	}
-
-	FSerializedPropertyPathIndexScope(const FSerializedPropertyPathIndexScope&) = delete;
-	FSerializedPropertyPathIndexScope& operator=(const FSerializedPropertyPathIndexScope&) = delete;
-
-private:
-	FUObjectSerializeContext* Context = nullptr;
-};
-
 } // UE
diff --git a/Engine/Source/Runtime/CoreUObject/Private/UObject/PropertyMap.cpp b/Engine/Source/Runtime/CoreUObject/Private/UObject/PropertyMap.cpp
--- a/Engine/Source/Runtime/CoreUObject/Private/UObject/PropertyMap.cpp
+++ b/Engine/Source/Runtime/CoreUObject/Private/UObject/PropertyMap.cpp
@@ -340,14 +340,13 @@ void FMapProperty::SerializeItem(FStructuredArchive::FSlot Slot, void* Value, co
 						FStructuredArchive::FRecord EntryRecord = ReplacedArray.EnterElement().EnterRecord();
 						int32 Index = MapHelper.AddDefaultValue_Invalid_NeedsRehash();
 						uint8* PairPtr = MapHelper.GetPairPtr(Index);
-						UE::FSerializedPropertyPathIndexScope SerializedPropertyPathIndex(Context, Index, UE::ESerializedPropertyPathNotify::Yes);
 						{
-							UE::FSerializedPropertyPathScope SerializedPropertyPath(Context, {NAME_MapKey}, UE::ESerializedPropertyPathNotify::No);
+							UE::FSerializedPropertyPathScope SerializedPropertyPath(Context, {NAME_MapKey});
 							FSerializedPropertyScope SerializedProperty(UnderlyingArchive, KeyProp, this);
 							KeyProp->SerializeItem(EntryRecord.EnterField(TEXT("Key")), PairPtr);
 						}
 						{
-							UE::FSerializedPropertyPathScope SerializedPropertyPath(Context, {NAME_MapValue}, UE::ESerializedPropertyPathNotify::No);
+							UE::FSerializedPropertyPathScope SerializedPropertyPath(Context, {NAME_MapValue});
 							FSerializedPropertyScope SerializedProperty(UnderlyingArchive, ValueProp, this);
 							ValueProp->SerializeItem(EntryRecord.EnterField(TEXT("Value")), PairPtr + MapLayout.ValueOffset);
 						}
@@ -383,7 +382,7 @@ void FMapProperty::SerializeItem(FStructuredArchive::FSlot Slot, void* Value, co
 						for (int32 i = 0; i < NumRemoved; ++i)
 						{
 							{
-								UE::FSerializedPropertyPathScope SerializedPropertyPath(Context, {NAME_MapKey}, UE::ESerializedPropertyPathNotify::No);
+								UE::FSerializedPropertyPathScope SerializedPropertyPath(Context, {NAME_MapKey});
 								FSerializedPropertyScope SerializedProperty(UnderlyingArchive, KeyProp, this);
 								KeyProp->SerializeItem(RemovedArray.EnterElement().EnterRecord().EnterField(TEXT("Key")), TempKeyValueStorage);
 							}
@@ -416,7 +415,7 @@ void FMapProperty::SerializeItem(FStructuredArchive::FSlot Slot, void* Value, co
 
 							// Read key into temporary storage
 							{
-								UE::FSerializedPropertyPathScope SerializedPropertyPath(Context, {NAME_MapKey}, UE::ESerializedPropertyPathNotify::No);
+								UE::FSerializedPropertyPathScope SerializedPropertyPath(Context, {NAME_MapKey});
 								FSerializedPropertyScope SerializedProperty(UnderlyingArchive, KeyProp, this);
 								KeyProp->SerializeItem(EntryRecord.EnterField(TEXT("Key")), TempKeyValueStorage);
 							}
@@ -426,7 +425,7 @@ void FMapProperty::SerializeItem(FStructuredArchive::FSlot Slot, void* Value, co
 
 							// Deserialize value into hash map-owned memory
 							{
-								UE::FSerializedPropertyPathScope SerializedPropertyPath(Context, {NAME_MapValue}, UE::ESerializedPropertyPathNotify::No);
+								UE::FSerializedPropertyPathScope SerializedPropertyPath(Context, {NAME_MapValue});
 								FSerializedPropertyScope SerializedProperty(UnderlyingArchive, ValueProp, this);
 								ValueProp->SerializeItem(EntryRecord.EnterField(TEXT("Value")), ValuePtr);
 							}
@@ -450,7 +449,7 @@ void FMapProperty::SerializeItem(FStructuredArchive::FSlot Slot, void* Value, co
 
 							// Read key into temporary storage
 							{
-								UE::FSerializedPropertyPathScope SerializedPropertyPath(Context, {NAME_MapKey}, UE::ESerializedPropertyPathNotify::No);
+								UE::FSerializedPropertyPathScope SerializedPropertyPath(Context, {NAME_MapKey});
 								FSerializedPropertyScope SerializedProperty(UnderlyingArchive, KeyProp, this);
 								KeyProp->SerializeItem(EntryRecord.EnterField(TEXT("Key")), TempKeyValueStorage);
 							}
@@ -459,7 +458,7 @@ void FMapProperty::SerializeItem(FStructuredArchive::FSlot Slot, void* Value, co
 
 							// Deserialize value into hash map-owned memory
 							{
-								UE::FSerializedPropertyPathScope SerializedPropertyPath(Context, {NAME_MapValue}, UE::ESerializedPropertyPathNotify::No);
+								UE::FSerializedPropertyPathScope SerializedPropertyPath(Context, {NAME_MapValue});
 								FSerializedPropertyScope SerializedProperty(UnderlyingArchive, ValueProp, this);
 								ValueProp->SerializeItem(EntryRecord.EnterField(TEXT("Value")), ValuePtr);
 							}
@@ -610,12 +609,12 @@ void FMapProperty::SerializeItem(FStructuredArchive::FSlot Slot, void* Value, co
 						FStructuredArchive::FRecord EntryRecord = ReplacedArray.EnterElement().EnterRecord();
 						uint8* PairPtr = MapHelper.GetPairPtr(It.GetInternalIndex());
 						{
-							UE::FSerializedPropertyPathScope SerializedPropertyPath(Context, {NAME_MapKey}, UE::ESerializedPropertyPathNotify::No);
+							UE::FSerializedPropertyPathScope SerializedPropertyPath(Context, {NAME_MapKey});
 							FSerializedPropertyScope SerializedProperty(UnderlyingArchive, KeyProp, this);
 							KeyProp->SerializeItem(EntryRecord.EnterField(TEXT("Key")), PairPtr);
 						}
 						{
-							UE::FSerializedPropertyPathScope SerializedPropertyPath(Context, {NAME_MapValue}, UE::ESerializedPropertyPathNotify::No);
+							UE::FSerializedPropertyPathScope SerializedPropertyPath(Context, {NAME_MapValue});
 							FSerializedPropertyScope SerializedProperty(UnderlyingArchive, ValueProp, this);
 							ValueProp->SerializeItem(EntryRecord.EnterField(TEXT("Value")), PairPtr + MapLayout.ValueOffset);
 						}
@@ -631,7 +630,7 @@ void FMapProperty::SerializeItem(FStructuredArchive::FSlot Slot, void* Value, co
 					for (int32 i : RemovedIndices)
 					{
 						FStructuredArchive::FRecord EntryRecord = RemovedArray.EnterElement().EnterRecord();
-						UE::FSerializedPropertyPathScope SerializedPropertyPath(Context, {NAME_MapKey}, UE::ESerializedPropertyPathNotify::No);
+						UE::FSerializedPropertyPathScope SerializedPropertyPath(Context, {NAME_MapKey});
 						FSerializedPropertyScope SerializedProperty(UnderlyingArchive, KeyProp, this);
 						KeyProp->SerializeItem(EntryRecord.EnterField(TEXT("Key")), DefaultsMapHelper.GetKeyPtr(i));
 					}
@@ -643,12 +642,12 @@ void FMapProperty::SerializeItem(FStructuredArchive::FSlot Slot, void* Value, co
 						FStructuredArchive::FRecord EntryRecord = ModifiedArray.EnterElement().EnterRecord();
 						uint8* PairPtr = MapHelper.GetPairPtr(i);
 						{
-							UE::FSerializedPropertyPathScope SerializedPropertyPath(Context, {NAME_MapKey}, UE::ESerializedPropertyPathNotify::No);
+							UE::FSerializedPropertyPathScope SerializedPropertyPath(Context, {NAME_MapKey});
 							FSerializedPropertyScope SerializedProperty(UnderlyingArchive, KeyProp, this);
 							KeyProp->SerializeItem(EntryRecord.EnterField(TEXT("Key")), PairPtr);
 						}
 						{
-							UE::FSerializedPropertyPathScope SerializedPropertyPath(Context, {NAME_MapValue}, UE::ESerializedPropertyPathNotify::No);
+							UE::FSerializedPropertyPathScope SerializedPropertyPath(Context, {NAME_MapValue});
 							FSerializedPropertyScope SerializedProperty(UnderlyingArchive, ValueProp, this);
 							ValueProp->SerializeItem(EntryRecord.EnterField(TEXT("Value")), PairPtr + MapLayout.ValueOffset);
 						}
@@ -662,12 +661,12 @@ void FMapProperty::SerializeItem(FStructuredArchive::FSlot Slot, void* Value, co
 						FStructuredArchive::FRecord EntryRecord = AddedArray.EnterElement().EnterRecord();
 						uint8* PairPtr = MapHelper.GetPairPtr(i);
 						{
-							UE::FSerializedPropertyPathScope SerializedPropertyPath(Context, {NAME_MapKey}, UE::ESerializedPropertyPathNotify::No);
+							UE::FSerializedPropertyPathScope SerializedPropertyPath(Context, {NAME_MapKey});
 							FSerializedPropertyScope SerializedProperty(UnderlyingArchive, KeyProp, this);
 							KeyProp->SerializeItem(EntryRecord.EnterField(TEXT("Key")), PairPtr);
 						}
 						{
-							UE::FSerializedPropertyPathScope SerializedPropertyPath(Context, {NAME_MapValue}, UE::ESerializedPropertyPathNotify::No);
+							UE::FSerializedPropertyPathScope SerializedPropertyPath(Context, {NAME_MapValue});
 							FSerializedPropertyScope SerializedProperty(UnderlyingArchive, ValueProp, this);
 							ValueProp->SerializeItem(EntryRecord.EnterField(TEXT("Value")), PairPtr + MapLayout.ValueOffset);
 						}
@@ -720,14 +719,13 @@ void FMapProperty::SerializeItem(FStructuredArchive::FSlot Slot, void* Value, co
 			{
 				FStructuredArchive::FRecord EntryRecord = EntriesArray.EnterElement().EnterRecord();
 				int32 Index = MapHelper.AddDefaultValue_Invalid_NeedsRehash();
-				UE::FSerializedPropertyPathIndexScope SerializedPropertyPathIndex(Context, Index, UE::ESerializedPropertyPathNotify::Yes);
 				{
-					UE::FSerializedPropertyPathScope SerializedPropertyPath(Context, {NAME_MapKey}, UE::ESerializedPropertyPathNotify::No);
+					UE::FSerializedPropertyPathScope SerializedPropertyPath(Context, {NAME_MapKey});
 					FSerializedPropertyScope SerializedProperty(UnderlyingArchive, KeyProp, this);
 					KeyProp->SerializeItem(EntryRecord.EnterField(TEXT("Key")), MapHelper.GetKeyPtr(Index));
 				}
 				{
-					UE::FSerializedPropertyPathScope SerializedPropertyPath(Context, {NAME_MapValue}, UE::ESerializedPropertyPathNotify::No);
+					UE::FSerializedPropertyPathScope SerializedPropertyPath(Context, {NAME_MapValue});
 					FSerializedPropertyScope SerializedProperty(UnderlyingArchive, ValueProp, this);
 					ValueProp->SerializeItem(EntryRecord.EnterField(TEXT("Value")), MapHelper.GetValuePtr(Index));
 				}
@@ -755,7 +753,7 @@ void FMapProperty::SerializeItem(FStructuredArchive::FSlot Slot, void* Value, co
 				KeyProp->InitializeValue(TempKeyValueStorage);
 
 				
-				UE::FSerializedPropertyPathScope SerializedPropertyPath(Context, {NAME_MapKey}, UE::ESerializedPropertyPathNotify::No);
+				UE::FSerializedPropertyPathScope SerializedPropertyPath(Context, {NAME_MapKey});
 				FSerializedPropertyScope SerializedProperty(UnderlyingArchive, KeyProp, this);
 				for (; NumKeysToRemove; --NumKeysToRemove)
 				{
@@ -787,7 +785,7 @@ void FMapProperty::SerializeItem(FStructuredArchive::FSlot Slot, void* Value, co
 
 				// Read key into temporary storage
 				{
-					UE::FSerializedPropertyPathScope SerializedPropertyPath(Context, {NAME_MapKey}, UE::ESerializedPropertyPathNotify::No);
+					UE::FSerializedPropertyPathScope SerializedPropertyPath(Context, {NAME_MapKey});
 					FSerializedPropertyScope SerializedProperty(UnderlyingArchive, KeyProp, this);
 					KeyProp->SerializeItem(EntryRecord.EnterField(TEXT("Key")), TempKeyValueStorage);
 				}
@@ -796,7 +794,7 @@ void FMapProperty::SerializeItem(FStructuredArchive::FSlot Slot, void* Value, co
 
 				// Deserialize value into hash map-owned memory
 				{
-					UE::FSerializedPropertyPathScope SerializedPropertyPath(Context, {NAME_MapValue}, UE::ESerializedPropertyPathNotify::No);
+					UE::FSerializedPropertyPathScope SerializedPropertyPath(Context, {NAME_MapValue});
 					FSerializedPropertyScope SerializedProperty(UnderlyingArchive, ValueProp, this);
 					ValueProp->SerializeItem(EntryRecord.EnterField(TEXT("Value")), ValuePtr);
 				}
@@ -861,12 +859,12 @@ void FMapProperty::SerializeItem(FStructuredArchive::FSlot Slot, void* Value, co
 				FStructuredArchive::FRecord EntryRecord = EntriesArray.EnterElement().EnterRecord();
 
 				{
-					UE::FSerializedPropertyPathScope SerializedPropertyPath(Context, {NAME_MapKey}, UE::ESerializedPropertyPathNotify::No);
+					UE::FSerializedPropertyPathScope SerializedPropertyPath(Context, {NAME_MapKey});
 					FSerializedPropertyScope SerializedProperty(UnderlyingArchive, KeyProp, this);
 					KeyProp->SerializeItem(EntryRecord.EnterField(TEXT("Key")), ValuePairPtr);
 				}
 				{
-					UE::FSerializedPropertyPathScope SerializedPropertyPath(Context, {NAME_MapValue}, UE::ESerializedPropertyPathNotify::No);
+					UE::FSerializedPropertyPathScope SerializedPropertyPath(Context, {NAME_MapValue});
 					FSerializedPropertyScope SerializedProperty(UnderlyingArchive, ValueProp, this);
 					ValueProp->SerializeItem(EntryRecord.EnterField(TEXT("Value")), ValuePairPtr + MapLayout.ValueOffset);
 				}
@@ -883,12 +881,12 @@ void FMapProperty::SerializeItem(FStructuredArchive::FSlot Slot, void* Value, co
 
 				uint8* ValuePairPtr = MapHelper.GetPairPtr(Iterator);
 				{
-					UE::FSerializedPropertyPathScope SerializedPropertyPath(Context, {NAME_MapKey}, UE::ESerializedPropertyPathNotify::No);
+					UE::FSerializedPropertyPathScope SerializedPropertyPath(Context, {NAME_MapKey});
 					FSerializedPropertyScope SerializedProperty(UnderlyingArchive, KeyProp, this);
 					KeyProp->SerializeItem(EntryRecord.EnterField(TEXT("Key")), ValuePairPtr);
 				}
 				{
-					UE::FSerializedPropertyPathScope SerializedPropertyPath(Context, {NAME_MapValue}, UE::ESerializedPropertyPathNotify::No);
+					UE::FSerializedPropertyPathScope SerializedPropertyPath(Context, {NAME_MapValue});
 					FSerializedPropertyScope SerializedProperty(UnderlyingArchive, ValueProp, this);
 					ValueProp->SerializeItem(EntryRecord.EnterField(TEXT("Value")), ValuePairPtr + MapLayout.ValueOffset);
 				}
@@ -1500,7 +1498,7 @@ EConvertFromTypeResult FMapProperty::ConvertFromType(const FPropertyTag& Tag, FS
 
 	const auto SerializeOrConvert = [Context = FUObjectThreadContext::Get().GetSerializeContext()](bool bCanSerialize, FProperty* Inner, const FPropertyTag& InnerTag, FName InnerName, FStructuredArchive::FSlot InnerSlot, uint8* InnerData, UStruct* InnerDefaultsStruct) -> bool
 	{
-		UE::FSerializedPropertyPathScope SerializedPropertyPath(Context, {InnerName}, UE::ESerializedPropertyPathNotify::No);
+		UE::FSerializedPropertyPathScope SerializedPropertyPath(Context, {InnerName});
 
 		if (!bCanSerialize)
 		{
diff --git a/Engine/Source/Runtime/CoreUObject/Private/UObject/PropertySet.cpp b/Engine/Source/Runtime/CoreUObject/Private/UObject/PropertySet.cpp
--- a/Engine/Source/Runtime/CoreUObject/Private/UObject/PropertySet.cpp
+++ b/Engine/Source/Runtime/CoreUObject/Private/UObject/PropertySet.cpp
@@ -286,7 +286,6 @@ void FSetProperty::SerializeItem(FStructuredArchive::FSlot Slot, void* Value, co
 			for (; Num; --Num)
 			{
 				int32 Index = SetHelper.AddDefaultValue_Invalid_NeedsRehash();
-				UE::FSerializedPropertyPathIndexScope SerializedPropertyPathIndex(Context, Index, UE::ESerializedPropertyPathNotify::Yes);
 				ElementProp->SerializeItem(ElementsArray.EnterElement(), SetHelper.GetElementPtrWithoutCheck(Index));
 			}
 		}
