Serialize unknown properties for removed map and set keys

Unknown properties need to be considered in these contexts to make the hashes of removed structs to match the hashes of the structs from the defaults.

#rnx

diff --git a/Engine/Source/Runtime/CoreUObject/Private/UObject/PropertyMap.cpp b/Engine/Source/Runtime/CoreUObject/Private/UObject/PropertyMap.cpp
--- a/Engine/Source/Runtime/CoreUObject/Private/UObject/PropertyMap.cpp
+++ b/Engine/Source/Runtime/CoreUObject/Private/UObject/PropertyMap.cpp
@@ -694,8 +694,6 @@ void FMapProperty::SerializeItem(FStructuredArchive::FSlot Slot, void* Value, co
 		{
 			if (NumKeysToRemove && !bReplaceMap)
 			{
-				TGuardValue<bool> SerializeUnknownProperty(Context->bSerializeUnknownProperty, false);
-
 				// Load and discard keys to remove, map is empty
 				void* TempKeyValueStorage = FMemory::Malloc(MapLayout.SetLayout.Size);
 				KeyProp->InitializeValue(TempKeyValueStorage);
@@ -747,12 +745,9 @@ void FMapProperty::SerializeItem(FStructuredArchive::FSlot Slot, void* Value, co
 
 			if (NumKeysToRemove)
 			{
-				TGuardValue<bool> SerializeUnknownProperty(Context->bSerializeUnknownProperty, false);
-
 				TempKeyValueStorage = (uint8*)FMemory::Malloc(MapLayout.SetLayout.Size);
 				KeyProp->InitializeValue(TempKeyValueStorage);
 
-				
 				UE::FSerializedPropertyPathScope SerializedPropertyPath(Context, {NAME_MapKey});
 				FSerializedPropertyScope SerializedProperty(UnderlyingArchive, KeyProp, this);
 				for (; NumKeysToRemove; --NumKeysToRemove)
diff --git a/Engine/Source/Runtime/CoreUObject/Private/UObject/PropertySet.cpp b/Engine/Source/Runtime/CoreUObject/Private/UObject/PropertySet.cpp
--- a/Engine/Source/Runtime/CoreUObject/Private/UObject/PropertySet.cpp
+++ b/Engine/Source/Runtime/CoreUObject/Private/UObject/PropertySet.cpp
@@ -262,8 +262,6 @@ void FSetProperty::SerializeItem(FStructuredArchive::FSlot Slot, void* Value, co
 		{
 			if (NumElementsToRemove)
 			{
-				TGuardValue<bool> SerializeUnknownProperty(Context->bSerializeUnknownProperty, false);
-
 				// Load and discard elements to remove, set is empty
 				void* TempElementStorage = FMemory::Malloc(SetLayout.Size);
 				ElementProp->InitializeValue(TempElementStorage);
@@ -303,8 +301,6 @@ void FSetProperty::SerializeItem(FStructuredArchive::FSlot Slot, void* Value, co
 
 			if (NumElementsToRemove)
 			{
-				TGuardValue<bool> SerializeUnknownProperty(Context->bSerializeUnknownProperty, false);
-
 				TempElementStorage = (uint8*)FMemory::Malloc(SetLayout.Size);
 				ElementProp->InitializeValue(TempElementStorage);
 
