Renamed bSerializeUnknownProperty to bTrackUnknownProperties

#rnx

diff --git a/Engine/Plugins/Experimental/InstanceDataObjectFixupTool/Source/Private/Tests/InstanceDataObjectTests.cpp b/Engine/Plugins/Experimental/InstanceDataObjectFixupTool/Source/Private/Tests/InstanceDataObjectTests.cpp
--- a/Engine/Plugins/Experimental/InstanceDataObjectFixupTool/Source/Private/Tests/InstanceDataObjectTests.cpp
+++ b/Engine/Plugins/Experimental/InstanceDataObjectFixupTool/Source/Private/Tests/InstanceDataObjectTests.cpp
@@ -332,9 +332,9 @@ static UObject* GenerateTestInstanceDataObject(UObject* ObjectOld, UClass* NewCl
 	{
 		FUObjectSerializeContext* LoadContext = FUObjectThreadContext::Get().GetSerializeContext();
 
-		// scoped flag changes to load context to support property bag generation
+		// scoped flag changes to load context to support unknown property tracking
 		TGuardValue<bool> ScopedTrackSerializedPropertyPath(LoadContext->bTrackSerializedPropertyPath, true);
-		TGuardValue<bool> ScopedSerializeUnknownProperty(LoadContext->bSerializeUnknownProperty, true);
+		TGuardValue<bool> ScopedSerializeUnknownProperty(LoadContext->bTrackUnknownProperties, true);
 		TGuardValue<bool> ScopedImpersonateProperties(LoadContext->bImpersonateProperties, true);
 		
 		UPackage* TempPackage = CreateTestPackage(NewClass);
diff --git a/Engine/Source/Runtime/CoreUObject/Private/UObject/Class.cpp b/Engine/Source/Runtime/CoreUObject/Private/UObject/Class.cpp
--- a/Engine/Source/Runtime/CoreUObject/Private/UObject/Class.cpp
+++ b/Engine/Source/Runtime/CoreUObject/Private/UObject/Class.cpp
@@ -1741,7 +1741,7 @@ void UStruct::SerializeVersionedTaggedProperties(FStructuredArchive::FSlot Slot,
 					if (!bSearchedForUnknownPropertyTree)
 					{
 						bSearchedForUnknownPropertyTree = true;
-						if (SerializeContext->bSerializeUnknownProperty)
+						if (SerializeContext->bTrackUnknownProperties)
 						{
 							if (UObject* Object = SerializeContext->SerializedObject)
 							{
diff --git a/Engine/Source/Runtime/CoreUObject/Private/UObject/LinkerLoad.cpp b/Engine/Source/Runtime/CoreUObject/Private/UObject/LinkerLoad.cpp
--- a/Engine/Source/Runtime/CoreUObject/Private/UObject/LinkerLoad.cpp
+++ b/Engine/Source/Runtime/CoreUObject/Private/UObject/LinkerLoad.cpp
@@ -4692,13 +4692,13 @@ void FLinkerLoad::Preload( UObject* Object )
 
 					// Toggle support for IDOs
 					FUObjectSerializeContext* LoadContext = FUObjectThreadContext::Get().GetSerializeContext();
-					// Enable property path tracking when IDO support is enabled. Both property bags and IDO creation require the paths.
+					// Enable property path tracking when IDO support is enabled. Both the property path name tree and IDO creation require the paths.
 					const bool bHasIDOSupport = UE::IsInstanceDataObjectSupportEnabled(Object);
 					TGuardValue<bool> ScopedTrackSerializedPropertyPath(LoadContext->bTrackSerializedPropertyPath, bHasIDOSupport);
-					// Enable creation of a property bag to hold any property that does not match the current class schema,
+					// Enable creation of a property path name tree to track any property that does not match the current class schema,
 					// except when impersonation is enabled because that implies we are deserializing an IDO.
 					const bool bCreateIDO = bHasIDOSupport && !LoadContext->bImpersonateProperties;
-					TGuardValue<bool> ScopedSerializeUnknownProperty(LoadContext->bSerializeUnknownProperty, bCreateIDO);
+					TGuardValue<bool> ScopedTrackUnknownProperties(LoadContext->bTrackUnknownProperties, bCreateIDO);
 
 					if (Object->HasAnyFlags(RF_ClassDefaultObject))
 					{
diff --git a/Engine/Source/Runtime/CoreUObject/Private/UObject/UObjectThreadContext.cpp b/Engine/Source/Runtime/CoreUObject/Private/UObject/UObjectThreadContext.cpp
--- a/Engine/Source/Runtime/CoreUObject/Private/UObject/UObjectThreadContext.cpp
+++ b/Engine/Source/Runtime/CoreUObject/Private/UObject/UObjectThreadContext.cpp
@@ -45,7 +45,7 @@ FUObjectSerializeContext::FUObjectSerializeContext()
 	, SerializedExportIndex(0)
 	, SerializedExportLinker(nullptr)
 	, bTrackSerializedPropertyPath(false)
-	, bSerializeUnknownProperty(false)
+	, bTrackUnknownProperties(false)
 	, bImpersonateProperties(false)
 {}
 
diff --git a/Engine/Source/Runtime/CoreUObject/Public/UObject/UObjectThreadContext.h b/Engine/Source/Runtime/CoreUObject/Public/UObject/UObjectThreadContext.h
--- a/Engine/Source/Runtime/CoreUObject/Public/UObject/UObjectThreadContext.h
+++ b/Engine/Source/Runtime/CoreUObject/Public/UObject/UObjectThreadContext.h
@@ -172,8 +172,8 @@ public:
 	UE_INTERNAL UE::FPropertyPathName SerializedPropertyPath;
 	/** True when SerializedPropertyPath is being tracked during serialization. */
 	UE_INTERNAL bool bTrackSerializedPropertyPath;
-	/** True when unknown properties will be serialized to or from a property bag for the serialized object. */
-	UE_INTERNAL bool bSerializeUnknownProperty;
+	/** True when unknown properties will tracked for the serialized object. */
+	UE_INTERNAL bool bTrackUnknownProperties;
 	/** True when the SerializedObject properties are being impersonated. */
 	UE_INTERNAL bool bImpersonateProperties;
 
diff --git a/Engine/Source/Runtime/Engine/Private/UnrealEngine.cpp b/Engine/Source/Runtime/Engine/Private/UnrealEngine.cpp
--- a/Engine/Source/Runtime/Engine/Private/UnrealEngine.cpp
+++ b/Engine/Source/Runtime/Engine/Private/UnrealEngine.cpp
@@ -17240,7 +17240,7 @@ public:
 		{
 			// this will had property path tracking and create a property bag to hold data not matching the current class schema
 			ScopedTrackSerializedPropertyPath.Emplace(LoadContext->bTrackSerializedPropertyPath, true);
-			ScopedSerializeUnknownProperty.Emplace(LoadContext->bSerializeUnknownProperty, true);
+			ScopedSerializeUnknownProperty.Emplace(LoadContext->bTrackUnknownProperties, true);
 			ScopedSerializedObject.Emplace(LoadContext->SerializedObject, DstObject);
 		}
 
