[ALT]
 - Fix TSAN warnings when FGCObject is destroyed from a worker thread while holding a gc scope guard

diff --git a/Engine/Source/Runtime/CoreUObject/Private/Misc/GCObjectReferencer.cpp b/Engine/Source/Runtime/CoreUObject/Private/Misc/GCObjectReferencer.cpp
--- a/Engine/Source/Runtime/CoreUObject/Private/Misc/GCObjectReferencer.cpp
+++ b/Engine/Source/Runtime/CoreUObject/Private/Misc/GCObjectReferencer.cpp
@@ -13,9 +13,11 @@
 #include "UObject/GarbageCollectionVerification.h"
 #include "HAL/IConsoleManager.h"
 
+#include <atomic>
+
 // Global GC state flags
-extern bool GObjIncrementalPurgeIsInProgress;
-extern bool GObjUnhashUnreachableIsInProgress;
+extern std::atomic<bool> GObjIncrementalPurgeIsInProgress;
+extern std::atomic<bool> GObjUnhashUnreachableIsInProgress;
 
 namespace UE::GC
 {
diff --git a/Engine/Source/Runtime/CoreUObject/Private/UObject/GarbageCollection.cpp b/Engine/Source/Runtime/CoreUObject/Private/UObject/GarbageCollection.cpp
--- a/Engine/Source/Runtime/CoreUObject/Private/UObject/GarbageCollection.cpp
+++ b/Engine/Source/Runtime/CoreUObject/Private/UObject/GarbageCollection.cpp
@@ -75,11 +75,11 @@ extern COREUOBJECT_API float GVerifyGCAssumptionsChance;
 /** Object count during last mark phase																				*/
 FThreadSafeCounter		GObjectCountDuringLastMarkPhase;
 /** Whether UObject hash tables are locked by GC */
-bool GIsGarbageCollectingAndLockingUObjectHashTables = false;
+std::atomic<bool> GIsGarbageCollectingAndLockingUObjectHashTables = false;
 /** Whether incremental object purge is in progress										*/
-bool GObjIncrementalPurgeIsInProgress = false;
+std::atomic<bool> GObjIncrementalPurgeIsInProgress = false;
 /** Whether GC is currently routing BeginDestroy to objects										*/
-bool GObjUnhashUnreachableIsInProgress = false;
+std::atomic<bool> GObjUnhashUnreachableIsInProgress = false;
 /** Time the GC started, needs to be reset on return from being in the background on some OSs */
 double GCStartTime = 0.;
 /** Whether FinishDestroy has already been routed to all unreachable objects. */
@@ -92,7 +92,7 @@ static TArray<UObject *> GGCObjectsPendingDestruction;
 /** Number of objects actually still pending destruction */
 static int32 GGCObjectsPendingDestructionCount = 0;
 /** Whether we need to purge objects or not.											*/
-static bool GObjPurgeIsRequired = false;
+static std::atomic<bool> GObjPurgeIsRequired = false;
 /** Current object index for incremental purge.											*/
 static int32 GObjCurrentPurgeObjectIndex = 0;
 /** Current object index for incremental purge.											*/
@@ -4625,14 +4625,12 @@ void IncrementalPurgeGarbage(bool bUseTimeLimit, double TimeLimit)
 		{
 			// Incremental purge is now in progress.
 			GObjIncrementalPurgeIsInProgress = true;
-			FPlatformMisc::MemoryBarrier();
 		}
 		~FResetPurgeProgress()
 		{
 			if (bCompletedRef)
 			{
 				GObjIncrementalPurgeIsInProgress = false;
-				FPlatformMisc::MemoryBarrier();
 			}
 		}
 
@@ -6028,7 +6026,7 @@ bool UnhashUnreachableObjects(bool bUseTimeLimit, double TimeLimit)
 	TRACE_CPUPROFILER_EVENT_SCOPE(UnhashUnreachableObjects);
 	DECLARE_SCOPE_CYCLE_COUNTER(TEXT("UnhashUnreachableObjects"), STAT_UnhashUnreachableObjects, STATGROUP_GC);
 
-	TGuardValue<bool> GuardObjUnhashUnreachableIsInProgress(GObjUnhashUnreachableIsInProgress, true);
+	TGuardValue GuardObjUnhashUnreachableIsInProgress(GObjUnhashUnreachableIsInProgress, true);
 
 	{
 		TRACE_CPUPROFILER_EVENT_SCOPE(BroadcastGarbageCollectConditionalBeginDestroy);
