Cache ValuesSetBySerialization on the IDO class/struct to avoid an expensive lookup on every access

#rnx

diff --git a/Engine/Source/Runtime/CoreUObject/Private/UObject/InstanceDataObjectUtils.cpp b/Engine/Source/Runtime/CoreUObject/Private/UObject/InstanceDataObjectUtils.cpp
--- a/Engine/Source/Runtime/CoreUObject/Private/UObject/InstanceDataObjectUtils.cpp
+++ b/Engine/Source/Runtime/CoreUObject/Private/UObject/InstanceDataObjectUtils.cpp
@@ -7,6 +7,7 @@
 #include "UObject/Class.h"
 #include "UObject/EnumProperty.h"
 #include "UObject/Field.h"
+#include "UObject/Package.h"
 #include "UObject/PropertyBagRepository.h"
 #include "UObject/PropertyHelper.h"
 #include "UObject/PropertyOptional.h"
@@ -15,6 +16,19 @@
 
 static const FName NAME_ValuesSetBySerialization(ANSITEXTVIEW("_ValuesSetBySerialization"));
 
+/** Type used for InstanceDataObject classes. */
+class UInstanceDataObjectClass final : public UClass
+{
+public:
+	DECLARE_CASTED_CLASS_INTRINSIC(UInstanceDataObjectClass, UClass, CLASS_Transient, TEXT("/Script/CoreUObject"), CASTCLASS_UClass)
+
+	FByteProperty* ValuesSetBySerializationProperty = nullptr;
+};
+
+IMPLEMENT_CORE_INTRINSIC_CLASS(UInstanceDataObjectClass, UClass,
+{
+});
+
 /** Type used for InstanceDataObject structs to provide support for hashing and custom guids. */
 class UInstanceDataObjectStruct final : public UScriptStruct
 {
@@ -24,6 +38,7 @@ public:
 	uint32 GetStructTypeHash(const void* Src) const final;
 	FGuid GetCustomGuid() const final { return Guid; }
 
+	FByteProperty* ValuesSetBySerializationProperty = nullptr;
 	FGuid Guid;
 };
 
@@ -312,13 +327,14 @@ namespace UE
 		TSet<FPropertyPathName> SuperPropertyPathsFromTree;
 
 		// UClass is required to inherit from UObject
-		UStruct* Super = (StructClass == UClass::StaticClass()) ? UObject::StaticClass() : nullptr;
+		UStruct* Super = StructClass->IsChildOf<UClass>() ? UObject::StaticClass() : nullptr;
 
 		if (TemplateStruct)
 		{
 			{
 				const FName SuperName(WriteToString<128>(TemplateStruct->GetName(), TEXTVIEW("_Super")));
-				UStruct* NewSuper = NewObject<UStruct>(Outer, StructClass, MakeUniqueObjectName(nullptr, StructClass, SuperName));
+				const UClass* SuperStructClass = StructClass->GetSuperClass();
+				UStruct* NewSuper = NewObject<UStruct>(Outer, SuperStructClass, MakeUniqueObjectName(nullptr, SuperStructClass, SuperName));
 				NewSuper->SetSuperStruct(Super);
 				Super = NewSuper;
 			}
@@ -402,6 +418,16 @@ namespace UE
 			Result->AddCppProperty(ValuesSetBySerializationProperty);
 		}
 
+		// Store generated properties to avoid scanning every property to find it when it is needed.
+		if (UInstanceDataObjectClass* IdoClass = Cast<UInstanceDataObjectClass>(Result))
+		{
+			IdoClass->ValuesSetBySerializationProperty = ValuesSetBySerializationProperty;
+		}
+		else if (UInstanceDataObjectStruct* IdoStruct = Cast<UInstanceDataObjectStruct>(Result))
+		{
+			IdoStruct->ValuesSetBySerializationProperty = ValuesSetBySerializationProperty;
+		}
+
 		// AddCppProperty expects reverse property order for StaticLink to work correctly
 		for (FProperty* Property : ReverseIterate(LooseInstanceDataObjectProperties))
 		{
@@ -458,7 +484,7 @@ namespace UE
 
 	UClass* CreateInstanceDataObjectClass(const FPropertyPathNameTree* PropertyTree, UClass* OwnerClass, UObject* Outer)
 	{
-		UClass* Result = CreateInstanceDataObjectStructRec<UClass>(OwnerClass, Outer, PropertyTree);
+		UClass* Result = CreateInstanceDataObjectStructRec<UInstanceDataObjectClass>(OwnerClass, Outer, PropertyTree);
 #if WITH_EDITORONLY_DATA
 		if (const FString& DisplayName = OwnerClass->GetMetaData(NAME_DisplayName); !DisplayName.IsEmpty())
 		{
@@ -477,10 +503,23 @@ namespace UE
 		return Result;
 	}
 
+	static const FByteProperty* FindValuesSetBySerializationProperty(const UStruct* Struct)
+	{
+		if (const UInstanceDataObjectClass* IdoClass = Cast<UInstanceDataObjectClass>(Struct))
+		{
+			return IdoClass->ValuesSetBySerializationProperty;
+		}
+		if (const UInstanceDataObjectStruct* IdoStruct = Cast<UInstanceDataObjectStruct>(Struct))
+		{
+			return IdoStruct->ValuesSetBySerializationProperty;
+		}
+		return CastField<FByteProperty>(Struct->FindPropertyByName(NAME_ValuesSetBySerialization));
+	}
+
 	void MarkPropertySetBySerialization(const UStruct* Struct, void* StructData, const FProperty* Property, int32 ArrayIndex)
 	{
 	#if WITH_EDITORONLY_DATA
-		if (const FByteProperty* ValuesSetBySerializationProperty = CastField<FByteProperty>(Struct->FindPropertyByName(NAME_ValuesSetBySerialization)))
+		if (const FByteProperty* ValuesSetBySerializationProperty = FindValuesSetBySerializationProperty(Struct))
 		{
 			const int32 PropertyIndex = Property->GetIndexInOwner() + ArrayIndex;
 			const int32 ByteIndex = PropertyIndex / 8;
@@ -497,7 +536,7 @@ namespace UE
 	bool WasPropertySetBySerialization(const UStruct* Struct, const void* StructData, const FProperty* Property, int32 ArrayIndex)
 	{
 	#if WITH_EDITORONLY_DATA
-		if (const FByteProperty* ValuesSetBySerializationProperty = CastField<FByteProperty>(Struct->FindPropertyByName(NAME_ValuesSetBySerialization)))
+		if (const FByteProperty* ValuesSetBySerializationProperty = FindValuesSetBySerializationProperty(Struct))
 		{
 			const int32 PropertyIndex = Property->GetIndexInOwner() + ArrayIndex;
 			const int32 ByteIndex = PropertyIndex / 8;
