Added tests for InstanceDataObject tracking of values set by serialization

#rnx

diff --git a/Engine/Source/Runtime/CoreUObject/Tests/UObject/InstanceDataObjectUtilsTest.cpp b/Engine/Source/Runtime/CoreUObject/Tests/UObject/InstanceDataObjectUtilsTest.cpp
new file mode 100644
--- /dev/null
+++ b/Engine/Source/Runtime/CoreUObject/Tests/UObject/InstanceDataObjectUtilsTest.cpp
@@ -0,0 +1,54 @@
+// Copyright Epic Games, Inc. All Rights Reserved.
+
+#include "InstanceDataObjectUtilsTest.h"
+
+#if WITH_TESTS
+
+#include "Tests/TestHarnessAdapter.h"
+#include "UObject/InstanceDataObjectUtils.h"
+#include "UObject/Package.h"
+
+namespace UE
+{
+
+TEST_CASE_NAMED(FInstanceDataObjectUtilsTest, "CoreUObject::InstanceDataObjectUtils", "[Core][UObject][EngineFilter]")
+{
+	UTestInstanceDataObjectClass* BaseObject = NewObject<UTestInstanceDataObjectClass>();
+	UClass* TestClass = CreateInstanceDataObjectClass(nullptr, BaseObject->GetClass(), BaseObject->GetOuter());
+
+	FIntProperty* Int32Property = FindFProperty<FIntProperty>(TestClass, TEXT("Int32"));
+	FStructProperty* StructProperty = FindFProperty<FStructProperty>(TestClass, TEXT("Struct"));
+	REQUIRE(Int32Property);
+	REQUIRE(StructProperty);
+	REQUIRE(StructProperty->Struct);
+
+	FIntProperty* AProperty = FindFProperty<FIntProperty>(StructProperty->Struct, TEXT("A"));
+	FIntProperty* BProperty = FindFProperty<FIntProperty>(StructProperty->Struct, TEXT("B"));
+	FIntProperty* CProperty = FindFProperty<FIntProperty>(StructProperty->Struct, TEXT("C"));
+	FIntProperty* DProperty = FindFProperty<FIntProperty>(StructProperty->Struct, TEXT("D"));
+	REQUIRE(AProperty);
+	REQUIRE(BProperty);
+	REQUIRE(CProperty);
+	REQUIRE(DProperty);
+
+	UObject* TestObject = NewObject<UObject>(GetTransientPackage(), TestClass);
+	void* StructData = StructProperty->ContainerPtrToValuePtr<void>(TestObject);
+
+	CHECK_FALSE(WasPropertySetBySerialization(TestClass, TestObject, StructProperty));
+	MarkPropertySetBySerialization(TestClass, TestObject, StructProperty);
+	CHECK(WasPropertySetBySerialization(TestClass, TestObject, StructProperty));
+	CHECK_FALSE(WasPropertySetBySerialization(TestClass, TestObject, Int32Property));
+
+	CHECK_FALSE(WasPropertySetBySerialization(StructProperty->Struct, StructData, AProperty));
+	MarkPropertySetBySerialization(StructProperty->Struct, StructData, AProperty);
+	CHECK(WasPropertySetBySerialization(StructProperty->Struct, StructData, AProperty));
+	CHECK_FALSE(WasPropertySetBySerialization(StructProperty->Struct, StructData, BProperty));
+	MarkPropertySetBySerialization(StructProperty->Struct, StructData, BProperty);
+	CHECK(WasPropertySetBySerialization(StructProperty->Struct, StructData, BProperty));
+	CHECK_FALSE(WasPropertySetBySerialization(StructProperty->Struct, StructData, CProperty));
+	CHECK_FALSE(WasPropertySetBySerialization(StructProperty->Struct, StructData, DProperty));
+}
+
+} // UE
+
+#endif // WITH_TESTS
diff --git a/Engine/Source/Runtime/CoreUObject/Tests/UObject/InstanceDataObjectUtilsTest.h b/Engine/Source/Runtime/CoreUObject/Tests/UObject/InstanceDataObjectUtilsTest.h
new file mode 100644
--- /dev/null
+++ b/Engine/Source/Runtime/CoreUObject/Tests/UObject/InstanceDataObjectUtilsTest.h
@@ -0,0 +1,39 @@
+// Copyright Epic Games, Inc. All Rights Reserved.
+
+#pragma once
+
+#include "CoreTypes.h"
+#include "UObject/Object.h"
+#include "InstanceDataObjectUtilsTest.generated.h"
+
+USTRUCT()
+struct FTestInstanceDataObjectStruct
+{
+	GENERATED_BODY()
+
+public:
+	UPROPERTY()
+	int32 A = -1;
+
+	UPROPERTY()
+	int32 B = -1;
+
+	UPROPERTY()
+	int32 C = -1;
+
+	UPROPERTY()
+	int32 D = -1;
+};
+
+UCLASS()
+class UTestInstanceDataObjectClass : public UObject
+{
+	GENERATED_BODY()
+
+public:
+	UPROPERTY()
+	int32 Int32 = -1;
+
+	UPROPERTY()
+	FTestInstanceDataObjectStruct Struct;
+};
