IncrementalCook: Extend UStruct's schema hash to include the name of the class or struct. Type names are referred to by name in package headers so if the struct name changes then the schema hash needs to change as well. This fixes the false positive iterative skips that occurred when /Script/Engine.PerPlatformProperties was renamed to /Script/Core.PerPlatformProperties.
#rnx

diff --git a/Engine/Source/Runtime/CoreUObject/Private/Serialization/UnversionedPropertySerialization.cpp b/Engine/Source/Runtime/CoreUObject/Private/Serialization/UnversionedPropertySerialization.cpp
--- a/Engine/Source/Runtime/CoreUObject/Private/Serialization/UnversionedPropertySerialization.cpp
+++ b/Engine/Source/Runtime/CoreUObject/Private/Serialization/UnversionedPropertySerialization.cpp
@@ -326,6 +326,12 @@ struct FUnversionedStructSchema
 	{
 #if WITH_EDITORONLY_DATA
 		FBlake3 HashBuilder;
+		// Append the full name of the struct. To improve performance, append the FNames of its outers individually
+		// rather than calculating the full name as a string for performance
+		for (const UObject* Outer = Struct; Outer; Outer = Outer->GetOuter())
+		{
+			AppendHash(HashBuilder, Outer->GetFName());
+		}
 #endif
 		TArray<FUnversionedPropertySerializer, TInlineAllocator<256>> Serializers;
 		for (FProperty* Property = Struct->PropertyLink; Property; Property = Property->PropertyLinkNext)
