Handle diamond interface inheritance in FindFunctionByName
- when invoking a function of an interface, a superclass implementation of the interface function should be preferred over the base interface function itself

diff --git a/Engine/Source/Runtime/CoreUObject/Private/UObject/Class.cpp b/Engine/Source/Runtime/CoreUObject/Private/UObject/Class.cpp
--- a/Engine/Source/Runtime/CoreUObject/Private/UObject/Class.cpp
+++ b/Engine/Source/Runtime/CoreUObject/Private/UObject/Class.cpp
@@ -6475,9 +6475,16 @@ UFunction* UClass::FindFunctionByName(FName InName, EIncludeSuperFlag::Type Incl
 						}
 					}
 
-					if (Result == nullptr && SuperClass != nullptr )
+					UFunction* SuperResult = nullptr;
+					if (SuperClass != nullptr)
 					{
-						Result = SuperClass->FindFunctionByName(InName);
+						SuperResult = SuperClass->FindFunctionByName(InName);
+					}
+
+					// Check for multiple inheritance: If a superclass implements the same interface, use its implementation instead
+					if (!Result || (SuperResult && SuperResult->GetOwnerClass()->ImplementsInterface(Result->GetOwnerClass())))
+					{
+						Result = SuperResult;
 					}
 
 					{
