FUnversionedStructSchema::Create: Add CA_ASSUME to fix static analysis warning C6011: Dereferencing NULL pointer 'Struct'.
#rnx

diff --git a/Engine/Source/Runtime/CoreUObject/Private/Serialization/UnversionedPropertySerialization.cpp b/Engine/Source/Runtime/CoreUObject/Private/Serialization/UnversionedPropertySerialization.cpp
--- a/Engine/Source/Runtime/CoreUObject/Private/Serialization/UnversionedPropertySerialization.cpp
+++ b/Engine/Source/Runtime/CoreUObject/Private/Serialization/UnversionedPropertySerialization.cpp
@@ -324,6 +324,7 @@ struct FUnversionedStructSchema
 
 	FORCEINLINE static FUnversionedStructSchema* Create(const UStruct* Struct, bool bSkipEditorOnly)
 	{
+		CA_ASSUME(Struct != nullptr);
 #if WITH_EDITORONLY_DATA
 		FBlake3 HashBuilder;
 		// Append the full name of the struct. To improve performance, append the FNames of its outers individually
