Improve error reporting when making an invalid bulkdata stream request.

#rnx

- If a user tries to make a read request from FBulkDataBatchRequest that goes past the end of the bulkdata's payload we will eventually assert, but only at the point of reading which is often quite far away (context wise) from the code that initiated the request.
- Previously we were asserting if the user made a read request that exceeded the size of the bulkdata payload but this did not take into account reading from an offset.
- Fixed the assert to take offset into account so we now asset on invalid reads at the point that the read is first requested. In addition I have added more debug info to the assert making it easier to tell from the log file what was incorrect about the read request.
- Changed the assert to an ensure because although a too large size usually indicates a higher level bug, the lower level file system code will generally prevent us from reading past the end of the payload and return the data up until the end of the payload. So far in every case where we have seen this happening, reading to the end of the payload was acceptable and the higher level code worked. An ensure will allow the problem to remain visibile to programmers but not block content teams should the error make it to production.
-- This logic needs to be revisited as we should be mroe consistent with async and sync file reads and how we return errors, ideally allowing the higher level code to decide if a failure is important or not as we do not really have enough context at lower level to decide.

diff --git a/Engine/Source/Runtime/CoreUObject/Private/Serialization/BulkDataStreaming.cpp b/Engine/Source/Runtime/CoreUObject/Private/Serialization/BulkDataStreaming.cpp
--- a/Engine/Source/Runtime/CoreUObject/Private/Serialization/BulkDataStreaming.cpp
+++ b/Engine/Source/Runtime/CoreUObject/Private/Serialization/BulkDataStreaming.cpp
@@ -1088,10 +1088,15 @@ FBulkDataBatchRequest::FBatchBuilder& FBulkDataBatchRequest::FBatchBuilder::Read
 	uint64 Offset,
 	uint64 Size,
 	EAsyncIOPriorityAndFlags Priority,
-	FIoBuffer& Dst, 
+	FIoBuffer& Dst,
 	FBulkDataBatchReadRequest* OutRequest)
 {
-	check(Size == MAX_uint64 || Size <= uint64(BulkData.GetBulkDataSize()));
+	ensureMsgf(Size == MAX_uint64 || (Offset + Size) <= uint64(BulkData.GetBulkDataSize()),
+		TEXT("%s: Trying to read past the end of the payload, Offset: %llu, ReadSize: %llu, Payload Size: %lld"),
+		*BulkData.GetDebugName(),
+		Size,
+		Offset,
+		BulkData.GetBulkDataSize());
 
 	const uint64 ReadOffset = BulkData.GetBulkDataOffsetInFile() + Offset;
 	const uint64 ReadSize	= FMath::Min(uint64(BulkData.GetBulkDataSize()), Size);
