Fixed issue where if the default class value for a var is changed in verse, then the level would fail to cook.

diff --git a/Engine/Source/Runtime/CoreUObject/Private/UObject/PropertyBaseObject.cpp b/Engine/Source/Runtime/CoreUObject/Private/UObject/PropertyBaseObject.cpp
--- a/Engine/Source/Runtime/CoreUObject/Private/UObject/PropertyBaseObject.cpp
+++ b/Engine/Source/Runtime/CoreUObject/Private/UObject/PropertyBaseObject.cpp
@@ -648,7 +648,7 @@ UObject* FObjectPropertyBase::ConstructDefaultObjectValueIfNecessary(UObject* Ex
 				ExistingValue->GetOuter() != Outer) // Unless the template's Outer is the same as the new Outer in which case the template (ExistingValue) IS the object we can reuse
 			{
 				// We probably got here because an object value failed to load (missing import class) and the property is left with a template of default subobject
-				NewDefaultObjectValue = NewObject<UObject>(Outer, ExistingValue->GetClass(), ExistingValue->GetFName(), RF_NoFlags, ExistingValue);
+				NewDefaultObjectValue = NewObject<UObject>(Outer, ExistingValue->GetClass(), NAME_None, RF_NoFlags, ExistingValue);
 			}
 			else
 			{
