[ZenLoader]
 - Increase stall detection as current default can trigger in apparently normal case.

diff --git a/Engine/Source/Runtime/CoreUObject/Private/Serialization/AsyncLoading2.cpp b/Engine/Source/Runtime/CoreUObject/Private/Serialization/AsyncLoading2.cpp
--- a/Engine/Source/Runtime/CoreUObject/Private/Serialization/AsyncLoading2.cpp
+++ b/Engine/Source/Runtime/CoreUObject/Private/Serialization/AsyncLoading2.cpp
@@ -418,7 +418,7 @@ static FAutoConsoleVariableRef CVarGOnlyProcessRequiredPackagesWhenSyncLoading(
 	ECVF_Default
 );
 
-static float GStallDetectorTimeout = 30.0f;
+static float GStallDetectorTimeout = 120.0f;
 static FAutoConsoleVariableRef CVarGStallDetectorTimeout(
 	TEXT("s.StallDetectorTimeout"),
 	GStallDetectorTimeout,
