Fixed crash on Programs shutdown due to a UObjects deallocation from the persistent allocator

diff --git a/Engine/Source/Runtime/CoreUObject/Private/UObject/UObjectBase.cpp b/Engine/Source/Runtime/CoreUObject/Private/UObject/UObjectBase.cpp
--- a/Engine/Source/Runtime/CoreUObject/Private/UObject/UObjectBase.cpp
+++ b/Engine/Source/Runtime/CoreUObject/Private/UObject/UObjectBase.cpp
@@ -1023,6 +1023,11 @@ void UObjectBaseInit()
 		GUObjectAllocator.DisablePersistentAllocator();
 	}
 
+#if IS_PROGRAM
+	// always disable peristent allocator for programs even if they require cooked data
+	GUObjectAllocator.DisablePersistentAllocator();
+#endif
+
 	// Log what we're doing to track down what really happens as log in LaunchEngineLoop doesn't report those settings in pristine form.
 	UE_LOG(LogInit, Log, TEXT("%s for max %d objects, including %i objects not considered by GC."), 
 		bPreAllocateUObjectArray ? TEXT("Pre-allocating") : TEXT("Presizing"), MaxUObjects, MaxObjectsNotConsideredByGC);
