Add ability to set temporary transient state as save override
- Allows the caching of the transient state for an object onto the SaveContext using the existing Save Override list
- Allows the ability to prevent an object from being saved by adding this save override during a particular PreSave

#tests A|B tested cooking with and without change under Matt.Peters guidance. A|B tested Editor package save.

diff --git a/Engine/Source/Runtime/CoreUObject/Private/UObject/SavePackage/PackageHarvester.cpp b/Engine/Source/Runtime/CoreUObject/Private/UObject/SavePackage/PackageHarvester.cpp
--- a/Engine/Source/Runtime/CoreUObject/Private/UObject/SavePackage/PackageHarvester.cpp
+++ b/Engine/Source/Runtime/CoreUObject/Private/UObject/SavePackage/PackageHarvester.cpp
@@ -614,7 +614,7 @@ void FPackageHarvester::TryHarvestExportInternal(UObject* InObject)
 
 	// Check whether the object is unsaveable and skip adding it as an export to any realm if so
 	SaveContext.MarkUnsaveable(InObject);
-	if (InObject->HasAnyFlags(RF_Transient))
+	if (SaveContext.IsTransient(InObject))
 	{
 		return;
 	}
@@ -975,6 +975,10 @@ void FPackageHarvester::ResolveOverrides()
 		{
 			TransientPropertyOverrides.Add(PairObjectOverrides.Key, MoveTemp(Props));
 		}
+		if (!SaveContext.IsTransient(PairObjectOverrides.Key) && PairObjectOverrides.Value.bForceTransient)
+		{
+			SaveContext.TransientAssignments.Add(PairObjectOverrides.Key, EMarkedTransientReason::TransientOverride);
+		}
 	}
 }
 
@@ -1121,7 +1125,7 @@ ESaveableStatus FPackageHarvester::GetSaveableStatusForRealm(UObject* Obj, ESave
 	}
 
 	ESaveableStatus CulpritStatus;
-	ESaveableStatus Status = SaveContext.GetSaveableStatus(Obj, &OutCulprit, &CulpritStatus);
+	ESaveableStatus Status = SaveContext.GetSaveableStatus(Obj, &OutCulprit, &CulpritStatus, EIgnoreMarkUnsaveable::Yes);
 	if (Status != ESaveableStatus::Success)
 	{
 		if (Status == ESaveableStatus::OuterUnsaveable)
diff --git a/Engine/Source/Runtime/CoreUObject/Private/UObject/SavePackage/SaveContext.cpp b/Engine/Source/Runtime/CoreUObject/Private/UObject/SavePackage/SaveContext.cpp
--- a/Engine/Source/Runtime/CoreUObject/Private/UObject/SavePackage/SaveContext.cpp
+++ b/Engine/Source/Runtime/CoreUObject/Private/UObject/SavePackage/SaveContext.cpp
@@ -32,16 +32,18 @@ void FSaveContext::MarkUnsaveable(UObject* InObject)
 {
 	if (IsUnsaveable(InObject))
 	{
-		// TODO: We should not be modifying objects during the save. Besides interfering with the objects outside of the save,
-		// it also prevents us from gathering reasons about why an object was marked unsaveable.
-		InObject->SetFlags(RF_Transient);
+		EMarkedTransientReason& Reason = TransientAssignments.FindOrAdd(InObject);
+		if (Reason == EMarkedTransientReason::Uninitialized)
+		{
+			Reason = EMarkedTransientReason::Unsaveable;
+		}
 	}
 
 	// if this is the class default object, make sure it's not
 	// marked transient for any reason, as we need it to be saved
 	// to disk (unless it's associated with a transient generated class)
 #if WITH_EDITORONLY_DATA
-	ensureAlways(!InObject->HasAllFlags(RF_ClassDefaultObject | RF_Transient) || (InObject->GetClass()->ClassGeneratedBy != nullptr && InObject->GetClass()->HasAnyFlags(RF_Transient)));
+	ensureAlways(!InObject->HasAllFlags(RF_ClassDefaultObject) || !IsTransient(InObject) || (InObject->GetClass()->ClassGeneratedBy != nullptr && InObject->GetClass()->HasAnyFlags(RF_Transient)));
 #endif
 }
 
@@ -68,12 +70,12 @@ bool FSaveContext::IsUnsaveable(TObjectPtr<UObject> InObject, bool bEmitWarning)
 	return true;
 }
 
-ESaveableStatus FSaveContext::GetSaveableStatus(TObjectPtr<UObject> InObject, TObjectPtr<UObject>* OutCulprit, ESaveableStatus* OutCulpritStatus) const
+ESaveableStatus FSaveContext::GetSaveableStatus(TObjectPtr<UObject> InObject, TObjectPtr<UObject>* OutCulprit, ESaveableStatus* OutCulpritStatus, EIgnoreMarkUnsaveable IgnoreMarkUnsaveable) const
 {
 	TObjectPtr<UObject> Obj = InObject;
 	while (Obj)
 	{
-		ESaveableStatus Status = GetSaveableStatusNoOuter(Obj);
+		ESaveableStatus Status = GetSaveableStatusNoOuter(Obj, IgnoreMarkUnsaveable);
 		if (Status != ESaveableStatus::Success)
 		{
 			if (OutCulprit)
@@ -99,7 +101,7 @@ ESaveableStatus FSaveContext::GetSaveableStatus(TObjectPtr<UObject> InObject, TO
 	return ESaveableStatus::Success;
 }
 
-ESaveableStatus FSaveContext::GetSaveableStatusNoOuter(TObjectPtr<UObject> Obj) const
+ESaveableStatus FSaveContext::GetSaveableStatusNoOuter(TObjectPtr<UObject> Obj, EIgnoreMarkUnsaveable IgnoreMarkUnsaveable) const
 {
 	// pending kill object are unsaveable
 	if (Obj.IsResolved() && !IsValidChecked(Obj))
@@ -108,9 +110,25 @@ ESaveableStatus FSaveContext::GetSaveableStatusNoOuter(TObjectPtr<UObject> Obj)
 	}
 
 	// transient object are considered unsaveable if non native
-	if (Obj.IsResolved() && Obj->HasAnyFlags(RF_Transient) && !Obj->IsNative())
+	if (Obj.IsResolved() && IsTransient(Obj, IgnoreMarkUnsaveable) && !Obj->IsNative())
 	{
-		return ESaveableStatus::Transient;
+		if (Obj->HasAnyFlags(RF_Transient))
+		{
+			return ESaveableStatus::TransientFlag;
+		}
+
+		if (const EMarkedTransientReason* Reason = TransientAssignments.Find(Obj))
+		{
+			switch (*Reason)
+			{
+			case EMarkedTransientReason::TransientOverride:
+				return ESaveableStatus::TransientOverride;
+			case EMarkedTransientReason::Unsaveable:
+				return ESaveableStatus::MarkedUnsaveable;
+			default:
+				return ESaveableStatus::MarkedUnsaveable;
+			}
+		}
 	}
 
 	UClass* Class = Obj.GetClass();
@@ -129,6 +147,30 @@ ESaveableStatus FSaveContext::GetSaveableStatusNoOuter(TObjectPtr<UObject> Obj)
 	return ESaveableStatus::Success;
 }
 
+bool FSaveContext::IsTransient(const TObjectPtr<UObject> InObject, EIgnoreMarkUnsaveable IgnoreMarkUnsaveable) const
+{
+	if (!InObject)
+	{
+		return false;
+	}
+
+	if (InObject->HasAnyFlags(RF_Transient))
+	{
+		return true;
+	}
+
+	if (const EMarkedTransientReason* Reason = TransientAssignments.Find(InObject))
+	{
+		if (*Reason == EMarkedTransientReason::Unsaveable && IgnoreMarkUnsaveable == EIgnoreMarkUnsaveable::Yes)
+		{
+			return false;
+		}
+		return true;
+	}
+
+	return false;
+}
+
 FSavePackageResultStruct FSaveContext::GetFinalResult()
 {
 	if (Result != ESavePackageResult::Success)
@@ -225,12 +267,14 @@ EObjectMark FSaveContext::GetExcludedObjectMarksForGameRealm(const ITargetPlatfo
 
 const TCHAR* LexToString(ESaveableStatus Status)
 {
-	static_assert(static_cast<int32>(ESaveableStatus::__Count) == 9);
+	static_assert(static_cast<int32>(ESaveableStatus::__Count) == 11);
 	switch (Status)
 	{
 	case ESaveableStatus::Success: return TEXT("is saveable");
 	case ESaveableStatus::PendingKill: return TEXT("is pendingkill");
-	case ESaveableStatus::Transient: return TEXT("is transient");
+	case ESaveableStatus::TransientFlag: return TEXT("is transient");
+	case ESaveableStatus::TransientOverride: return TEXT("is Overriden as transient");
+	case ESaveableStatus::MarkedUnsaveable: return TEXT("has been marked unsaveable");
 	case ESaveableStatus::AbstractClass: return TEXT("has a Class with CLASS_Abstract");
 	case ESaveableStatus::DeprecatedClass: return TEXT("has a Class with CLASS_Deprecated");
 	case ESaveableStatus::NewerVersionExistsClass: return TEXT("has a Class with CLASS_NewerVersionExists");
diff --git a/Engine/Source/Runtime/CoreUObject/Private/UObject/SavePackage/SaveContext.h b/Engine/Source/Runtime/CoreUObject/Private/UObject/SavePackage/SaveContext.h
--- a/Engine/Source/Runtime/CoreUObject/Private/UObject/SavePackage/SaveContext.h
+++ b/Engine/Source/Runtime/CoreUObject/Private/UObject/SavePackage/SaveContext.h
@@ -107,7 +107,9 @@ enum class ESaveableStatus
 {
 	Success,
 	PendingKill,
-	Transient,
+	TransientFlag,
+	TransientOverride,
+	MarkedUnsaveable,
 	AbstractClass,
 	DeprecatedClass,
 	NewerVersionExistsClass,
@@ -117,6 +119,19 @@ enum class ESaveableStatus
 	__Count,
 };
 
+enum class EIgnoreMarkUnsaveable
+{
+	No = 0,
+	Yes,
+};
+
+enum class EMarkedTransientReason
+{
+	Uninitialized = 0,
+	Unsaveable,
+	TransientOverride,
+};
+
 /** Hold the harvested exports and imports for a realm */
 struct FHarvestedRealm
 {
@@ -781,8 +796,10 @@ public:
 	void MarkUnsaveable(UObject* InObject);
 
 	bool IsUnsaveable(TObjectPtr<UObject> InObject, bool bEmitWarning = true) const;
-	ESaveableStatus GetSaveableStatus(TObjectPtr<UObject> InObject, TObjectPtr<UObject>* OutCulprit = nullptr, ESaveableStatus* OutCulpritStatus = nullptr) const;
-	ESaveableStatus GetSaveableStatusNoOuter(TObjectPtr<UObject> InObject) const;
+	ESaveableStatus GetSaveableStatus(TObjectPtr<UObject> InObject, TObjectPtr<UObject>* OutCulprit = nullptr, ESaveableStatus* OutCulpritStatus = nullptr, EIgnoreMarkUnsaveable IgnoreMarkUnsaveable = EIgnoreMarkUnsaveable::No) const;
+	ESaveableStatus GetSaveableStatusNoOuter(TObjectPtr<UObject> InObject, EIgnoreMarkUnsaveable IgnoreMarkUnsaveable = EIgnoreMarkUnsaveable::No) const;
+
+	bool IsTransient(TObjectPtr<UObject> InObject, EIgnoreMarkUnsaveable IgnoreMarkUnsaveable = EIgnoreMarkUnsaveable::No) const;
 
 	void RecordIllegalReference(UObject* InFrom, UObject* InTo, EIllegalRefReason InReason, FString&& InOptionalReasonText = FString())
 	{
@@ -1170,8 +1187,11 @@ private:
 	// Set of AssetDatas created for the Assets saved into the package
 	TArray<FAssetData> SavedAssets;
 
-	// Overrided properties for each export that should be treated as transient, and nulled out when serializing
+	// Overridden properties for each export that should be treated as transient, and nulled out when serializing
 	TMap<UObject*, TSet<FProperty*>> TransientPropertyOverrides;
+
+	// Overridden objects that should be treated as transient, and skipped when serializing
+	TMap<UObject*, EMarkedTransientReason> TransientAssignments;
 };
 
 const TCHAR* LexToString(ESaveableStatus Status);
diff --git a/Engine/Source/Runtime/CoreUObject/Private/UObject/SavePackage2.cpp b/Engine/Source/Runtime/CoreUObject/Private/UObject/SavePackage2.cpp
--- a/Engine/Source/Runtime/CoreUObject/Private/UObject/SavePackage2.cpp
+++ b/Engine/Source/Runtime/CoreUObject/Private/UObject/SavePackage2.cpp
@@ -341,6 +341,13 @@ ESavePackageResult HarvestPackage(FSaveContext& SaveContext)
 
 	auto TryHarvestRootObject = [&Harvester, &SaveContext](UObject* InRoot)
 	{
+		// ForEachObjectWithPackage will filter out Transient objects but not objects with the bForceTransient Save Override
+		// Perform one more check to cover this case
+		if (SaveContext.IsTransient(InRoot))
+		{
+			return;
+		}
+
 		Harvester.TryHarvestExport(InRoot);
 		// if we are automatically generating an optional package output, re-harvest objects with that realm as the default
 		if (SaveContext.IsSaveAutoOptional())
@@ -862,7 +869,7 @@ ESavePackageResult ValidateImports(FSaveContext& SaveContext)
 			&& SaveContext.GetCurrentHarvestingRealm() != ESaveRealm::Optional;
 		if (bWrongImport)
 		{
-			if (!Import->HasAllFlags(RF_Transient) || !Import->IsNative())
+			if (!SaveContext.IsTransient(Import) || !Import->IsNative())
 			{
 				UE_LOG(LogSavePackage, Warning, TEXT("Bad Object=%s"), *Import->GetFullName());
 			}
@@ -872,10 +879,10 @@ ESavePackageResult ValidateImports(FSaveContext& SaveContext)
 				// a property of an intrinsic class.  Only properties of intrinsic classes will have
 				// an Outer that passes the check for "GetOuter()->IsInPackage(InOuter)" (thus ending up in this
 				// block of code).  Just verify that the Outer for this property is also marked RF_Transient and Native
-				check(Import->GetOuter()->HasAllFlags(RF_Transient) && Import->GetOuter()->IsNative());
+				check(SaveContext.IsTransient(Import->GetOuter()) && Import->GetOuter()->IsNative());
 			}
 		}
-		check(!bWrongImport || Import->HasAllFlags(RF_Transient) || Import->IsNative());
+		check(!bWrongImport || SaveContext.IsTransient(Import) || Import->IsNative());
 
 		// if this import shares a outer with top level object of this package then the reference is acceptable
 		if ((!SaveContext.IsCooking() || SaveContext.GetCurrentHarvestingRealm() == ESaveRealm::Optional) &&
@@ -1468,7 +1475,7 @@ void SavePreloadDependencies(FStructuredArchive::FRecord& StructuredArchiveRoot,
 			{
 				return;
 			}
-			if (!Index.IsNull() && (ToTest->HasAllFlags(RF_Transient) && !ToTest->IsNative()))
+			if (!Index.IsNull() && (SaveContext.IsTransient(ToTest) && !ToTest->IsNative()))
 			{
 				UE_LOG(LogSavePackage, Warning, TEXT("A dependency '%s' of '%s' is in the linker table, but is transient. We will keep the dependency anyway (%d)."), *ToTest.GetFullName(), *ForObj->GetFullName(), CallSite);
 			}
@@ -1547,7 +1554,7 @@ void SavePreloadDependencies(FStructuredArchive::FRecord& StructuredArchiveRoot,
 					// Only include subobject archetypes
 					if (SubObj->HasAnyFlags(RF_DefaultSubObject | RF_ArchetypeObject))
 					{
-						while (SubObj->HasAnyFlags(RF_Transient)) // transient components are stripped by the ICH, so find the one it will really use at runtime
+						while (SaveContext.IsTransient(SubObj)) // transient components are stripped by the ICH, so find the one it will really use at runtime
 						{
 							UObject* SubObjArch = SubObj->GetArchetype();
 							if (SubObjArch->GetClass()->HasAnyClassFlags(CLASS_Native | CLASS_Intrinsic))
@@ -1598,7 +1605,7 @@ void SavePreloadDependencies(FStructuredArchive::FRecord& StructuredArchiveRoot,
 							}
 
 							SubObj = SubObj->GetArchetype();
-							while (SubObj->HasAnyFlags(RF_Transient)) // transient components are stripped by the ICH, so find the one it will really use at runtime
+							while (SaveContext.IsTransient(SubObj)) // transient components are stripped by the ICH, so find the one it will really use at runtime
 							{
 								UObject* SubObjArch = SubObj->GetArchetype();
 								if (SubObjArch->GetClass()->HasAnyClassFlags(CLASS_Native | CLASS_Intrinsic))
diff --git a/Engine/Source/Runtime/CoreUObject/Public/UObject/ObjectSaveContext.h b/Engine/Source/Runtime/CoreUObject/Public/UObject/ObjectSaveContext.h
--- a/Engine/Source/Runtime/CoreUObject/Public/UObject/ObjectSaveContext.h
+++ b/Engine/Source/Runtime/CoreUObject/Public/UObject/ObjectSaveContext.h
@@ -170,8 +170,7 @@ public:
 	bool IsFirstConcurrentSave() const { return Data.bOuterConcurrentSave; }
 
 	/**
-	 * Add a save override to specific object. (i.e. mark certain properties transient for this save)
-	 * @note only object property are supported at the moment
+	 * Add a save override to specific object. (i.e. mark certain objects or properties transient for this save)
 	 */
 	void AddSaveOverride(UObject* Target, FObjectSaveOverride InOverride)
 	{
diff --git a/Engine/Source/Runtime/CoreUObject/Public/UObject/ObjectSaveOverride.h b/Engine/Source/Runtime/CoreUObject/Public/UObject/ObjectSaveOverride.h
--- a/Engine/Source/Runtime/CoreUObject/Public/UObject/ObjectSaveOverride.h
+++ b/Engine/Source/Runtime/CoreUObject/Public/UObject/ObjectSaveOverride.h
@@ -21,4 +21,7 @@ struct FPropertySaveOverride
 struct FObjectSaveOverride
 {
 	TArray<FPropertySaveOverride> PropOverrides;
+
+	// Treats the object as RF_Transient for the duration of the save
+	bool bForceTransient = false;
 };
\ No newline at end of file
