[ZenLoader]
 - Improve stall detector by monitoring activity on ALT when GT is waiting.

diff --git a/Engine/Source/Runtime/CoreUObject/Private/Serialization/AsyncLoading2.cpp b/Engine/Source/Runtime/CoreUObject/Private/Serialization/AsyncLoading2.cpp
--- a/Engine/Source/Runtime/CoreUObject/Private/Serialization/AsyncLoading2.cpp
+++ b/Engine/Source/Runtime/CoreUObject/Private/Serialization/AsyncLoading2.cpp
@@ -2483,6 +2483,21 @@ struct FAsyncLoadingThreadState2
 		return bUseTimeLimit;
 	}
 
+	void MarkAsActive()
+	{
+		bIsActive.store(true, std::memory_order_relaxed);
+	}
+	
+	void ResetActivity()
+	{
+		bIsActive.store(false, std::memory_order_relaxed);
+	}
+
+	bool IsActive()
+	{
+		return bIsActive.load(std::memory_order_relaxed);
+	}
+
 	FAsyncLoadEventGraphAllocator& GraphAllocator;
 	TArray<FEventLoadNode2*> NodesToFire;
 	TArray<FEventLoadNode2*> CurrentlyExecutingEventNodeStack;
@@ -2492,6 +2507,9 @@ struct FAsyncLoadingThreadState2
 	TSpscQueue<FAsyncPackage2*> PackagesToReprioritize;
 	bool bIsAsyncLoadingThread = false;
 	bool bCanAccessAsyncLoadingThreadData = true;
+	// used to probe activity for stall detection by the game thread
+	// we use relaxed memory ordering with a simple store to avoid any form of costly interlocked operations.
+	std::atomic<bool> bIsActive = false;
 	bool bShouldFireNodes = true;
 	bool bUseTimeLimit = false;
 	double TimeLimit = 0.0;
@@ -4423,6 +4441,8 @@ bool FAsyncLoadingThread2::CreateAsyncPackagesFromQueue(FAsyncLoadingThreadState
 
 	for (auto It = PendingPackages.CreateIterator(); It; ++It)
 	{
+		ThreadState.MarkAsActive();
+
 		FAsyncPackage2* PendingPackage = *It;
 		FPackageStoreEntry PackageEntry;
 		EPackageStoreEntryStatus PendingPackageStatus = PackageStore.GetPackageStoreEntry(PendingPackage->Desc.PackageIdToLoad,
@@ -4454,6 +4474,8 @@ bool FAsyncLoadingThread2::CreateAsyncPackagesFromQueue(FAsyncLoadingThreadState
 		int32 NumDequeued = 0;
 		while (NumDequeued < TimeSliceGranularity)
 		{
+			ThreadState.MarkAsActive();
+
 			TOptional<FPackageRequest> OptionalRequest = PackageRequestQueue.Dequeue();
 			if (!OptionalRequest.IsSet())
 			{
@@ -4711,6 +4733,8 @@ void FEventLoadNode2::Fire(FAsyncLoadingThreadState2* ThreadState)
 EEventLoadNodeExecutionResult FEventLoadNode2::Execute(FAsyncLoadingThreadState2& ThreadState)
 {
 	//TRACE_CPUPROFILER_EVENT_SCOPE(ExecuteEvent);
+	ThreadState.MarkAsActive();
+
 	check(BarrierCount.load(std::memory_order_relaxed) == 0);
 	EEventLoadNodeExecutionResult Result;
 	{
@@ -5744,6 +5768,8 @@ bool FAsyncPackage2::PreloadLinkerLoadExports(FAsyncLoadingThreadState2& ThreadS
 	check(LinkerLoadState->Linker->ExportMap.Num() == Data.Exports.Num());
 	while (LinkerLoadState->SerializeExportIndex < ExportCount)
 	{
+		ThreadState.MarkAsActive();
+
 		const int32 ExportIndex = LinkerLoadState->SerializeExportIndex++;
 		FExportObject& ExportObject = Data.Exports[ExportIndex];
 		FObjectExport& LinkerExport = LinkerLoadState->Linker->ExportMap[ExportIndex];
@@ -8345,6 +8371,8 @@ uint32 FAsyncLoadingThread2::Run()
 			bool bShouldWaitForExternalReads = false;
 			while (!bStopRequested.load(std::memory_order_relaxed))
 			{
+				ThreadState.MarkAsActive();
+
 				if (bShouldSuspend || SuspendRequestedCount.load(std::memory_order_relaxed) > 0 || IsGarbageCollectionWaiting())
 				{
 					TRACE_CPUPROFILER_EVENT_SCOPE(SuspendAsyncLoading);
@@ -9586,6 +9614,12 @@ void FAsyncLoadingThread2::FlushLoading(TConstArrayView<int32> RequestIDs)
 
 						// The loop count offers additional protection against timeout that could arise during debugging.
 						IdleLoopCount = 0;
+
+						// Mark the ALT as inactive, we'll then monitor IsActive() to see if the thread has done anything since the reset.
+						if (AsyncLoadingThreadState)
+						{
+							AsyncLoadingThreadState->ResetActivity();
+						}
 					}
 					else if (FPlatformTime::Seconds() - LastActivity > GStallDetectorTimeout && ++IdleLoopCount > GStallDetectorIdleLoops)
 					{
@@ -9651,6 +9685,12 @@ void FAsyncLoadingThread2::FlushLoading(TConstArrayView<int32> RequestIDs)
 						// Reset the manual event right after we wake up so we don't miss any trigger.
 						// Worst case, we'll do an empty spin before going back to sleep.
 						MainThreadWakeEvent.Reset();
+
+						// Reset the stall detector if there was any activity on the ALT.
+						if (AsyncLoadingThreadState && AsyncLoadingThreadState->IsActive())
+						{
+							LastActivity = 0.0;
+						}
 					}
 
 					// Flush logging when running cook-on-the-fly and waiting for packages
