autortfm: proactive crash fixes guided by static analysis

- CheckVerifyFailedImpl2: flushing the output device uses atomics
- DumpStackTraceToLog: flushing the output device uses atomics

- BasicFatalLog: invokes several debug utilities to print diagnostics and shut down the process
- BasicLog: uses atomics when GConvertBasicLogToLogRecord is true

- IsEnsureAllowed: generic platform init stuff

- GetNamePool: static object initialization

- GetCurrentThreadCallstack: internally takes locks

- CsvProfiler: push/pop-style profiler

- MakeUniqueObjectName: use of atomics

diff --git a/Engine/Source/Runtime/Core/Private/GenericPlatform/GenericPlatformMisc.cpp b/Engine/Source/Runtime/Core/Private/GenericPlatform/GenericPlatformMisc.cpp
--- a/Engine/Source/Runtime/Core/Private/GenericPlatform/GenericPlatformMisc.cpp
+++ b/Engine/Source/Runtime/Core/Private/GenericPlatform/GenericPlatformMisc.cpp
@@ -42,6 +42,7 @@
 #include "Internationalization/Culture.h"
 #include "ProfilingDebugging/CsvProfiler.h"
 #include "Misc/CoreDelegates.h"
+#include "AutoRTFM/AutoRTFM.h"
 
 #if UE_ENABLE_ICU
 	THIRD_PARTY_INCLUDES_START
@@ -1725,6 +1726,7 @@ namespace GenericPlatformMisc
 	bool GEnsureSettingsEverUpdated = false;
 }
 
+UE_AUTORTFM_ALWAYS_OPEN
 bool FGenericPlatformMisc::IsEnsureAllowed()
 {
 	// not all targets call FEngineLoop::Tick() or we might be here early
diff --git a/Engine/Source/Runtime/Core/Private/Logging/StructuredLog.cpp b/Engine/Source/Runtime/Core/Private/Logging/StructuredLog.cpp
--- a/Engine/Source/Runtime/Core/Private/Logging/StructuredLog.cpp
+++ b/Engine/Source/Runtime/Core/Private/Logging/StructuredLog.cpp
@@ -23,6 +23,7 @@
 #include "Serialization/CompactBinarySerialization.h"
 #include "Serialization/CompactBinaryValue.h"
 #include "Serialization/VarInt.h"
+#include "AutoRTFM/AutoRTFM.h"
 
 void VARARGS StaticFailDebug (const TCHAR* Error, const ANSICHAR* Expression, const ANSICHAR* File, int32 Line, bool bIsEnsure, void* ProgramCounter, const TCHAR* DescriptionFormat, ...);
 void         StaticFailDebugV(const TCHAR* Error, const ANSICHAR* Expression, const ANSICHAR* File, int32 Line, bool bIsEnsure, void* ProgramCounter, const TCHAR* DescriptionFormat, va_list DescriptionArgs);
@@ -1023,6 +1024,7 @@ FORCENOINLINE static FCbObject SerializeBasicLogMessage(const FStaticBasicLogRec
 	return Writer.Save().AsObject();
 }
 
+UE_AUTORTFM_ALWAYS_OPEN
 void BasicLog(const FLogCategoryBase& Category, const FStaticBasicLogRecord* Log, ...)
 {
 #if !NO_LOGGING
@@ -1077,6 +1079,7 @@ void BasicLog(const FLogCategoryBase& Category, const FStaticBasicLogRecord* Log
 #endif
 }
 
+UE_AUTORTFM_ALWAYS_OPEN
 void BasicFatalLog(const FLogCategoryBase& Category, const FStaticBasicLogRecord* Log, ...)
 {
 #if !NO_LOGGING
diff --git a/Engine/Source/Runtime/Core/Private/Misc/AssertionMacros.cpp b/Engine/Source/Runtime/Core/Private/Misc/AssertionMacros.cpp
--- a/Engine/Source/Runtime/Core/Private/Misc/AssertionMacros.cpp
+++ b/Engine/Source/Runtime/Core/Private/Misc/AssertionMacros.cpp
@@ -728,11 +728,13 @@ FORCENOINLINE void UE_DEBUG_SECTION VARARGS LowLevelFatalErrorHandler(const ANSI
 	FDebug::ProcessFatalError(PLATFORM_RETURN_ADDRESS());
 }
 
+UE_AUTORTFM_ALWAYS_OPEN
 void FDebug::DumpStackTraceToLog(const ELogVerbosity::Type LogVerbosity)
 {
 	DumpStackTraceToLog(TEXT("=== FDebug::DumpStackTrace(): ==="), LogVerbosity);
 }
 
+UE_AUTORTFM_ALWAYS_OPEN
 FORCENOINLINE void FDebug::DumpStackTraceToLog(const TCHAR* Heading, const ELogVerbosity::Type LogVerbosity)
 {
 #if !NO_LOGGING
diff --git a/Engine/Source/Runtime/Core/Private/UObject/UnrealNames.cpp b/Engine/Source/Runtime/Core/Private/UObject/UnrealNames.cpp
--- a/Engine/Source/Runtime/Core/Private/UObject/UnrealNames.cpp
+++ b/Engine/Source/Runtime/Core/Private/UObject/UnrealNames.cpp
@@ -2015,9 +2015,15 @@ static FNamePool& GetNamePool()
 		return *(FNamePool*)NamePoolData;
 	}
 
-	FNamePool* Singleton = new (NamePoolData) FNamePool;
-	bNamePoolInitialized = true;
-	LLM(FLowLevelMemTracker::Get().FinishInitialise());
+	FNamePool* Singleton = nullptr;
+	
+	UE_AUTORTFM_OPEN(
+	{
+		Singleton = new (NamePoolData) FNamePool;
+		bNamePoolInitialized = true;
+		LLM(FLowLevelMemTracker::Get().FinishInitialise());
+	});
+	
 	return *Singleton;
 }
 
diff --git a/Engine/Source/Runtime/Core/Public/AutoRTFM/AutoRTFM.h b/Engine/Source/Runtime/Core/Public/AutoRTFM/AutoRTFM.h
--- a/Engine/Source/Runtime/Core/Public/AutoRTFM/AutoRTFM.h
+++ b/Engine/Source/Runtime/Core/Public/AutoRTFM/AutoRTFM.h
@@ -24,6 +24,10 @@
 #define UE_AUTOSTM 0
 #endif
 
+#if !defined(UE_AUTORTFM_STATIC_VERIFIER)
+#define UE_AUTORTFM_STATIC_VERIFIER 0
+#endif
+
 #if UE_AUTORTFM
 #define UE_AUTORTFM_AUTORTFM(F) [[clang::autortfm(F), clang::noinline]]
 #define UE_AUTORTFM_NOAUTORTFM [[clang::noautortfm, clang::noinline]]
@@ -34,6 +38,14 @@
 #define UE_AUTORTFM_ALWAYS_OPEN
 #endif
 
+#if UE_AUTORTFM && UE_AUTORTFM_STATIC_VERIFIER
+#define UE_AUTORTFM_ENSURE_SAFE [[clang::autortfm_ensure_safe]]
+#define UE_AUTORTFM_ASSUME_SAFE [[clang::autortfm_assume_safe]]
+#else
+#define UE_AUTORTFM_ENSURE_SAFE
+#define UE_AUTORTFM_ASSUME_SAFE
+#endif
+
 #if defined(UE_AUTORTFM_STANDALONE)
 #include <stdlib.h>
 #include <memory.h>
diff --git a/Engine/Source/Runtime/Core/Public/Misc/MTAccessDetector.h b/Engine/Source/Runtime/Core/Public/Misc/MTAccessDetector.h
--- a/Engine/Source/Runtime/Core/Public/Misc/MTAccessDetector.h
+++ b/Engine/Source/Runtime/Core/Public/Misc/MTAccessDetector.h
@@ -15,6 +15,7 @@
 #include "HAL/PreprocessorHelpers.h"
 #include "HAL/PlatformStackWalk.h"
 #include "Misc/Build.h"
+#include "AutoRTFM/AutoRTFM.h"
 #include <atomic>
 
 
@@ -364,7 +365,7 @@ private:
 	{
 		const SIZE_T StackTraceSize = 65536;
 		ANSICHAR StackTrace[StackTraceSize] = { 0 };
-		FPlatformStackWalk::StackWalkAndDump(StackTrace, StackTraceSize, 0);
+		UE_AUTORTFM_OPEN({ FPlatformStackWalk::StackWalkAndDump(StackTrace, StackTraceSize, 0); });
 		return ANSI_TO_TCHAR(StackTrace);
 	}
 
@@ -383,7 +384,7 @@ private:
 
 		const SIZE_T StackTraceSize = 65536;
 		ANSICHAR StackTrace[StackTraceSize] = { 0 };
-		FPlatformStackWalk::ThreadStackWalkAndDump(StackTrace, StackTraceSize, 0, ThreadId);
+		UE_AUTORTFM_OPEN({ FPlatformStackWalk::ThreadStackWalkAndDump(StackTrace, StackTraceSize, 0, ThreadId); });
 		return ANSI_TO_TCHAR(StackTrace);
 	}
 
diff --git a/Engine/Source/Runtime/Core/Public/ProfilingDebugging/CsvProfiler.h b/Engine/Source/Runtime/Core/Public/ProfilingDebugging/CsvProfiler.h
--- a/Engine/Source/Runtime/Core/Public/ProfilingDebugging/CsvProfiler.h
+++ b/Engine/Source/Runtime/Core/Public/ProfilingDebugging/CsvProfiler.h
@@ -31,6 +31,7 @@
 #include "Templates/UnrealTemplate.h"
 #include "Traits/IsCharEncodingCompatibleWith.h"
 #include "UObject/NameTypes.h"
+#include "AutoRTFM/AutoRTFM.h"
 
 #if UE_ENABLE_INCLUDE_ORDER_DEPRECATED_IN_5_2
 #include "Async/TaskGraphInterfaces.h"
@@ -538,12 +539,14 @@ public:
 		: StatName(InStatName)
 		, CategoryIndex(InCategoryIndex)
 	{
-		FCsvProfiler::BeginStat(StatName, CategoryIndex, InNamedEventName);
+		UE_AUTORTFM_OPEN({ FCsvProfiler::BeginStat(StatName, CategoryIndex, InNamedEventName); });
+		AutoRTFM::PushOnAbortHandler(this, [SN = StatName, CI = CategoryIndex](){ FCsvProfiler::EndStat(SN, CI); });
 	}
 
 	~FScopedCsvStat()
 	{
-		FCsvProfiler::EndStat(StatName, CategoryIndex);
+		AutoRTFM::PopOnAbortHandler(this);
+		UE_AUTORTFM_OPEN({ FCsvProfiler::EndStat(StatName, CategoryIndex); });
 	}
 	const char * StatName;
 	uint32 CategoryIndex;
@@ -555,12 +558,14 @@ public:
 	FScopedCsvStatExclusive(const char * InStatName, const char* InNamedEventName = nullptr)
 		: StatName(InStatName)
 	{
-		FCsvProfiler::BeginExclusiveStat(StatName, InNamedEventName);
+		UE_AUTORTFM_OPEN({ FCsvProfiler::BeginExclusiveStat(StatName, InNamedEventName); });
+		AutoRTFM::PushOnAbortHandler(this, [SN = StatName](){ FCsvProfiler::EndExclusiveStat(SN); });
 	}
 
 	~FScopedCsvStatExclusive()
 	{
-		FCsvProfiler::EndExclusiveStat(StatName);
+		AutoRTFM::PopOnAbortHandler(this);
+		UE_AUTORTFM_OPEN({ FCsvProfiler::EndExclusiveStat(StatName); });
 	}
 	const char * StatName;
 };
diff --git a/Engine/Source/Runtime/CoreUObject/Private/UObject/UObjectGlobals.cpp b/Engine/Source/Runtime/CoreUObject/Private/UObject/UObjectGlobals.cpp
--- a/Engine/Source/Runtime/CoreUObject/Private/UObject/UObjectGlobals.cpp
+++ b/Engine/Source/Runtime/CoreUObject/Private/UObject/UObjectGlobals.cpp
@@ -2863,37 +2863,40 @@ FName MakeUniqueObjectName(UObject* Parent, const UClass* Class, FName InBaseNam
 			else
 			{
 				int32 NameNumber = 0;
-				if (Parent && (Parent != ANY_PACKAGE_DEPRECATED) && !(Options & EUniqueObjectNameOptions::GloballyUnique))
+				UE_AUTORTFM_OPEN(
 				{
-					if (!FPlatformProperties::HasEditorOnlyData() && GFastPathUniqueNameGeneration)
+					if (Parent && (Parent != ANY_PACKAGE_DEPRECATED) && !(Options & EUniqueObjectNameOptions::GloballyUnique))
 					{
-						/*   Fast Path Name Generation
-						* A significant fraction of object creation time goes into verifying that the a chosen unique name is really unique.
-						* The idea here is to generate unique names using very high numbers and only in situations where collisions are
-						* impossible for other reasons.
-						*
-						* Rationale for uniqueness as used here.
-						* - Consoles do not save objects in general, and certainly not animation trees. So we could never load an object that would later clash.
-						* - We assume that we never load or create any object with a "name number" as large as, say, MAX_int32 / 2, other than via
-						*   HACK_FastPathUniqueNameGeneration.
-						* - After using one of these large "name numbers", we decrement the static UniqueIndex, this no two names generated this way, during the
-						*   same run, could ever clash.
-						* - We assume that we could never create anywhere near MAX_int32/2 total objects at runtime, within a single run.
-						* - We require an outer for these items, thus outers must themselves be unique. Therefore items with unique names created on the fast path
-						*   could never clash with anything with a different outer. For animation trees, these outers are never saved or loaded, thus clashes are
-						*   impossible.
-						*/
-						NameNumber = --NameNumberUniqueIndex;
+						if (!FPlatformProperties::HasEditorOnlyData() && GFastPathUniqueNameGeneration)
+						{
+							/*   Fast Path Name Generation
+							* A significant fraction of object creation time goes into verifying that the a chosen unique name is really unique.
+							* The idea here is to generate unique names using very high numbers and only in situations where collisions are
+							* impossible for other reasons.
+							*
+							* Rationale for uniqueness as used here.
+							* - Consoles do not save objects in general, and certainly not animation trees. So we could never load an object that would later clash.
+							* - We assume that we never load or create any object with a "name number" as large as, say, MAX_int32 / 2, other than via
+							*   HACK_FastPathUniqueNameGeneration.
+							* - After using one of these large "name numbers", we decrement the static UniqueIndex, this no two names generated this way, during the
+							*   same run, could ever clash.
+							* - We assume that we could never create anywhere near MAX_int32/2 total objects at runtime, within a single run.
+							* - We require an outer for these items, thus outers must themselves be unique. Therefore items with unique names created on the fast path
+							*   could never clash with anything with a different outer. For animation trees, these outers are never saved or loaded, thus clashes are
+							*   impossible.
+							*/
+							NameNumber = --NameNumberUniqueIndex;
+						}
+						else
+						{
+							NameNumber = UpdateSuffixForNextNewObject(Parent, Class, [](int32& Index) { ++Index; });
+						}
 					}
 					else
 					{
-						NameNumber = UpdateSuffixForNextNewObject(Parent, Class, [](int32& Index) { ++Index; });
+						NameNumber = ++Class->ClassUnique;
 					}
-				}
-				else
-				{
-					NameNumber = ++Class->ClassUnique;
-				}
+				});
 				TestName = FName(BaseName, NameNumber);
 			}
 
@@ -3387,7 +3390,7 @@ static thread_local FRestoreForUObjectOverwrite* ObjectRestoreAfterInitProps = n
 extern const FName NAME_UniqueObjectNameForCooking(TEXT("UniqueObjectNameForCooking"));
 COREUOBJECT_API bool GOutputCookingWarnings = false;
 
-
+UE_AUTORTFM_ASSUME_SAFE
 UObject* StaticAllocateObject
 (
 	const UClass*	InClass,
@@ -3897,6 +3900,7 @@ FObjectInitializer::~FObjectInitializer()
 	}
 }
 
+UE_AUTORTFM_ASSUME_SAFE
 void FObjectInitializer::PostConstructInit()
 {
 	// we clear the Obj pointer at the end of this function, so if it is null 
