Const correct IsInContainer

#rnx

diff --git a/Engine/Source/Runtime/CoreUObject/Public/UObject/UnrealType.h b/Engine/Source/Runtime/CoreUObject/Public/UObject/UnrealType.h
--- a/Engine/Source/Runtime/CoreUObject/Public/UObject/UnrealType.h
+++ b/Engine/Source/Runtime/CoreUObject/Public/UObject/UnrealType.h
@@ -784,7 +784,7 @@ public:
 		return Offset_Internal + GetSize() <= ContainerSize;
 	}
 	/** See if the offset of this property is below the supplied container size */
-	FORCEINLINE bool IsInContainer(UStruct* ContainerClass) const
+	FORCEINLINE bool IsInContainer(const UStruct* ContainerClass) const
 	{
 		return Offset_Internal + GetSize() <= (ContainerClass ? ContainerClass->GetPropertiesSize() : MAX_int32);
 	}
