Correct documentation on FWeakObjectPtr::IsValid(bool)
[RN] fix Engine/Core

diff --git a/Engine/Source/Runtime/CoreUObject/Public/UObject/WeakObjectPtr.h b/Engine/Source/Runtime/CoreUObject/Public/UObject/WeakObjectPtr.h
--- a/Engine/Source/Runtime/CoreUObject/Public/UObject/WeakObjectPtr.h
+++ b/Engine/Source/Runtime/CoreUObject/Public/UObject/WeakObjectPtr.h
@@ -190,7 +190,8 @@ public:
 	 * Test if this points to a live UObject
 	 * This should be done only when needed as excess resolution of the underlying pointer can cause performance issues.
 	 *
-	 * @param bEvenIfGarbage if this is true, Garbage objects are considered invalid
+	 * @param bEvenIfGarbage if this is true, Garbage objects are considered valid. When false, an object is considered invalid as soon as it is
+	 *                       marked for destruction.
 	 * @param bThreadsafeTest if true then function will just give you information whether referenced
 	 *							UObject is gone forever (return false) or if it is still there (return true, no object flags checked).
 	 *							This is required as without it IsValid can return false during the mark phase of the GC
