Fixed the IDO tests being enabled for unsupported targets

#rnx

diff --git a/Engine/Source/Runtime/CoreUObject/Tests/UObject/InstanceDataObjectUtilsTest.cpp b/Engine/Source/Runtime/CoreUObject/Tests/UObject/InstanceDataObjectUtilsTest.cpp
--- a/Engine/Source/Runtime/CoreUObject/Tests/UObject/InstanceDataObjectUtilsTest.cpp
+++ b/Engine/Source/Runtime/CoreUObject/Tests/UObject/InstanceDataObjectUtilsTest.cpp
@@ -2,7 +2,7 @@
 
 #include "InstanceDataObjectUtilsTest.h"
 
-#if WITH_TESTS && (PLATFORM_WINDOWS || PLATFORM_LINUX || PLATFORM_MAC)
+#if WITH_TESTS && WITH_EDITORONLY_DATA
 
 #include "Tests/TestHarnessAdapter.h"
 #include "UObject/InstanceDataObjectUtils.h"
@@ -51,4 +51,4 @@ TEST_CASE_NAMED(FInstanceDataObjectUtilsTest, "CoreUObject::InstanceDataObjectUt
 
 } // UE
 
-#endif // WITH_TESTS
+#endif // WITH_TESTS && WITH_EDITORONLY_DATA
