Added a way to convert a PropertyVisitorPath to a FArchieveSerializationChain

diff --git a/Engine/Source/Runtime/CoreUObject/Private/UObject/PropertyVisitor.cpp b/Engine/Source/Runtime/CoreUObject/Private/UObject/PropertyVisitor.cpp
--- a/Engine/Source/Runtime/CoreUObject/Private/UObject/PropertyVisitor.cpp
+++ b/Engine/Source/Runtime/CoreUObject/Private/UObject/PropertyVisitor.cpp
@@ -4,6 +4,7 @@
 #include "Containers/UnrealString.h"
 #include "UObject/NameTypes.h"
 #include "UObject/UnrealType.h"
+#include "Serialization/ArchiveSerializedPropertyChain.h"
 
 FString FPropertyVisitorPath::ToString(const TCHAR* Separator /*= TEXT(".")*/) const
 {
@@ -134,4 +135,14 @@ void* FPropertyVisitorPath::GetPropertyDataPtr(UObject* Object) const
 	});
 
 	return DataPtr;
+}
+
+FArchiveSerializedPropertyChain FPropertyVisitorPath::ToSerializedPropertyChain() const
+{
+	FArchiveSerializedPropertyChain Chain;
+	for (const FPropertyVisitorInfo& Info : Path)
+	{
+		Chain.PushProperty(const_cast<FProperty*>(Info.Property), Info.Property->IsEditorOnlyProperty());
+	}
+	return Chain;
 }
\ No newline at end of file
diff --git a/Engine/Source/Runtime/CoreUObject/Public/UObject/PropertyVisitor.h b/Engine/Source/Runtime/CoreUObject/Public/UObject/PropertyVisitor.h
--- a/Engine/Source/Runtime/CoreUObject/Public/UObject/PropertyVisitor.h
+++ b/Engine/Source/Runtime/CoreUObject/Public/UObject/PropertyVisitor.h
@@ -5,6 +5,7 @@
 #include "Containers/Array.h"
 
 class FProperty;
+struct FArchiveSerializedPropertyChain;
 
 enum class EPropertyVisitorControlFlow : uint8
 {
@@ -25,7 +26,7 @@ enum class EPropertyVisitorInfoType : uint8
 
 struct FPropertyVisitorInfo
 {
-	explicit FPropertyVisitorInfo( const FProperty* InProperty, int32 InIndex = INDEX_NONE, EPropertyVisitorInfoType InPropertyInfo = EPropertyVisitorInfoType::None)
+	explicit FPropertyVisitorInfo(const FProperty* InProperty, int32 InIndex = INDEX_NONE, EPropertyVisitorInfoType InPropertyInfo = EPropertyVisitorInfoType::None)
 		: Property(InProperty)
 		, Index(InIndex)
 		, PropertyInfo(InPropertyInfo)
@@ -87,9 +88,9 @@ public:
 		Path.Push(Info);
 	}
 
-	void Pop()
+	FPropertyVisitorInfo Pop()
 	{
-		Path.Pop(EAllowShrinking::No);
+		return Path.Pop(EAllowShrinking::No);
 	}
 
 	int32 Num() const
@@ -134,6 +135,15 @@ public:
 		return Path.CreateConstIterator();
 	}
 
+	/**
+	 * Converts path to an archive serialized property chain
+	 * This method useful when APIs like in the FOverridableManager are taking in a
+	 * FArchiveSerializedPropertyChain parameter and would like to use them.
+	 * It is used in many places in the FProperty::ImportText
+	 * @return the archive serialized property chain built from the property path
+	 */
+	COREUOBJECT_API FArchiveSerializedPropertyChain ToSerializedPropertyChain() const;
+
 protected:
 	TArray<FPropertyVisitorInfo> Path;
 };
