[ALT]
 - Fix TSAN warning on FApp::IsGame() changing on GT while creating linkers by moving the check as the very last condition where it is highly unlikely to ever get reached in editor (the only place it can change).
 - Since the value should never changes in -game and GIsEditor should never changes in editor we should never get into a state where that value dynamically affects the outcome of the loading but TSAN complained because we read it in all cases.

diff --git a/Engine/Source/Runtime/CoreUObject/Private/Serialization/AsyncLoading2.cpp b/Engine/Source/Runtime/CoreUObject/Private/Serialization/AsyncLoading2.cpp
--- a/Engine/Source/Runtime/CoreUObject/Private/Serialization/AsyncLoading2.cpp
+++ b/Engine/Source/Runtime/CoreUObject/Private/Serialization/AsyncLoading2.cpp
@@ -5404,7 +5404,7 @@ void FAsyncPackage2::CreateLinker(const FLinkerInstancingContext* InstancingCont
 {
 	uint32 LinkerFlags = (LOAD_Async | LOAD_NoVerify | LOAD_SkipLoadImportedPackages | Desc.LoadFlags);
 #if WITH_EDITOR
-	if ((!FApp::IsGame() || GIsEditor) && (Desc.PackageFlags & PKG_PlayInEditor) != 0)
+	if ((Desc.PackageFlags & PKG_PlayInEditor) != 0 && (GIsEditor || !FApp::IsGame()))
 	{
 		LinkerFlags |= LOAD_PackageForPIE;
 	}
