Replaced some TEnableIfs and UE traits.

diff --git a/Engine/Source/Runtime/CoreUObject/Public/UObject/UnrealType.h b/Engine/Source/Runtime/CoreUObject/Public/UObject/UnrealType.h
--- a/Engine/Source/Runtime/CoreUObject/Public/UObject/UnrealType.h
+++ b/Engine/Source/Runtime/CoreUObject/Public/UObject/UnrealType.h
@@ -6834,8 +6834,11 @@ struct TFieldRange
 	Field templates.
 -----------------------------------------------------------------------------*/
 
-template <class T> 
-typename TEnableIf<TIsDerivedFrom<T, UField>::IsDerived, T*>::Type FindUField(const UStruct* Owner, FName FieldName, EFieldIterationFlags IterationFlags = EFieldIterationFlags::Default)
+template <
+	typename T
+	UE_REQUIRES(std::is_base_of_v<UField, T>)
+>
+T* FindUField(const UStruct* Owner, FName FieldName, EFieldIterationFlags IterationFlags = EFieldIterationFlags::Default)
 {
 	static_assert(sizeof(T) > 0, "T must not be an incomplete type");
 
@@ -6858,8 +6861,11 @@ typename TEnableIf<TIsDerivedFrom<T, UField>::IsDerived, T*>::Type FindUField(co
 	return nullptr;
 }
 
-template <class T> 
-typename TEnableIf<TIsDerivedFrom<T, UField>::IsDerived, T*>::Type FindUField(const UStruct* Owner, const TCHAR* FieldName, EFieldIterationFlags IterationFlags = EFieldIterationFlags::Default)
+template <
+	typename T
+	UE_REQUIRES(std::is_base_of_v<UField, T>)
+>
+T* FindUField(const UStruct* Owner, const TCHAR* FieldName, EFieldIterationFlags IterationFlags = EFieldIterationFlags::Default)
 {
 	static_assert(sizeof(T) > 0, "T must not be an incomplete type");
 
@@ -6868,8 +6874,11 @@ typename TEnableIf<TIsDerivedFrom<T, UField>::IsDerived, T*>::Type FindUField(co
 	return FindUField<T>(Owner, Name, IterationFlags);
 }
 
-template <class T>
-typename TEnableIf<TIsDerivedFrom<T, FField>::IsDerived, T*>::Type FindFProperty(const UStruct* Owner, FName FieldName, EFieldIterationFlags IterationFlags = EFieldIterationFlags::Default)
+template <
+	typename T
+	UE_REQUIRES(std::is_base_of_v<FField, T>)
+>
+T* FindFProperty(const UStruct* Owner, FName FieldName, EFieldIterationFlags IterationFlags = EFieldIterationFlags::Default)
 {
 	static_assert(sizeof(T) > 0, "T must not be an incomplete type");
 
@@ -6892,8 +6901,11 @@ typename TEnableIf<TIsDerivedFrom<T, FField>::IsDerived, T*>::Type FindFProperty
 	return nullptr;
 }
 
-template <class T>
-typename TEnableIf<TIsDerivedFrom<T, FField>::IsDerived, T*>::Type FindFProperty(const UStruct* Owner, const TCHAR* FieldName, EFieldIterationFlags IterationFlags = EFieldIterationFlags::Default)
+template <
+	typename T
+	UE_REQUIRES(std::is_base_of_v<FField, T>)
+>
+T* FindFProperty(const UStruct* Owner, const TCHAR* FieldName, EFieldIterationFlags IterationFlags = EFieldIterationFlags::Default)
 {
 	static_assert(sizeof(T) > 0, "T must not be an incomplete type");
 
@@ -6922,26 +6934,38 @@ inline FFieldVariant FindUFieldOrFProperty(const UStruct* Owner, const TCHAR* Fi
 	return FindUFieldOrFProperty(Owner, Name, IterationFlags);
 }
 
-template <class T>
-typename TEnableIf<TIsDerivedFrom<T, UField>::IsDerived, T*>::Type FindUFieldOrFProperty(const UStruct* Owner, FName FieldName, EFieldIterationFlags IterationFlags = EFieldIterationFlags::Default)
+template <
+	typename T
+	UE_REQUIRES(std::is_base_of_v<UField, T>)
+>
+T* FindUFieldOrFProperty(const UStruct* Owner, FName FieldName, EFieldIterationFlags IterationFlags = EFieldIterationFlags::Default)
 {
 	return FindUField<T>(Owner, FieldName, IterationFlags);
 }
 
-template <class T>
-typename TEnableIf<TIsDerivedFrom<T, UField>::IsDerived, T*>::Type FindUFieldOrFProperty(const UStruct* Owner, const TCHAR* FieldName, EFieldIterationFlags IterationFlags = EFieldIterationFlags::Default)
+template <
+	typename T
+	UE_REQUIRES(std::is_base_of_v<UField, T>)
+>
+T* FindUFieldOrFProperty(const UStruct* Owner, const TCHAR* FieldName, EFieldIterationFlags IterationFlags = EFieldIterationFlags::Default)
 {
 	return FindUField<T>(Owner, FieldName, IterationFlags);
 }
 
-template <class T>
-typename TEnableIf<TIsDerivedFrom<T, FField>::IsDerived, T*>::Type FindUFieldOrFProperty(const UStruct* Owner, FName FieldName, EFieldIterationFlags IterationFlags = EFieldIterationFlags::Default)
+template <
+	typename T
+	UE_REQUIRES(std::is_base_of_v<FField, T>)
+>
+T* FindUFieldOrFProperty(const UStruct* Owner, FName FieldName, EFieldIterationFlags IterationFlags = EFieldIterationFlags::Default)
 {
 	return FindFProperty<T>(Owner, FieldName, IterationFlags);
 }
 
-template <class T>
-typename TEnableIf<TIsDerivedFrom<T, FField>::IsDerived, T*>::Type FindUFieldOrFProperty(const UStruct* Owner, const TCHAR* FieldName, EFieldIterationFlags IterationFlags = EFieldIterationFlags::Default)
+template <
+	typename T
+	UE_REQUIRES(std::is_base_of_v<FField, T>)
+>
+T* FindUFieldOrFProperty(const UStruct* Owner, const TCHAR* FieldName, EFieldIterationFlags IterationFlags = EFieldIterationFlags::Default)
 {
 	return FindFProperty<T>(Owner, FieldName, IterationFlags);
 }
