FAssetPackageData: Add compactbinary functions NetworkWrite and TryNetworkRead for use by MPCook.
#rnx

diff --git a/Engine/Source/Runtime/Core/Private/UObject/ObjectVersion.cpp b/Engine/Source/Runtime/Core/Private/UObject/ObjectVersion.cpp
--- a/Engine/Source/Runtime/Core/Private/UObject/ObjectVersion.cpp
+++ b/Engine/Source/Runtime/Core/Private/UObject/ObjectVersion.cpp
@@ -9,6 +9,7 @@
 #include "Misc/AssertionMacros.h"
 #include "Serialization/Archive.h"
 #include "Serialization/CompactBinary.h"
+#include "Serialization/CompactBinarySerialization.h"
 #include "Serialization/CompactBinaryWriter.h"
 
 // @see ObjectVersion.h for the list of changes/defines
@@ -40,22 +41,18 @@ FArchive& operator<<(FArchive& Ar, FPackageFileVersion& Version)
 	return Ar;
 }
 
-FCbWriter& operator<<(FCbWriter& Writer, const FPackageFileVersion& Version)
+FCbWriter& FPackageFileVersion::Write(FCbWriter& Writer) const
 {
 	Writer.BeginObject();
-	Writer << "ue4version" << Version.FileVersionUE4;
-	Writer << "ue5version" << Version.FileVersionUE5;
+	Writer << "ue4version" << FileVersionUE4;
+	Writer << "ue5version" << FileVersionUE5;
 	Writer.EndObject();
-
 	return Writer;
 }
 
-FPackageFileVersion FromCbObject(const FCbObject& Obj)
+bool FPackageFileVersion::TryRead(const FCbFieldView& FieldView)
 {
-	FPackageFileVersion Version;
-
-	Version.FileVersionUE4 = Obj["ue4version"].AsInt32();
-	Version.FileVersionUE5 = Obj["ue5version"].AsInt32();
-
-	return Version;
+	bool bOk = LoadFromCompactBinary(FieldView["ue4version"], FileVersionUE4);
+	bOk = LoadFromCompactBinary(FieldView["ue5version"], FileVersionUE5) & bOk;
+	return bOk;
 }
diff --git a/Engine/Source/Runtime/Core/Public/UObject/ObjectVersion.h b/Engine/Source/Runtime/Core/Public/UObject/ObjectVersion.h
--- a/Engine/Source/Runtime/Core/Public/UObject/ObjectVersion.h
+++ b/Engine/Source/Runtime/Core/Public/UObject/ObjectVersion.h
@@ -5,6 +5,7 @@
 #include "CoreTypes.h"
 
 class FArchive;
+class FCbFieldView;
 class FCbObject;
 class FCbWriter;
 
@@ -833,8 +834,19 @@ struct FPackageFileVersion
 
 	/** Serialization members */
 	CORE_API friend FArchive& operator<<(FArchive& Ar, FPackageFileVersion& Version);
-	CORE_API friend FCbWriter& operator<<(FCbWriter& Writer, const FPackageFileVersion& Version);
-	CORE_API static FPackageFileVersion FromCbObject(const FCbObject& Obj);
+	/** Serialization members implemented as hidden friends*/
+private:
+	friend FCbWriter& operator<<(FCbWriter& Writer, const FPackageFileVersion& Version)
+	{
+		return Version.Write(Writer);
+	}
+	CORE_API FCbWriter& Write(FCbWriter& Writer) const;
+	friend bool LoadFromCompactBinary(const FCbFieldView& FieldView, FPackageFileVersion& Version)
+	{
+		return Version.TryRead(FieldView);
+	}
+	CORE_API bool TryRead(const FCbFieldView& FieldView);
+public:
 
 	/* UE4 file version*/
 	int32		FileVersionUE4 = 0;
diff --git a/Engine/Source/Runtime/CoreUObject/Private/AssetRegistry/AssetData.cpp b/Engine/Source/Runtime/CoreUObject/Private/AssetRegistry/AssetData.cpp
--- a/Engine/Source/Runtime/CoreUObject/Private/AssetRegistry/AssetData.cpp
+++ b/Engine/Source/Runtime/CoreUObject/Private/AssetRegistry/AssetData.cpp
@@ -892,6 +892,27 @@ bool FAssetRegistryVersion::SerializeVersion(FArchive& Ar, FAssetRegistryVersion
 	return !Ar.IsError();
 }
 
+namespace UE::AssetRegistry
+{
+
+FCbWriter& FPackageCustomVersion::Write(FCbWriter& Writer) const
+{
+	Writer.BeginArray();
+	Writer << Key << Version;
+	Writer.EndArray();
+	return Writer;
+}
+
+bool FPackageCustomVersion::TryRead(const FCbFieldView& Field)
+{
+	FCbFieldViewIterator Iter = Field.CreateViewIterator();
+	bool bOk = LoadFromCompactBinary(*Iter++, Key);
+	bOk = LoadFromCompactBinary(*Iter++, Version) & bOk;
+	return bOk;
+}
+
+}
+
 void FAssetPackageData::SerializeForCacheInternal(FArchive& Ar, FAssetPackageData& PackageData, FAssetRegistryVersion::Type Version)
 {
 	Ar << PackageData.DiskSize;
@@ -982,6 +1003,92 @@ void FAssetPackageData::SetPackageSavedHash(const FIoHash& InHash)
 	PRAGMA_ENABLE_DEPRECATION_WARNINGS;
 }
 
+inline FCbWriter& operator<<(FCbWriter& Writer, const TPair<FIoChunkId, FIoHash>& Value)
+{
+	Writer.BeginArray();
+	Writer << Value.Key << Value.Value;
+	Writer.EndArray();
+	return Writer;
+}
+
+inline bool LoadFromCompactBinary(FCbFieldView Field, TPair<FIoChunkId, FIoHash>& Value)
+{
+	FCbFieldViewIterator Iter = Field.CreateViewIterator();
+	bool bOk = LoadFromCompactBinary(*Iter++, Value.Key);
+	bOk = LoadFromCompactBinary(*Iter++, Value.Value) & bOk;
+	return bOk;
+}
+
+void FAssetPackageData::NetworkWrite(FCbWriter& Writer) const
+{
+	Writer.BeginArray();
+	bool bCookedHash = CookedHash.IsValid();
+	Writer << bCookedHash;
+	if (bCookedHash)
+	{
+		Writer << CookedHash;
+	}
+	Writer << ChunkHashes.Array();
+	Writer << ImportedClasses;
+	Writer << DiskSize;
+	Writer << FileVersionUE;
+	Writer << FileVersionLicenseeUE;
+	TArray<UE::AssetRegistry::FPackageCustomVersion> LocalCustomVersions;
+	LocalCustomVersions.Append(GetCustomVersions());
+	Writer << LocalCustomVersions;
+	Writer << Flags;
+	Writer << static_cast<uint8>(Extension);
+	Writer.EndArray();
+}
+
+bool FAssetPackageData::TryNetworkRead(FCbFieldView Field)
+{
+	FCbFieldViewIterator Iter = Field.CreateViewIterator();
+	bool bCookedHash = false;
+	bool bOk = LoadFromCompactBinary(*Iter++, bCookedHash);
+	if (bCookedHash)
+	{
+		bOk = LoadFromCompactBinary(*Iter++, CookedHash) & bOk;
+	}
+	else
+	{
+		CookedHash = FMD5Hash();
+	}
+	TArray<TPair<FIoChunkId, FIoHash>> ChunkHashesArray;
+	if (LoadFromCompactBinary(*Iter++, ChunkHashesArray))
+	{
+		ChunkHashes.Empty(ChunkHashesArray.Num());
+		for (TPair<FIoChunkId, FIoHash>& Pair : ChunkHashesArray)
+		{
+			ChunkHashes.Add(Pair.Key, Pair.Value);
+		}
+	}
+	else
+	{
+		bOk = false;
+	}
+	bOk = LoadFromCompactBinary(*Iter++, ImportedClasses) & bOk;
+	bOk = LoadFromCompactBinary(*Iter++, DiskSize) & bOk;
+	bOk = LoadFromCompactBinary(*Iter++, FileVersionUE) & bOk;
+	bOk = LoadFromCompactBinary(*Iter++, FileVersionLicenseeUE) & bOk;
+	TArray<UE::AssetRegistry::FPackageCustomVersion> LocalCustomVersions;
+	if (LoadFromCompactBinary(*Iter++, LocalCustomVersions))
+	{
+		SetCustomVersions(LocalCustomVersions);
+	}
+	bOk = LoadFromCompactBinary(*Iter++, Flags) & bOk;
+	uint8 ExtensionInt = 0;
+	if (LoadFromCompactBinary(*Iter++, ExtensionInt) && ExtensionInt < static_cast<uint8>(EPackageExtension::Count))
+	{
+		Extension = static_cast<EPackageExtension>(ExtensionInt);
+	}
+	else
+	{
+		bOk = false;
+	}
+	return bOk;
+}
+
 void FARFilter::PostSerialize(const FArchive& Ar)
 {
 	PRAGMA_DISABLE_DEPRECATION_WARNINGS
diff --git a/Engine/Source/Runtime/CoreUObject/Public/AssetRegistry/AssetData.h b/Engine/Source/Runtime/CoreUObject/Public/AssetRegistry/AssetData.h
--- a/Engine/Source/Runtime/CoreUObject/Public/AssetRegistry/AssetData.h
+++ b/Engine/Source/Runtime/CoreUObject/Public/AssetRegistry/AssetData.h
@@ -709,7 +709,7 @@ PRAGMA_ENABLE_DEPRECATION_WARNINGS
 			});
 	}
 
-	// Note: these functions should only be used for live communication between processing running the same version of the engine.
+	// Note: these functions should only be used for live communication between processes running the same version of the engine.
 	// There is no versioning support
 	COREUOBJECT_API void NetworkWrite(FCbWriter& Writer, bool bWritePackageName) const;
 	COREUOBJECT_API bool TryNetworkRead(FCbFieldView Field, bool bReadPackageName, FName InPackageName);
@@ -853,9 +853,7 @@ inline FAssetRegistryExportPath FAssetData::GetTagValueRef<FAssetRegistryExportP
 	return FoundValue.IsSet() ? FoundValue.AsExportPath() : FAssetRegistryExportPath();
 }
 
-namespace UE
-{
-namespace AssetRegistry
+namespace UE::AssetRegistry
 {
 
 /** Low-memory version of FCustomVersion; holds only Guid and integer version. */
@@ -884,6 +882,18 @@ struct FPackageCustomVersion
 	{
 		return Ar << CustomVersion.Key << CustomVersion.Version;
 	}
+
+private:
+	friend FCbWriter& operator<<(FCbWriter& Writer, const FPackageCustomVersion& Value)
+	{
+		return Value.Write(Writer);
+	}
+	COREUOBJECT_API FCbWriter& Write(FCbWriter& Writer) const;
+	friend bool LoadFromCompactBinary(const FCbFieldView& Field, FPackageCustomVersion& Value)
+	{
+		return Value.TryRead(Field);
+	}
+	COREUOBJECT_API bool TryRead(const FCbFieldView& Field);
 };
 
 /** A handle to a deduplicated, sorted array of FPackageCustomVersion. */
@@ -902,7 +912,6 @@ private:
 	friend class FPackageCustomVersionRegistry;
 };
 
-}
 }
 
 /** A class to hold data about a package on disk, this data is updated on save/load and is not updated when an asset changes in memory */
@@ -1000,6 +1009,11 @@ public:
 	COREUOBJECT_API FIoHash GetPackageSavedHash() const;
 	COREUOBJECT_API void SetPackageSavedHash(const FIoHash& InHash);
 
+	// Note: these functions should only be used for live communication between processes running the same version of the engine.
+	// There is no versioning support
+	COREUOBJECT_API void NetworkWrite(FCbWriter& Writer) const;
+	COREUOBJECT_API bool TryNetworkRead(FCbFieldView Field);
+
 private:
 	FORCEINLINE void SerializeForCacheInternal(FArchive& Ar, FAssetPackageData& PackageData, FAssetRegistryVersion::Type Version);
 
diff --git a/Engine/Source/Runtime/CoreUObject/Public/Misc/PackagePath.h b/Engine/Source/Runtime/CoreUObject/Public/Misc/PackagePath.h
--- a/Engine/Source/Runtime/CoreUObject/Public/Misc/PackagePath.h
+++ b/Engine/Source/Runtime/CoreUObject/Public/Misc/PackagePath.h
@@ -66,6 +66,7 @@ enum class EPackageExtension : uint8
 	BulkDataOptional,
 	BulkDataMemoryMapped,
 	PayloadSidecar,
+	Count,
 };
 inline constexpr int EPackageExtensionCount = static_cast<int>(EPackageExtension::PayloadSidecar) + 1;
 COREUOBJECT_API const TCHAR* LexToString(EPackageExtension PackageExtension);
