Optimized FPropertyVisitorPath::GetPropertyDataPtr

The previous implementation used the Visit API directly to iteratively walk each property in turn until it found a match, which was very slow when compared to being able to access the property data directly.

The optimized implementation requires that an object also provide a ResolveVisitedPathInfo function. These resolver functions are implemented to use direct offsets where possible, which vastly improves performance.

In a synthesized test case this took performance of 16.8m calls to GetPropertyDataPtr from ~7s to ~0.5s.

If a given object has inner properties but can't provide a more optimized implementation than its Visit function, then it can implement ResolveVisitedPathInfo using PropertyVisitorHelpers::ResolveVisitedPathInfo_Generic instead.

#rnx

diff --git a/Engine/Plugins/Experimental/StructUtils/Source/StructUtils/Private/InstancedStruct.cpp b/Engine/Plugins/Experimental/StructUtils/Source/StructUtils/Private/InstancedStruct.cpp
--- a/Engine/Plugins/Experimental/StructUtils/Source/StructUtils/Private/InstancedStruct.cpp
+++ b/Engine/Plugins/Experimental/StructUtils/Source/StructUtils/Private/InstancedStruct.cpp
@@ -602,3 +602,13 @@ EPropertyVisitorControlFlow FInstancedStruct::Visit(FPropertyVisitorPath& Path,
 
 	return ScriptStruct->Visit(Path, StructMemory, InFunc);
 }
+
+void* FInstancedStruct::ResolveVisitedPathInfo(const FPropertyVisitorInfo& Info)
+{
+	if (!ScriptStruct || !StructMemory)
+	{
+		return nullptr;
+	}
+
+	return ScriptStruct->ResolveVisitedPathInfo(StructMemory, Info);
+}
diff --git a/Engine/Plugins/Experimental/StructUtils/Source/StructUtils/Public/InstancedStruct.h b/Engine/Plugins/Experimental/StructUtils/Source/StructUtils/Public/InstancedStruct.h
--- a/Engine/Plugins/Experimental/StructUtils/Source/StructUtils/Public/InstancedStruct.h
+++ b/Engine/Plugins/Experimental/StructUtils/Source/StructUtils/Public/InstancedStruct.h
@@ -157,6 +157,7 @@ public:
 	bool NetSerialize(FArchive& Ar, class UPackageMap* Map, bool& bOutSuccess);
 	bool FindInnerPropertyInstance(FName PropertyName, const FProperty*& OutProp, const void*& OutData) const;
 	EPropertyVisitorControlFlow Visit(FPropertyVisitorPath& Path, const TFunctionRef<EPropertyVisitorControlFlow(const FPropertyVisitorPath& /*Path*/, void* /*Data*/)> InFunc);
+	void* ResolveVisitedPathInfo(const FPropertyVisitorInfo& Info);
 
 	/** Returns struct type. */
 	const UScriptStruct* GetScriptStruct() const
diff --git a/Engine/Source/Runtime/CoreUObject/Private/UObject/Class.cpp b/Engine/Source/Runtime/CoreUObject/Private/UObject/Class.cpp
--- a/Engine/Source/Runtime/CoreUObject/Private/UObject/Class.cpp
+++ b/Engine/Source/Runtime/CoreUObject/Private/UObject/Class.cpp
@@ -812,6 +812,20 @@ EPropertyVisitorControlFlow UStruct::Visit(FPropertyVisitorPath& Path, void* Dat
 	return RetVal;
 }
 
+void* UStruct::ResolveVisitedPathInfo(void* Data, const FPropertyVisitorInfo& Info) const
+{
+	if (Info.ParentStructType) // We don't care if this matches, but we do care that it's set as it identifies a UStruct info
+	{
+		if (const UStruct* PropertyOwnerStruct = Info.Property->GetOwnerStruct();
+			PropertyOwnerStruct && IsChildOf(PropertyOwnerStruct))
+		{
+			return Info.Property->ContainerPtrToValuePtr<void>(Data, Info.PropertyInfo == EPropertyVisitorInfoType::StaticArrayIndex ? Info.Index : 0);
+		}
+	}
+
+	return nullptr;
+}
+
 void UStruct::PreloadChildren(FArchive& Ar)
 {
 	for (UField* Field = Children; Field; Field = Field->Next)
@@ -3544,6 +3558,22 @@ EPropertyVisitorControlFlow UScriptStruct::Visit(FPropertyVisitorPath& Path, voi
 	return RetVal;
 }
 
+void* UScriptStruct::ResolveVisitedPathInfo(void* Data, const FPropertyVisitorInfo& Info) const
+{
+	if (void* InnerData = Super::ResolveVisitedPathInfo(Data, Info))
+	{
+		return InnerData;
+	}
+
+	if (StructFlags & STRUCT_Visitor)
+	{
+		checkf(CppStructOps && CppStructOps->HasVisitor(), TEXT("Expecting to have a visitor implementation when STRUCT_Visitor is set"));
+		return CppStructOps->ResolveVisitedPathInfo(Data, Info);
+	}
+
+	return nullptr;
+}
+
 void UScriptStruct::ClearScriptStruct(void* Dest, int32 ArrayDim) const
 {
 	uint8 *Data = (uint8*)Dest;
diff --git a/Engine/Source/Runtime/CoreUObject/Private/UObject/Property.cpp b/Engine/Source/Runtime/CoreUObject/Private/UObject/Property.cpp
--- a/Engine/Source/Runtime/CoreUObject/Private/UObject/Property.cpp
+++ b/Engine/Source/Runtime/CoreUObject/Private/UObject/Property.cpp
@@ -1350,6 +1350,12 @@ EPropertyVisitorControlFlow FProperty::Visit(FPropertyVisitorPath& Path, void* D
 	return InFunc(Path, Data);
 }
 
+void* FProperty::ResolveVisitedPathInfo(void* Data, const FPropertyVisitorInfo& Info) const
+{
+	FPropertyVisitorPath Path(FPropertyVisitorInfo(this));
+	return PropertyVisitorHelpers::ResolveVisitedPathInfo_Generic(this, Path, Data, Info);
+}
+
 int32 FProperty::SetupOffset()
 {
 	UObject* OwnerUObject = GetOwner<UObject>();
diff --git a/Engine/Source/Runtime/CoreUObject/Private/UObject/PropertyArray.cpp b/Engine/Source/Runtime/CoreUObject/Private/UObject/PropertyArray.cpp
--- a/Engine/Source/Runtime/CoreUObject/Private/UObject/PropertyArray.cpp
+++ b/Engine/Source/Runtime/CoreUObject/Private/UObject/PropertyArray.cpp
@@ -1370,3 +1370,13 @@ EPropertyVisitorControlFlow FArrayProperty::Visit(FPropertyVisitorPath& Path, vo
 	}
 	return RetVal;
 }
+
+void* FArrayProperty::ResolveVisitedPathInfo(void* Data, const FPropertyVisitorInfo& Info) const
+{
+	if (Info.PropertyInfo == EPropertyVisitorInfoType::ContainerIndex && Info.Property == Inner)
+	{
+		return GetValueAddressAtIndex_Direct(Info.Property, Data, Info.Index);
+	}
+
+	return nullptr;
+}
diff --git a/Engine/Source/Runtime/CoreUObject/Private/UObject/PropertyBaseObject.cpp b/Engine/Source/Runtime/CoreUObject/Private/UObject/PropertyBaseObject.cpp
--- a/Engine/Source/Runtime/CoreUObject/Private/UObject/PropertyBaseObject.cpp
+++ b/Engine/Source/Runtime/CoreUObject/Private/UObject/PropertyBaseObject.cpp
@@ -765,8 +765,18 @@ EPropertyVisitorControlFlow FObjectPropertyBase::Visit(FPropertyVisitorPath& Pat
 	{
 		if (const TObjectPtr<UObject> Object = GetObjectPropertyValue(Data))
 		{
-			RetVal = Object->GetClass()->Visit(Path, Object, InFunc);
+			RetVal = Object.GetClass()->Visit(Path, Object.Get(), InFunc);
 		}
 	}
 	return RetVal;
 }
+
+void* FObjectPropertyBase::ResolveVisitedPathInfo(void* Data, const FPropertyVisitorInfo& Info) const
+{
+	if (const TObjectPtr<UObject> Object = GetObjectPropertyValue(Data))
+	{
+		return Object.GetClass()->ResolveVisitedPathInfo(Object.Get(), Info);
+	}
+
+	return nullptr;
+}
diff --git a/Engine/Source/Runtime/CoreUObject/Private/UObject/PropertyMap.cpp b/Engine/Source/Runtime/CoreUObject/Private/UObject/PropertyMap.cpp
--- a/Engine/Source/Runtime/CoreUObject/Private/UObject/PropertyMap.cpp
+++ b/Engine/Source/Runtime/CoreUObject/Private/UObject/PropertyMap.cpp
@@ -1859,3 +1859,13 @@ EPropertyVisitorControlFlow FMapProperty::Visit(FPropertyVisitorPath& Path, void
 	}
 	return RetVal;
 }
+
+void* FMapProperty::ResolveVisitedPathInfo(void* Data, const FPropertyVisitorInfo& Info) const
+{
+	if ((Info.PropertyInfo == EPropertyVisitorInfoType::MapKey && Info.Property == KeyProp) || (Info.PropertyInfo == EPropertyVisitorInfoType::MapValue && Info.Property == ValueProp))
+	{
+		return GetValueAddressAtIndex_Direct(Info.Property, Data, Info.Index);
+	}
+
+	return nullptr;
+}
diff --git a/Engine/Source/Runtime/CoreUObject/Private/UObject/PropertyOptional.cpp b/Engine/Source/Runtime/CoreUObject/Private/UObject/PropertyOptional.cpp
--- a/Engine/Source/Runtime/CoreUObject/Private/UObject/PropertyOptional.cpp
+++ b/Engine/Source/Runtime/CoreUObject/Private/UObject/PropertyOptional.cpp
@@ -708,3 +708,13 @@ EPropertyVisitorControlFlow FOptionalProperty::Visit(FPropertyVisitorPath& Path,
 	}
 	return RetVal;
 }
+
+void* FOptionalProperty::ResolveVisitedPathInfo(void* Data, const FPropertyVisitorInfo& Info) const
+{
+	if (Info.Property == ValueProperty && IsSet(Data))
+	{
+		return Data;
+	}
+
+	return nullptr;
+}
diff --git a/Engine/Source/Runtime/CoreUObject/Private/UObject/PropertySet.cpp b/Engine/Source/Runtime/CoreUObject/Private/UObject/PropertySet.cpp
--- a/Engine/Source/Runtime/CoreUObject/Private/UObject/PropertySet.cpp
+++ b/Engine/Source/Runtime/CoreUObject/Private/UObject/PropertySet.cpp
@@ -1160,3 +1160,13 @@ EPropertyVisitorControlFlow FSetProperty::Visit(FPropertyVisitorPath& Path, void
 	}
 	return RetVal;
 }
+
+void* FSetProperty::ResolveVisitedPathInfo(void* Data, const FPropertyVisitorInfo& Info) const
+{
+	if (Info.PropertyInfo == EPropertyVisitorInfoType::ContainerIndex && Info.Property == ElementProp)
+	{
+		return GetValueAddressAtIndex_Direct(Info.Property, Data, Info.Index);
+	}
+
+	return nullptr;
+}
diff --git a/Engine/Source/Runtime/CoreUObject/Private/UObject/PropertyStruct.cpp b/Engine/Source/Runtime/CoreUObject/Private/UObject/PropertyStruct.cpp
--- a/Engine/Source/Runtime/CoreUObject/Private/UObject/PropertyStruct.cpp
+++ b/Engine/Source/Runtime/CoreUObject/Private/UObject/PropertyStruct.cpp
@@ -647,4 +647,9 @@ EPropertyVisitorControlFlow FStructProperty::Visit(FPropertyVisitorPath& Path, v
 		RetVal = Struct->Visit(Path, Data, InFunc);
 	}
 	return RetVal;
-}
\ No newline at end of file
+}
+
+void* FStructProperty::ResolveVisitedPathInfo(void* Data, const FPropertyVisitorInfo& Info) const
+{
+	return Struct->ResolveVisitedPathInfo(Data, Info);
+}
diff --git a/Engine/Source/Runtime/CoreUObject/Private/UObject/PropertyVisitor.cpp b/Engine/Source/Runtime/CoreUObject/Private/UObject/PropertyVisitor.cpp
--- a/Engine/Source/Runtime/CoreUObject/Private/UObject/PropertyVisitor.cpp
+++ b/Engine/Source/Runtime/CoreUObject/Private/UObject/PropertyVisitor.cpp
@@ -109,32 +109,8 @@ bool FPropertyVisitorPath::Contained(const FPropertyVisitorPath& Other, bool* bI
 
 void* FPropertyVisitorPath::GetPropertyDataPtr(UObject* Object) const
 {
-	void* DataPtr = nullptr;
 	checkf(Object, TEXT("Expecting an valid object"));
-
-	int32 MatchedPathDepth = 0;
-	Object->GetClass()->Visit(Object, [this, &DataPtr, &MatchedPathDepth](const FPropertyVisitorPath& InPath, void* Data)
-	{
-		if (InPath.Num() <= MatchedPathDepth)
-		{
-			// We've returned a level that we previously found a match in; we can stop now
-			return EPropertyVisitorControlFlow::Stop;
-		}
-		bool bIsEqual = false;
-		if (InPath.Contained(*this, &bIsEqual))
-		{
-			MatchedPathDepth = InPath.Num();
-			if(bIsEqual)
-			{
-				DataPtr = Data;
-				return EPropertyVisitorControlFlow::Stop;
-			}
-			return EPropertyVisitorControlFlow::StepInto;
-		}
-		return EPropertyVisitorControlFlow::StepOver;
-	});
-
-	return DataPtr;
+	return PropertyVisitorHelpers::ResolveVisitedPath(Object->GetClass(), Object, *this);
 }
 
 FArchiveSerializedPropertyChain FPropertyVisitorPath::ToSerializedPropertyChain() const
@@ -145,4 +121,4 @@ FArchiveSerializedPropertyChain FPropertyVisitorPath::ToSerializedPropertyChain(
 		Chain.PushProperty(const_cast<FProperty*>(Info.Property), Info.Property->IsEditorOnlyProperty());
 	}
 	return Chain;
-}
\ No newline at end of file
+}
diff --git a/Engine/Source/Runtime/CoreUObject/Public/UObject/Class.h b/Engine/Source/Runtime/CoreUObject/Public/UObject/Class.h
--- a/Engine/Source/Runtime/CoreUObject/Public/UObject/Class.h
+++ b/Engine/Source/Runtime/CoreUObject/Public/UObject/Class.h
@@ -749,6 +749,12 @@ public:
 	 */
 	COREUOBJECT_API virtual EPropertyVisitorControlFlow Visit(FPropertyVisitorPath& Path, void* Data, const TFunctionRef<EPropertyVisitorControlFlow(const FPropertyVisitorPath& /*Path*/, void* /*Data*/)> InFunc) const;
 
+	/**
+	 * Attempt to resolve the given inner path info against this outer struct to get the inner property value.
+	 * @return The inner property value, or null if the path info is incompatible/missing on this struct.
+	 */
+	COREUOBJECT_API virtual void* ResolveVisitedPathInfo(void* Data, const FPropertyVisitorInfo& Info) const;
+
 protected:
 
 	/** Returns the property name from the guid */
@@ -1252,6 +1258,7 @@ public:
 		}
 		/** Structs property visitor signature */
 		virtual EPropertyVisitorControlFlow Visit(FPropertyVisitorPath& Path, void* Data, const TFunctionRef<EPropertyVisitorControlFlow(const FPropertyVisitorPath& /*Path*/, void* /*Data*/)> InFunc) const = 0;
+		virtual void* ResolveVisitedPathInfo(void* Data, const FPropertyVisitorInfo& Info) const = 0;
 
 	private:
 		/** sizeof() of the structure **/
@@ -1650,6 +1657,18 @@ public:
 				return EPropertyVisitorControlFlow::StepOver;
 			}
 		}
+
+		virtual void* ResolveVisitedPathInfo(void* Data, const FPropertyVisitorInfo& Info) const override
+		{
+			if constexpr (TStructOpsTypeTraits<CPPSTRUCT>::WithVisitor)
+			{
+				return ((CPPSTRUCT*)Data)->ResolveVisitedPathInfo(Info);
+			}
+			else
+			{
+				return nullptr;
+			}
+		}
 	};
 
 	/** Template for noexport classes to autoregister before main starts **/
@@ -1899,7 +1918,9 @@ public:
 	virtual COREUOBJECT_API bool FindInnerPropertyInstance(FName PropertyName, const void* Data, const FProperty*& OutProp, const void*& OutData) const;
 
 	/* Custom visit implementation for structs */
+	using Super::Visit;
 	virtual COREUOBJECT_API EPropertyVisitorControlFlow Visit(FPropertyVisitorPath& Path, void* Data, const TFunctionRef<EPropertyVisitorControlFlow(const FPropertyVisitorPath& /*Path*/, void* /*Data*/)> InFunc) const override;
+	virtual COREUOBJECT_API void* ResolveVisitedPathInfo(void* Data, const FPropertyVisitorInfo& Info) const override;
 };
 
 /*-----------------------------------------------------------------------------
diff --git a/Engine/Source/Runtime/CoreUObject/Public/UObject/PropertyOptional.h b/Engine/Source/Runtime/CoreUObject/Public/UObject/PropertyOptional.h
--- a/Engine/Source/Runtime/CoreUObject/Public/UObject/PropertyOptional.h
+++ b/Engine/Source/Runtime/CoreUObject/Public/UObject/PropertyOptional.h
@@ -240,5 +240,6 @@ public:
 	virtual bool SameType(const FProperty* Other) const override;
 
 	virtual EPropertyVisitorControlFlow Visit(FPropertyVisitorPath& Path, void* Data, const TFunctionRef<EPropertyVisitorControlFlow(const FPropertyVisitorPath& /*Path*/, void* /*Data*/)> InFunc) const override;
+	virtual void* ResolveVisitedPathInfo(void* Data, const FPropertyVisitorInfo& Info) const override;
 	// End of FProperty interface
 };
diff --git a/Engine/Source/Runtime/CoreUObject/Public/UObject/PropertyVisitor.h b/Engine/Source/Runtime/CoreUObject/Public/UObject/PropertyVisitor.h
--- a/Engine/Source/Runtime/CoreUObject/Public/UObject/PropertyVisitor.h
+++ b/Engine/Source/Runtime/CoreUObject/Public/UObject/PropertyVisitor.h
@@ -117,7 +117,7 @@ public:
 
 	/**
 	 * Is this property path contained in the specified one
-	 * @param OtherPath property path to check if it is contianed in
+	 * @param OtherPath property path to check if it is contained in
 	 * @param bIsEqual optional parameter to know if it fully match
 	 * @return true if it is contained in the specified path
 	 */
@@ -164,4 +164,71 @@ public:
 
 protected:
 	FPropertyVisitorPath& Path;
-};
\ No newline at end of file
+};
+
+namespace PropertyVisitorHelpers
+{
+
+namespace Private
+{
+
+template <typename Type>
+void* ResolveVisitedPathInfo(const Type* This, void* Data, const FPropertyVisitorInfo& Info)
+{
+	return This->ResolveVisitedPathInfo(Data, Info);
+}
+
+} // namespace Private
+
+/**
+ * Given a FPropertyVisitorPath, attempt to resolve that to a valid data pointer.
+ * RootObject is required to implement ResolveVisitedPathInfo to provide the resolver logic.
+ */
+template <typename Type>
+void* ResolveVisitedPath(const Type* RootObject, void* RootData, const FPropertyVisitorPath& Path)
+{
+	void* FoundPropertyData = nullptr;
+	if (const TArray<FPropertyVisitorInfo>& PathArray = Path.GetPath();
+		PathArray.Num() > 0)
+	{
+		FoundPropertyData = Private::ResolveVisitedPathInfo(RootObject, RootData, PathArray[0]);
+		for (int32 PathIndex = 1; FoundPropertyData && PathIndex < PathArray.Num(); ++PathIndex)
+		{
+			const FPropertyVisitorInfo& PreviousInfo = PathArray[PathIndex - 1];
+			FoundPropertyData = Private::ResolveVisitedPathInfo(PreviousInfo.Property, FoundPropertyData, PathArray[PathIndex]);
+		}
+	}
+	return FoundPropertyData;
+}
+
+/**
+ * A generic implementation of ResolveVisitedPathInfo that uses Visit to find the property data pointer.
+ * This may be used as the ResolveVisitedPathInfo implementation for your type if it doesn't have a more optimized version.
+ */
+template <typename Type>
+void* ResolveVisitedPathInfo_Generic(Type* This, FPropertyVisitorPath& Path, void* Data, const FPropertyVisitorInfo& Info)
+{
+	void* FoundInnerData = nullptr;
+	This->Visit(Path, Data, [&FoundInnerData, &Info, InnerPathDepth = Path.Num() + 1](const FPropertyVisitorPath& InnerPath, void* InnerData)
+	{
+		if (InnerPath.Num() < InnerPathDepth)
+		{
+			return EPropertyVisitorControlFlow::StepInto;
+		}
+		if (Info.Property == InnerPath.Top().Property && Info.PropertyInfo == InnerPath.Top().PropertyInfo && Info.Index == InnerPath.Top().Index)
+		{
+			FoundInnerData = InnerData;
+			return EPropertyVisitorControlFlow::Stop;
+		}
+		return EPropertyVisitorControlFlow::StepOver;
+	});
+	return FoundInnerData;
+}
+template <typename Type>
+void* ResolveVisitedPathInfo_Generic(Type* This, void* Data, const FPropertyVisitorInfo& Info)
+{
+	FPropertyVisitorPath Path;
+	return ResolveVisitedPathInfo_Generic(This, Path, Data, Info);
+}
+
+} // namespace PropertyVisitorHelpers
diff --git a/Engine/Source/Runtime/CoreUObject/Public/UObject/UnrealType.h b/Engine/Source/Runtime/CoreUObject/Public/UObject/UnrealType.h
--- a/Engine/Source/Runtime/CoreUObject/Public/UObject/UnrealType.h
+++ b/Engine/Source/Runtime/CoreUObject/Public/UObject/UnrealType.h
@@ -342,6 +342,12 @@ public:
 	 */
 	COREUOBJECT_API virtual EPropertyVisitorControlFlow Visit(FPropertyVisitorPath& Path, void* Data, const TFunctionRef<EPropertyVisitorControlFlow(const FPropertyVisitorPath& /*Path*/, void* /*Data*/)> InFunc) const;
 
+	/**
+	 * Attempt to resolve the given inner path info against this outer property to get the inner property value.
+	 * @return The inner property value, or null if the path info is incompatible/missing on this property.
+	 */
+	COREUOBJECT_API virtual void* ResolveVisitedPathInfo(void* Data, const FPropertyVisitorInfo& Info) const;
+
 private:
 	/** Set the alignment offset for this property 
 	 * @return the size of the structure including this newly added property
@@ -2633,6 +2639,7 @@ public:
 	virtual void InstanceSubobjects( void* Data, void const* DefaultData, UObject* Owner, struct FObjectInstancingGraph* InstanceGraph ) override;
 	virtual bool SameType(const FProperty* Other) const override;
 	virtual EPropertyVisitorControlFlow Visit(FPropertyVisitorPath& Path, void* Data, const TFunctionRef<EPropertyVisitorControlFlow(const FPropertyVisitorPath& /*Path*/, void* /*Data*/)> InFunc) const override;
+	virtual void* ResolveVisitedPathInfo(void* Data, const FPropertyVisitorInfo& Info) const override;
 	// End of FProperty interface
 
 	// FObjectPropertyBase interface
@@ -3627,6 +3634,7 @@ public:
 	virtual void SaveTypeName(UE::FPropertyTypeNameBuilder& Type) const override;
 	virtual bool CanSerializeFromTypeName(UE::FPropertyTypeName Type) const override;
 	virtual EPropertyVisitorControlFlow Visit(FPropertyVisitorPath& Path, void* Data, const TFunctionRef<EPropertyVisitorControlFlow(const FPropertyVisitorPath& /*Path*/, void* /*Data*/)> InFunc) const override;
+	virtual void* ResolveVisitedPathInfo(void* Data, const FPropertyVisitorInfo& Info) const override;
 	// End of FProperty interface
 
 	FString GetCPPTypeCustom(FString* ExtendedTypeText, uint32 CPPExportFlags, const FString& InnerTypeText, const FString& InInnerExtendedTypeText) const;
@@ -3759,6 +3767,7 @@ public:
 	virtual void SaveTypeName(UE::FPropertyTypeNameBuilder& Type) const override;
 	virtual bool CanSerializeFromTypeName(UE::FPropertyTypeName Type) const override;
 	virtual EPropertyVisitorControlFlow Visit(FPropertyVisitorPath& Path, void* Data, const TFunctionRef<EPropertyVisitorControlFlow(const FPropertyVisitorPath& /*Path*/, void* /*Data*/)> InFunc) const override;
+	virtual void* ResolveVisitedPathInfo(void* Data, const FPropertyVisitorInfo& Info) const override;
 	// End of FProperty interface
 
 	FString GetCPPTypeCustom(FString* ExtendedTypeText, uint32 CPPExportFlags, const FString& KeyTypeText, const FString& InKeyExtendedTypeText, const FString& ValueTypeText, const FString& InValueExtendedTypeText) const;
@@ -3886,6 +3895,7 @@ public:
 	virtual void SaveTypeName(UE::FPropertyTypeNameBuilder& Type) const override;
 	virtual bool CanSerializeFromTypeName(UE::FPropertyTypeName Type) const override;
 	virtual EPropertyVisitorControlFlow Visit(FPropertyVisitorPath& Path, void* Data, const TFunctionRef<EPropertyVisitorControlFlow(const FPropertyVisitorPath& /*Path*/, void* /*Data*/)> InFunc) const override;
+	virtual void* ResolveVisitedPathInfo(void* Data, const FPropertyVisitorInfo& Info) const override;
 	// End of FProperty interface
 
 	FString GetCPPTypeCustom(FString* ExtendedTypeText, uint32 CPPExportFlags, const FString& ElementTypeText, const FString& InElementExtendedTypeText) const;
@@ -6010,6 +6020,7 @@ public:
 	virtual void SaveTypeName(UE::FPropertyTypeNameBuilder& Type) const override;
 	virtual bool CanSerializeFromTypeName(UE::FPropertyTypeName Type) const override;
 	virtual EPropertyVisitorControlFlow Visit(FPropertyVisitorPath& Path, void* Data, const TFunctionRef<EPropertyVisitorControlFlow(const FPropertyVisitorPath& /*Path*/, void* /*Data*/)> InFunc) const override;
+	virtual void* ResolveVisitedPathInfo(void* Data, const FPropertyVisitorInfo& Info) const override;
 	// End of FProperty interface
 
 	bool FindInnerPropertyInstance(FName PropertyName, const void* Data, const FProperty*& OutProp, const void*& OutData) const;
