Added loading of missing structs into IDOs

#rnx

diff --git a/Engine/Source/Runtime/CoreUObject/Internal/UObject/PropertyPathNameTree.h b/Engine/Source/Runtime/CoreUObject/Internal/UObject/PropertyPathNameTree.h
--- a/Engine/Source/Runtime/CoreUObject/Internal/UObject/PropertyPathNameTree.h
+++ b/Engine/Source/Runtime/CoreUObject/Internal/UObject/PropertyPathNameTree.h
@@ -6,6 +6,7 @@
 #include "Containers/Map.h"
 #include "Templates/TypeHash.h"
 #include "UObject/PropertyPathName.h"
+#include "UObject/PropertyTag.h"
 #include "UObject/PropertyTypeName.h"
 
 #define UE_API COREUOBJECT_API
@@ -46,7 +47,30 @@ class FPropertyPathNameTree
 		}
 	};
 
+	struct FValue
+	{
+		TUniquePtr<FPropertyPathNameTree> SubTree;
+		FPropertyTag Tag;
+	};
+
 public:
+	struct FConstNode
+	{
+		const FValue* Value = nullptr;
+		inline explicit operator bool() const { return !!Value; }
+		inline const FPropertyPathNameTree* GetSubTree() const { return Value ? Value->SubTree.Get() : nullptr; }
+		inline const FPropertyTag* GetTag() const { return Value && !Value->Tag.Name.IsNone() ? &Value->Tag : nullptr; }
+	};
+
+	struct FNode : FConstNode
+	{
+		using FConstNode::GetSubTree;
+		using FConstNode::GetTag;
+		inline FPropertyPathNameTree* GetSubTree() { return const_cast<FPropertyPathNameTree*>(FConstNode::GetSubTree()); }
+		inline FPropertyTag* GetTag() { return const_cast<FPropertyTag*>(FConstNode::GetTag()); }
+		UE_API void SetTag(const FPropertyTag& Tag);
+	};
+
 	FPropertyPathNameTree() = default;
 	FPropertyPathNameTree(FPropertyPathNameTree&&) = default;
 	FPropertyPathNameTree(const FPropertyPathNameTree&) = delete;
@@ -60,7 +84,7 @@ public:
 	UE_API void Empty();
 
 	/** Adds the path to the tree. Keeps any existing nodes that match both name and type. */
-	UE_API void Add(const FPropertyPathName& Path, int32 StartIndex = 0);
+	UE_API FNode Add(const FPropertyPathName& Path, int32 StartIndex = 0);
 
 	/**
 	 * Finds the path within the tree.
@@ -68,18 +92,18 @@ public:
 	 * @param OutSubTree   If non-null, this is set to the maybe-null sub-tree at the path if the path is found.
 	 * @param Path         Path to search within the tree.
 	 * @param StartIndex   Index of the segment within Path at which to start searching.
-	 * @return true if the path exists within the tree, otherwise false.
+	 * @return An accessor for the node at the path if it exists.
 	 */
-	UE_API bool Find(FPropertyPathNameTree** OutSubTree, const FPropertyPathName& Path, int32 StartIndex = 0);
-	inline bool Find(const FPropertyPathNameTree** OutSubTree, const FPropertyPathName& Path, int32 StartIndex = 0) const
+	UE_API FNode Find(const FPropertyPathName& Path, int32 StartIndex = 0);
+	inline FConstNode Find(const FPropertyPathName& Path, int32 StartIndex = 0) const
 	{
-		return const_cast<FPropertyPathNameTree*>(this)->Find(const_cast<FPropertyPathNameTree**>(OutSubTree), Path, StartIndex);
+		return {const_cast<FPropertyPathNameTree*>(this)->Find(Path, StartIndex).Value};
 	}
 
 	class FConstIterator
 	{
 	public:
-		inline explicit FConstIterator(const TMap<FKey, TUniquePtr<FPropertyPathNameTree>>::TConstIterator& InNodeIt)
+		inline explicit FConstIterator(const TMap<FKey, FValue>::TConstIterator& InNodeIt)
 			: NodeIt(InNodeIt)
 		{
 		}
@@ -93,16 +117,16 @@ public:
 
 		inline FName GetName() const { return NodeIt.Key().Name; }
 		inline FPropertyTypeName GetType() const { return NodeIt.Key().Type; }
-		inline const FPropertyPathNameTree* GetSubTree() const { return NodeIt.Value().Get(); }
+		inline FConstNode GetNode() const { return {&NodeIt.Value()}; }
 
 	private:
-		TMap<FKey, TUniquePtr<FPropertyPathNameTree>>::TConstIterator NodeIt;
+		TMap<FKey, FValue>::TConstIterator NodeIt;
 	};
 
 	inline FConstIterator CreateConstIterator() const { return FConstIterator(Nodes.CreateConstIterator()); }
 
 private:
-	TMap<FKey, TUniquePtr<FPropertyPathNameTree>> Nodes;
+	TMap<FKey, FValue> Nodes;
 };
 
 } // UE
diff --git a/Engine/Source/Runtime/CoreUObject/Private/UObject/Class.cpp b/Engine/Source/Runtime/CoreUObject/Private/UObject/Class.cpp
--- a/Engine/Source/Runtime/CoreUObject/Private/UObject/Class.cpp
+++ b/Engine/Source/Runtime/CoreUObject/Private/UObject/Class.cpp
@@ -1766,7 +1766,7 @@ void UStruct::SerializeVersionedTaggedProperties(FStructuredArchive::FSlot Slot,
 
 					if (UnknownPropertyTree)
 					{
-						UnknownPropertyTree->Add(SerializeContext->SerializedPropertyPath);
+						UnknownPropertyTree->Add(SerializeContext->SerializedPropertyPath).SetTag(Tag);
 
 						const bool bSerializeValue = PropertyTypeContainsStructOrEnum(Tag.GetType());
 						const bool bRestoreOverrideOperation = Tag.OverrideOperation != EOverriddenPropertyOperation:: None && !Property && !UnderlyingArchive.IsTransacting();
@@ -1775,7 +1775,7 @@ void UStruct::SerializeVersionedTaggedProperties(FStructuredArchive::FSlot Slot,
 						if (bSerializeValue || bRestoreOverrideOperation)
 						{
 							TUniquePtr<FField> TempField(FField::TryConstruct(Tag.Type, {}, Tag.Name, RF_NoFlags));
-							if (FProperty* TempProperty = CastField<FProperty>(TempField.Get()); TempProperty && TempProperty->LoadTypeName(Tag.GetType()))
+							if (FProperty* TempProperty = CastField<FProperty>(TempField.Get()); TempProperty && TempProperty->LoadTypeName(Tag.GetType(), &Tag))
 							{
 								TempProperty->Link(UnderlyingArchive);
 
diff --git a/Engine/Source/Runtime/CoreUObject/Private/UObject/InstanceDataObjectUtils.cpp b/Engine/Source/Runtime/CoreUObject/Private/UObject/InstanceDataObjectUtils.cpp
--- a/Engine/Source/Runtime/CoreUObject/Private/UObject/InstanceDataObjectUtils.cpp
+++ b/Engine/Source/Runtime/CoreUObject/Private/UObject/InstanceDataObjectUtils.cpp
@@ -269,10 +269,10 @@ namespace UE
 			{
 				FPropertyPathName Path;
 				Path.Push({NAME_IDOMapKey});
-				PropertyTree->Find(&KeyTree, Path);
+				KeyTree = PropertyTree->Find(Path).GetSubTree();
 				Path.Pop();
 				Path.Push({NAME_IDOMapValue});
-				PropertyTree->Find(&ValueTree, Path);
+				ValueTree = PropertyTree->Find(Path).GetSubTree();
 				Path.Pop();
 			}
 
@@ -363,8 +363,9 @@ namespace UE
 				{
 					FPropertyPathName Path;
 					Path.Push({TemplateProperty->GetFName(), Type});
-					if (PropertyTree->Find(&SubTree, Path))
+					if (FPropertyPathNameTree::FConstNode Node = PropertyTree->Find(Path))
 					{
+						SubTree = Node.GetSubTree();
 						SuperPropertyPathsFromTree.Add(MoveTemp(Path));
 					}
 				}
@@ -399,10 +400,10 @@ namespace UE
 				{
 					// Construct a property from the type and try to use it to serialize the value.
 					FField* Field = FField::TryConstruct(Type.GetName(), Result, Name, RF_NoFlags);
-					if (FProperty* Property = CastField<FProperty>(Field); Property && Property->LoadTypeName(Type))
+					if (FProperty* Property = CastField<FProperty>(Field); Property && Property->LoadTypeName(Type, It.GetNode().GetTag()))
 					{
 						MarkPropertyAsLoose(Property);
-						ConvertToInstanceDataObjectProperty(Property, Type, Result, It.GetSubTree());
+						ConvertToInstanceDataObjectProperty(Property, Type, Result, It.GetNode().GetSubTree());
 						LooseInstanceDataObjectProperties.Add(Property);
 						continue;
 					}
diff --git a/Engine/Source/Runtime/CoreUObject/Private/UObject/PropertyPathNameTree.cpp b/Engine/Source/Runtime/CoreUObject/Private/UObject/PropertyPathNameTree.cpp
--- a/Engine/Source/Runtime/CoreUObject/Private/UObject/PropertyPathNameTree.cpp
+++ b/Engine/Source/Runtime/CoreUObject/Private/UObject/PropertyPathNameTree.cpp
@@ -2,6 +2,7 @@
 
 #include "UObject/PropertyPathNameTree.h"
 
+#include "Misc/StringBuilder.h"
 #include "Templates/TypeHash.h"
 
 namespace UE
@@ -12,45 +13,65 @@ void FPropertyPathNameTree::Empty()
 	Nodes.Empty();
 }
 
-void FPropertyPathNameTree::Add(const FPropertyPathName& Path, int32 StartIndex)
+FPropertyPathNameTree::FNode FPropertyPathNameTree::Add(const FPropertyPathName& Path, int32 StartIndex)
 {
 	if (const int32 SegmentCount = Path.GetSegmentCount(); StartIndex < SegmentCount)
 	{
 		FPropertyPathNameSegment Segment = Path.GetSegment(StartIndex);
-		TUniquePtr<FPropertyPathNameTree>& Child = Nodes.FindOrAdd({Segment.Name, Segment.Type});
+		FValue& Child = Nodes.FindOrAdd({Segment.Name, Segment.Type});
 		if (++StartIndex < SegmentCount)
 		{
-			if (!Child)
+			if (!Child.SubTree)
 			{
-				Child = MakeUnique<FPropertyPathNameTree>();
+				Child.SubTree = MakeUnique<FPropertyPathNameTree>();
 			}
-			Child->Add(Path, StartIndex);
+			return Child.SubTree->Add(Path, StartIndex);
 		}
+		return {&Child};
 	}
+	return {};
 }
 
-bool FPropertyPathNameTree::Find(FPropertyPathNameTree** OutSubTree, const FPropertyPathName& Path, int32 StartIndex)
+FPropertyPathNameTree::FNode FPropertyPathNameTree::Find(const FPropertyPathName& Path, int32 StartIndex)
 {
 	if (const int32 SegmentCount = Path.GetSegmentCount(); StartIndex < SegmentCount)
 	{
 		FPropertyPathNameSegment Segment = Path.GetSegment(StartIndex);
-		if (TUniquePtr<FPropertyPathNameTree>* Child = Nodes.Find({Segment.Name, Segment.Type}))
+		if (FValue* Child = Nodes.Find({Segment.Name, Segment.Type}))
 		{
 			if (++StartIndex >= SegmentCount)
 			{
-				if (OutSubTree)
-				{
-					*OutSubTree = Child->Get();
-				}
-				return true;
+				return {Child};
 			}
-			if (*Child)
+			if (Child->SubTree)
 			{
-				return (*Child)->Find(OutSubTree, Path, StartIndex);
+				return Child->SubTree->Find(Path, StartIndex);
 			}
 		}
 	}
-	return false;
+	return {};
+}
+
+void FPropertyPathNameTree::FNode::SetTag(const FPropertyTag& Tag)
+{
+	if (FValue* LocalValue = const_cast<FValue*>(Value))
+	{
+		if (LocalValue->Tag.Name.IsNone())
+		{
+			// Copy the tag and reset values that may vary between tags with the same path.
+			LocalValue->Tag = Tag;
+			LocalValue->Tag.Size = 0;
+			LocalValue->Tag.ArrayIndex = INDEX_NONE;
+			LocalValue->Tag.SizeOffset = INDEX_NONE;
+			LocalValue->Tag.BoolVal = 0;
+		}
+		else
+		{
+			ensureMsgf(LocalValue->Tag.GetType() == Tag.GetType() && LocalValue->Tag.SerializeType == Tag.SerializeType,
+				TEXT("Tag mismatch in property path name tree for property %s of type %s."),
+				*WriteToString<32>(Tag.Name), *WriteToString<64>(Tag.GetType()));
+		}
+	}
 }
 
 } // UE
diff --git a/Engine/Source/Runtime/CoreUObject/Private/UObject/PropertyStruct.cpp b/Engine/Source/Runtime/CoreUObject/Private/UObject/PropertyStruct.cpp
--- a/Engine/Source/Runtime/CoreUObject/Private/UObject/PropertyStruct.cpp
+++ b/Engine/Source/Runtime/CoreUObject/Private/UObject/PropertyStruct.cpp
@@ -571,10 +571,23 @@ bool FStructProperty::LoadTypeName(UE::FPropertyTypeName Type, const FPropertyTa
 		{
 			return true;
 		}
-		// TODO: Look up the struct based on the guid.
-		// TODO: Use the fallback struct if allowed.
 	}
 
+	// TODO: Look up the struct based on the guid.
+	//const FName StructGuidName = Type.GetParameterName(1);
+	//if (FGuid StructGuid; !StructGuidName.IsNone() && FGuid::Parse(StructGuidName.ToString(), StructGuid) && StructGuid.IsValid())
+	//{
+	//}
+
+#if WITH_EDITORONLY_DATA
+	if (Tag && Tag->SerializeType == EPropertyTagSerializeType::Property)
+	{
+		Struct = GetFallbackStruct();
+		SetMetaData(UE::NAME_OriginalType, *WriteToString<256>(Type.GetParameter(0)));
+		return true;
+	}
+#endif // WITH_EDITORONLY_DATA
+
 	return false;
 }
 
diff --git a/Engine/Source/Runtime/CoreUObject/Tests/UObject/PropertyPathNameTreeTest.cpp b/Engine/Source/Runtime/CoreUObject/Tests/UObject/PropertyPathNameTreeTest.cpp
--- a/Engine/Source/Runtime/CoreUObject/Tests/UObject/PropertyPathNameTreeTest.cpp
+++ b/Engine/Source/Runtime/CoreUObject/Tests/UObject/PropertyPathNameTreeTest.cpp
@@ -49,9 +49,9 @@ TEST_CASE_NAMED(FPropertyPathNameTreeTest, "CoreUObject::PropertyPathNameTree",
 		PathName.Push({CountName});
 		FPropertyPathNameTree Tree;
 		Tree.Add(PathName);
-		FPropertyPathNameTree* SubTree = &Tree;
-		CHECK(Tree.Find(&SubTree, PathName));
-		CHECK_FALSE(SubTree);
+		FPropertyPathNameTree::FNode Node = Tree.Find(PathName);
+		CHECK(Node);
+		CHECK_FALSE(Node.GetSubTree());
 		const FPropertyPathNameTree::FConstIterator First = Tree.CreateConstIterator();
 		FPropertyPathNameTree::FConstIterator It = Tree.CreateConstIterator();
 		CHECK(First == It);
@@ -60,7 +60,9 @@ TEST_CASE_NAMED(FPropertyPathNameTreeTest, "CoreUObject::PropertyPathNameTree",
 		{
 			CHECK(It.GetName() == CountName);
 			CHECK(It.GetType().IsEmpty());
-			CHECK_FALSE(It.GetSubTree());
+			CHECK(It.GetNode());
+			CHECK_FALSE(It.GetNode().GetSubTree());
+			CHECK_FALSE(It.GetNode().GetTag());
 			++It;
 			CHECK_FALSE(It);
 			CHECK_FALSE(First == It);
@@ -74,9 +76,9 @@ TEST_CASE_NAMED(FPropertyPathNameTreeTest, "CoreUObject::PropertyPathNameTree",
 		PathName.Push({CountName, IntType});
 		FPropertyPathNameTree Tree;
 		Tree.Add(PathName);
-		FPropertyPathNameTree* SubTree = &Tree;
-		CHECK(Tree.Find(&SubTree, PathName));
-		CHECK_FALSE(SubTree);
+		FPropertyPathNameTree::FNode Node = Tree.Find(PathName);
+		CHECK(Node);
+		CHECK_FALSE(Node.GetSubTree());
 		const FPropertyPathNameTree::FConstIterator First = Tree.CreateConstIterator();
 		FPropertyPathNameTree::FConstIterator It = Tree.CreateConstIterator();
 		CHECK(First == It);
@@ -85,7 +87,8 @@ TEST_CASE_NAMED(FPropertyPathNameTreeTest, "CoreUObject::PropertyPathNameTree",
 		{
 			CHECK(It.GetName() == CountName);
 			CHECK(It.GetType() == IntType);
-			CHECK_FALSE(It.GetSubTree());
+			CHECK(It.GetNode());
+			CHECK_FALSE(It.GetNode().GetSubTree());
 			++It;
 			CHECK_FALSE(It);
 			CHECK_FALSE(First == It);
@@ -99,15 +102,15 @@ TEST_CASE_NAMED(FPropertyPathNameTreeTest, "CoreUObject::PropertyPathNameTree",
 		PathName.Push({CountName, IntType, 7});
 		FPropertyPathNameTree Tree;
 		Tree.Add(PathName);
-		CHECK(Tree.Find(nullptr, PathName));
+		CHECK(Tree.Find(PathName));
 		PathName.SetIndex(5);
-		CHECK(Tree.Find(nullptr, PathName));
+		CHECK(Tree.Find(PathName));
 		PathName.SetIndex(3);
 		Tree.Add(PathName);
 		PathName.SetIndex(INDEX_NONE);
-		const FPropertyPathNameTree* SubTree = &Tree;
-		CHECK(Tree.Find(&SubTree, PathName));
-		CHECK_FALSE(SubTree);
+		FPropertyPathNameTree::FConstNode Node = Tree.Find(PathName);
+		CHECK(Node);
+		CHECK_FALSE(Node.GetSubTree());
 		const FPropertyPathNameTree::FConstIterator First = Tree.CreateConstIterator();
 		FPropertyPathNameTree::FConstIterator It = Tree.CreateConstIterator();
 		CHECK(First == It);
@@ -116,7 +119,8 @@ TEST_CASE_NAMED(FPropertyPathNameTreeTest, "CoreUObject::PropertyPathNameTree",
 		{
 			CHECK(It.GetName() == CountName);
 			CHECK(It.GetType() == IntType);
-			CHECK_FALSE(It.GetSubTree());
+			CHECK(It.GetNode());
+			CHECK_FALSE(It.GetNode().GetSubTree());
 			++It;
 			CHECK_FALSE(It);
 			CHECK_FALSE(First == It);
@@ -133,19 +137,21 @@ TEST_CASE_NAMED(FPropertyPathNameTreeTest, "CoreUObject::PropertyPathNameTree",
 		FPropertyPathNameTree Tree;
 		Tree.Add(PathNameInt);
 		Tree.Add(PathNameFloat);
-		CHECK(Tree.Find(nullptr, PathNameInt));
-		CHECK(Tree.Find(nullptr, PathNameFloat));
+		CHECK(Tree.Find(PathNameInt));
+		CHECK(Tree.Find(PathNameFloat));
 
 		FPropertyPathNameTree::FConstIterator It = Tree.CreateConstIterator();
 		CHECKED_IF(It)
 		{
 			CHECK(It.GetType() == IntType);
-			CHECK_FALSE(It.GetSubTree());
+			CHECK(It.GetNode());
+			CHECK_FALSE(It.GetNode().GetSubTree());
 			++It;
 			CHECKED_IF(It)
 			{
 				CHECK(It.GetType() == FloatType);
-				CHECK_FALSE(It.GetSubTree());
+				CHECK(It.GetNode());
+				CHECK_FALSE(It.GetNode().GetSubTree());
 				++It;
 				CHECK_FALSE(It);
 			}
@@ -161,19 +167,21 @@ TEST_CASE_NAMED(FPropertyPathNameTreeTest, "CoreUObject::PropertyPathNameTree",
 		FPropertyPathNameTree Tree;
 		Tree.Add(PathNameCount);
 		Tree.Add(PathNameSize);
-		CHECK(Tree.Find(nullptr, PathNameCount));
-		CHECK(Tree.Find(nullptr, PathNameSize));
+		CHECK(Tree.Find(PathNameCount));
+		CHECK(Tree.Find(PathNameSize));
 
 		FPropertyPathNameTree::FConstIterator It = Tree.CreateConstIterator();
 		CHECKED_IF(It)
 		{
 			CHECK(It.GetName() == CountName);
-			CHECK_FALSE(It.GetSubTree());
+			CHECK(It.GetNode());
+			CHECK_FALSE(It.GetNode().GetSubTree());
 			++It;
 			CHECKED_IF(It)
 			{
 				CHECK(It.GetName() == SizeName);
-				CHECK_FALSE(It.GetSubTree());
+				CHECK(It.GetNode());
+				CHECK_FALSE(It.GetNode().GetSubTree());
 				++It;
 				CHECK_FALSE(It);
 			}
@@ -191,16 +199,16 @@ TEST_CASE_NAMED(FPropertyPathNameTreeTest, "CoreUObject::PropertyPathNameTree",
 		FPropertyPathNameTree Tree;
 		Tree.Add(PathNameA);
 		Tree.Add(PathNameB);
-		CHECK(Tree.Find(nullptr, ParentPathName));
-		CHECK(Tree.Find(nullptr, PathNameA));
-		CHECK(Tree.Find(nullptr, PathNameB));
+		CHECK(Tree.Find(ParentPathName));
+		CHECK(Tree.Find(PathNameA));
+		CHECK(Tree.Find(PathNameB));
 
 		FPropertyPathNameTree::FConstIterator ParentIt = Tree.CreateConstIterator();
 		CHECKED_IF(ParentIt)
 		{
 			CHECK(ParentIt.GetName() == NAME_Vector);
 			CHECK(ParentIt.GetType() == VectorType);
-			const FPropertyPathNameTree* ChildTree = ParentIt.GetSubTree();
+			const FPropertyPathNameTree* ChildTree = ParentIt.GetNode().GetSubTree();
 			++ParentIt;
 			CHECK_FALSE(ParentIt);
 			CHECKED_IF(ChildTree)
@@ -209,12 +217,14 @@ TEST_CASE_NAMED(FPropertyPathNameTreeTest, "CoreUObject::PropertyPathNameTree",
 				CHECKED_IF(ChildIt)
 				{
 					CHECK(ChildIt.GetName() == CountName);
-					CHECK_FALSE(ChildIt.GetSubTree());
+					CHECK(ChildIt.GetNode());
+					CHECK_FALSE(ChildIt.GetNode().GetSubTree());
 					++ChildIt;
 					CHECKED_IF(ChildIt)
 					{
 						CHECK(ChildIt.GetName() == SizeName);
-						CHECK_FALSE(ChildIt.GetSubTree());
+						CHECK(ChildIt.GetNode());
+						CHECK_FALSE(ChildIt.GetNode().GetSubTree());
 						++ChildIt;
 						CHECK_FALSE(ChildIt);
 					}
@@ -222,17 +232,16 @@ TEST_CASE_NAMED(FPropertyPathNameTreeTest, "CoreUObject::PropertyPathNameTree",
 			}
 		}
 
-		FPropertyPathNameTree* ChildTree = nullptr;
-		CHECK(Tree.Find(&ChildTree, ParentPathName));
+		FPropertyPathNameTree::FNode ChildNode = Tree.Find(ParentPathName);
+		CHECK(ChildNode);
+		const FPropertyPathNameTree* ChildTree = ChildNode.GetSubTree();
 		CHECKED_IF(ChildTree)
 		{
-			CHECK(ChildTree->Find(nullptr, PathNameA, 1));
-			CHECK(ChildTree->Find(nullptr, PathNameB, 1));
+			CHECK(ChildTree->Find(PathNameA, 1));
+			CHECK(ChildTree->Find(PathNameB, 1));
 		}
 
-		FPropertyPathNameTree* MissingTree = &Tree;
-		CHECK_FALSE(Tree.Find(&MissingTree, PathNameA, 1));
-		CHECK(MissingTree == &Tree);
+		CHECK_FALSE(Tree.Find(PathNameA, 1));
 	}
 }
 
