Fixes two issues in zenloader around package export invalidation:
- Previously CollectUnreachableObjects would mark GC'd sub-objects as invalid. This could break future loads of the package while the package (but not all its sub-objects) are still in memory. (UE-213934)
- Previously importing a package which may convert/upgrade it's exports could result in the old objects lingering. These lingering objects might cause issues during cook/saving (UE-213910, note while this bug is specifically for -nozenloader, the changes in this CL fix the same issue for -zenloader)

To fix the above, when a package is finished loading we synchronize the FAsyncPackage2 exports with LFLinkerLoad's exports and ensure the FAsyncPackage2 contains at least the same invalid/bExportLoadFailed=true exports. These exports and any export that is a sub-object to an invalid export, are stored in the GlobalImportStore. When we next directly load or import a package, we will filter out any invalid exports. Note, we synchronize the invalid exports at the end of package loading because Preload and PostLoad codepaths have opportunities to alter exports, so we wait for those to complete before we do any recording. When a package is removed from memory, we remove all invalid exports from the GlobalImportStoreto prevent the InvalidExport map growing unbounded.

CollectUnreachableObjects no longer marks GC'd objects as invalid, and instead we rely on the logic described above. This prevents us from incorrectly marking valid objects, we should reload when later requested, as invalid.

Two tests were added to validate the two issues above, and in doing so the existing AsyncLoading test files that were scattered in the test directory have been moved to a test specific directory for clarity

#virtualized

diff --git a/Engine/Source/Runtime/CoreUObject/Private/Serialization/AsyncLoading2.cpp b/Engine/Source/Runtime/CoreUObject/Private/Serialization/AsyncLoading2.cpp
--- a/Engine/Source/Runtime/CoreUObject/Private/Serialization/AsyncLoading2.cpp
+++ b/Engine/Source/Runtime/CoreUObject/Private/Serialization/AsyncLoading2.cpp
@@ -690,7 +690,7 @@ private:
 
 		FPublicExportMap(const FPublicExportMap&) = delete;
 
-		FPublicExportMap(FPublicExportMap&& Other)
+		FPublicExportMap(FPublicExportMap&& Other) noexcept
 		{
 			Allocation = Other.Allocation;
 			Count = Other.Count;
@@ -701,7 +701,7 @@ private:
 
 		FPublicExportMap& operator=(const FPublicExportMap&) = delete;
 
-		FPublicExportMap& operator=(FPublicExportMap&& Other)
+		FPublicExportMap& operator=(FPublicExportMap&& Other) noexcept
 		{
 			if (Count > 1)
 			{
@@ -1072,6 +1072,10 @@ private:
 	TMap<FPackageObjectIndex, UObject*> ScriptObjects;
 	// All currently loaded public export objects from any loaded package
 	TMap<int32, FPublicExportKey> ObjectIndexToPublicExport;
+#if WITH_EDITOR
+	// All exports that have been marked invalid and should not be loaded
+	TMultiMap<FPackageId, uint64> InvalidExports;
+#endif
 
 	// Process all the deferred deletion for packages that are finished loading
 	void FlushDeferredDeletePackagesQueue();
@@ -1082,6 +1086,10 @@ public:
 		Packages.Reserve(32768);
 		ScriptObjects.Reserve(32768);
 		ObjectIndexToPublicExport.Reserve(32768);
+
+#if WITH_EDITOR
+		InvalidExports.Reserve(1024);
+#endif
 	}
 
 	int32 GetStoredPackagesCount() const
@@ -1349,6 +1357,10 @@ public:
 					ObjectIndexToPublicExport.Remove(ObjectIndex);
 				}
 			}
+#if WITH_EDITOR
+			// Now that the package is removed, remove all export invalidations.
+			InvalidExports.Remove(PackageId);
+#endif
 		}
 	}
 
@@ -1503,6 +1515,19 @@ public:
 		int32 ObjectIndex = GUObjectArray.ObjectToIndex(Object);
 		FPublicExportKey Key = FPublicExportKey::MakeKey(PackageId, ExportHash);
 
+#if WITH_EDITOR
+		if (IsInvalidExport(Key))
+		{
+			UE_LOG(LogStreaming, Verbose,
+				TEXT("FGlobalImportStore::StoreGlobalObject: The constructed public export object '%s' with index %d and id %s:0x%llX was marked as an invalid export and will not be stored."),
+				Object ? *Object->GetFullName() : TEXT("null"),
+				ObjectIndex,
+				*FormatPackageId(Key.GetPackageId()), 
+				Key.GetExportHash());
+			return;
+		}
+#endif
+
 		UObject* ExistingObject = FindPublicExportObjectUnchecked(Key);
 		if (ExistingObject && ExistingObject != Object)
 		{
@@ -1513,7 +1538,8 @@ public:
 				TEXT("The existing object will be replaced since it or its package was most likely renamed after it was loaded the first time."),
 				Object ? *Object->GetFullName() : TEXT("null"),
 				ObjectIndex,
-				*FormatPackageId(Key.GetPackageId()), Key.GetExportHash(),
+				*FormatPackageId(Key.GetPackageId()), 
+				Key.GetExportHash(),
 				*ExistingObject->GetFullName(),
 				ExistingObject->GetFlags(),
 				int(ExistingObject->GetInternalFlags()),
@@ -1549,6 +1575,19 @@ public:
 		ObjectIndexToPublicExport.Add(ObjectIndex, Key);
 	}
 
+#if WITH_EDITOR
+	bool IsInvalidExport(FPublicExportKey Key)
+	{
+		uint64* Pair = InvalidExports.FindPair(Key.GetPackageId(), Key.GetExportHash());
+		return Pair != nullptr;
+	}
+
+	void StoreInvalidExport(FPackageId PackageId, uint64 ExportHash)
+	{
+		InvalidExports.Add(PackageId, ExportHash);
+	}
+#endif
+
 	void FindAllScriptObjects(bool bVerifyOnly);
 	void RegistrationComplete();
 
@@ -1691,6 +1730,18 @@ struct FPackageImportStore
 		GlobalImportStore.StoreGlobalObject(PackageId, ExportHash, Object);
 	}
 
+#if WITH_EDITOR
+	inline bool IsInvalidExport(FPackageId PackageId, uint64 ExportHash)
+	{
+		return GlobalImportStore.IsInvalidExport(FPublicExportKey::MakeKey(PackageId, ExportHash));
+	}
+
+	inline void StoreInvalidExport(FPackageId PackageId, uint64 ExportHash)
+	{
+		return GlobalImportStore.StoreInvalidExport(PackageId, ExportHash);
+	}
+#endif
+
 public:
 	FLoadedPackageRef& AddImportedPackageReference(FPackageId ImportedPackageId, FName PackageNameIfKnown)
 	{
@@ -2830,6 +2881,14 @@ struct FAsyncPackage2
 
 	void ClearImportedPackages();
 
+#if WITH_EDITOR
+	/*
+	* Marks exports as failed if LinkerLoad(if using) has marked the same export as failed. 
+	* All found failed exports in the FAsyncPackage2 and FLinkerLoad are recorded in the GlobalImportStore
+	*/
+	void SyncAndStoreInvalidExports();
+#endif
+
 	/**
 	 * @return Time load begun. This is NOT the time the load was requested in the case of other pending requests.
 	 */
@@ -5914,8 +5973,19 @@ bool FAsyncPackage2::CreateLinkerLoadExports(FAsyncLoadingThreadState2& ThreadSt
 			continue;
 		}
 #endif
+
 		FObjectExport& LinkerExport = LinkerLoadState->Linker->ExportMap[ExportIndex];
 		FExportObject& ExportObject = Data.Exports[ExportIndex];
+		uint64_t HeaderPublicExportHash = HeaderData.ExportMap[ExportIndex].PublicExportHash;
+#if WITH_EDITOR
+		if (ImportStore.IsInvalidExport(Desc.UPackageId, HeaderPublicExportHash))
+		{
+			ExportObject.Object = nullptr;
+			ExportObject.bExportLoadFailed = true;
+			continue;
+		}
+#endif
+
 		if (UObject* Object = LinkerLoadState->Linker->CreateExport(ExportIndex))
 		{
 			checkf(!Object->IsUnreachable(), TEXT("Trying to store an unreachable object '%s' in the import store"), *Object->GetFullName());
@@ -5923,6 +5993,8 @@ bool FAsyncPackage2::CreateLinkerLoadExports(FAsyncLoadingThreadState2& ThreadSt
 			ExportObject.bWasFoundInMemory = true; // Make sure that the async flags are cleared in ClearConstructedObjects
 			EInternalObjectFlags FlagsToSet = EInternalObjectFlags::Async;
 			uint64 PublicExportHash = LinkerLoadState->LinkerLoadHeaderData.ExportMap[ExportIndex].PublicExportHash;
+			checkf(HeaderPublicExportHash == PublicExportHash, TEXT("PublicExportHash for object %s doesn't match the hash in the header for package %s"), *Object->GetFullName(), *HeaderData.PackageName.ToString());
+
 			if (Desc.bCanBeImported && PublicExportHash)
 			{
 				FlagsToSet |= EInternalObjectFlags::LoaderImport;
@@ -8028,6 +8100,8 @@ EAsyncPackageState::Type FAsyncLoadingThread2::ProcessLoadedPackagesFromGameThre
 							}
 						}
 					}
+
+					Package->SyncAndStoreInvalidExports();
 				}
 #endif
 
@@ -8678,12 +8752,6 @@ void FAsyncLoadingThread2::CollectUnreachableObjects(
 			{
 				const int32 CachedLinkerIndex = Object->GetLinkerIndex();
 				Object->SetLinker(nullptr, INDEX_NONE);
-				// As we are garbaging the object, mark it as invalid in the linker
-				// Either it is now truly invalid to access this object 
-				// (i.e the asset ran upgraded, migrated away from this object and doesn't hold any references to it anymore.)
-				// or we are gc'ing the entire asset in which case the linker will eventually get purged and those entry won't be marked invalid anymore on recreation
-				FObjectExport& ObjExport = ObjectLinker->ExportMap[CachedLinkerIndex];
-				ObjExport.bExportLoadFailed = true;
 			}
 #endif // ALT2_ENABLE_LINKERLOAD_SUPPORT
 		});
@@ -8969,6 +9037,96 @@ void FAsyncPackage2::ClearImportedPackages()
 	}
 }
 
+#if WITH_EDITOR
+void FAsyncPackage2::SyncAndStoreInvalidExports()
+{
+	TArrayView<FExportObject> PackageExports = Data.Exports;
+	const int32 ExportCount = PackageExports.Num();
+
+#if ALT2_ENABLE_LINKERLOAD_SUPPORT
+	FLinkerLoad* Linker = LinkerLoadState->Linker;
+	check(Linker->ExportMap.Num() == ExportCount);
+#endif
+
+	// The visited state is inverted from what you'd normally expect so we may more efficiently iterate later
+	constexpr int InlineElementCount = 32;
+	TBitArray<TInlineAllocator<InlineElementCount>> NotVisitedIndices(true, ExportCount);
+	TSet<int32, DefaultKeyFuncs<int32>, TInlineSetAllocator<InlineElementCount>> InvalidIndices;
+	InvalidIndices.Reserve(32); // There are normally not very many invalid exports per package
+
+	// First scan all our exports ensure we list all the same invalid exports as LinkerLoad
+	for (int32 ExportIndex = 0; ExportIndex < ExportCount; ++ExportIndex)
+	{
+		FExportObject& ExportObject = PackageExports[ExportIndex];
+#if ALT2_ENABLE_LINKERLOAD_SUPPORT
+		ExportObject.bExportLoadFailed |= Linker->ExportMap[ExportIndex].bExportLoadFailed;
+#endif
+		if (ExportObject.bExportLoadFailed)
+		{
+			// Only mark invalid indices as visited since we know we don't need to check this export's Outer
+			NotVisitedIndices[ExportIndex] = false;
+			InvalidIndices.Add(ExportIndex);
+		}
+	}
+	
+#if ALT2_ENABLE_LINKERLOAD_SUPPORT
+	if (!InvalidIndices.IsEmpty())
+	{
+		TRACE_CPUPROFILER_EVENT_SCOPE(SyncAndStoreInvalidExports_GatherAndStoreInvalidChildExports);
+
+		for (TConstSetBitIterator<TInlineAllocator<InlineElementCount>> It(NotVisitedIndices); It;)
+		{
+			const int32 ExportIndex = It.GetIndex();
+
+			TArray<int32, TInlineAllocator<InlineElementCount>> Stack;
+			Stack.Add(ExportIndex);
+			NotVisitedIndices[ExportIndex] = false;
+
+			// Since we pre-visited all invalid exports we are only iterating over known valid exports 
+			// and only need to check their outers (i.e. are they a child of an invalid export)
+			bool bIsInvalid = false;
+			FPackageIndex Outer = Linker->ExportMap[ExportIndex].OuterIndex;
+			while (Outer.IsExport())
+			{
+				const int32 OuterIndex = Outer.ToExport();
+				Stack.Add(OuterIndex);
+
+				if (NotVisitedIndices[OuterIndex])
+				{
+					bIsInvalid = InvalidIndices.Contains(OuterIndex);
+					break;
+				}
+				NotVisitedIndices[OuterIndex] = false;
+
+				if (Linker->ExportMap[OuterIndex].bExportLoadFailed)
+				{
+					bIsInvalid = true;
+					break;
+				}
+				Outer = Linker->ExportMap[OuterIndex].OuterIndex;
+			}
+
+			if (bIsInvalid)
+			{
+				InvalidIndices.Append(Stack);
+			}
+
+			// We modified It so re-create it (there isn't a non-const iterator form)
+			new(&It) TConstSetBitIterator<TInlineAllocator<InlineElementCount>>(NotVisitedIndices, ExportIndex);
+		}
+#elif WITH_EDITOR
+	UE_LOG(LogStreaming, Warning, TEXT("Building WITH_EDITOR but not ALT2_ENABLE_LINKERLOAD_SUPPORT is missing support for marking child exports of invalid exports. We may load invalid exports unnecessarily."));
+#endif
+
+		for (int32 Index : InvalidIndices)
+		{ 
+			const uint64 ExportHash = HeaderData.ExportMap[Index].PublicExportHash;
+			ImportStore.GlobalImportStore.StoreInvalidExport(Desc.UPackageId, ExportHash);
+		}
+	}
+}
+#endif
+
 void FAsyncPackage2::ClearConstructedObjects()
 {
 	TRACE_CPUPROFILER_EVENT_SCOPE(ClearConstructedObjects);
@@ -9179,7 +9337,7 @@ EAsyncPackageState::Type FAsyncPackage2::PostLoadInstances(FAsyncLoadingThreadSt
 	{
 		const FExportObject& Export = Data.Exports[PostLoadInstanceIndex++];
 
-		if (!(Export.bFiltered | Export.bExportLoadFailed))
+		if (!(Export.bFiltered || Export.bExportLoadFailed))
 		{
 			UClass* ObjClass = Export.Object->GetClass();
 			ObjClass->PostLoadInstance(Export.Object);
diff --git a/Engine/Source/Runtime/CoreUObject/Private/UObject/Class.cpp b/Engine/Source/Runtime/CoreUObject/Private/UObject/Class.cpp
--- a/Engine/Source/Runtime/CoreUObject/Private/UObject/Class.cpp
+++ b/Engine/Source/Runtime/CoreUObject/Private/UObject/Class.cpp
@@ -1102,6 +1102,7 @@ void UStruct::Link(FArchive& Ar, bool bRelinkExistingProperties)
 	// Discard old wrapper objects used by property grids
 	for (UPropertyWrapper* Wrapper : PropertyWrappers)
 	{
+		FLinkerLoad::InvalidateExport(Wrapper);
 		Wrapper->Rename(nullptr, GetTransientPackage(), REN_ForceNoResetLoaders | REN_DoNotDirty | REN_DontCreateRedirectors | REN_NonTransactional);
 		Wrapper->RemoveFromRoot();
 	}
@@ -2015,7 +2016,8 @@ void UStruct::ConvertUFieldsToFFields()
 			{
 				Children = OldField->Next;
 			}
-			// Move the old UProperty to the transient package and rename it to something unique
+			// Mark the old UProperty as invalid, move it to the transient package and rename it to something unique
+			FLinkerLoad::InvalidateExport(OldField);
 			OldField->Rename(*MakeUniqueObjectName(GetTransientPackage(), OldField->GetClass()).ToString(), GetTransientPackage(), REN_ForceNoResetLoaders | REN_DoNotDirty | REN_DontCreateRedirectors | REN_NonTransactional);
 			OldField->RemoveFromRoot();
 		}
diff --git a/Engine/Source/Runtime/Engine/Private/Tests/Loading/AsyncLoadingTests.cpp b/Engine/Source/Runtime/Engine/Private/Tests/Loading/AsyncLoadingTests.cpp
--- a/Engine/Source/Runtime/Engine/Private/Tests/Loading/AsyncLoadingTests.cpp
+++ b/Engine/Source/Runtime/Engine/Private/Tests/Loading/AsyncLoadingTests.cpp
@@ -6,6 +6,9 @@
 #include "AssetRegistry/AssetRegistryState.h"
 #include "AssetRegistry/IAssetRegistry.h"
 #include "Async/ParallelFor.h"
+#include "AsyncLoadingTests_Shared.h"
+#include "Engine/UserDefinedStruct.h"
+#include "UObject/UObjectHash.h"
 
 #if WITH_DEV_AUTOMATION_TESTS
 
@@ -62,7 +65,7 @@ bool FThreadSafeAsyncLoadingTest::RunTest(const FString& Parameters)
 IMPLEMENT_SIMPLE_AUTOMATION_TEST(FCircularDependencyLoadingTest, TEXT(TEST_NAME_ROOT ".LoadBlueprintWithCircularDependencyTest"), EAutomationTestFlags::EditorContext | EAutomationTestFlags::EngineFilter)
 bool FCircularDependencyLoadingTest::RunTest(const FString& Parameters)
 {
-	static constexpr const TCHAR* ActorWithCircularReferences = TEXT("/Game/Tests/Core/AssetLoading/BlueprintActorWithCircularReferences");
+	static constexpr const TCHAR* ActorWithCircularReferences = TEXT("/Game/Tests/Core/AssetLoading/RecursiveLoads/BlueprintActorWithCircularReferences");
 
 	UPackage* Package = LoadPackage(nullptr, ActorWithCircularReferences, LOAD_None);
 	TestTrue(TEXT("The object should have been properly loaded recursively"), Package != nullptr);
@@ -75,5 +78,167 @@ bool FCircularDependencyLoadingTest::RunTest(const FString& Parameters)
 	return true;
 }
 
+/**
+ * PackagePath1 leads to a user defined structure containing UFields which will be forcibly converted to FField when loading. The UFields should be marked as invalid and the loader should not consider them for loading.
+ */
+IMPLEMENT_SIMPLE_AUTOMATION_TEST(FLoadingTests_ImportPackageConvertedOnLoad, TEXT(TEST_NAME_ROOT ".InvalidExports.ImportPackageConvertedOnLoad"), EAutomationTestFlags::EditorContext | EAutomationTestFlags::EngineFilter)
+bool FLoadingTests_ImportPackageConvertedOnLoad::RunTest(const FString& Parameters)
+{
+	static constexpr const TCHAR* PackagePath1 = TEXT("/Game/Tests/Core/AssetLoading/InvalidatingExports/OldProcMeshData");
+	static constexpr const TCHAR* PackagePath2 = TEXT("/Game/Tests/Core/AssetLoading/InvalidatingExports/BP_UsingUStructWithUnconvertedFields");
+	static constexpr const TCHAR* ObjectName = TEXT("ProcMeshData");
+	typedef TSet<TPair<FName, FName>> FObjPathAndClassNameSet;
+	FObjPathAndClassNameSet ExpectedObjects;
+	FObjPathAndClassNameSet ActualObjects;
+
+	auto VerifyLoadedObjects = [this](const FObjPathAndClassNameSet& ExpectedObjects, const FObjPathAndClassNameSet& ActualObjects)
+		{
+			TestTrue(FString::Printf(TEXT("Expected the same number of objects in package %s after reloading. Expected: %d != Actual: %d"), PackagePath1, ExpectedObjects.Num(), ActualObjects.Num()), ExpectedObjects.Num() == ActualObjects.Num());
+			for (const TPair<FName, FName>& Pair : ExpectedObjects)
+			{
+				TestTrue(FString::Printf(TEXT("Missing obj %s (type: %s) after reloading %s"), *Pair.Key.ToString(), *Pair.Value.ToString(), PackagePath1), ActualObjects.Find(Pair) != nullptr);
+			}
+		};
+
+	// Load our package and populate the ExpectedObjects set
+	{
+		UPackage* Package = LoadPackage(nullptr, PackagePath1, LOAD_None);
+		TestTrue(FString::Printf(TEXT("Failed to load package at %s"), PackagePath1), Package != nullptr);
+
+		ForEachObjectWithPackage(Package,
+			[&ExpectedObjects](UObject* Object)
+			{
+				ExpectedObjects.Add({ Object->GetFName(), Object->GetClass()->GetFName() });
+				return true;
+			}
+		);
+
+		FLoadingTestsScope::GarbageCollect({ {PackagePath1} }, *this);
+	}
+
+	// Load a package that imports the objects from the package at PackagePath1 and ensure we still load only the objects we expect
+	{
+		UPackage* Package = LoadPackage(nullptr, PackagePath2, LOAD_None);
+		TestTrue(FString::Printf(TEXT("Failed to load package at %s"), PackagePath2), Package != nullptr);
+
+		// Path2 imports PackagePath1 so we should be able to still find our Obj and our PackagePath1 package
+		UUserDefinedStruct* Obj = FindFirstObject<UUserDefinedStruct>(ObjectName);
+		TestTrue(FString::Printf(TEXT("Failed to find expected object %s in package %s"), ObjectName, PackagePath1), Obj != nullptr);
+		if (Obj == nullptr)
+		{
+			return false;
+		}
+
+		UPackage* PackagePath1Package = Obj->GetPackage();
+		ForEachObjectWithPackage(PackagePath1Package,
+			[&ActualObjects](UObject* Object)
+			{
+				ActualObjects.Add({ Object->GetFName(), Object->GetClass()->GetFName() });
+				return true;
+			}
+		);
+
+		VerifyLoadedObjects(ExpectedObjects, ActualObjects);
+		ActualObjects.Empty();
+
+		// Note, we are only GC'ing the importing package.
+		Package->ClearFlags(RF_Standalone);
+		FLoadingTestsScope::GarbageCollect({ { PackagePath2 } }, *this);
+	}
+
+	// Load the old package from memory and ensure we don't have any new objects created from importing
+	{
+		UPackage* Package = LoadPackage(nullptr, PackagePath1, LOAD_None);
+		TestTrue(FString::Printf(TEXT("Failed to load package at %s"), PackagePath1), Package != nullptr);
+
+		ForEachObjectWithPackage(Package,
+			[&ActualObjects](UObject* Object)
+			{
+				ActualObjects.Add({ Object->GetFName(), Object->GetClass()->GetFName() });
+				return true;
+			}
+		);
+
+		VerifyLoadedObjects(ExpectedObjects, ActualObjects);
+		ActualObjects.Empty();
+
+		FLoadingTestsScope::GarbageCollect({ { PackagePath2 } }, *this);
+	}
+
+	return true;
+}
+
+
+/**
+ * PackagePath1 leads to a user defined structure containing UFields which will be forcibly converted to FField when loading.
+ * When a sub-object is deleted, if the package is requested again we should reload the missing object and not create objects marked as invalid
+ */
+IMPLEMENT_SIMPLE_AUTOMATION_TEST(FLoadingTests_ReloadDestroyedExport, TEXT(TEST_NAME_ROOT ".InvalidExports.ReloadDestroyedExport"), EAutomationTestFlags::EditorContext | EAutomationTestFlags::EngineFilter)
+bool FLoadingTests_ReloadDestroyedExport::RunTest(const FString& Parameters)
+{
+	static constexpr const TCHAR* PackagePath1 = TEXT("/Game/Tests/Core/AssetLoading/InvalidatingExports/OldProcMeshData");
+	static constexpr const TCHAR* PackagePath2 = TEXT("/Game/Tests/Core/AssetLoading/InvalidatingExports/BP_UsingUStructWithUnconvertedFields");
+	static constexpr const TCHAR* ObjectName = TEXT("ProcMeshData");
+	typedef TSet<TPair<FName, FName>> FObjPathAndClassNameSet;
+	FObjPathAndClassNameSet ExpectedObjects;
+	FObjPathAndClassNameSet ActualObjects;
+
+	auto VerifyLoadedObjects = [this](const FObjPathAndClassNameSet& ExpectedObjects, const FObjPathAndClassNameSet& ActualObjects)
+		{
+			TestTrue(FString::Printf(TEXT("Expected the same number of objects in package %s after reloading. Expected: %d != Actual: %d"), PackagePath1, ExpectedObjects.Num(), ActualObjects.Num()), ExpectedObjects.Num() == ActualObjects.Num());
+			for (const TPair<FName, FName>& Pair : ExpectedObjects)
+			{
+				TestTrue(FString::Printf(TEXT("Missing obj %s (type: %s) after reloading %s"), *Pair.Key.ToString(), *Pair.Value.ToString(), PackagePath1), ActualObjects.Find(Pair) != nullptr);
+			}
+		};
+
+	// Load our package and populate the ExpectedObjects set
+	{
+		UPackage* Package = LoadPackage(nullptr, PackagePath1, LOAD_None);
+		TestTrue(FString::Printf(TEXT("Failed to load package at %s"), PackagePath1), Package != nullptr);
+
+		ForEachObjectWithPackage(Package,
+			[&ExpectedObjects](UObject* Object)
+			{
+				ExpectedObjects.Add({ Object->GetFName(), Object->GetClass()->GetFName() });
+				return true;
+			}
+		);
+
+		// Find a sub-object from the package and mark it for garbage collection, but do not destroy the package entirely
+		UUserDefinedStruct* Obj = FindObject<UUserDefinedStruct>(Package, ObjectName);
+		TestTrue(FString::Printf(TEXT("Failed to find expected object %s in package %s"), ObjectName, PackagePath1), Obj != nullptr);
+		Obj->ClearFlags(RF_Standalone);
+
+		// Note we are not using FLoadingTestsScope::GarbageCollect as that will mark all sub-objects and the package for collection
+		CollectGarbage(RF_Standalone);
+	}
+
+	// Reloading the package should re-populate any deleted sub-objects as long as the objects weren't marked as invalid to load
+	{
+		UPackage* Package = LoadPackage(nullptr, PackagePath1, LOAD_None);
+		TestTrue(FString::Printf(TEXT("Failed to load package at %s"), PackagePath1), Package != nullptr);
+
+		ForEachObjectWithPackage(Package,
+			[&ActualObjects](UObject* Object)
+			{
+				ActualObjects.Add({ Object->GetFName(), Object->GetClass()->GetFName() });
+				return true;
+			}
+		);
+
+		// Our deleted object should have been restored
+		UUserDefinedStruct* Obj = FindObject<UUserDefinedStruct>(Package, ObjectName);
+		TestTrue(FString::Printf(TEXT("Failed to find expected object %s in package %s"), ObjectName, PackagePath1), Obj != nullptr);
+
+		VerifyLoadedObjects(ExpectedObjects, ActualObjects);
+		ActualObjects.Empty();
+
+		FLoadingTestsScope::GarbageCollect({ {PackagePath1} }, *this);
+	}
+
+	return true;
+}
+
 #undef TEST_NAME_ROOT
 #endif // WITH_DEV_AUTOMATION_TESTS
\ No newline at end of file
diff --git a/Engine/Source/Runtime/Engine/Private/Tests/Loading/AsyncLoadingTests_Shared.cpp b/Engine/Source/Runtime/Engine/Private/Tests/Loading/AsyncLoadingTests_Shared.cpp
--- a/Engine/Source/Runtime/Engine/Private/Tests/Loading/AsyncLoadingTests_Shared.cpp
+++ b/Engine/Source/Runtime/Engine/Private/Tests/Loading/AsyncLoadingTests_Shared.cpp
@@ -50,6 +50,11 @@ void FLoadingTestsScope::SavePackages()
 }
 
 void FLoadingTestsScope::GarbageCollect()
+{
+	GarbageCollect(PackageNames, AutomationTest);
+}
+
+void FLoadingTestsScope::GarbageCollect(const TArray<FString>& PackageNames, FAutomationTestBase& AutomationTest)
 {
 	TArray<FString> ObjectPaths;
 
@@ -75,7 +80,7 @@ void FLoadingTestsScope::GarbageCollect()
 	// Make sure everything we gathered can be properly found
 	for (const FString& ObjectPath : ObjectPaths)
 	{
-		checkf(FindObject<UObject>(nullptr, *ObjectPath) != nullptr, TEXT("%s should be present in memory"), *ObjectPath);
+		AutomationTest.TestTrue(FString::Printf(TEXT("%s should be present in memory"), *ObjectPath), FindObject<UObject>(nullptr, *ObjectPath) != nullptr);
 	}
 
 	// GC and make sure everything gets cleaned up before loading
@@ -84,7 +89,7 @@ void FLoadingTestsScope::GarbageCollect()
 	// Now make sure everything is gone
 	for (const FString& ObjectPath : ObjectPaths)
 	{
-		checkf(FindObject<UObject>(nullptr, *ObjectPath) == nullptr, TEXT("%s should have been garbage collected"), *ObjectPath);
+		AutomationTest.TestTrue(FString::Printf(TEXT("%s should have been garbage collected"), *ObjectPath), FindObject<UObject>(nullptr, *ObjectPath) == nullptr);
 	}
 }
 
diff --git a/Engine/Source/Runtime/Engine/Private/Tests/Loading/AsyncLoadingTests_Shared.h b/Engine/Source/Runtime/Engine/Private/Tests/Loading/AsyncLoadingTests_Shared.h
--- a/Engine/Source/Runtime/Engine/Private/Tests/Loading/AsyncLoadingTests_Shared.h
+++ b/Engine/Source/Runtime/Engine/Private/Tests/Loading/AsyncLoadingTests_Shared.h
@@ -105,6 +105,7 @@ public:
 	void LoadObjects();
 	void CleanupObjects();
 	void GarbageCollect();
+	static void GarbageCollect(const TArray<FString>& PackageNames, FAutomationTestBase& AutomationTest);
 
 	FLoadingTestsScope(FAutomationTestBase* InAutomationTest, TFunction<void (FLoadingTestsScope&)> InMutateObjects = nullptr)
 		: AutomationTest(*InAutomationTest)
