Disabled persistent UObjects pool when gc.MaxObjectsNotConsideredByGC ini setting is 0

diff --git a/Engine/Source/Runtime/CoreUObject/Private/UObject/UObjectBase.cpp b/Engine/Source/Runtime/CoreUObject/Private/UObject/UObjectBase.cpp
--- a/Engine/Source/Runtime/CoreUObject/Private/UObject/UObjectBase.cpp
+++ b/Engine/Source/Runtime/CoreUObject/Private/UObject/UObjectBase.cpp
@@ -1019,15 +1019,14 @@ void UObjectBaseInit()
 		// Maximum number of UObjects in the editor
 		GConfig->GetInt(TEXT("/Script/Engine.GarbageCollectionSettings"), TEXT("gc.MaxObjectsInEditor"), MaxUObjects, GEngineIni);
 #endif
+	}
 
+	if (!MaxObjectsNotConsideredByGC)
+	{
+		//Disable persistent UObjects pool if there are 0 objects not considered by GC
 		GUObjectAllocator.DisablePersistentAllocator();
 	}
 
-#if IS_PROGRAM
-	// always disable peristent allocator for programs even if they require cooked data
-	GUObjectAllocator.DisablePersistentAllocator();
-#endif
-
 	// Log what we're doing to track down what really happens as log in LaunchEngineLoop doesn't report those settings in pristine form.
 	UE_LOG(LogInit, Log, TEXT("%s for max %d objects, including %i objects not considered by GC."), 
 		bPreAllocateUObjectArray ? TEXT("Pre-allocating") : TEXT("Presizing"), MaxUObjects, MaxObjectsNotConsideredByGC);
