Disable persistent uobjects pool when  FUObjectArray::DisableDisregardForGC() is called

diff --git a/Engine/Source/Runtime/CoreUObject/Private/UObject/UObjectArray.cpp b/Engine/Source/Runtime/CoreUObject/Private/UObject/UObjectArray.cpp
--- a/Engine/Source/Runtime/CoreUObject/Private/UObject/UObjectArray.cpp
+++ b/Engine/Source/Runtime/CoreUObject/Private/UObject/UObjectArray.cpp
@@ -203,6 +203,7 @@ void FUObjectArray::DisableDisregardForGC()
 	}
 
 	MaxObjectsNotConsideredByGC = 0;
+	GUObjectAllocator.DisablePersistentAllocator();
 	ObjFirstGCIndex = 0;
 	ObjLastNonGCIndex = -1;
 	if (IsOpenForDisregardForGC())
diff --git a/Engine/Source/Runtime/CoreUObject/Private/UObject/UObjectBase.cpp b/Engine/Source/Runtime/CoreUObject/Private/UObject/UObjectBase.cpp
--- a/Engine/Source/Runtime/CoreUObject/Private/UObject/UObjectBase.cpp
+++ b/Engine/Source/Runtime/CoreUObject/Private/UObject/UObjectBase.cpp
@@ -950,14 +950,6 @@ static FAutoConsoleVariableRef CMaxObjectsNotConsideredByGC(
 	ECVF_Default
 	);
 
-static int32 GSizeOfPermanentObjectPool;
-static FAutoConsoleVariableRef CSizeOfPermanentObjectPool(
-	TEXT("gc.SizeOfPermanentObjectPool"),
-	GSizeOfPermanentObjectPool,
-	TEXT("Placeholder console variable, currently not used in runtime."),
-	ECVF_Default
-	);
-
 static int32 GMaxObjectsInEditor;
 static FAutoConsoleVariableRef CMaxObjectsInEditor(
 	TEXT("gc.MaxObjectsInEditor"),
@@ -1021,7 +1013,7 @@ void UObjectBaseInit()
 #endif
 	}
 
-	if (!MaxObjectsNotConsideredByGC)
+	if (MaxObjectsNotConsideredByGC == 0)
 	{
 		//Disable persistent UObjects pool if there are 0 objects not considered by GC
 		GUObjectAllocator.DisablePersistentAllocator();
