Fix use of potentially unset optional LinkerLoadState in FAsyncPackage2::SyncAndStoreInvalidExports.

#rnx

diff --git a/Engine/Source/Runtime/CoreUObject/Private/Serialization/AsyncLoading2.cpp b/Engine/Source/Runtime/CoreUObject/Private/Serialization/AsyncLoading2.cpp
--- a/Engine/Source/Runtime/CoreUObject/Private/Serialization/AsyncLoading2.cpp
+++ b/Engine/Source/Runtime/CoreUObject/Private/Serialization/AsyncLoading2.cpp
@@ -9044,8 +9044,12 @@ void FAsyncPackage2::SyncAndStoreInvalidExports()
 	const int32 ExportCount = PackageExports.Num();
 
 #if ALT2_ENABLE_LINKERLOAD_SUPPORT
-	FLinkerLoad* Linker = LinkerLoadState->Linker;
-	check(Linker->ExportMap.Num() == ExportCount);
+	FLinkerLoad* Linker = nullptr; 
+	if (LinkerLoadState.IsSet())
+	{
+		Linker = LinkerLoadState->Linker;
+		check(Linker->ExportMap.Num() == ExportCount);
+	}
 #endif
 
 	// The visited state is inverted from what you'd normally expect so we may more efficiently iterate later
@@ -9059,7 +9063,7 @@ void FAsyncPackage2::SyncAndStoreInvalidExports()
 	{
 		FExportObject& ExportObject = PackageExports[ExportIndex];
 #if ALT2_ENABLE_LINKERLOAD_SUPPORT
-		ExportObject.bExportLoadFailed |= Linker->ExportMap[ExportIndex].bExportLoadFailed;
+		ExportObject.bExportLoadFailed |= Linker ? Linker->ExportMap[ExportIndex].bExportLoadFailed : false;
 #endif
 		if (ExportObject.bExportLoadFailed)
 		{
@@ -9070,7 +9074,7 @@ void FAsyncPackage2::SyncAndStoreInvalidExports()
 	}
 	
 #if ALT2_ENABLE_LINKERLOAD_SUPPORT
-	if (!InvalidIndices.IsEmpty())
+	if (!InvalidIndices.IsEmpty() && Linker)
 	{
 		TRACE_CPUPROFILER_EVENT_SCOPE(SyncAndStoreInvalidExports_GatherAndStoreInvalidChildExports);
 
