Made UObjectBaseUtility::GetVersePath() virtual

#rnx

diff --git a/Engine/Source/Runtime/CoreUObject/Public/UObject/UObjectBaseUtility.h b/Engine/Source/Runtime/CoreUObject/Public/UObject/UObjectBaseUtility.h
--- a/Engine/Source/Runtime/CoreUObject/Public/UObject/UObjectBaseUtility.h
+++ b/Engine/Source/Runtime/CoreUObject/Public/UObject/UObjectBaseUtility.h
@@ -477,7 +477,7 @@ public:
 	 *
 	 * @return The VersePath of the object
 	 */
-	COREUOBJECT_API UE::Core::FVersePath GetVersePath() const;
+	COREUOBJECT_API virtual UE::Core::FVersePath GetVersePath() const;
 
 	/** 
 	 * Legacy function, has the same behavior as GetPackage
