Fixed constraints in TFieldPath and TWeakFieldPtr.
Fixed various constness violations.

diff --git a/Engine/Plugins/Experimental/Web/WebAPI/Source/WebAPIBlueprintGraph/Private/K2Node_WebAPIOperation.cpp b/Engine/Plugins/Experimental/Web/WebAPI/Source/WebAPIBlueprintGraph/Private/K2Node_WebAPIOperation.cpp
--- a/Engine/Plugins/Experimental/Web/WebAPI/Source/WebAPIBlueprintGraph/Private/K2Node_WebAPIOperation.cpp
+++ b/Engine/Plugins/Experimental/Web/WebAPI/Source/WebAPIBlueprintGraph/Private/K2Node_WebAPIOperation.cpp
@@ -162,7 +162,7 @@ void UK2Node_WebAPIOperation::AllocateDefaultPins()
 		CreatePin(EGPD_Output, UEdGraphSchema_K2::PC_Exec, UEdGraphSchema_K2::PN_Then);
 	}
 
-	TArray<FMulticastDelegateProperty*, TFixedAllocator<2>> DelegateProperties = {
+	TArray<const FMulticastDelegateProperty*, TFixedAllocator<2>> DelegateProperties = {
 		UE::WebAPI::Operation::GetPositiveOutcomeDelegate(OperationClass),
 		UE::WebAPI::Operation::GetNegativeOutcomeDelegate(OperationClass)
 	};
diff --git a/Engine/Plugins/Experimental/Web/WebAPI/Source/WebAPIBlueprintGraph/Private/WebAPIBlueprintGraphUtilities.cpp b/Engine/Plugins/Experimental/Web/WebAPI/Source/WebAPIBlueprintGraph/Private/WebAPIBlueprintGraphUtilities.cpp
--- a/Engine/Plugins/Experimental/Web/WebAPI/Source/WebAPIBlueprintGraph/Private/WebAPIBlueprintGraphUtilities.cpp
+++ b/Engine/Plugins/Experimental/Web/WebAPI/Source/WebAPIBlueprintGraph/Private/WebAPIBlueprintGraphUtilities.cpp
@@ -12,7 +12,7 @@ namespace UE::WebAPI
 	namespace Operation
 	{
 		template <uint32 Index, typename TEnableIf<Index < 2>::Type* = nullptr>
-		FMulticastDelegateProperty* GetOutcomeDelegate(const TSubclassOf<UWebAPIOperationObject>& InOperationClass, FName InOutcomeName)
+		const FMulticastDelegateProperty* GetOutcomeDelegate(const TSubclassOf<UWebAPIOperationObject>& InOperationClass, FName InOutcomeName)
 		{
 			check(InOperationClass);
 
@@ -51,14 +51,14 @@ namespace UE::WebAPI
 			return nullptr;
 		}
 		
-		FMulticastDelegateProperty* GetPositiveOutcomeDelegate(const TSubclassOf<UWebAPIOperationObject>& InOperationClass)
+		const FMulticastDelegateProperty* GetPositiveOutcomeDelegate(const TSubclassOf<UWebAPIOperationObject>& InOperationClass)
 		{
 			check(InOperationClass);
 
 			return GetOutcomeDelegate<0>(InOperationClass, PositiveOutcomeName);
 		}
 
-		FMulticastDelegateProperty* GetNegativeOutcomeDelegate(const TSubclassOf<UWebAPIOperationObject>& InOperationClass)
+		const FMulticastDelegateProperty* GetNegativeOutcomeDelegate(const TSubclassOf<UWebAPIOperationObject>& InOperationClass)
 		{
 			check(InOperationClass);
 
diff --git a/Engine/Plugins/Experimental/Web/WebAPI/Source/WebAPIBlueprintGraph/Private/WebAPIBlueprintGraphUtilities.h b/Engine/Plugins/Experimental/Web/WebAPI/Source/WebAPIBlueprintGraph/Private/WebAPIBlueprintGraphUtilities.h
--- a/Engine/Plugins/Experimental/Web/WebAPI/Source/WebAPIBlueprintGraph/Private/WebAPIBlueprintGraphUtilities.h
+++ b/Engine/Plugins/Experimental/Web/WebAPI/Source/WebAPIBlueprintGraph/Private/WebAPIBlueprintGraphUtilities.h
@@ -17,11 +17,11 @@ namespace UE::WebAPI
 		struct FCachedOutcome
 		{
 		public:
-			TWeakFieldPtr<FMulticastDelegateProperty> PositiveOutcomeDelegate;
-			TWeakFieldPtr<FMulticastDelegateProperty> NegativeOutcomeDelegate;
+			TWeakFieldPtr<const FMulticastDelegateProperty> PositiveOutcomeDelegate;
+			TWeakFieldPtr<const FMulticastDelegateProperty> NegativeOutcomeDelegate;
 
 			template <uint32 Index, typename TEnableIf<Index < 2>::Type* = nullptr>
-			const TWeakFieldPtr<FMulticastDelegateProperty>& Get() const
+			const TWeakFieldPtr<const FMulticastDelegateProperty>& Get() const
 			{
 				if constexpr (Index == 0)
 				{
@@ -56,20 +56,20 @@ namespace UE::WebAPI
 		/** Cached positive and negative (in order) delegates per operation class (name). */
 		static TMap<FName, FCachedOutcome> CachedOutcomeDelegates;
 
-		FMulticastDelegateProperty* GetPositiveOutcomeDelegate(const TSubclassOf<UWebAPIOperationObject>& InOperationClass);
+		const FMulticastDelegateProperty* GetPositiveOutcomeDelegate(const TSubclassOf<UWebAPIOperationObject>& InOperationClass);
 
 		template <class OperationType = UWebAPIOperationObject>
-		FMulticastDelegateProperty* GetPositiveOutcomeDelegate()
+		const FMulticastDelegateProperty* GetPositiveOutcomeDelegate()
 		{
-			return GetPositiveOutcomeDelegate(OperationType::StaticClass());			
+			return GetPositiveOutcomeDelegate(OperationType::StaticClass());
 		}
 		
-		FMulticastDelegateProperty* GetNegativeOutcomeDelegate(const TSubclassOf<UWebAPIOperationObject>& InOperationClass);
+		const FMulticastDelegateProperty* GetNegativeOutcomeDelegate(const TSubclassOf<UWebAPIOperationObject>& InOperationClass);
 
 		template <class OperationType = UWebAPIOperationObject>
-		FMulticastDelegateProperty* GetNegativeOutcomeDelegate()
+		const FMulticastDelegateProperty* GetNegativeOutcomeDelegate()
 		{
-			return GetNegativeOutcomeDelegate(OperationType::StaticClass());			
+			return GetNegativeOutcomeDelegate(OperationType::StaticClass());
 		}
 
 		UFunction* GetOutcomeDelegateSignatureFunction(const TSubclassOf<UWebAPIOperationObject>& InOperationClass);
diff --git a/Engine/Plugins/VirtualProduction/RemoteControl/Source/RemoteControl/Private/RemoteControlPropertyHandle.cpp b/Engine/Plugins/VirtualProduction/RemoteControl/Source/RemoteControl/Private/RemoteControlPropertyHandle.cpp
--- a/Engine/Plugins/VirtualProduction/RemoteControl/Source/RemoteControl/Private/RemoteControlPropertyHandle.cpp
+++ b/Engine/Plugins/VirtualProduction/RemoteControl/Source/RemoteControl/Private/RemoteControlPropertyHandle.cpp
@@ -172,7 +172,7 @@ namespace
 			return false;
 		}
 
-		const FProperty* Property = PropertyHandle.GetProperty();
+		FProperty* Property = PropertyHandle.GetProperty();
 		const FProperty* ParentProperty = PropertyHandle.GetParentProperty();
 		if (!ensure(Property))
 		{
@@ -235,7 +235,7 @@ namespace
 			return false;
 		}
 	
-		const FProperty* Property = InPropertyHandle.GetProperty();
+		FProperty* Property = InPropertyHandle.GetProperty();
 		if (!ensure(Property))
 		{
 			return false;
diff --git a/Engine/Plugins/VirtualProduction/RemoteControl/Source/RemoteControlCommon/Public/RCPropertyUtilities.h b/Engine/Plugins/VirtualProduction/RemoteControl/Source/RemoteControlCommon/Public/RCPropertyUtilities.h
--- a/Engine/Plugins/VirtualProduction/RemoteControl/Source/RemoteControlCommon/Public/RCPropertyUtilities.h
+++ b/Engine/Plugins/VirtualProduction/RemoteControl/Source/RemoteControlCommon/Public/RCPropertyUtilities.h
@@ -13,7 +13,7 @@
 
 namespace RemoteControlPropertyUtilities
 {
-	static TMap<TWeakFieldPtr<FProperty>, TWeakObjectPtr<UFunction>> CachedSetterFunctions;
+	static TMap<TWeakFieldPtr<const FProperty>, TWeakObjectPtr<UFunction>> CachedSetterFunctions;
 	static const FName NAME_BlueprintGetter(TEXT("BlueprintGetter"));
 	static const FName NAME_BlueprintSetter(TEXT("BlueprintSetter"));
 
@@ -23,7 +23,7 @@ namespace RemoteControlPropertyUtilities
 #if WITH_EDITOR
 		TSharedPtr<IPropertyHandle> PropertyHandle = nullptr;
 #endif
-		TWeakFieldPtr<FProperty> Property = nullptr;
+		TWeakFieldPtr<const FProperty> Property = nullptr;
 		void* PropertyData = nullptr;
 		TArray<uint8>* PropertyContainer = nullptr;
 
@@ -157,14 +157,14 @@ namespace RemoteControlPropertyUtilities
 		std::is_same_v<PropertyType, FNumericProperty>, bool>::Type
 	Serialize(const FRCPropertyVariant& InSrc, FRCPropertyVariant& OutDst);
 
-	static FProperty* FindSetterArgument(UFunction* SetterFunction, FProperty* PropertyToModify);
+	static FProperty* FindSetterArgument(UFunction* SetterFunction, const FProperty* PropertyToModify);
 
 	/** LightComponent derived components use a lot of property setters, without specifying BlueprintSetter, so handle here.  */
-	static FName FindLightSetterFunctionInternal(FProperty* Property, UClass* OwnerClass);
+	static FName FindLightSetterFunctionInternal(const FProperty* Property, UClass* OwnerClass);
 
-	static UFunction* FindSetterFunctionInternal(FProperty* Property, UClass* OwnerClass);
+	static UFunction* FindSetterFunctionInternal(const FProperty* Property, UClass* OwnerClass);
 	
-	static UFunction* FindSetterFunction(FProperty* Property, UClass* OwnerClass = nullptr);
+	static UFunction* FindSetterFunction(const FProperty* Property, UClass* OwnerClass = nullptr);
 }
 
 #include "RCPropertyUtilities.inl"
diff --git a/Engine/Plugins/VirtualProduction/RemoteControl/Source/RemoteControlCommon/Public/RCPropertyUtilities.inl b/Engine/Plugins/VirtualProduction/RemoteControl/Source/RemoteControlCommon/Public/RCPropertyUtilities.inl
--- a/Engine/Plugins/VirtualProduction/RemoteControl/Source/RemoteControlCommon/Public/RCPropertyUtilities.inl
+++ b/Engine/Plugins/VirtualProduction/RemoteControl/Source/RemoteControlCommon/Public/RCPropertyUtilities.inl
@@ -50,7 +50,8 @@ namespace RemoteControlPropertyUtilities
 	{
 		if(Property.IsValid())
 		{
-			return Property.Get();
+			// API is not const-correct
+			return const_cast<FProperty*>(Property.Get());
 		}
 
 #if WITH_EDITOR
@@ -362,7 +363,7 @@ namespace RemoteControlPropertyUtilities
 		return true;
 	}
 
-	static FProperty* FindSetterArgument(UFunction* SetterFunction, FProperty* PropertyToModify)
+	static FProperty* FindSetterArgument(UFunction* SetterFunction, const FProperty* PropertyToModify)
 	{
 		FProperty* SetterArgument = nullptr;
 
@@ -389,7 +390,7 @@ namespace RemoteControlPropertyUtilities
 	}
 
 	/** LightComponent derived components use a lot of property setters, without specifying BlueprintSetter, so handle here.  */
-	static FName FindLightSetterFunctionInternal(FProperty* Property, UClass* OwnerClass)
+	static FName FindLightSetterFunctionInternal(const FProperty* Property, UClass* OwnerClass)
 	{
 		static const FName LightAffectDynamicIndirectLightingPropertyName = TEXT("bAffectDynamicIndirectLighting");
 		static const FName LightAffectTranslucentLightingPropertyName = TEXT("bAffectTranslucentLighting");
@@ -461,7 +462,7 @@ namespace RemoteControlPropertyUtilities
 		return NAME_None;	
 	}
 
-	static UFunction* FindSetterFunctionInternal(FProperty* Property, UClass* OwnerClass)
+	static UFunction* FindSetterFunctionInternal(const FProperty* Property, UClass* OwnerClass)
 	{
 		// Check if the property setter is already cached.
 		const TWeakObjectPtr<UFunction> SetterPtr = CachedSetterFunctions.FindRef(Property);
@@ -547,7 +548,7 @@ namespace RemoteControlPropertyUtilities
 		return SetterFunction;
 	}
 
-	static UFunction* FindSetterFunction(FProperty* Property, UClass* OwnerClass)
+	static UFunction* FindSetterFunction(const FProperty* Property, UClass* OwnerClass)
 	{
 		// UStruct properties cannot have setters.
 		if (!ensure(Property) || !Property->GetOwnerClass())
diff --git a/Engine/Source/Runtime/CoreUObject/Public/UObject/FieldPath.h b/Engine/Source/Runtime/CoreUObject/Public/UObject/FieldPath.h
--- a/Engine/Source/Runtime/CoreUObject/Public/UObject/FieldPath.h
+++ b/Engine/Source/Runtime/CoreUObject/Public/UObject/FieldPath.h
@@ -12,9 +12,6 @@
 #include "CoreTypes.h"
 #include "HAL/ThreadSafeCounter.h"
 #include "Misc/AssertionMacros.h"
-#include "Templates/EnableIf.h"
-#include "Templates/LosesQualifiersFromTo.h"
-#include "Templates/PointerIsConvertibleFromTo.h"
 #include "Templates/Requires.h"
 #include "Templates/TypeHash.h"
 #include "Templates/UnrealTemplate.h"
@@ -329,6 +326,9 @@ public:
 	TFieldPath(UField* InField)
 		: FFieldPath(InField, PropertyType::StaticClass()->GetFName())
 	{
+		// This static assert is in here rather than in the body of the class because we want
+		// to be able to define TWeakFieldPtr<UUndefinedClass>.
+		static_assert(std::is_convertible_v<PropertyType*, const volatile FField*>, "TFieldPath can only be constructed with FField types");
 	}
 #endif
 
@@ -345,30 +345,38 @@ public:
 	{
 		// This static assert is in here rather than in the body of the class because we want
 		// to be able to define TFieldPath<UUndefinedClass>.
-		static_assert(TPointerIsConvertibleFromTo<PropertyType, const volatile FField>::Value, "TFieldPath can only be constructed with FField types");
+		static_assert(std::is_convertible_v<PropertyType*, const volatile FField*>, "TFieldPath can only be constructed with FField types");
 	}
 
 	/**
 	* Construct from another weak pointer of another type, intended for derived-to-base conversions
 	* @param Other weak pointer to copy from
 	**/
-	template <typename OtherPropertyType>
+	template <
+		typename OtherPropertyType
+		UE_REQUIRES(std::is_convertible_v<OtherPropertyType*, PropertyType*>)
+	>
 	FORCEINLINE TFieldPath(const TFieldPath<OtherPropertyType>& Other)
 		: FFieldPath(Other)
 	{
-		// It's also possible that this static_assert may fail for valid conversions because
-		// one or both of the types have only been forward-declared.
-		static_assert(TPointerIsConvertibleFromTo<OtherPropertyType, PropertyType>::Value, "Unable to convert TFieldPath - types are incompatible");
+		// This static assert is in here rather than in the body of the class because we want
+		// to be able to define TFieldPath<UUndefinedClass>.
+		static_assert(std::is_convertible_v<PropertyType*, const volatile FField*>, "TFieldPath can only be constructed with FField types");
 	}
 
 	/**
 	* Copy from an object pointer
 	* @param Object object to create a weak pointer to
 	**/
-	template<class OtherPropertyType>
-	FORCEINLINE typename TEnableIf<!TLosesQualifiersFromTo<OtherPropertyType, PropertyType>::Value>::Type operator=(OtherPropertyType* InProperty)
+	template <
+		typename OtherPropertyType
+		UE_REQUIRES(std::is_convertible_v<OtherPropertyType*, PropertyType*>)
+	>
+	FORCEINLINE void operator=(OtherPropertyType* InProperty)
 	{
-		ResolvedField = InProperty; 
+		// This (FField*) cast is effectively a const_cast, as we've already validated the convertibility in the
+		// constraint, and that PropertyType is an FField type in the constructors.
+		ResolvedField = (FField*)InProperty; 
 		Generate(ResolvedField);
 	}
 
@@ -376,13 +384,12 @@ public:
 	* Assign from another weak pointer, intended for derived-to-base conversions
 	* @param Other weak pointer to copy from
 	**/
-	template <typename OtherPropertyType>
+	template <
+		typename OtherPropertyType
+		UE_REQUIRES(std::is_convertible_v<OtherPropertyType*, PropertyType*>)
+	>
 	FORCEINLINE void operator=(const TFieldPath<OtherPropertyType>& Other)
 	{
-		// It's also possible that this static_assert may fail for valid conversions because
-		// one or both of the types have only been forward-declared.
-		static_assert(TPointerIsConvertibleFromTo<OtherPropertyType, PropertyType>::Value, "Unable to convert TFieldPath - types are incompatible");
-
 		// First make sure the Other path has the serial number up to date, otherwise we'll keep having to
 		// reevealuate this path because it gets the serial number copied from the Other path
 		Other.Get();
diff --git a/Engine/Source/Runtime/CoreUObject/Public/UObject/WeakFieldPtr.h b/Engine/Source/Runtime/CoreUObject/Public/UObject/WeakFieldPtr.h
--- a/Engine/Source/Runtime/CoreUObject/Public/UObject/WeakFieldPtr.h
+++ b/Engine/Source/Runtime/CoreUObject/Public/UObject/WeakFieldPtr.h
@@ -6,6 +6,7 @@
 
 #pragma once
 
+#include "Templates/Requires.h"
 #include "UObject/WeakObjectPtrTemplates.h"
 #include "UObject/WeakObjectPtr.h"
 #include "UObject/FieldPath.h"
@@ -54,21 +55,24 @@ public:
 	{
 		// This static assert is in here rather than in the body of the class because we want
 		// to be able to define TWeakFieldPtr<UUndefinedClass>.
-		static_assert(TPointerIsConvertibleFromTo<T, const volatile FField>::Value, "TWeakFieldPtr can only be constructed with FField types");
+		static_assert(std::is_convertible_v<T*, const volatile FField*>, "TWeakFieldPtr can only be constructed with FField types");
 	}
 
 	/**
 	* Construct from another weak pointer of another type, intended for derived-to-base conversions
 	* @param Other weak pointer to copy from
 	**/
-	template <typename OtherT>
+	template <
+		typename OtherT
+		UE_REQUIRES(std::is_convertible_v<OtherT*, T*>)
+	>
 	FORCEINLINE TWeakFieldPtr(const TWeakFieldPtr<OtherT>& Other)
 		: Owner(Other.Owner)
 		, Field(Other.Field)
 	{
-		// It's also possible that this static_assert may fail for valid conversions because
-		// one or both of the types have only been forward-declared.
-		static_assert(TPointerIsConvertibleFromTo<OtherT, T>::Value, "Unable to convert TWeakFieldPtr - types are incompatible");
+		// This static assert is in here rather than in the body of the class because we want
+		// to be able to define TWeakFieldPtr<UUndefinedClass>.
+		static_assert(std::is_convertible_v<T*, const volatile FField*>, "TWeakFieldPtr can only be constructed with FField types");
 	}
 
 	/**
@@ -84,8 +88,11 @@ public:
 	* Copy from an object pointer
 	* @param Object object to create a weak pointer to
 	**/
-	template<class U>
-	FORCEINLINE typename TEnableIf<!TLosesQualifiersFromTo<U, T>::Value>::Type operator=(const U* InField)
+	template <
+		typename U
+		UE_REQUIRES(std::is_convertible_v<U*, T*>)
+	>
+	FORCEINLINE void operator=(U* InField)
 	{
 		Owner = InField ? InField->GetOwnerUObject() : (UObject*)nullptr;
 		Field = (U*)InField;
@@ -95,13 +102,12 @@ public:
 	* Assign from another weak pointer, intended for derived-to-base conversions
 	* @param Other weak pointer to copy from
 	**/
-	template <typename OtherT>
+	template <
+		typename OtherT
+		UE_REQUIRES(std::is_convertible_v<OtherT*, T*>)
+	>
 	FORCEINLINE void operator=(const TWeakFieldPtr<OtherT>& Other)
 	{
-		// It's also possible that this static_assert may fail for valid conversions because
-		// one or both of the types have only been forward-declared.
-		static_assert(TPointerIsConvertibleFromTo<OtherT, T>::Value, "Unable to convert TWeakFieldPtr - types are incompatible");
-
 		Owner = Other.Owner;
 		Field = Other.Field;
 	}
