Ensure internal index is valid before resolving an FPackedObjectRef value from an object reference.

diff --git a/Engine/Source/Runtime/CoreUObject/Private/UObject/ObjectHandleIndex.cpp b/Engine/Source/Runtime/CoreUObject/Private/UObject/ObjectHandleIndex.cpp
--- a/Engine/Source/Runtime/CoreUObject/Private/UObject/ObjectHandleIndex.cpp
+++ b/Engine/Source/Runtime/CoreUObject/Private/UObject/ObjectHandleIndex.cpp
@@ -590,11 +590,12 @@ namespace UE::CoreUObject::Private
 
 	FPackedObjectRef FindExistingPackedObjectRef(const UObject* Object)
 	{
-		int32 ObjectIndex = GUObjectArray.ObjectToIndex(Object);
-		if (ObjectIndex == INDEX_NONE)
+		if (!GUObjectArray.IsValidIndex(Object))
 		{
 			return FPackedObjectRef {0};
 		}
+
+		const int32 ObjectIndex = GUObjectArray.ObjectToIndex(Object);
 		return UE::CoreUObject::Private::GObjectHandleIndex.ObjectIndexToPackedObjectRef[ObjectIndex];
 	}
 }
