Allow TObjectPtr to resolve to placeholder-typed objects regardless of 'const' qualifier when bound to the UObject base type.

Additional changes:
- Consolidated TObjectPtr natvis into a single visualizer block with an additional IsDerivedObjectType() intrinsic.
- Updated type safety unit test coverage to include 'const UObject' pointer types.

diff --git a/Engine/Source/Runtime/CoreUObject/Public/UObject/ObjectPtr.h b/Engine/Source/Runtime/CoreUObject/Public/UObject/ObjectPtr.h
--- a/Engine/Source/Runtime/CoreUObject/Public/UObject/ObjectPtr.h
+++ b/Engine/Source/Runtime/CoreUObject/Public/UObject/ObjectPtr.h
@@ -389,7 +389,7 @@ namespace ObjectPtr_Private
 	template <
 		typename T
 #if UE_WITH_OBJECT_HANDLE_TYPE_SAFETY
-		UE_REQUIRES(std::is_same_v<T, UObject>)
+		UE_REQUIRES(std::is_same_v<std::remove_const_t<T>, UObject>)
 #endif
 	>
 	FORCEINLINE bool IsObjectPtrNull(const FObjectPtr& ObjectPtr)
@@ -401,7 +401,7 @@ namespace ObjectPtr_Private
 	template <
 		typename T
 #if UE_WITH_OBJECT_HANDLE_TYPE_SAFETY
-		UE_REQUIRES(std::is_same_v<T, UObject>)
+		UE_REQUIRES(std::is_same_v<std::remove_const_t<T>, UObject>)
 #endif
 	>
 	FORCEINLINE T* Get(const FObjectPtr& ObjectPtr)
@@ -413,7 +413,7 @@ namespace ObjectPtr_Private
 	/** Check for NULL without resolving the handle. Always returns true if the handle is not type safe (only when T != UObject). */
 	template <
 		typename T
-		UE_REQUIRES(!std::is_same_v<T, UObject>)
+		UE_REQUIRES(!std::is_same_v<std::remove_const_t<T>, UObject>)
 	>
 	FORCEINLINE bool IsObjectPtrNull(const FObjectPtr& ObjectPtr)
 	{
@@ -429,7 +429,7 @@ namespace ObjectPtr_Private
 	/** Resolve and return the underlying reference. Always returns NULL if the handle is not type safe (only when T != UObject). */
 	template <
 		typename T
-		UE_REQUIRES(!std::is_same_v<T, UObject>)
+		UE_REQUIRES(!std::is_same_v<std::remove_const_t<T>, UObject>)
 	>
 	FORCEINLINE T* Get(const FObjectPtr& ObjectPtr)
 	{
diff --git a/Engine/Source/Runtime/CoreUObject/Tests/ObjectHandleTest.cpp b/Engine/Source/Runtime/CoreUObject/Tests/ObjectHandleTest.cpp
--- a/Engine/Source/Runtime/CoreUObject/Tests/ObjectHandleTest.cpp
+++ b/Engine/Source/Runtime/CoreUObject/Tests/ObjectHandleTest.cpp
@@ -464,12 +464,16 @@ TEST_CASE_METHOD(FObjectHandleTestBase, "CoreUObject::FObjectHandle::Type Safety
 	// construct object pointers for testing intentionally different behaviors of UObject-type vs. non-UObject-type bindings
 	TObjectPtr<UObject> NullObjectPtr(nullptr);
 	TObjectPtr<UObject> TestSafeObjectPtr(TestUnsafeObject);								// type safe pointer to placeholder (bound to UObject type)
+	TObjectPtr<const UObject> TestSafeConstObjectPtr(TestUnsafeObject);						// type safe const pointer to placeholder (bound to UObject type)
 	FObjectPtr TestUnsafeObjectPtr_Untyped(TestUnsafeObject);
 	TObjectPtr<UObjectPtrTestClass> TestUnsafeObjectPtr(TestUnsafeObjectPtr_Untyped);		// unsafe pointer to placeholder object (bound to non-UObject type)
+	TObjectPtr<const UObjectPtrTestClass> TestUnsafeConstObjectPtr(TestUnsafeObjectPtr_Untyped);	// unsafe const pointer to placeholder object (bound to non-UObject type)
 #if UE_WITH_OBJECT_HANDLE_LATE_RESOLVE
 	// note: "safe" in this context means the pointer should be type safe because it's bound to the UObject base type, but both reference the same "unsafe" object
 	TObjectPtr<UObject> TestLateResolveSafeObjectPtr(FObjectPtr({ UE::CoreUObject::Private::MakePackedObjectRef(TestUnsafeObject).EncodedRef }));
+	TObjectPtr<const UObject> TestLateResolveSafeConstObjectPtr(FObjectPtr({ UE::CoreUObject::Private::MakePackedObjectRef(TestUnsafeObject).EncodedRef }));
 	TObjectPtr<UObjectPtrTestClass> TestLateResolveUnsafeObjectPtr(FObjectPtr({ UE::CoreUObject::Private::MakePackedObjectRef(TestUnsafeObject).EncodedRef }));
+	TObjectPtr<const UObjectPtrTestClass> TestLateResolveUnsafeConstObjectPtr(FObjectPtr({ UE::CoreUObject::Private::MakePackedObjectRef(TestUnsafeObject).EncodedRef }));
 #endif
 
 	// an unsafe packed object reference should now exist with the unsafe type bit set for an unsafe object reference
@@ -484,6 +488,12 @@ TEST_CASE_METHOD(FObjectHandleTestBase, "CoreUObject::FObjectHandle::Type Safety
 	CHECK(TestSafeObjectPtr != NULL);
 	CHECK(nullptr != TestSafeObjectPtr);
 	CHECK(TestSafeObjectPtr != nullptr);
+	CHECK(TestSafeConstObjectPtr);
+	CHECK(!!TestSafeConstObjectPtr);
+	CHECK(NULL != TestSafeConstObjectPtr);
+	CHECK(TestSafeConstObjectPtr != NULL);
+	CHECK(nullptr != TestSafeConstObjectPtr);
+	CHECK(TestSafeConstObjectPtr != nullptr);
 #if UE_WITH_OBJECT_HANDLE_LATE_RESOLVE
 	CHECK(TestLateResolveSafeObjectPtr);
 	CHECK(!!TestLateResolveSafeObjectPtr);
@@ -491,6 +501,12 @@ TEST_CASE_METHOD(FObjectHandleTestBase, "CoreUObject::FObjectHandle::Type Safety
 	CHECK(TestLateResolveSafeObjectPtr != NULL);
 	CHECK(nullptr != TestLateResolveSafeObjectPtr);
 	CHECK(TestLateResolveSafeObjectPtr != nullptr);
+	CHECK(TestLateResolveSafeConstObjectPtr);
+	CHECK(!!TestLateResolveSafeConstObjectPtr);
+	CHECK(NULL != TestLateResolveSafeConstObjectPtr);
+	CHECK(TestLateResolveSafeConstObjectPtr != NULL);
+	CHECK(nullptr != TestLateResolveSafeConstObjectPtr);
+	CHECK(TestLateResolveSafeConstObjectPtr != nullptr);
 #endif
 
 	// unsafe type object pointers should evaluate to NULL/false (for type safety), but only when late resolve is enabled (required for support)
@@ -501,12 +517,24 @@ TEST_CASE_METHOD(FObjectHandleTestBase, "CoreUObject::FObjectHandle::Type Safety
 	CHECK(TestUnsafeObjectPtr == NULL);
 	CHECK(nullptr == TestUnsafeObjectPtr);
 	CHECK(TestUnsafeObjectPtr == nullptr);
+	CHECK_FALSE(TestUnsafeConstObjectPtr);
+	CHECK_FALSE(!!TestUnsafeConstObjectPtr);
+	CHECK(NULL == TestUnsafeConstObjectPtr);
+	CHECK(TestUnsafeConstObjectPtr == NULL);
+	CHECK(nullptr == TestUnsafeConstObjectPtr);
+	CHECK(TestUnsafeConstObjectPtr == nullptr);
 	CHECK_FALSE(TestLateResolveUnsafeObjectPtr);
 	CHECK_FALSE(!!TestLateResolveUnsafeObjectPtr);
 	CHECK(NULL == TestLateResolveUnsafeObjectPtr);
 	CHECK(TestLateResolveUnsafeObjectPtr == NULL);
 	CHECK(nullptr == TestLateResolveUnsafeObjectPtr);
 	CHECK(TestLateResolveUnsafeObjectPtr == nullptr);
+	CHECK_FALSE(TestLateResolveUnsafeConstObjectPtr);
+	CHECK_FALSE(!!TestLateResolveUnsafeConstObjectPtr);
+	CHECK(NULL == TestLateResolveUnsafeConstObjectPtr);
+	CHECK(TestLateResolveUnsafeConstObjectPtr == NULL);
+	CHECK(nullptr == TestLateResolveUnsafeConstObjectPtr);
+	CHECK(TestLateResolveUnsafeConstObjectPtr == nullptr);
 #else
 	CHECK(TestUnsafeObjectPtr);
 	CHECK(!!TestUnsafeObjectPtr);
@@ -514,26 +542,47 @@ TEST_CASE_METHOD(FObjectHandleTestBase, "CoreUObject::FObjectHandle::Type Safety
 	CHECK(TestUnsafeObjectPtr != NULL);
 	CHECK(nullptr != TestUnsafeObjectPtr);
 	CHECK(TestUnsafeObjectPtr != nullptr);
+	CHECK(TestUnsafeConstObjectPtr);
+	CHECK(!!TestUnsafeConstObjectPtr);
+	CHECK(NULL != TestUnsafeConstObjectPtr);
+	CHECK(TestUnsafeConstObjectPtr != NULL);
+	CHECK(nullptr != TestUnsafeConstObjectPtr);
+	CHECK(TestUnsafeConstObjectPtr != nullptr);
 #endif
 
 	// an unsafe type object pointer should not equate to other unsafe type object pointers, excluding NULL and itself
 #if UE_WITH_OBJECT_HANDLE_LATE_RESOLVE
 	CHECK(NullObjectPtr == TestUnsafeObjectPtr);				// note: this intentionally differs from object *handles* (see above), and depends on late resolve
 	CHECK(TestUnsafeObjectPtr == NullObjectPtr);				// see note directly above
+	CHECK(NullObjectPtr == TestUnsafeConstObjectPtr);
+	CHECK(TestUnsafeConstObjectPtr == NullObjectPtr);
 	CHECK(NullObjectPtr != TestLateResolveSafeObjectPtr);
 	CHECK(TestLateResolveSafeObjectPtr != NullObjectPtr);
 	CHECK(NullObjectPtr == TestLateResolveUnsafeObjectPtr);
 	CHECK(TestLateResolveUnsafeObjectPtr == NullObjectPtr);
 	CHECK(TestSafeObjectPtr != TestLateResolveUnsafeObjectPtr);
 	CHECK(TestLateResolveUnsafeObjectPtr != TestSafeObjectPtr);
+	CHECK(NullObjectPtr != TestLateResolveSafeConstObjectPtr);
+	CHECK(TestLateResolveSafeConstObjectPtr != NullObjectPtr);
+	CHECK(NullObjectPtr == TestLateResolveUnsafeConstObjectPtr);
+	CHECK(TestLateResolveUnsafeConstObjectPtr == NullObjectPtr);
+	CHECK(TestSafeConstObjectPtr != TestLateResolveUnsafeConstObjectPtr);
+	CHECK(TestLateResolveUnsafeConstObjectPtr != TestSafeConstObjectPtr);
 	CHECK(TestUnsafeObjectPtr == TestLateResolveUnsafeObjectPtr);
 	CHECK(TestLateResolveUnsafeObjectPtr == TestUnsafeObjectPtr);
+	CHECK(TestUnsafeConstObjectPtr == TestLateResolveUnsafeConstObjectPtr);
+	CHECK(TestLateResolveUnsafeConstObjectPtr == TestUnsafeConstObjectPtr);
 	CHECK(TestLateResolveSafeObjectPtr != TestLateResolveUnsafeObjectPtr);
 	CHECK(TestLateResolveUnsafeObjectPtr != TestLateResolveSafeObjectPtr);
+	CHECK(TestLateResolveSafeConstObjectPtr != TestLateResolveUnsafeConstObjectPtr);
+	CHECK(TestLateResolveUnsafeConstObjectPtr != TestLateResolveSafeConstObjectPtr);
 	CHECK(TestLateResolveUnsafeObjectPtr == TestLateResolveUnsafeObjectPtr);
+	CHECK(TestLateResolveUnsafeConstObjectPtr == TestLateResolveUnsafeConstObjectPtr);
 #else
 	CHECK(NullObjectPtr != TestUnsafeObjectPtr);				// note: if late resolve is disabled, these will no longer be equivalent
 	CHECK(TestUnsafeObjectPtr != NullObjectPtr);				// see note directly above
+	CHECK(NullObjectPtr != TestUnsafeConstObjectPtr);
+	CHECK(TestUnsafeConstObjectPtr != NullObjectPtr);
 #endif
 
 	// an unsafe type object should evaluate the object's attributes correctly
@@ -558,28 +607,39 @@ TEST_CASE_METHOD(FObjectHandleTestBase, "CoreUObject::FObjectHandle::Type Safety
 #if UE_WITH_OBJECT_HANDLE_LATE_RESOLVE
 	CHECK_FALSE(TestLateResolveSafeObjectPtr.IsResolved());
 	CHECK_FALSE(TestLateResolveUnsafeObjectPtr.IsResolved());
+	CHECK_FALSE(TestLateResolveSafeConstObjectPtr.IsResolved());
+	CHECK_FALSE(TestLateResolveUnsafeConstObjectPtr.IsResolved());
 #endif
 
 	// a type safe object pointer should resolve to a non-NULL value when dereferenced
 	CHECK(TestSafeObjectPtr.Get() == TestUnsafeObject);
+	CHECK(TestSafeConstObjectPtr.Get() == TestUnsafeObject);
 #if UE_WITH_OBJECT_HANDLE_LATE_RESOLVE
 	CHECK(TestLateResolveSafeObjectPtr.Get() == TestUnsafeObject);
+	CHECK(TestLateResolveSafeConstObjectPtr.Get() == TestUnsafeObject);
 #endif
 
 	// an unsafe type object pointer should resolve to NULL when dereferenced (for type safety), but only when late resolve is enabled (required for support)
 #if UE_WITH_OBJECT_HANDLE_LATE_RESOLVE
 	CHECK(TestUnsafeObjectPtr.Get() == nullptr);
+	CHECK(TestUnsafeConstObjectPtr.Get() == nullptr);
 	CHECK(TestLateResolveUnsafeObjectPtr.Get() == nullptr);
+	CHECK(TestLateResolveUnsafeConstObjectPtr.Get() == nullptr);
 #else
 	CHECK(TestUnsafeObjectPtr.Get() == TestUnsafeObject);
+	CHECK(TestUnsafeConstObjectPtr.Get() == TestUnsafeObject);
 #endif
 
 	// all pointers should now be resolved
 	CHECK(TestSafeObjectPtr.IsResolved());
 	CHECK(TestUnsafeObjectPtr.IsResolved());
+	CHECK(TestSafeConstObjectPtr.IsResolved());
+	CHECK(TestUnsafeConstObjectPtr.IsResolved());
 #if UE_WITH_OBJECT_HANDLE_LATE_RESOLVE
 	CHECK(TestLateResolveSafeObjectPtr.IsResolved());
 	CHECK(TestLateResolveUnsafeObjectPtr.IsResolved());
+	CHECK(TestLateResolveSafeConstObjectPtr.IsResolved());
+	CHECK(TestLateResolveUnsafeConstObjectPtr.IsResolved());
 #endif
 }
 #endif
