Remove transient object flag from property bag placeholder types and make them public when created by the linker in place of a missing type import.

This allows instances of those types to be saved as if they were the original export object when IDO impersonation is enabled.

Additional changes:
- Flag outer placeholder type package instances as memory-only (if created).
- Remove placeholder type objects from the property bag type registry if the reference is replaced with a valid type when script classes are regenerated.

diff --git a/Engine/Source/Runtime/CoreUObject/Private/UObject/LinkerLoad.cpp b/Engine/Source/Runtime/CoreUObject/Private/UObject/LinkerLoad.cpp
--- a/Engine/Source/Runtime/CoreUObject/Private/UObject/LinkerLoad.cpp
+++ b/Engine/Source/Runtime/CoreUObject/Private/UObject/LinkerLoad.cpp
@@ -4111,6 +4111,9 @@ UClass* FLinkerLoad::TryCreatePlaceholderClassImport(int32 ImportIndex)
 					{
 						ClassObjectPackage = CreatePackage(*OuterImport.ObjectName.ToString());
 
+						// Flag that this package exists in memory only (i.e. it's not being loaded from disk).
+						ClassObjectPackage->SetPackageFlags(PKG_InMemoryOnly);
+
 						// Patch it into the import table so that we resolve to this package for future reference.
 						OuterImport.XObject = ClassObjectPackage;
 					}
@@ -4118,8 +4121,8 @@ UClass* FLinkerLoad::TryCreatePlaceholderClassImport(int32 ImportIndex)
 
 				if (ClassObjectPackage)
 				{
-					// Create an opaque, non-native transient type object that has no reflected properties.
-					ClassObject = UE::FPropertyBagRepository::CreatePropertyBagPlaceholderClass(ClassObjectPackage, ImportClass, Import.ObjectName);
+					// Create an opaque, non-native public type object that has no reflected properties.
+					ClassObject = UE::FPropertyBagRepository::CreatePropertyBagPlaceholderClass(ClassObjectPackage, ImportClass, Import.ObjectName, RF_Public);
 
 					// Patch it into the import table so that we resolve to this class for any future exports of this type.
 					Import.XObject = ClassObject;
diff --git a/Engine/Source/Runtime/CoreUObject/Private/UObject/PropertyBagRepository.cpp b/Engine/Source/Runtime/CoreUObject/Private/UObject/PropertyBagRepository.cpp
--- a/Engine/Source/Runtime/CoreUObject/Private/UObject/PropertyBagRepository.cpp
+++ b/Engine/Source/Runtime/CoreUObject/Private/UObject/PropertyBagRepository.cpp
@@ -48,6 +48,11 @@ public:
 		PendingPlaceholderTypes.Enqueue(Type);
 	}
 
+	void Remove(UStruct* Type)
+	{
+		PlaceholderTypes.Remove(Type);
+	}
+
 	bool Contains(const UStruct* Type)
 	{
 		ConsumePendingPlaceholderTypes();
@@ -139,7 +144,7 @@ void FPropertyBagRepository::ReassociateObjects(const TMap<UObject*, UObject*>&
 	FPropertyBagAssociationData OldBagData;
 	for (const TPair<UObject*, UObject*>& Pair : ReplacedObjects)
 	{
-		if(AssociatedData.RemoveAndCopyValue(Pair.Key, OldBagData))
+		if (AssociatedData.RemoveAndCopyValue(Pair.Key, OldBagData))
 		{
 			InstanceDataObjectToOwner.Remove(OldBagData.InstanceDataObject);
 			if (Pair.Value != nullptr) // Pair.Value can be nullptr when an object was destroyed like for example a UClass when it's deleted
@@ -154,6 +159,13 @@ void FPropertyBagRepository::ReassociateObjects(const TMap<UObject*, UObject*>&
 			}
 			OldBagData.Destroy();
 		}
+		else if (UStruct* TypeObject = Cast<UStruct>(Pair.Value))
+		{
+			if (IsPropertyBagPlaceholderType(TypeObject))
+			{
+				FPropertyBagPlaceholderTypeRegistry::Get().Remove(TypeObject);
+			}
+		}
 		Namespaces.Remove(Pair.Key);
 	}
 }
@@ -737,8 +749,8 @@ bool FPropertyBagRepository::IsInstanceDataObjectSupportEnabled(UObject* InObjec
 
 UStruct* FPropertyBagRepository::CreatePropertyBagPlaceholderType(UObject* Outer, UClass* Class, FName Name, EObjectFlags Flags, UStruct* SuperStruct)
 {
-	// Generate and link a transient type object using the given SuperStruct as its base.
-	UStruct* PlaceholderType = NewObject<UClass>(Outer, Class, Name, Flags | RF_Transient);
+	// Generate and link a new type object using the given SuperStruct as its base.
+	UStruct* PlaceholderType = NewObject<UClass>(Outer, Class, Name, Flags);
 	PlaceholderType->SetSuperStruct(SuperStruct);
 	PlaceholderType->Bind();
 	PlaceholderType->StaticLink(/*bRelinkExistingProperties =*/ true);
