[TSAN]
 - Separate bIsCookedForEditor from the bitfield so it doesn't race with bHasBeenFullyLoaded

#rnx

diff --git a/Engine/Source/Runtime/CoreUObject/Public/UObject/Package.h b/Engine/Source/Runtime/CoreUObject/Public/UObject/Package.h
--- a/Engine/Source/Runtime/CoreUObject/Public/UObject/Package.h
+++ b/Engine/Source/Runtime/CoreUObject/Public/UObject/Package.h
@@ -211,12 +211,6 @@ private:
 #endif
 
 public:
-
-#if WITH_EDITORONLY_DATA
-	/** True if this packages has been cooked for the editor / opened cooked by the editor */
-	uint8 bIsCookedForEditor:1;
-#endif
-
 	/** Whether this package has been fully loaded (aka had all it's exports created) at some point. */
 	mutable uint8 bHasBeenFullyLoaded:1;
 
@@ -229,6 +223,11 @@ public:
 	uint8 bCanBeImported:1;
 
 #if WITH_EDITORONLY_DATA
+	/** True if this packages has been cooked for the editor / opened cooked by the editor
+        *   Note: This flag is accessed on different threads, do not mix with other bitfields above as the value returned might get corrupted.
+        */
+	bool bIsCookedForEditor{ false };
+
 private:
 	/** True if this package is only referenced by editor-only properties 
 	*   Note: This flag is manipulated on different threads, do not mix with other bitfields above as they might get corrupted.
