PropertyVisitor utils for operating on a subset of the full path

#rnx

diff --git a/Engine/Source/Runtime/CoreUObject/Private/UObject/PropertyVisitor.cpp b/Engine/Source/Runtime/CoreUObject/Private/UObject/PropertyVisitor.cpp
--- a/Engine/Source/Runtime/CoreUObject/Private/UObject/PropertyVisitor.cpp
+++ b/Engine/Source/Runtime/CoreUObject/Private/UObject/PropertyVisitor.cpp
@@ -7,8 +7,42 @@
 #include "Serialization/ArchiveSerializedPropertyChain.h"
 
 FString FPropertyVisitorPath::ToString(const TCHAR* Separator /*= TEXT(".")*/) const
+{
+	return PropertyVisitorHelpers::PathToString(Path, Separator);
+}
+
+void FPropertyVisitorPath::ToString(FWideStringBuilderBase& Out, const TCHAR* Separator /*= TEXT(".")*/) const
+{
+	PropertyVisitorHelpers::PathToString(Path, Out, Separator);
+}
+
+bool FPropertyVisitorPath::Contained(const FPropertyVisitorPath& Other, bool* bIsEqual) const
+{
+	return PropertyVisitorHelpers::PathIsContainedWithin(Path, Other.Path, bIsEqual);
+}
+
+void* FPropertyVisitorPath::GetPropertyDataPtr(UObject* Object) const
+{
+	checkf(Object, TEXT("Expecting an valid object"));
+	return PropertyVisitorHelpers::ResolveVisitedPath(Object->GetClass(), Object, *this);
+}
+
+FArchiveSerializedPropertyChain FPropertyVisitorPath::ToSerializedPropertyChain() const
+{
+	return PropertyVisitorHelpers::PathToSerializedPropertyChain(Path);
+}
+
+FString PropertyVisitorHelpers::PathToString(TArrayView<const FPropertyVisitorInfo> Path, const TCHAR* Separator)
 {
 	TStringBuilder<FName::StringBufferSize> PropertyPath;
+	PropertyVisitorHelpers::PathToString(Path, PropertyPath, Separator);
+	return PropertyPath.ToString();
+}
+
+void PropertyVisitorHelpers::PathToString(TArrayView<const FPropertyVisitorInfo> Path, FWideStringBuilderBase& Out, const TCHAR* Separator)
+{
+	Out.Reset();
+
 	bool bFirstEntry = true;
 	for (const FPropertyVisitorInfo& Entry : Path)
 	{
@@ -47,38 +81,36 @@ FString FPropertyVisitorPath::ToString(const TCHAR* Separator /*= TEXT(".")*/) c
 			}
 			else
 			{
-				PropertyPath.Append(Separator);
+				Out.Append(Separator);
 			}
-			PropertyPath.Append(Entry.Property->GetAuthoredName());
+			Out.Append(Entry.Property->GetAuthoredName());
 		}
 		if (bDisplayIndex)
 		{
 			checkf(Entry.Index != INDEX_NONE, TEXT("Expecting the index to be valid"));
-			PropertyPath.Appendf(TEXT("[%d]"), Entry.Index);
+			Out.Appendf(TEXT("[%d]"), Entry.Index);
 		}
 		if (Suffix.Len())
 		{
-			PropertyPath.Append(Separator);
-			PropertyPath.Append(Suffix);
+			Out.Append(Separator);
+			Out.Append(Suffix);
 		}
 	}
-
-	return PropertyPath.ToString();
 }
 
-bool FPropertyVisitorPath::Contained(const FPropertyVisitorPath& Other, bool* bIsEqual) const
+bool PropertyVisitorHelpers::PathIsContainedWithin(TArrayView<const FPropertyVisitorInfo> Path, TArrayView<const FPropertyVisitorInfo> OtherPath, bool* bIsEqual)
 {
 	int32 i;
 	for (i = 0; i < Path.Num(); ++i)
 	{
-		if (i >= Other.Path.Num())
+		if (i >= OtherPath.Num())
 		{
 			// The other path is smaller than this one, so not contained in
 			break;
 		}
 
 		const FPropertyVisitorInfo& PathInfo = Path[i];
-		const FPropertyVisitorInfo& OtherPathInfo = Other.Path[i];
+		const FPropertyVisitorInfo& OtherPathInfo = OtherPath[i];
 
 		if (PathInfo.Property != OtherPathInfo.Property)
 		{
@@ -102,18 +134,12 @@ bool FPropertyVisitorPath::Contained(const FPropertyVisitorPath& Other, bool* bI
 	}
 	if (bIsEqual)
 	{
-		*bIsEqual = i == Other.Path.Num();
+		*bIsEqual = i == OtherPath.Num();
 	}
 	return i == Path.Num();
 }
 
-void* FPropertyVisitorPath::GetPropertyDataPtr(UObject* Object) const
-{
-	checkf(Object, TEXT("Expecting an valid object"));
-	return PropertyVisitorHelpers::ResolveVisitedPath(Object->GetClass(), Object, *this);
-}
-
-FArchiveSerializedPropertyChain FPropertyVisitorPath::ToSerializedPropertyChain() const
+FArchiveSerializedPropertyChain PropertyVisitorHelpers::PathToSerializedPropertyChain(TArrayView<const FPropertyVisitorInfo> Path)
 {
 	FArchiveSerializedPropertyChain Chain;
 	for (const FPropertyVisitorInfo& Info : Path)
diff --git a/Engine/Source/Runtime/CoreUObject/Public/UObject/PropertyVisitor.h b/Engine/Source/Runtime/CoreUObject/Public/UObject/PropertyVisitor.h
--- a/Engine/Source/Runtime/CoreUObject/Public/UObject/PropertyVisitor.h
+++ b/Engine/Source/Runtime/CoreUObject/Public/UObject/PropertyVisitor.h
@@ -114,6 +114,7 @@ public:
 	}
 
 	COREUOBJECT_API FString ToString(const TCHAR* Separator = TEXT(".")) const;
+	COREUOBJECT_API void ToString(FWideStringBuilderBase& Out, const TCHAR* Separator = TEXT(".")) const;
 
 	/**
 	 * Is this property path contained in the specified one
@@ -121,7 +122,7 @@ public:
 	 * @param bIsEqual optional parameter to know if it fully match
 	 * @return true if it is contained in the specified path
 	 */
-	COREUOBJECT_API bool Contained(const FPropertyVisitorPath& OtherPath, bool* bIsEqual) const;
+	COREUOBJECT_API bool Contained(const FPropertyVisitorPath& OtherPath, bool* bIsEqual = nullptr) const;
 
 	/**
 	 * Retrieves the data using the specified root object
@@ -231,4 +232,19 @@ void* ResolveVisitedPathInfo_Generic(Type* This, void* Data, const FPropertyVisi
 	return ResolveVisitedPathInfo_Generic(This, Path, Data, Info);
 }
 
+/** Convert the given path to a string */
+COREUOBJECT_API FString PathToString(TArrayView<const FPropertyVisitorInfo> Path, const TCHAR* Separator = TEXT("."));
+COREUOBJECT_API void PathToString(TArrayView<const FPropertyVisitorInfo> Path, FWideStringBuilderBase& Out, const TCHAR* Separator = TEXT("."));
+
+/**
+ * Is this property path contained in the specified one
+ * @param OtherPath property path to check if it is contained in
+ * @param bIsEqual optional parameter to know if it fully match
+ * @return true if it is contained in the specified path
+ */
+COREUOBJECT_API bool PathIsContainedWithin(TArrayView<const FPropertyVisitorInfo> Path, TArrayView<const FPropertyVisitorInfo> OtherPath, bool* bIsEqual = nullptr);
+
+/** Convert the given path to a serialized property chain */
+COREUOBJECT_API FArchiveSerializedPropertyChain PathToSerializedPropertyChain(TArrayView<const FPropertyVisitorInfo> Path);
+
 } // namespace PropertyVisitorHelpers
