Allow FPropertyVisitorPath to be used with a TSet

#rnx

diff --git a/Engine/Source/Runtime/CoreUObject/Private/UObject/PropertyVisitor.cpp b/Engine/Source/Runtime/CoreUObject/Private/UObject/PropertyVisitor.cpp
--- a/Engine/Source/Runtime/CoreUObject/Private/UObject/PropertyVisitor.cpp
+++ b/Engine/Source/Runtime/CoreUObject/Private/UObject/PropertyVisitor.cpp
@@ -6,6 +6,15 @@
 #include "UObject/UnrealType.h"
 #include "Serialization/ArchiveSerializedPropertyChain.h"
 
+bool FPropertyVisitorInfo::Identical(const FPropertyVisitorInfo& Other) const
+{
+	return Property == Other.Property
+		&& ParentStructType == Other.ParentStructType
+		&& Index == Other.Index
+		&& PropertyInfo == Other.PropertyInfo
+		&& bContainsInnerProperties == Other.bContainsInnerProperties;
+}
+
 FString FPropertyVisitorPath::ToString(const TCHAR* Separator /*= TEXT(".")*/) const
 {
 	return PropertyVisitorHelpers::PathToString(Path, Separator);
diff --git a/Engine/Source/Runtime/CoreUObject/Public/UObject/PropertyVisitor.h b/Engine/Source/Runtime/CoreUObject/Public/UObject/PropertyVisitor.h
--- a/Engine/Source/Runtime/CoreUObject/Public/UObject/PropertyVisitor.h
+++ b/Engine/Source/Runtime/CoreUObject/Public/UObject/PropertyVisitor.h
@@ -3,6 +3,7 @@
 #pragma once
 
 #include "Containers/Array.h"
+#include "Containers/ArrayView.h"
 
 class FProperty;
 struct FArchiveSerializedPropertyChain;
@@ -45,6 +46,23 @@ struct FPropertyVisitorInfo
 		PropertyInfo = InPropertyInfo;
 	}
 
+	COREUOBJECT_API bool Identical(const FPropertyVisitorInfo& Other) const;
+
+	/** @note The default comparison only compares the key data to match GetTypeHash and is for use with hashed containers and ResolveVisitedPathInfo_Generic; use Identical for an exact comparison */
+	friend bool operator==(const FPropertyVisitorInfo& A, const FPropertyVisitorInfo& B)
+	{
+		return A.Property == B.Property && A.PropertyInfo == B.PropertyInfo && A.Index == B.Index;
+	}
+	friend bool operator!=(const FPropertyVisitorInfo& A, const FPropertyVisitorInfo& B)
+	{
+		return !(A == B);
+	}
+
+	friend uint32 GetTypeHash(const FPropertyVisitorInfo& A)
+	{
+		return HashCombine(HashCombine(GetTypeHash(A.Property), GetTypeHash(A.PropertyInfo)), GetTypeHash(A.Index));
+	}
+
 	/** The property currently being visited */
 	const FProperty* Property;
 
@@ -83,6 +101,11 @@ public:
 		Push(Info);
 	}
 
+	explicit FPropertyVisitorPath(TArrayView<const FPropertyVisitorInfo> InPath)
+	{
+		Path = InPath;
+	}
+
 	void Push(const FPropertyVisitorInfo& Info)
 	{
 		Path.Push(Info);
@@ -113,6 +136,20 @@ public:
 		return Path;
 	}
 
+	friend bool operator==(const FPropertyVisitorPath& A, const FPropertyVisitorPath& B)
+	{
+		return A.Path == B.Path;
+	}
+	friend bool operator!=(const FPropertyVisitorPath& A, const FPropertyVisitorPath& B)
+	{
+		return !(A == B);
+	}
+
+	friend uint32 GetTypeHash(const FPropertyVisitorPath& A)
+	{
+		return GetTypeHash(A.Path);
+	}
+
 	COREUOBJECT_API FString ToString(const TCHAR* Separator = TEXT(".")) const;
 	COREUOBJECT_API void ToString(FWideStringBuilderBase& Out, const TCHAR* Separator = TEXT(".")) const;
 
@@ -216,7 +253,7 @@ void* ResolveVisitedPathInfo_Generic(Type* This, FPropertyVisitorPath& Path, voi
 		{
 			return EPropertyVisitorControlFlow::StepInto;
 		}
-		if (Info.Property == InnerPath.Top().Property && Info.PropertyInfo == InnerPath.Top().PropertyInfo && Info.Index == InnerPath.Top().Index)
+		if (Info == InnerPath.Top())
 		{
 			FoundInnerData = InnerData;
 			return EPropertyVisitorControlFlow::Stop;
