Use FStringBuilderBase and added AppendString

#rnx

diff --git a/Engine/Source/Runtime/CoreUObject/Private/UObject/PropertyVisitor.cpp b/Engine/Source/Runtime/CoreUObject/Private/UObject/PropertyVisitor.cpp
--- a/Engine/Source/Runtime/CoreUObject/Private/UObject/PropertyVisitor.cpp
+++ b/Engine/Source/Runtime/CoreUObject/Private/UObject/PropertyVisitor.cpp
@@ -20,11 +20,16 @@ FString FPropertyVisitorPath::ToString(const TCHAR* Separator /*= TEXT(".")*/) c
 	return PropertyVisitorHelpers::PathToString(Path, Separator);
 }
 
-void FPropertyVisitorPath::ToString(FWideStringBuilderBase& Out, const TCHAR* Separator /*= TEXT(".")*/) const
+void FPropertyVisitorPath::ToString(FStringBuilderBase& Out, const TCHAR* Separator /*= TEXT(".")*/) const
 {
 	PropertyVisitorHelpers::PathToString(Path, Out, Separator);
 }
 
+void FPropertyVisitorPath::AppendString(FStringBuilderBase& Out, const TCHAR* Separator /*= TEXT(".")*/) const
+{
+	PropertyVisitorHelpers::PathAppendString(Path, Out, Separator);
+}
+
 bool FPropertyVisitorPath::Contained(const FPropertyVisitorPath& Other, bool* bIsEqual) const
 {
 	return PropertyVisitorHelpers::PathIsContainedWithin(Path, Other.Path, bIsEqual);
@@ -44,14 +49,18 @@ FArchiveSerializedPropertyChain FPropertyVisitorPath::ToSerializedPropertyChain(
 FString PropertyVisitorHelpers::PathToString(TArrayView<const FPropertyVisitorInfo> Path, const TCHAR* Separator)
 {
 	TStringBuilder<FName::StringBufferSize> PropertyPath;
-	PropertyVisitorHelpers::PathToString(Path, PropertyPath, Separator);
+	PathAppendString(Path, PropertyPath, Separator);
 	return PropertyPath.ToString();
 }
 
-void PropertyVisitorHelpers::PathToString(TArrayView<const FPropertyVisitorInfo> Path, FWideStringBuilderBase& Out, const TCHAR* Separator)
+void PropertyVisitorHelpers::PathToString(TArrayView<const FPropertyVisitorInfo> Path, FStringBuilderBase& Out, const TCHAR* Separator)
 {
 	Out.Reset();
+	PathAppendString(Path, Out, Separator);
+}
 
+void PropertyVisitorHelpers::PathAppendString(TArrayView<const FPropertyVisitorInfo> Path, FStringBuilderBase& Out, const TCHAR* Separator)
+{
 	bool bFirstEntry = true;
 	for (const FPropertyVisitorInfo& Entry : Path)
 	{
diff --git a/Engine/Source/Runtime/CoreUObject/Public/UObject/PropertyVisitor.h b/Engine/Source/Runtime/CoreUObject/Public/UObject/PropertyVisitor.h
--- a/Engine/Source/Runtime/CoreUObject/Public/UObject/PropertyVisitor.h
+++ b/Engine/Source/Runtime/CoreUObject/Public/UObject/PropertyVisitor.h
@@ -151,7 +151,8 @@ public:
 	}
 
 	COREUOBJECT_API FString ToString(const TCHAR* Separator = TEXT(".")) const;
-	COREUOBJECT_API void ToString(FWideStringBuilderBase& Out, const TCHAR* Separator = TEXT(".")) const;
+	COREUOBJECT_API void ToString(FStringBuilderBase& Out, const TCHAR* Separator = TEXT(".")) const;
+	COREUOBJECT_API void AppendString(FStringBuilderBase& Out, const TCHAR* Separator = TEXT(".")) const;
 
 	/**
 	 * Is this property path contained in the specified one
@@ -271,7 +272,8 @@ void* ResolveVisitedPathInfo_Generic(Type* This, void* Data, const FPropertyVisi
 
 /** Convert the given path to a string */
 COREUOBJECT_API FString PathToString(TArrayView<const FPropertyVisitorInfo> Path, const TCHAR* Separator = TEXT("."));
-COREUOBJECT_API void PathToString(TArrayView<const FPropertyVisitorInfo> Path, FWideStringBuilderBase& Out, const TCHAR* Separator = TEXT("."));
+COREUOBJECT_API void PathToString(TArrayView<const FPropertyVisitorInfo> Path, FStringBuilderBase& Out, const TCHAR* Separator = TEXT("."));
+COREUOBJECT_API void PathAppendString(TArrayView<const FPropertyVisitorInfo> Path, FStringBuilderBase& Out, const TCHAR* Separator = TEXT("."));
 
 /**
  * Is this property path contained in the specified one
