Fixed operators on TFieldPath and TWeakFieldPtr.
Worked around UE_REQUIRES_FRIEND not working under Clang 16.0.

diff --git a/Engine/Source/Runtime/Core/Public/Templates/Requires.h b/Engine/Source/Runtime/Core/Public/Templates/Requires.h
--- a/Engine/Source/Runtime/Core/Public/Templates/Requires.h
+++ b/Engine/Source/Runtime/Core/Public/Templates/Requires.h
@@ -77,7 +77,12 @@
 	the body of a UE_REQUIRES() macro - standalone concept checks must still be
 	done via something like a TModels-type concept.
  -----------------------------------------------------------------------------*/
-#if __cplusplus < 202000
+#if __cplusplus < 202000 || (defined(__clang__) && __clang_major__ == 16)
+	// Clang 16 treats a UE_REQUIRES_FRIEND declaration as a separate function, causing ambiguous overload resolution,
+	// so fall back to the non-concept implementation.
+	//
+	// [REDACTED_PROFILE_URL]/llvm-project/issues/60749
+
 	#define UE_REQUIRES(...) , std::enable_if_t<(__VA_ARGS__), int> = 0
 	#define UE_REQUIRES_FRIEND(...) , std::enable_if_t<(__VA_ARGS__), int>
 	#define UE_REQUIRES_DEFINITION(...) , std::enable_if_t<(__VA_ARGS__), int>
diff --git a/Engine/Source/Runtime/CoreUObject/Public/UObject/FieldPath.h b/Engine/Source/Runtime/CoreUObject/Public/UObject/FieldPath.h
--- a/Engine/Source/Runtime/CoreUObject/Public/UObject/FieldPath.h
+++ b/Engine/Source/Runtime/CoreUObject/Public/UObject/FieldPath.h
@@ -429,104 +429,121 @@ public:
 	{
 		return Get();
 	}
+};
 
-	/**
-	* Compare weak pointers for equality
-	* @param Other weak pointer to compare to
-	**/
-	template <typename OtherPropertyType>
-	FORCEINLINE bool operator==(const TFieldPath<OtherPropertyType> &Other) const
-	{
-		static_assert(TPointerIsConvertibleFromTo<OtherPropertyType, const FField>::Value, "TFieldPath can only be compared with FField types");
-		static_assert(TPointerIsConvertibleFromTo<PropertyType, OtherPropertyType>::Value, "Unable to compare TFieldPath with raw pointer - types are incompatible");
+/**
+* Compare weak pointers for equality
+* @param Lhs weak pointer to compare
+* @param Rhs weak pointer to compare
+**/
+template <typename LhsType, typename RhsType>
+FORCEINLINE auto operator==(const TFieldPath<LhsType>& Lhs, const TFieldPath<LhsType>& Rhs)
+	-> decltype((LhsType*)nullptr == (RhsType*)nullptr)
+{
+	return *(const FFieldPath*)&Lhs == *(const FFieldPath*)&Rhs;
+}
 
-		return FFieldPath::operator==(Other);
-	}
+/**
+* Compare weak pointers for equality
+* @param Lhs weak pointer to compare
+* @param Rhs pointer to compare
+**/
+template <typename LhsType, typename RhsType>
+FORCEINLINE auto operator==(const TFieldPath<LhsType>& Lhs, const RhsType* Rhs)
+	-> decltype((LhsType*)nullptr == Rhs)
+{
+	return Lhs.Get() == Rhs;
+}
 
-	/**
-	* Compare weak pointers for equality
-	* @param Other pointer to compare to
-	**/
-	template <typename OtherPropertyType>
-	FORCEINLINE bool operator==(const OtherPropertyType* Other) const
-	{
-		static_assert(TPointerIsConvertibleFromTo<OtherPropertyType, const FField>::Value, "TFieldPath can only be compared with FField types");
-		static_assert(TPointerIsConvertibleFromTo<PropertyType, const OtherPropertyType>::Value, "Unable to compare TFieldPath with raw pointer - types are incompatible");
+/**
+* Compare weak pointers for equality
+* @param Lhs pointer to compare
+* @param Rhs weak pointer to compare
+**/
+template <typename LhsType, typename RhsType>
+FORCEINLINE auto operator==(const LhsType* Lhs, const TFieldPath<RhsType>& Rhs)
+	-> decltype(Lhs == (RhsType*)nullptr)
+{
+	return Lhs == Rhs.Get();
+}
 
-		return Get() == Other;
-	}
+/**
+* Test weak pointer for null
+* @param Lhs pointer to test
+**/
+template <typename LhsType>
+FORCEINLINE bool operator==(TFieldPath<LhsType>& Lhs, TYPE_OF_NULLPTR)
+{
+	return !Lhs.Get();
+}
 
-	FORCENOINLINE bool operator==(TYPE_OF_NULLPTR) const
-	{
-		return !Get();
-	}
+/**
+* Test weak pointer for null
+* @param Rhs pointer to test
+**/
+template <typename RhsType>
+FORCEINLINE bool operator==(TYPE_OF_NULLPTR, TFieldPath<RhsType>& Rhs)
+{
+	return !Rhs.Get();
+}
 
 #if !PLATFORM_COMPILER_HAS_GENERATED_COMPARISON_OPERATORS
+/**
+* Compare weak pointers for inequality
+* @param Lhs weak pointer to compare
+* @param Rhs weak pointer to compare
+**/
+template <typename LhsType, typename RhsType>
+FORCEINLINE auto operator!=(const TFieldPath<LhsType>& Lhs, const TFieldPath<LhsType>& Rhs)
+	-> decltype((LhsType*)nullptr != (RhsType*)nullptr)
+{
+	return !(Lhs == Rhs);
+}
 
-	/**
-	* Compare weak pointers for inequality
-	* @param Other weak pointer to compare to
-	**/
-	template <typename OtherPropertyType>
-	FORCEINLINE bool operator!=(const TFieldPath<OtherPropertyType> &Other) const
-	{
-		static_assert(TPointerIsConvertibleFromTo<OtherPropertyType, const FField>::Value, "TFieldPath can only be compared with FField types");
-		static_assert(TPointerIsConvertibleFromTo<PropertyType, const OtherPropertyType>::Value, "Unable to compare TFieldPath with raw pointer - types are incompatible");
-
-		return FFieldPath::operator!=(Other);
-	}
-
-	/**
-	* Compare weak pointers for inequality
-	* @param Other pointer to compare to
-	**/
-	template <typename OtherPropertyType>
-	FORCEINLINE bool operator!=(const OtherPropertyType* Other) const
-	{
-		static_assert(TPointerIsConvertibleFromTo<OtherPropertyType, const FField>::Value, "TFieldPath can only be compared with FField types");
-		static_assert(TPointerIsConvertibleFromTo<PropertyType, const OtherPropertyType>::Value, "Unable to compare TFieldPath with raw pointer - types are incompatible");
-
-		return Get() != Other;
-	}
-
-	template <typename LhsT>
-	friend FORCENOINLINE bool operator==(const LhsT* Lhs, const TFieldPath<PropertyType>& Rhs)
-	{
-		// It's also possible that these static_asserts may fail for valid conversions because
-		// one or both of the types have only been forward-declared.
-		static_assert(TPointerIsConvertibleFromTo<LhsT, const FField>::Value, "TFieldPath can only be compared with FField types");
-		static_assert(TPointerIsConvertibleFromTo<LhsT, const PropertyType>::Value || TPointerIsConvertibleFromTo<PropertyType, LhsT>::Value, "Unable to compare TFieldPath with raw pointer - types are incompatible");
-
-		return Rhs == Lhs;
-	}
-
-	friend FORCENOINLINE bool operator==(TYPE_OF_NULLPTR, const TFieldPath<PropertyType>& Rhs)
-	{
-		return !Rhs.Get();
-	}
-
-	template <typename LhsT>
-	friend FORCENOINLINE bool operator!=(const LhsT* Lhs, const TFieldPath<PropertyType>& Rhs)
-	{
-		// It's also possible that these static_asserts may fail for valid conversions because
-		// one or both of the types have only been forward-declared.
-		static_assert(TPointerIsConvertibleFromTo<LhsT, const FField>::Value, "TFieldPath can only be compared with FField types");
-		static_assert(TPointerIsConvertibleFromTo<LhsT, const PropertyType>::Value || TPointerIsConvertibleFromTo<PropertyType, LhsT>::Value, "Unable to compare TFieldPath with raw pointer - types are incompatible");
+/**
+* Compare weak pointers for inequality
+* @param Lhs weak pointer to compare
+* @param Rhs pointer to compare
+**/
+template <typename LhsType, typename RhsType>
+FORCEINLINE auto operator!=(const TFieldPath<LhsType>& Lhs, const RhsType* Rhs)
+	-> decltype((LhsType*)nullptr != Rhs)
+{
+	return !(Lhs == Rhs);
+}
 
-		return Rhs != Lhs;
-	}
+/**
+* Compare weak pointers for inequality
+* @param Lhs pointer to compare
+* @param Rhs weak pointer to compare
+**/
+template <typename LhsType, typename RhsType>
+FORCEINLINE auto operator!=(const LhsType* Lhs, const TFieldPath<RhsType>& Rhs)
+	-> decltype(Lhs != (RhsType*)nullptr)
+{
+	return !(Lhs == Rhs);
+}
 
-	FORCENOINLINE bool operator!=(TYPE_OF_NULLPTR) const
-	{
-		return !!Get();
-	}
+/**
+* Test weak pointer for non-null
+* @param Lhs pointer to test
+**/
+template <typename LhsType>
+FORCEINLINE bool operator!=(TFieldPath<LhsType>& Lhs, TYPE_OF_NULLPTR)
+{
+	return !(Lhs == nullptr);
+}
 
-	friend FORCENOINLINE bool operator!=(TYPE_OF_NULLPTR, const TFieldPath<PropertyType>& Rhs)
-	{
-		return !!Rhs.Get();
-	}
+/**
+* Test weak pointers for non-null
+* @param Rhs pointer to test
+**/
+template <typename RhsType>
+FORCEINLINE bool operator!=(TYPE_OF_NULLPTR, TFieldPath<RhsType>& Rhs)
+{
+	return !(nullptr == Rhs);
+}
 #endif
-};
 
 // Helper function which deduces the type of the initializer
 template <typename PropertyType>
diff --git a/Engine/Source/Runtime/CoreUObject/Public/UObject/WeakFieldPtr.h b/Engine/Source/Runtime/CoreUObject/Public/UObject/WeakFieldPtr.h
--- a/Engine/Source/Runtime/CoreUObject/Public/UObject/WeakFieldPtr.h
+++ b/Engine/Source/Runtime/CoreUObject/Public/UObject/WeakFieldPtr.h
@@ -13,9 +13,103 @@
 
 #include <type_traits>
 
+template<class T>
+struct TWeakFieldPtr;
+
+template <
+	typename LhsType,
+	typename RhsType
+	UE_REQUIRES(UE_REQUIRES_EXPR((LhsType*)nullptr == (RhsType*)nullptr))
+>
+bool operator==(const TWeakFieldPtr<LhsType>& Lhs, const TWeakFieldPtr<RhsType>& Rhs);
+
+template <
+	typename LhsType,
+	typename RhsType
+	UE_REQUIRES(UE_REQUIRES_EXPR((LhsType*)nullptr == (const RhsType*)nullptr))
+>
+bool operator==(const TWeakFieldPtr<LhsType>& Lhs, const RhsType* Rhs);
+
+template <
+	typename LhsType,
+	typename RhsType
+	UE_REQUIRES(UE_REQUIRES_EXPR((const LhsType*)nullptr == (RhsType*)nullptr))
+>
+bool operator==(const LhsType* Lhs, const TWeakFieldPtr<RhsType>& Rhs);
+
+#if !PLATFORM_COMPILER_HAS_GENERATED_COMPARISON_OPERATORS
+template <
+	typename LhsType,
+	typename RhsType
+	UE_REQUIRES(UE_REQUIRES_EXPR((LhsType*)nullptr != (RhsType*)nullptr))
+>
+bool operator!=(const TWeakFieldPtr<LhsType>& Lhs, const TWeakFieldPtr<RhsType>& Rhs);
+
+template <
+	typename LhsType,
+	typename RhsType
+	UE_REQUIRES(UE_REQUIRES_EXPR((LhsType*)nullptr != (const RhsType*)nullptr))
+>
+bool operator!=(const TWeakFieldPtr<LhsType>& Lhs, const RhsType* Rhs);
+
+template <
+	typename LhsType,
+	typename RhsType
+	UE_REQUIRES(UE_REQUIRES_EXPR((const LhsType*)nullptr != (RhsType*)nullptr))
+>
+bool operator!=(const LhsType* Lhs, const TWeakFieldPtr<RhsType>& Rhs);
+#endif
+
 template<class T>
 struct TWeakFieldPtr
 {
+	template <typename>
+	friend struct TWeakFieldPtr;
+
+	template <
+		typename LhsType,
+		typename RhsType
+		UE_REQUIRES_FRIEND(UE_REQUIRES_EXPR((LhsType*)nullptr == (RhsType*)nullptr))
+	>
+	friend bool operator==(const TWeakFieldPtr<LhsType>& Lhs, const TWeakFieldPtr<RhsType>& Rhs);
+
+	template <
+		typename LhsType,
+		typename RhsType
+		UE_REQUIRES_FRIEND(UE_REQUIRES_EXPR((LhsType*)nullptr == (const RhsType*)nullptr))
+	>
+	friend bool operator==(const TWeakFieldPtr<LhsType>& Lhs, const RhsType* Rhs);
+
+	template <
+		typename LhsType,
+		typename RhsType
+		UE_REQUIRES_FRIEND(UE_REQUIRES_EXPR((const LhsType*)nullptr == (RhsType*)nullptr))
+	>
+	friend bool operator==(const LhsType* Lhs, const TWeakFieldPtr<RhsType>& Rhs);
+
+#if !PLATFORM_COMPILER_HAS_GENERATED_COMPARISON_OPERATORS
+	template <
+		typename LhsType,
+		typename RhsType
+		UE_REQUIRES_FRIEND(UE_REQUIRES_EXPR((LhsType*)nullptr != (RhsType*)nullptr))
+	>
+	friend bool operator!=(const TWeakFieldPtr<LhsType>& Lhs, const TWeakFieldPtr<RhsType>& Rhs);
+
+	template <
+		typename LhsType,
+		typename RhsType
+		UE_REQUIRES_FRIEND(UE_REQUIRES_EXPR((LhsType*)nullptr != (const RhsType*)nullptr))
+	>
+	friend bool operator!=(const TWeakFieldPtr<LhsType>& Lhs, const RhsType* Rhs);
+
+	template <
+		typename LhsType,
+		typename RhsType
+		UE_REQUIRES_FRIEND(UE_REQUIRES_EXPR((const LhsType*)nullptr != (RhsType*)nullptr))
+	>
+	friend bool operator!=(const LhsType* Lhs, const TWeakFieldPtr<RhsType>& Rhs);
+#endif
+
 private:
 
 	// These exists only to disambiguate the two constructors below
@@ -227,111 +321,145 @@ public:
 		Ar << Owner;
 		Ar << Field;
 	}
+};
 
-	/**
-	* Compare weak pointers for equality
-	* @param Other weak pointer to compare to
-	**/
-	template <typename TOther>
-	FORCEINLINE bool operator==(const TWeakFieldPtr<TOther> &Other) const
-	{
-		static_assert(TPointerIsConvertibleFromTo<TOther, const FField>::Value, "TWeakFieldPtr can only be compared with FField types");
-		static_assert(TPointerIsConvertibleFromTo<T, const TOther>::Value, "Unable to compare TWeakFieldPtr with raw pointer - types are incompatible");
-
-		return Field == Other.Field;
-	}
-
-	/**
-	* Compare weak pointers for inequality
-	* @param Other weak pointer to compare to
-	**/
-	template <typename TOther>
-	FORCEINLINE bool operator!=(const TWeakFieldPtr<TOther> &Other) const
-	{
-		static_assert(TPointerIsConvertibleFromTo<TOther, const FField>::Value, "TWeakFieldPtr can only be compared with FField types");
-		static_assert(TPointerIsConvertibleFromTo<T, const TOther>::Value, "Unable to compare TWeakFieldPtr with raw pointer - types are incompatible");
-
-		return Field != Other.Field;
-	}
-
-	/**
-	* Compare weak pointers for equality
-	* @param Other pointer to compare to
-	**/
-	template <typename TOther>
-	FORCEINLINE bool operator==(const TOther* Other) const
-	{
-		static_assert(TPointerIsConvertibleFromTo<TOther, const FField>::Value, "TWeakFieldPtr can only be compared with FField types");
-		static_assert(TPointerIsConvertibleFromTo<T, const TOther>::Value, "Unable to compare TWeakFieldPtr with raw pointer - types are incompatible");
-
-		return Field == Other;
-	}
-
-	/**
-	* Compare weak pointers for inequality
-	* @param Other pointer to compare to
-	**/
-	template <typename TOther>
-	FORCEINLINE bool operator!=(const TOther* Other) const
-	{
-		static_assert(TPointerIsConvertibleFromTo<TOther, const FField>::Value, "TWeakFieldPtr can only be compared with FField types");
-		static_assert(TPointerIsConvertibleFromTo<T, const TOther>::Value, "Unable to compare TWeakFieldPtr with raw pointer - types are incompatible");
+/**
+* Compare weak pointers for equality
+* @param Lhs weak pointer to compare
+* @param Rhs weak pointer to compare
+**/
+template <
+	typename LhsType,
+	typename RhsType
+	UE_REQUIRES_DEFINITION(UE_REQUIRES_EXPR((LhsType*)nullptr == (RhsType*)nullptr))
+>
+FORCEINLINE bool operator==(const TWeakFieldPtr<LhsType>& Lhs, const TWeakFieldPtr<RhsType>& Rhs)
+{
+	return Lhs.Field == Rhs.Field;
+}
 
-		return Field != Other;
-	}
-};
+/**
+* Compare weak pointers for equality
+* @param Lhs weak pointer to compare
+* @param Rhs pointer to compare
+**/
+template <
+	typename LhsType,
+	typename RhsType
+	UE_REQUIRES_DEFINITION(UE_REQUIRES_EXPR((LhsType*)nullptr == (const RhsType*)nullptr))
+>
+FORCEINLINE bool operator==(const TWeakFieldPtr<LhsType>& Lhs, const RhsType* Rhs)
+{
+	return Lhs.Field == Rhs;
+}
 
-// Helper function which deduces the type of the initializer
-template <typename T>
-FORCEINLINE TWeakFieldPtr<T> MakeWeakFieldPtr(T* Ptr)
+/**
+* Compare weak pointers for equality
+* @param Lhs pointer to compare
+* @param Rhs weak pointer to compare
+**/
+template <
+	typename LhsType,
+	typename RhsType
+	UE_REQUIRES_DEFINITION(UE_REQUIRES_EXPR((const LhsType*)nullptr == (RhsType*)nullptr))
+>
+FORCEINLINE bool operator==(const LhsType* Lhs, const TWeakFieldPtr<RhsType>& Rhs)
 {
-	return TWeakFieldPtr<T>(Ptr);
+	return Lhs == Rhs.Field;
 }
 
-template <typename LhsT, typename RhsT>
-FORCENOINLINE bool operator==(const LhsT* Lhs, const TWeakFieldPtr<RhsT>& Rhs)
+/**
+* Test weak pointer for null
+* @param Lhs weak pointer to test
+**/
+template <typename LhsType>
+FORCEINLINE bool operator==(const TWeakFieldPtr<LhsType>& Lhs, TYPE_OF_NULLPTR)
 {
-	// It's also possible that these static_asserts may fail for valid conversions because
-	// one or both of the types have only been forward-declared.
-	static_assert(TPointerIsConvertibleFromTo<LhsT, const FField>::Value, "TWeakFieldPtr can only be compared with FField types");
-	static_assert(TPointerIsConvertibleFromTo<LhsT, RhsT>::Value || TPointerIsConvertibleFromTo<RhsT, LhsT>::Value, "Unable to compare TWeakFieldPtr with raw pointer - types are incompatible");
+	return !Lhs.Get();
+}
 
-	return Rhs == Lhs;
+/**
+* Test weak pointer for null
+* @param Rhs weak pointer to test
+**/
+template <typename RhsType>
+FORCEINLINE bool operator==(TYPE_OF_NULLPTR, const TWeakFieldPtr<RhsType>& Rhs)
+{
+	return !Rhs.Get();
 }
 
-template <typename LhsT>
-FORCENOINLINE bool operator==(const TWeakFieldPtr<LhsT>& Lhs, TYPE_OF_NULLPTR)
+#if !PLATFORM_COMPILER_HAS_GENERATED_COMPARISON_OPERATORS
+/**
+* Compare weak pointers for inequality
+* @param Lhs weak pointer to compare
+* @param Rhs weak pointer to compare
+**/
+template <
+	typename LhsType,
+	typename RhsType
+	UE_REQUIRES_DEFINITION(UE_REQUIRES_EXPR((LhsType*)nullptr != (RhsType*)nullptr))
+>
+FORCEINLINE bool operator!=(const TWeakFieldPtr<LhsType>& Lhs, const TWeakFieldPtr<RhsType>& Rhs)
 {
-	return !Lhs.IsValid();
+	return !(Lhs == Rhs);
 }
 
-template <typename RhsT>
-FORCENOINLINE bool operator==(TYPE_OF_NULLPTR, const TWeakFieldPtr<RhsT>& Rhs)
+/**
+* Compare weak pointers for inequality
+* @param Lhs weak pointer to compare
+* @param Rhs pointer to compare
+**/
+template <
+	typename LhsType,
+	typename RhsType
+	UE_REQUIRES_DEFINITION(UE_REQUIRES_EXPR((LhsType*)nullptr != (const RhsType*)nullptr))
+>
+FORCEINLINE bool operator!=(const TWeakFieldPtr<LhsType>& Lhs, const RhsType* Rhs)
 {
-	return !Rhs.IsValid();
+	return !(Lhs == Rhs);
 }
 
-template <typename LhsT, typename RhsT>
-FORCENOINLINE bool operator!=(const LhsT* Lhs, const TWeakFieldPtr<RhsT>& Rhs)
+/**
+* Compare weak pointers for inequality
+* @param Lhs pointer to compare
+* @param Rhs weak pointer to compare
+**/
+template <
+	typename LhsType,
+	typename RhsType
+	UE_REQUIRES_DEFINITION(UE_REQUIRES_EXPR((const LhsType*)nullptr != (RhsType*)nullptr))
+>
+FORCEINLINE bool operator!=(const LhsType* Lhs, const TWeakFieldPtr<RhsType>& Rhs)
 {
-	// It's also possible that these static_asserts may fail for valid conversions because
-	// one or both of the types have only been forward-declared.
-	static_assert(TPointerIsConvertibleFromTo<LhsT, const FField>::Value, "TWeakFieldPtr can only be compared with FField types");
-	static_assert(TPointerIsConvertibleFromTo<LhsT, RhsT>::Value || TPointerIsConvertibleFromTo<RhsT, LhsT>::Value, "Unable to compare TWeakFieldPtr with raw pointer - types are incompatible");
+	return !(Lhs == Rhs);
+}
 
-	return Rhs != Lhs;
+/**
+* Test weak pointer for non-null
+* @param Lhs weak pointer to test
+**/
+template <typename LhsType>
+FORCEINLINE bool operator!=(const TWeakFieldPtr<LhsType>& Lhs, TYPE_OF_NULLPTR)
+{
+	return !(Lhs == nullptr);
 }
 
-template <typename LhsT>
-FORCENOINLINE bool operator!=(const TWeakFieldPtr<LhsT>& Lhs, TYPE_OF_NULLPTR)
+/**
+* Test weak pointer for non-null
+* @param Rhs weak pointer to test
+**/
+template <typename RhsType>
+FORCEINLINE bool operator!=(TYPE_OF_NULLPTR, const TWeakFieldPtr<RhsType>& Rhs)
 {
-	return Lhs.IsValid();
+	return !(nullptr == Rhs);
 }
+#endif
 
-template <typename RhsT>
-FORCENOINLINE bool operator!=(TYPE_OF_NULLPTR, const TWeakFieldPtr<RhsT>& Rhs)
+// Helper function which deduces the type of the initializer
+template <typename T>
+FORCEINLINE TWeakFieldPtr<T> MakeWeakFieldPtr(T* Ptr)
 {
-	return Rhs.IsValid();
+	return TWeakFieldPtr<T>(Ptr);
 }
 
 template<class T> struct TIsPODType<TWeakFieldPtr<T> > { enum { Value = true }; };
