Updated EMouseCursor in NoExportTypes to properly mirror ICursor

#rnx

diff --git a/Engine/Source/Runtime/CoreUObject/Public/UObject/NoExportTypes.h b/Engine/Source/Runtime/CoreUObject/Public/UObject/NoExportTypes.h
--- a/Engine/Source/Runtime/CoreUObject/Public/UObject/NoExportTypes.h
+++ b/Engine/Source/Runtime/CoreUObject/Public/UObject/NoExportTypes.h
@@ -352,6 +352,9 @@ namespace EMouseCursor
 
 		/** Eye-dropper cursor for picking colors. */
 		EyeDropper,
+
+		/** Custom cursor shape for platforms that support setting a native cursor shape. Same as specifying None if not set. */
+		Custom,
 	};
 }
 
