Avoid doing a MoveTemp in OPEN code as when we OnAbort we need to correctly unset this memory
Since these are only cached values its safe to only set them in none-transactional code while avoiding

diff --git a/Engine/Source/Runtime/CoreUObject/Public/UObject/UObjectAnnotation.h b/Engine/Source/Runtime/CoreUObject/Public/UObject/UObjectAnnotation.h
--- a/Engine/Source/Runtime/CoreUObject/Public/UObject/UObjectAnnotation.h
+++ b/Engine/Source/Runtime/CoreUObject/Public/UObject/UObjectAnnotation.h
@@ -93,10 +93,20 @@ private:
 				UE_AUTORTFM_OPEN(
 				{
 					FScopeLock AnnotationMapLock(&AnnotationMapCritical);
-					AnnotationCacheKey = Object;
-					AnnotationCacheValue = MoveTemp(LocalAnnotation);
 					bWasEmpty = (AnnotationMap.Num() == 0);
-					AnnotationMap.Add(AnnotationCacheKey, AnnotationCacheValue);
+
+					// If we are not in a trasncation update our cache values with the correct information
+					// While in a transaction avoid updating these as it pulls from closed code memory and can cause memory stomps if we abort here
+					if (!AutoRTFM::IsTransactional())
+					{
+						AnnotationCacheKey = Object;
+						AnnotationCacheValue = MoveTemp(LocalAnnotation);
+						AnnotationMap.Add(AnnotationCacheKey, AnnotationCacheValue);
+					}
+					else
+					{
+						AnnotationMap.Add(Object, LocalAnnotation);
+					}
 				});
 
 				UE_AUTORTFM_ONABORT(
