Restoring EInternalObjectFlags::Unreachable to make sure FObjectItem::IsUnreachable() is thread safe.

Deprecated GReachableObjectFlag, GMaybeUnreachableObjectFlag and GUnreachableObjectFlag. Moved the first two to an internal FGCFlags class which is the only thing that can manipulate GC reachability flags.

Deprecated a bunch of FUObjectItem functions that could manipulate reachability flags. These functions will be removed in UE 5.6.

diff --git a/Engine/Plugins/EnhancedInput/Source/InputEditor/Private/InputEditorModule.cpp b/Engine/Plugins/EnhancedInput/Source/InputEditor/Private/InputEditorModule.cpp
--- a/Engine/Plugins/EnhancedInput/Source/InputEditor/Private/InputEditorModule.cpp
+++ b/Engine/Plugins/EnhancedInput/Source/InputEditor/Private/InputEditorModule.cpp
@@ -509,7 +509,7 @@ void FInputEditorModule::AutoUpgradeDefaultInputClasses()
 const UPlayerMappableKeySettings* FInputEditorModule::FindMappingByName(const FName InName)
 {
 	// Instanced objects that have been deleted will have an outer of the transient package, and we don't want to include them
-	for (TObjectIterator<UPlayerMappableKeySettings> Itr(RF_ClassDefaultObject | RF_Transient, true, EInternalObjectFlags::Garbage | UE::GC::GUnreachableObjectFlag); Itr; ++Itr)
+	for (TObjectIterator<UPlayerMappableKeySettings> Itr(RF_ClassDefaultObject | RF_Transient, true, EInternalObjectFlags::Garbage | EInternalObjectFlags::Unreachable); Itr; ++Itr)
 	{
 		const UPlayerMappableKeySettings* Settings = *Itr;
 		if (Settings && Settings->GetOuter() != GetTransientPackage() && Settings->Name == InName)
diff --git a/Engine/Source/Editor/UnrealEd/Private/PackageTools.cpp b/Engine/Source/Editor/UnrealEd/Private/PackageTools.cpp
--- a/Engine/Source/Editor/UnrealEd/Private/PackageTools.cpp
+++ b/Engine/Source/Editor/UnrealEd/Private/PackageTools.cpp
@@ -164,7 +164,7 @@ UPackageTools::UPackageTools(const FObjectInitializer& ObjectInitializer)
 						Object->SetFlags(RF_Standalone);
 					}
 					return true;
-				}, true, RF_NoFlags, UE::GC::GUnreachableObjectFlag);
+				}, true, RF_NoFlags, EInternalObjectFlags::Unreachable);
 			}
 		}
 	}
diff --git a/Engine/Source/Runtime/CoreUObject/Private/UObject/GarbageCollection.cpp b/Engine/Source/Runtime/CoreUObject/Private/UObject/GarbageCollection.cpp
--- a/Engine/Source/Runtime/CoreUObject/Private/UObject/GarbageCollection.cpp
+++ b/Engine/Source/Runtime/CoreUObject/Private/UObject/GarbageCollection.cpp
@@ -37,6 +37,7 @@
 #include "HAL/RunnableThread.h"
 #include "UObject/FieldPathProperty.h"
 #include "UObject/GarbageCollectionHistory.h"
+#include "UObject/GarbageCollectionInternalFlags.h"
 #include "UObject/GarbageCollectionTesting.h"
 #include "UObject/PropertyOptional.h"
 #include "UObject/ExpandingChunkedList.h"
@@ -564,7 +565,7 @@ namespace UE::GC
 	EInternalObjectFlags GReachableObjectFlag = EInternalObjectFlags::ReachabilityFlag0;
 
 	/** EInternalObjectFlag value representing an unreachable object */
-	EInternalObjectFlags GUnreachableObjectFlag = EInternalObjectFlags::ReachabilityFlag1;
+	EInternalObjectFlags GUnreachableObjectFlag = EInternalObjectFlags::Unreachable;
 
 	/** EInternalObjectFlag value representing a maybe unreachable object */
 	EInternalObjectFlags GMaybeUnreachableObjectFlag = EInternalObjectFlags::ReachabilityFlag2;
@@ -584,6 +585,12 @@ namespace UE::GC::Private
 
 	static TSet<int32> GRoots;
 	static FCriticalSection GRootsCritical;
+
+	/** Current EInternalObjectFlags value representing a reachable object */
+	EInternalObjectFlags FGCFlags::ReachableObjectFlag = EInternalObjectFlags::ReachabilityFlag0;
+
+	/** Current EInternalObjectFlags value representing a maybe unreachable object */
+	EInternalObjectFlags FGCFlags::MaybeUnreachableObjectFlag = EInternalObjectFlags::ReachabilityFlag1;
 }
 
 static bool GatherUnreachableObjects(UE::GC::EGatherOptions Options, double TimeLimit = 0.0);
@@ -605,7 +612,7 @@ bool FUObjectItem::SetRootFlags(EInternalObjectFlags FlagsToSet)
 				GRoots.Add(GUObjectArray.ObjectToIndex(Object));
 			}
 		}
-		bIChangedIt = ThisThreadAtomicallySetFlag_ForGC(FlagsToSet);
+		bIChangedIt = AtomicallySetFlag_ForGC(FlagsToSet);
 	}
 	if (bIChangedIt & GIsIncrementalReachabilityPending) //-V792
 	{
@@ -636,9 +643,55 @@ bool FUObjectItem::ClearRootFlags(EInternalObjectFlags FlagsToClear)
 	{
 		GRoots.Remove(GUObjectArray.ObjectToIndex(Object));
 	}
-	return ThisThreadAtomicallyClearedFlag_ForGC((EInternalObjectFlags)FlagsToClear);
+	return AtomicallyClearFlag_ForGC((EInternalObjectFlags)FlagsToClear);
+}
+
+// BEGIN Deprecated FUObjectItem functions. These functions will be removed in the next release and are not used internally by the engine
+bool FUObjectItem::IsMaybeUnreachable() const
+{
+	return UE::GC::Private::FGCFlags::IsMaybeUnreachable_ForGC(this);
+}
+
+void FUObjectItem::SetMaybeUnreachable()
+{
+	UE::GC::Private::FGCFlags::SetMaybeUnreachable_ForGC(this);
+}
+
+void FUObjectItem::FastMarkAsReachableAndClearReachaleInClusterInterlocked_ForGC()
+{
+	UE::GC::Private::FGCFlags::FastMarkAsReachableAndClearReachableInClusterInterlocked_ForGC(this);
+}
+
+void FUObjectItem::FastMarkAsReachableInterlocked_ForGC()
+{
+	UE::GC::Private::FGCFlags::FastMarkAsReachableInterlocked_ForGC(this);
+}
 
+bool FUObjectItem::MarkAsReachableInterlocked_ForGC()
+{
+	return UE::GC::Private::FGCFlags::MarkAsReachableInterlocked_ForGC(this);
 }
+// END Deprecated FUObjectItem functions
+
+// BEGIN Hidden (but DLL exported) functions that can only be used for testing GC and GC flags
+namespace UE::GC::Test
+{
+	COREUOBJECT_API bool IsReachable(const UObject* Object)
+	{
+		const FUObjectItem* ObjectItem = GUObjectArray.IndexToObjectUnsafeForGC(GUObjectArray.ObjectToIndex(Object));
+		return UE::GC::Private::FGCFlags::IsReachable_ForGC(ObjectItem);
+	}
+	COREUOBJECT_API bool IsMaybeUnreachable(const UObject* Object)
+	{
+		return UE::GC::Private::FGCFlags::IsMaybeUnreachable_ForGC(Object);
+	}
+	COREUOBJECT_API void SetMaybeUnreachable(UObject* Object)
+	{
+		FUObjectItem* ObjectItem = GUObjectArray.IndexToObjectUnsafeForGC(GUObjectArray.ObjectToIndex(Object));
+		UE::GC::Private::FGCFlags::SetMaybeUnreachable_ForGC(ObjectItem);
+	}
+} // namespace UE::GC::Privates
+// END Hidden functions
 
 void OnDisregardForGCSetDisabled(int32 NumObjects)
 {
@@ -663,7 +716,7 @@ void OnDisregardForGCSetDisabled(int32 NumObjects)
 			FUObjectItem* RootItem = &GUObjectArray.GetObjectItemArrayUnsafe()[ObjectIndex];	
 			if (RootItem->Object)
 			{
-				RootItem->ThisThreadAtomicallySetFlag_ForGC(GReachableObjectFlag);
+				FGCFlags::SetReachable_ForGC(RootItem);
 				if (RootItem->HasAnyFlags(EInternalObjectFlags_RootFlags))
 				{
 					ThreadState.Payload.Add(ObjectIndex);
@@ -1149,6 +1202,8 @@ FORCEINLINE static void MarkReferencedClustersAsReachableThunk(int32 ClusterInde
 template<EGCOptions Options, class ContainerType>
 static bool MarkClusterMutableObjectsAsReachable(FUObjectCluster& Cluster, ContainerType& ObjectsToSerialize)
 {
+	using namespace UE::GC::Private;
+
 	// This is going to be the return value and basically means that we ran across some Garbage objects
 	bool bAddClusterObjectsToSerialize = false;
 	for (int32& ReferencedMutableObjectIndex : Cluster.MutableObjects)
@@ -1159,9 +1214,9 @@ static bool MarkClusterMutableObjectsAsReachable(FUObjectCluster& Cluster, Conta
 			UE::GC::GDetailedStats.IncClusterToObjectRefs(ReferencedMutableObjectItem);
 			if (!ReferencedMutableObjectItem->HasAnyFlags(EInternalObjectFlags::Garbage))
 			{
-				if (ReferencedMutableObjectItem->IsMaybeUnreachable())
+				if (FGCFlags::IsMaybeUnreachable_ForGC(ReferencedMutableObjectItem))
 				{
-					if (ReferencedMutableObjectItem->MarkAsReachableInterlocked_ForGC())
+					if (FGCFlags::MarkAsReachableInterlocked_ForGC(ReferencedMutableObjectItem))
 					{
 						// Needs doing because this is either a normal unclustered object (clustered objects are never unreachable) or a cluster root
 						ObjectsToSerialize.Add(static_cast<UObject*>(ReferencedMutableObjectItem->Object));
@@ -1176,14 +1231,14 @@ static bool MarkClusterMutableObjectsAsReachable(FUObjectCluster& Cluster, Conta
 				else if (ReferencedMutableObjectItem->GetOwnerIndex() > 0 && !ReferencedMutableObjectItem->HasAnyFlags(EInternalObjectFlags::ReachableInCluster))
 				{
 					// This is a clustered object that maybe hasn't been processed yet
-					if (ReferencedMutableObjectItem->ThisThreadAtomicallySetFlag_ForGC(EInternalObjectFlags::ReachableInCluster))
+					if (FGCFlags::ThisThreadAtomicallySetFlag_ForGC(ReferencedMutableObjectItem, EInternalObjectFlags::ReachableInCluster))
 					{
 						// Needs doing, we need to get its cluster root and process it too
 						FUObjectItem* ReferencedMutableObjectsClusterRootItem = GUObjectArray.IndexToObjectUnsafeForGC(ReferencedMutableObjectItem->GetOwnerIndex());
-						if (ReferencedMutableObjectsClusterRootItem->IsMaybeUnreachable())
+						if (FGCFlags::IsMaybeUnreachable_ForGC(ReferencedMutableObjectsClusterRootItem))
 						{
 							// The root is also maybe unreachable so process it and all the referenced clusters
-							if (ReferencedMutableObjectsClusterRootItem->MarkAsReachableInterlocked_ForGC())
+							if (FGCFlags::MarkAsReachableInterlocked_ForGC(ReferencedMutableObjectsClusterRootItem))
 							{
 								MarkReferencedClustersAsReachable<Options>(ReferencedMutableObjectsClusterRootItem->GetClusterIndex(), ObjectsToSerialize);
 							}
@@ -1225,6 +1280,8 @@ FORCEINLINE static void MarkClusterMutableCellsAsReachable(FUObjectCluster& Clus
 template<EGCOptions Options, class ContainerType>
 static FORCENOINLINE void MarkReferencedClustersAsReachable(int32 ClusterIndex, ContainerType& ObjectsToSerialize)
 {
+	using namespace UE::GC::Private;
+
 	UE::GC::GDetailedStats.IncNumClustersTraversed();
 
 	// If we run across some Garbage objects we need to add all objects from this cluster
@@ -1242,7 +1299,7 @@ static FORCENOINLINE void MarkReferencedClustersAsReachable(int32 ClusterIndex,
 			FUObjectItem* ReferencedClusterRootObjectItem = GUObjectArray.IndexToObjectUnsafeForGC(ReferncedClusterIndex);
 			if (!ReferencedClusterRootObjectItem->HasAnyFlags(EInternalObjectFlags::Garbage))
 			{
-				ReferencedClusterRootObjectItem->FastMarkAsReachableInterlocked_ForGC();
+				FGCFlags::FastMarkAsReachableInterlocked_ForGC(ReferencedClusterRootObjectItem);
 			}
 			else
 			{
@@ -1778,6 +1835,8 @@ private:
 	// Process all queued validated references
 	FORCEINLINE_DEBUGGABLE void DrainValidated(const uint32 Num)
 	{
+		using namespace UE::GC::Private;
+
 		check(Num <= ValidatedBatchSize);
 
 		// Prefetch metadata
@@ -1786,7 +1845,7 @@ private:
 		// shares cacheline with the raw UObject*. Not adding an offset leads to a bit faster and more compact code.
 		static constexpr uint32 InternalIndexPrefetchOffset = 0;
 		static constexpr uint32 PrefetchAhead = ValidatedPrefetchAhead;
-		static constexpr ::size_t OffsetOfFlags = FUObjectItem::OffsetOfFlags();
+		static constexpr ::size_t OffsetOfFlags = FGCFlags::OffsetOfFlags_ForGC();
 
 		uint32 ObjectIndices[ValidatedBatchSize];
 		if (Num > PrefetchAhead)
@@ -3080,7 +3139,9 @@ public:
 
 	FORCEINLINE static bool HandleValidReference(FWorkerContext& Context, FImmutableReference Reference, FReferenceMetadata Metadata)
 	{
-		if (Metadata.ObjectItem->MarkAsReachableInterlocked_ForGC())
+		using namespace UE::GC::Private;
+
+		if (FGCFlags::MarkAsReachableInterlocked_ForGC(Metadata.ObjectItem))
 		{
 			// Objects that are part of a GC cluster should never have the unreachable flag set!
 			checkSlow(Metadata.ObjectItem->GetOwnerIndex() <= 0);
@@ -3106,9 +3167,9 @@ public:
 				FUObjectItem* RootObjectItem = GUObjectArray.IndexToObjectUnsafeForGC(Metadata.ObjectItem->GetOwnerIndex());
 				checkSlow(RootObjectItem->HasAnyFlags(EInternalObjectFlags::ClusterRoot));
 
-				if (Metadata.ObjectItem->ThisThreadAtomicallySetFlag_ForGC(EInternalObjectFlags::ReachableInCluster))
+				if (FGCFlags::ThisThreadAtomicallySetFlag_ForGC(Metadata.ObjectItem, EInternalObjectFlags::ReachableInCluster))
 				{
-					if (RootObjectItem->MarkAsReachableInterlocked_ForGC())
+					if (FGCFlags::MarkAsReachableInterlocked_ForGC(RootObjectItem))
 					{
 						// Make sure all referenced clusters are marked as reachable too
 						MarkReferencedClustersAsReachableThunk<Options>(RootObjectItem->GetClusterIndex(), Context.ObjectsToSerialize);
@@ -4195,7 +4256,7 @@ public:
 						bool bKeepCluster = RootItem->HasAnyFlags(EInternalObjectFlags_RootFlags);
 						if (bKeepCluster)
 						{
-							RootItem->FastMarkAsReachableInterlocked_ForGC();
+							FGCFlags::FastMarkAsReachableInterlocked_ForGC(RootItem);
 							ThreadState.Payload.KeepClusters.Add(RootItem);
 						}
 
@@ -4203,7 +4264,7 @@ public:
 						{
 							FUObjectItem* ClusteredItem = &GUObjectArray.GetObjectItemArrayUnsafe()[ObjectIndex];
 
-							ClusteredItem->FastMarkAsReachableAndClearReachaleInClusterInterlocked_ForGC();
+							FGCFlags::FastMarkAsReachableAndClearReachableInClusterInterlocked_ForGC(ClusteredItem);
 
 							if (!bKeepCluster && ClusteredItem->HasAnyFlags(EInternalObjectFlags_RootFlags))
 							{
@@ -4278,7 +4339,7 @@ public:
 					checkCode(if (RootItem->HasAllFlags(EInternalObjectFlags::Garbage | EInternalObjectFlags::RootSet)) { UE_LOG(LogGarbage, Fatal, TEXT("Object %s is part of root set though has been marked as Garbage!"), *Object->GetFullName()); });
 #endif
 
-					RootItem->FastMarkAsReachableInterlocked_ForGC();
+					FGCFlags::FastMarkAsReachableInterlocked_ForGC(RootItem);
 					ThreadState.Payload.Add(Object);
 				}
 			}, (MarkRootsState.NumWorkerThreads() == 1) ? EParallelForFlags::ForceSingleThread : EParallelForFlags::None);			
@@ -4310,7 +4371,7 @@ public:
 						// IsValidLowLevel is extremely slow in this loop so only do it in debug
 						checkSlow(Object->IsValidLowLevel());
 
-						ObjectItem->FastMarkAsReachableInterlocked_ForGC();
+						FGCFlags::FastMarkAsReachableInterlocked_ForGC(ObjectItem);
 						ThreadState.Payload.Add(Object);
 					}
 				}
@@ -4331,12 +4392,13 @@ public:
 	FORCENOINLINE void MarkObjectsAsUnreachable(const EObjectFlags KeepFlags)
 	{
 		using namespace UE::GC;
+		using namespace UE::GC::Private;
 
 		// Don't swap the flags if we're re-entering this function to track garbage references
 		if (const bool bInitialMark = !Stats.bFoundGarbageRef)
 		{
 			// This marks all UObjects as MaybeUnreachable
-			Swap(GReachableObjectFlag, GMaybeUnreachableObjectFlag);
+			FGCFlags::SwapReachableAndMaybeUnreachable();
 		}
 
 		// Not counting the disregard for GC set to preserve legacy behavior
@@ -4718,6 +4780,8 @@ static FAutoConsoleVariableRef CVarAdditionalFinishDestroyTimeGC(
 
 bool IncrementalDestroyGarbage(bool bUseTimeLimit, double TimeLimit)
 {
+	using namespace UE::GC::Private;
+
 	TRACE_CPUPROFILER_EVENT_SCOPE(IncrementalDestroyGarbage);
 	const bool bMultithreadedPurge = !ShouldForceSingleThreadedGC() && GMultithreadedDestructionEnabled;
 	if (!GAsyncPurge)
@@ -4772,7 +4836,7 @@ bool IncrementalDestroyGarbage(bool bUseTimeLimit, double TimeLimit)
 
 				//@todo UE - A prefetch was removed here. Re-add it. It wasn't right anyway, since it was ten items ahead and the consoles on have 8 prefetch slots
 
-				check(!ObjectItem->IsMaybeUnreachable());
+				check(!FGCFlags::IsReachable_ForGC(ObjectItem) && FGCFlags::IsMaybeUnreachable_ForGC(ObjectItem));
 				if (ObjectItem->IsUnreachable())
 				{
 					UObject* Object = static_cast<UObject*>(ObjectItem->Object);
@@ -5106,14 +5170,14 @@ void DissolveUnreachableClusters(UE::GC::EGatherOptions Options)
 			if (Cluster.RootIndex >= 0)
 			{
 				FUObjectItem* RootItem = &GUObjectArray.GetObjectItemArrayUnsafe()[Cluster.RootIndex];
-				if (RootItem->IsUnreachable())
+				if (FGCFlags::IsMaybeUnreachable_ForGC(RootItem))
 				{
 					for (int32 ObjectIndex : Cluster.Objects)
 					{
 						FUObjectItem* ObjectItem = &GUObjectArray.GetObjectItemArrayUnsafe()[ObjectIndex];
 						if (!ObjectItem->HasAnyFlags(EInternalObjectFlags::ReachableInCluster))
 						{
-							ObjectItem->SetUnreachable();
+							FGCFlags::SetMaybeUnreachable_ForGC(ObjectItem);
 							ThisThreadClusteredObjects++;
 						}
 					}
@@ -5181,9 +5245,10 @@ bool GatherUnreachableObjects(UE::GC::EGatherOptions Options, double TimeLimit /
 		while (Iterator.Index <= Iterator.LastIndex)
 		{
 			FUObjectItem* ObjectItem = &GUObjectArray.GetObjectItemArrayUnsafe()[Iterator.Index++];
-			if (ObjectItem->IsUnreachable())
+			if (FGCFlags::IsMaybeUnreachable_ForGC(ObjectItem))
 			{
 				checkf(!ObjectItem->HasAnyFlags(EInternalObjectFlags::ClusterRoot), TEXT("Unreachable cluster root found. Unreachable clusters should have been dissolved in DissolveUnreachableClusters!"));
+				FGCFlags::SetUnreachable(ObjectItem);
 				Iterator.Payload.Add(ObjectItem);
 			}
 
@@ -5265,7 +5330,7 @@ public:
 		
 	virtual void HandleObjectReference(UObject*& InObject, const UObject* InReferencingObject, const FProperty* InReferencingProperty) override 
 	{
-		if (InObject && InObject->IsUnreachable())
+		if (InObject && UE::GC::Private::FGCFlags::IsMaybeUnreachable_ForGC(InObject))
 		{
 			InObject = nullptr;
 		}
@@ -5273,7 +5338,7 @@ public:
 	
 	virtual bool MarkWeakObjectReferenceForClearing(UObject** WeakReference, UObject* ReferenceOwner) override
 	{
-		if (WeakReference && *WeakReference && (*WeakReference)->IsUnreachable())
+		if (WeakReference && *WeakReference && UE::GC::Private::FGCFlags::IsMaybeUnreachable_ForGC(*WeakReference))
 		{
 			*WeakReference = nullptr;
 		}
@@ -5284,6 +5349,8 @@ public:
 template <bool bGatheredWithIncrementalReachability = true>
 static void ClearWeakReferences(TConstArrayView<TUniquePtr<FWorkerContext>> Contexts)
 {
+	using namespace UE::GC::Private;
+
 	TRACE_CPUPROFILER_EVENT_SCOPE(ClearWeakReferences);
 	TSet<UObject*> ObjectsThatNeedWeakReferenceClearing;
 	for (const TUniquePtr<FWorkerContext>& Context : Contexts)
@@ -5292,7 +5359,7 @@ static void ClearWeakReferences(TConstArrayView<TUniquePtr<FWorkerContext>> Cont
 		for (FWeakReferenceInfo& ReferenceInfo : Context->WeakReferences)
 		{
 			UObject* ReferencedObject = ReferenceInfo.ReferencedObject;
-			if (ReferencedObject && ReferencedObject->IsUnreachable())
+			if (ReferencedObject && FGCFlags::IsMaybeUnreachable_ForGC(ReferencedObject))
 			{
 				if constexpr (bGatheredWithIncrementalReachability)
 				{
@@ -5315,7 +5382,7 @@ static void ClearWeakReferences(TConstArrayView<TUniquePtr<FWorkerContext>> Cont
 		FWeakReferenceEliminator ReferenceEliminator;
 		for (UObject* Object : ObjectsThatNeedWeakReferenceClearing)
 		{
-			if (Object && !Object->IsUnreachable())
+			if (Object && !FGCFlags::IsMaybeUnreachable_ForGC(Object))
 			{
 				Object->GetClass()->CallAddReferencedObjects(Object, ReferenceEliminator);
 			}
@@ -5625,8 +5692,6 @@ void PostCollectGarbageImpl(EObjectFlags KeepFlags)
 
 		DumpGarbageReferencers(AllContexts);
 
-		Swap(GUnreachableObjectFlag, GMaybeUnreachableObjectFlag);
-
 		const EGatherOptions GatherOptions = GetObjectGatherOptions();
 		DissolveUnreachableClusters(GatherOptions);
 
@@ -5649,7 +5714,7 @@ void PostCollectGarbageImpl(EObjectFlags KeepFlags)
 
 		GGatherUnreachableObjectsState.Init();
 
-		if (bPerformFullPurge || !GAllowIncrementalGather)
+		if (bPerformFullPurge || !GAllowIncrementalGather || !FGCFlags::IsIncrementalGatherUnreachableSupported())
 		{
 			GatherUnreachableObjects(GatherOptions, /*TimeLimit =*/ 0.0);
 		}
@@ -6171,6 +6236,11 @@ bool UObject::IsDestructionThreadSafe() const
 	return false;
 }
 
+bool UObjectBaseUtility::ThisThreadAtomicallyClearedRFUnreachable()
+{
+	return UE::GC::Private::FGCFlags::ThisThreadAtomicallyClearedRFUnreachable(GUObjectArray.IndexToObject(InternalIndex));
+}
+
 template <bool bIsVerse>
 FORCEINLINE static void MarkObjectItemAsReachable(FUObjectItem* ObjectItem)
 {
@@ -6188,7 +6258,7 @@ FORCEINLINE static void MarkObjectItemAsReachable(FUObjectItem* ObjectItem)
 	{
 		checkf(GIsIncrementalReachabilityPending, TEXT("%s is marked as MaybeUnreachable but Incremental Reachability Analysis is not in progress"), *static_cast<UObject*>(ObjectItem->Object)->GetFullName());
 	}
-	if (ObjectItem->MarkAsReachableInterlocked_ForGC())
+	if (FGCFlags::MarkAsReachableInterlocked_ForGC(ObjectItem))
 	{
 		if (ObjectItem->GetOwnerIndex() >= 0)
 		{
@@ -6206,15 +6276,17 @@ FORCEINLINE static void MarkObjectItemAsReachable(FUObjectItem* ObjectItem)
 template <bool bIsVerse>
 FORCEINLINE static void MarkAsReachable(const UObjectBase* Obj)
 {
+	using namespace UE::GC::Private;
+
 	FUObjectItem* ObjectItem = GUObjectArray.ObjectToObjectItem(Obj);
-	if (ObjectItem->IsMaybeUnreachable())
+	if (FGCFlags::IsMaybeUnreachable_ForGC(ObjectItem))
 	{
 		MarkObjectItemAsReachable<bIsVerse>(ObjectItem);
 	}
 	else if (ObjectItem->GetOwnerIndex() > 0) // Clustered objects are never marked as MaybeUnreachable so we need to check if the cluster root is MaybeUnreachable
 	{
 		FUObjectItem* ClusterRootObjectItem = GUObjectArray.IndexToObject(ObjectItem->GetOwnerIndex());
-		if (ClusterRootObjectItem->IsMaybeUnreachable())
+		if (FGCFlags::IsMaybeUnreachable_ForGC(ClusterRootObjectItem))
 		{
 			MarkObjectItemAsReachable<bIsVerse>(ClusterRootObjectItem);
 		}
diff --git a/Engine/Source/Runtime/CoreUObject/Private/UObject/GarbageCollectionInternalFlags.h b/Engine/Source/Runtime/CoreUObject/Private/UObject/GarbageCollectionInternalFlags.h
new file mode 100644
--- /dev/null
+++ b/Engine/Source/Runtime/CoreUObject/Private/UObject/GarbageCollectionInternalFlags.h
@@ -0,0 +1,149 @@
+// Copyright Epic Games, Inc. All Rights Reserved.
+
+/*==============================================================================================
+	GarbageCollectionInternalFlags.h: Unreal realtime garbage collection internal flags helpers
+===============================================================================================*/
+
+#pragma once
+
+#include "CoreMinimal.h"
+#include "Stats/Stats.h"
+#include "UObject/GarbageCollection.h"
+#include "UObject/Object.h"
+#include "UObject/UObjectArray.h"
+
+namespace UE::GC::Private
+{
+
+/**
+* Access to internal garbage collector rachability flags. Only GC and GC related functions can use these.
+* NOTHING except GC should be manipulating reachability flags (including EInternalObjectFlags::Unreachable).
+* EInternalObjectFlags::Unreachable is the ONLY reachability flag that can be safely READ by non-GC functions.
+* Reading ReachableObjectFlag and MaybeUnreachableObjectFlag outside of GC is NOT THREAD SAFE.
+*/
+class FGCFlags
+{
+	/** Current EInternalObjectFlags value representing a reachable object */
+	static EInternalObjectFlags ReachableObjectFlag;
+
+	/** Current EInternalObjectFlags value representing a maybe unreachable object */
+	static EInternalObjectFlags MaybeUnreachableObjectFlag;
+
+public:
+
+	FORCEINLINE static void SetUnreachable(FUObjectItem* ObjectItem)
+	{
+		ObjectItem->AtomicallySetFlag_ForGC(EInternalObjectFlags::Unreachable);
+	}
+
+	FORCEINLINE static void SetReachable_ForGC(FUObjectItem* ObjectItem)
+	{
+		ObjectItem->AtomicallySetFlag_ForGC(ReachableObjectFlag);
+	}
+
+	FORCEINLINE static bool IsReachable_ForGC(const FUObjectItem* ObjectItem)
+	{
+		return !!(ObjectItem->GetFlagsInternal() & int32(ReachableObjectFlag));
+	}
+
+	FORCEINLINE static void SetMaybeUnreachable_ForGC(FUObjectItem* ObjectItem)
+	{
+		ObjectItem->AtomicallyClearFlag_ForGC(ReachableObjectFlag);
+		ObjectItem->AtomicallySetFlag_ForGC(MaybeUnreachableObjectFlag);
+	}
+
+	FORCEINLINE static void ClearMaybeUnreachable_ForGC(FUObjectItem* ObjectItem)
+	{
+		ObjectItem->AtomicallyClearFlag_ForGC(MaybeUnreachableObjectFlag);
+		ObjectItem->AtomicallySetFlag_ForGC(ReachableObjectFlag);
+	}
+
+	FORCEINLINE static bool IsMaybeUnreachable_ForGC(const FUObjectItem* ObjectItem)
+	{
+		return !!(ObjectItem->GetFlagsInternal() & int32(MaybeUnreachableObjectFlag));
+	}
+
+	FORCEINLINE static bool IsMaybeUnreachable_ForGC(const UObject* Object)
+	{
+		const FUObjectItem* ObjectItem = GUObjectArray.IndexToObjectUnsafeForGC(GUObjectArray.ObjectToIndex(Object));
+		return IsMaybeUnreachable_ForGC(ObjectItem);
+	}
+
+	FORCEINLINE static bool ThisThreadAtomicallyClearedRFUnreachable(FUObjectItem* ObjectItem)
+	{
+		return ObjectItem->AtomicallyClearFlag_ForGC(EInternalObjectFlags::Unreachable);
+	}
+
+	FORCEINLINE static bool ThisThreadAtomicallySetFlag_ForGC(FUObjectItem* ObjectItem, EInternalObjectFlags Flag)
+	{
+		return ObjectItem->AtomicallySetFlag_ForGC(Flag);
+	}
+
+	FORCEINLINE static bool ThisThreadAtomicallyClearedFlag_ForGC(FUObjectItem* ObjectItem, EInternalObjectFlags Flag)
+	{
+		return ObjectItem->AtomicallyClearFlag_ForGC(Flag);
+	}
+
+	FORCEINLINE static void FastMarkAsReachableInterlocked_ForGC(FUObjectItem* ObjectItem)
+	{
+		FPlatformAtomics::InterlockedAnd(&ObjectItem->Flags, ~int32(MaybeUnreachableObjectFlag));
+		FPlatformAtomics::InterlockedOr(&ObjectItem->Flags, int32(ReachableObjectFlag));
+	}
+
+	FORCEINLINE static void FastMarkAsReachableAndClearReachableInClusterInterlocked_ForGC(FUObjectItem* ObjectItem)
+	{
+		FPlatformAtomics::InterlockedAnd(&ObjectItem->Flags, ~int32(MaybeUnreachableObjectFlag | EInternalObjectFlags::ReachableInCluster));
+		FPlatformAtomics::InterlockedOr(&ObjectItem->Flags, int32(ReachableObjectFlag));
+	}
+
+	FORCEINLINE static bool MarkAsReachableInterlocked_ForGC(FUObjectItem* ObjectItem)
+	{
+		const int32 FlagToClear = int32(MaybeUnreachableObjectFlag);
+		if (FPlatformAtomics::AtomicRead_Relaxed(&ObjectItem->Flags) & FlagToClear)
+		{
+			int32 Old = FPlatformAtomics::InterlockedAnd(&ObjectItem->Flags, ~FlagToClear);
+			FPlatformAtomics::InterlockedOr(&ObjectItem->Flags, int32(ReachableObjectFlag));
+			return Old & FlagToClear;
+		}
+		return false;
+	}
+
+	FORCEINLINE static constexpr ::size_t OffsetOfFlags_ForGC()
+	{
+		return offsetof(FUObjectItem, Flags);
+	}
+
+	FORCEINLINE static void SwapReachableAndMaybeUnreachable()
+	{
+		// It's important to lock the global UObjectArray so that the flag swap doesn't occur while a new object is being created
+		// as we set the GReachableObjectFlag on all newly created objects
+		GUObjectArray.LockInternalArray();
+		
+		Swap(ReachableObjectFlag, MaybeUnreachableObjectFlag);
+
+		// Maintain the old flag variables for backwards compatibility
+		PRAGMA_DISABLE_DEPRECATION_WARNINGS
+		UE::GC::GReachableObjectFlag = ReachableObjectFlag;
+		UE::GC::GMaybeUnreachableObjectFlag = MaybeUnreachableObjectFlag;
+		PRAGMA_ENABLE_DEPRECATION_WARNINGS
+
+		GUObjectArray.UnlockInternalArray();
+	}
+
+	FORCEINLINE static EInternalObjectFlags GetReachableFlagValue_ForGC()
+	{
+		return ReachableObjectFlag;
+	}
+
+	FORCEINLINE static EInternalObjectFlags GetMaybeUnreachableFlagValue_ForGC()
+	{
+		return MaybeUnreachableObjectFlag;
+	}
+
+	FORCEINLINE static bool IsIncrementalGatherUnreachableSupported()
+	{
+		return false;
+	}
+};
+
+} // namespace UE::GC::Private
\ No newline at end of file
diff --git a/Engine/Source/Runtime/CoreUObject/Private/UObject/GarbageCollectionVerification.cpp b/Engine/Source/Runtime/CoreUObject/Private/UObject/GarbageCollectionVerification.cpp
--- a/Engine/Source/Runtime/CoreUObject/Private/UObject/GarbageCollectionVerification.cpp
+++ b/Engine/Source/Runtime/CoreUObject/Private/UObject/GarbageCollectionVerification.cpp
@@ -6,6 +6,7 @@
 
 #include "UObject/GarbageCollectionVerification.h"
 #include "UObject/GarbageCollection.h"
+#include "UObject/GarbageCollectionInternalFlags.h"
 #include "HAL/ThreadSafeBool.h"
 #include "Misc/TimeGuard.h"
 #include "HAL/IConsoleManager.h"
@@ -386,6 +387,8 @@ void VerifyClustersAssumptions()
 
 void VerifyObjectFlags()
 {
+	using namespace UE::GC::Private;
+
 	int32 MaxNumberOfObjects = GUObjectArray.GetObjectArrayNum();
 	int32 NumThreads = FMath::Max(1, FTaskGraphInterface::Get().GetNumWorkerThreads());
 	int32 NumberOfObjectsPerThread = (MaxNumberOfObjects / NumThreads) + 1;
@@ -415,7 +418,7 @@ void VerifyObjectFlags()
 					++NumErrors;
 				}
 
-				if (!ObjectItem.HasAnyFlags(UE::GC::GReachableObjectFlag) && !GUObjectArray.IsDisregardForGC(Object))
+				if (!FGCFlags::IsReachable_ForGC(&ObjectItem) && !GUObjectArray.IsDisregardForGC(Object))
 				{
 					UE_LOG(LogGarbage, Warning, TEXT("Object %s%s is NOT marked as Reachable at the beginning of GC"),
 						*FReferenceChainSearch::GetObjectFlags(FGCObjectInfo(Object)),
@@ -424,7 +427,7 @@ void VerifyObjectFlags()
 					++NumErrors;
 				}
 
-				if (ObjectItem.HasAnyFlags(UE::GC::GUnreachableObjectFlag| UE::GC::GMaybeUnreachableObjectFlag))
+				if (FGCFlags::IsMaybeUnreachable_ForGC(&ObjectItem) || ObjectItem.IsUnreachable())
 				{
 					UE_LOG(LogGarbage, Warning, TEXT("Object %s%s is marked with at least one of the unreachable flags at the beginning of GC"),
 						*FReferenceChainSearch::GetObjectFlags(FGCObjectInfo(Object)),
@@ -456,10 +459,12 @@ public:
 
 	FORCEINLINE_DEBUGGABLE void HandleTokenStreamObjectReference(FGCArrayStruct& ObjectsToSerializeStruct, UObject* ReferencingObject, UObject*& Object, FMemberId MemberId, EOrigin Origin, bool bAllowReferenceElimination)
 	{
+		using namespace UE::GC::Private;
+
 		if (Object)
 		{
 			FUObjectItem* ObjectItem = GUObjectArray.ObjectToObjectItem(Object);
-			if (ObjectItem->HasAnyFlags(UE::GC::GMaybeUnreachableObjectFlag | UE::GC::GUnreachableObjectFlag))
+			if (ObjectItem->HasAnyFlags(FGCFlags::GetMaybeUnreachableFlagValue_ForGC() | EInternalObjectFlags::Unreachable))
 			{
 				if (ReferencingObject)
 				{
@@ -510,6 +515,8 @@ public:
 
 void VerifyNoUnreachableObjects(int32 NumUnreachable)
 {
+	using namespace UE::GC::Private;
+
 	const double StartTime = FPlatformTime::Seconds();
 	const int32 MaxNumberOfReachableObjects = GUObjectArray.GetObjectArrayNum();
 	const int32 NumThreads = GetNumCollectReferenceWorkers();
@@ -532,16 +539,16 @@ void VerifyNoUnreachableObjects(int32 NumUnreachable)
 			if (ObjectItem.Object)
 			{
 				UObject* Object = static_cast<UObject*>(ObjectItem.Object);
-				if (!ObjectItem.HasAnyFlags(UE::GC::GUnreachableObjectFlag))
+				if (!ObjectItem.HasAnyFlags(EInternalObjectFlags::Unreachable))
 				{
-					if (ObjectItem.HasAnyFlags(UE::GC::GMaybeUnreachableObjectFlag))
+					if (FGCFlags::IsMaybeUnreachable_ForGC(&ObjectItem))
 					{
 						UE_LOG(LogGarbage, Warning, TEXT("Object %s%s is still marked as MaybeUnreachable after Reachability Analysis is complete."), 
 							*FReferenceChainSearch::GetObjectFlags(FGCObjectInfo(Object)),
 							*Object->GetFullName());
 						NumErrors++;
 					}
-					if (!ObjectItem.HasAnyFlags(UE::GC::GReachableObjectFlag) && !GUObjectArray.IsDisregardForGC(Object))
+					if (!FGCFlags::IsReachable_ForGC(&ObjectItem) && !GUObjectArray.IsDisregardForGC(Object))
 					{
 						UE_LOG(LogGarbage, Warning, TEXT("Object %s%s is NOT marked as Unreachable and NOT marked as Reachable."),
 							*FReferenceChainSearch::GetObjectFlags(FGCObjectInfo(Object)),
@@ -554,14 +561,7 @@ void VerifyNoUnreachableObjects(int32 NumUnreachable)
 				{
 					ThisThreadUnrachableObjectsNum++;
 
-					if (ObjectItem.HasAnyFlags(UE::GC::GMaybeUnreachableObjectFlag))
-					{
-						UE_LOG(LogGarbage, Warning, TEXT("Object %s%s is still marked as MaybeUnreachable after Reachability Analysis is complete."),
-							*FReferenceChainSearch::GetObjectFlags(FGCObjectInfo(Object)),
-							*Object->GetFullName());
-						NumErrors++;
-					}
-					if (ObjectItem.HasAnyFlags(UE::GC::GReachableObjectFlag))
+					if (FGCFlags::IsReachable_ForGC(&ObjectItem))
 					{
 						UE_LOG(LogGarbage, Warning, TEXT("Object %s%s is still marked as Reachable."),
 							*FReferenceChainSearch::GetObjectFlags(FGCObjectInfo(Object)),
diff --git a/Engine/Source/Runtime/CoreUObject/Private/UObject/Obj.cpp b/Engine/Source/Runtime/CoreUObject/Private/UObject/Obj.cpp
--- a/Engine/Source/Runtime/CoreUObject/Private/UObject/Obj.cpp
+++ b/Engine/Source/Runtime/CoreUObject/Private/UObject/Obj.cpp
@@ -25,6 +25,7 @@
 #include "UObject/Object.h"
 #include "Serialization/ArchiveUObject.h"
 #include "UObject/GarbageCollectionHistory.h"
+#include "UObject/GarbageCollectionInternalFlags.h"
 #include "UObject/Class.h"
 #include "UObject/EnumProperty.h"
 #include "UObject/UObjectIterator.h"
@@ -5453,11 +5454,11 @@ void StaticExit()
 			if (!Obj->IsA<UField>())
 			{
 				// Mark as unreachable so purge phase will kill it.
-				ObjItem->SetUnreachable();
+				UE::GC::Private::FGCFlags::SetMaybeUnreachable_ForGC(ObjItem);
 			}
 			else
 			{
-				ObjItem->ClearUnreachable();
+				UE::GC::Private::FGCFlags::ClearMaybeUnreachable_ForGC(ObjItem);
 			}
 		}
 	}
@@ -5477,7 +5478,7 @@ void StaticExit()
 		for (FRawObjectIterator It; It; ++It)
 		{
 			// Mark as unreachable so purge phase will kill it.
-			It->SetUnreachable();
+			UE::GC::Private::FGCFlags::SetMaybeUnreachable_ForGC(*It);
 		}
 
 		GatherUnreachableObjects(false);
diff --git a/Engine/Source/Runtime/CoreUObject/Private/UObject/ReferenceChainSearch.cpp b/Engine/Source/Runtime/CoreUObject/Private/UObject/ReferenceChainSearch.cpp
--- a/Engine/Source/Runtime/CoreUObject/Private/UObject/ReferenceChainSearch.cpp
+++ b/Engine/Source/Runtime/CoreUObject/Private/UObject/ReferenceChainSearch.cpp
@@ -10,6 +10,7 @@
 #include "HAL/ThreadHeartBeat.h"
 #include "UObject/FastReferenceCollector.h"
 #include "UObject/GCObject.h"
+#include "UObject/GarbageCollectionInternalFlags.h"
 #include "UObject/UObjectIterator.h"
 #include "UObject/UnrealType.h"
 
@@ -1265,28 +1266,30 @@ namespace UE::ReferenceChainSearch
 
 FString FReferenceChainSearch::GetObjectFlags(const FGCObjectInfo& InObject)
 {
+	using namespace UE::GC::Private;
+
 	FString Flags;
 
 	if (!InObject.IsDisregardForGC())
 	{
-		if (!InObject.HasAnyInternalFlags(UE::GC::GReachableObjectFlag | UE::GC::GMaybeUnreachableObjectFlag | UE::GC::GUnreachableObjectFlag))
+		if (!InObject.HasAnyInternalFlags(EInternalObjectFlags_ReachabilityFlags))
 		{
 			Flags += TEXT("(Error: No reachability flag) ");
 		}
 	}
-	else if (InObject.HasAnyInternalFlags(UE::GC::GReachableObjectFlag))
+	else if (InObject.HasAnyInternalFlags(FGCFlags::GetReachableFlagValue_ForGC()))
 	{
 		Flags += TEXT("(Error: Reachable but NeverGCed) ");
 	}
 
-	if (InObject.HasAnyInternalFlags(UE::GC::GMaybeUnreachableObjectFlag))
+	if (InObject.HasAnyInternalFlags(FGCFlags::GetMaybeUnreachableFlagValue_ForGC()))
 	{
-		Flags += FString::Printf(TEXT("(MaybeUnreachable<%d>) "), (int32)UE::GC::GMaybeUnreachableObjectFlag);
+		Flags += FString::Printf(TEXT("(MaybeUnreachable<%d>) "), FGCFlags::GetMaybeUnreachableFlagValue_ForGC());
 	}
 
-	if (InObject.HasAnyInternalFlags(UE::GC::GUnreachableObjectFlag))
+	if (InObject.HasAnyInternalFlags(EInternalObjectFlags::Unreachable))
 	{
-		Flags += FString::Printf(TEXT("(Unreachable<%d>) "), (int32)UE::GC::GUnreachableObjectFlag);
+		Flags += TEXT("(Unreachable) ");
 	}
 
 	if (InObject.HasAnyInternalFlags(EInternalObjectFlags::RefCounted))
diff --git a/Engine/Source/Runtime/CoreUObject/Private/UObject/UObjectArray.cpp b/Engine/Source/Runtime/CoreUObject/Private/UObject/UObjectArray.cpp
--- a/Engine/Source/Runtime/CoreUObject/Private/UObject/UObjectArray.cpp
+++ b/Engine/Source/Runtime/CoreUObject/Private/UObject/UObjectArray.cpp
@@ -11,6 +11,7 @@
 #include "ProfilingDebugging/AssetMetadataTrace.h"
 #include "UObject/UObjectAllocator.h"
 #include "UObject/Class.h"
+#include "UObject/GarbageCollectionInternalFlags.h"
 #include "UObject/UObjectIterator.h"
 #include "UObject/ReachabilityAnalysisState.h"
 
@@ -268,7 +269,9 @@ void FUObjectArray::AllocateUObjectIndex(UObjectBase* Object, EInternalObjectFla
 	ObjectItem->Flags = (int32)EInternalObjectFlags::PendingConstruction;
 	if (!(IsOpenForDisregardForGC() & GUObjectArray.DisregardForGCEnabled())) //-V792
 	{
-		ObjectItem->Flags |= (int32)UE::GC::GReachableObjectFlag;
+		// It's safe to access FGCFlags::GetReachableFlagValue_ForGC() here because creating new objects is being performed 
+		// under the same UObjectArray lock as swapping reachability flags inside of GC, see FGCFlags::SwapReachableAndMaybeUnreachable()
+		ObjectItem->Flags |= (int32)UE::GC::Private::FGCFlags::GetReachableFlagValue_ForGC();
 	}
 	ObjectItem->RefCount = 0;
 	ObjectItem->ClusterRootIndex = 0;
diff --git a/Engine/Source/Runtime/CoreUObject/Private/UObject/UObjectClusters.cpp b/Engine/Source/Runtime/CoreUObject/Private/UObject/UObjectClusters.cpp
--- a/Engine/Source/Runtime/CoreUObject/Private/UObject/UObjectClusters.cpp
+++ b/Engine/Source/Runtime/CoreUObject/Private/UObject/UObjectClusters.cpp
@@ -14,6 +14,7 @@
 #include "UObject/UObjectArray.h"
 #include "UObject/UObjectBaseUtility.h"
 #include "UObject/GarbageCollection.h"
+#include "UObject/GarbageCollectionInternalFlags.h"
 #include "UObject/UObjectIterator.h"
 #include "UObject/ReferenceChainSearch.h"
 #include "UObject/LinkerLoad.h"
@@ -148,7 +149,7 @@ void FUObjectClusterContainer::DissolveCluster(FUObjectCluster& Cluster)
 	TArray<int32> ReferencedByClusters = MoveTemp(Cluster.ReferencedByClusters);
 
 	// Unreachable clusters will be removed by GC during BeginDestroy phase (unhashing)
-	if (!RootObjectItem->IsMaybeUnreachable())
+	if (!UE::GC::Private::FGCFlags::IsMaybeUnreachable_ForGC(RootObjectItem))
 	{
 #if UE_GCCLUSTER_VERBOSE_LOGGING
 		UObject* ClusterRootObject = static_cast<UObject*>(RootObjectItem->Object);
@@ -193,7 +194,7 @@ void FUObjectClusterContainer::DissolveClusterAndMarkObjectsAsUnreachable(FUObje
 	{
 		FUObjectItem* ClusterObjectItem = GUObjectArray.IndexToObjectUnsafeForGC(ClusterObjectIndex);
 		ClusterObjectItem->SetOwnerIndex(0);
-		ClusterObjectItem->SetMaybeUnreachable();
+		UE::GC::Private::FGCFlags::SetMaybeUnreachable_ForGC(ClusterObjectItem);
 	}
 
 #if !UE_GCCLUSTER_VERBOSE_LOGGING
@@ -209,7 +210,7 @@ void FUObjectClusterContainer::DissolveClusterAndMarkObjectsAsUnreachable(FUObje
 		FUObjectItem* ReferencedByClusterRootItem = GUObjectArray.IndexToObjectUnsafeForGC(ReferencedByClusterRootIndex);
 		if (ReferencedByClusterRootItem->HasAnyFlags(EInternalObjectFlags::ClusterRoot))
 		{
-			ReferencedByClusterRootItem->SetMaybeUnreachable();
+			UE::GC::Private::FGCFlags::SetMaybeUnreachable_ForGC(ReferencedByClusterRootItem);
 			DissolveClusterAndMarkObjectsAsUnreachable(ReferencedByClusterRootItem);
 		}
 	}
diff --git a/Engine/Source/Runtime/CoreUObject/Private/UObject/UObjectGlobals.cpp b/Engine/Source/Runtime/CoreUObject/Private/UObject/UObjectGlobals.cpp
--- a/Engine/Source/Runtime/CoreUObject/Private/UObject/UObjectGlobals.cpp
+++ b/Engine/Source/Runtime/CoreUObject/Private/UObject/UObjectGlobals.cpp
@@ -29,6 +29,7 @@
 #include "UObject/UObjectHashPrivate.h"
 #include "UObject/Object.h"
 #include "UObject/GarbageCollection.h"
+#include "UObject/GarbageCollectionInternalFlags.h"
 #include "UObject/Class.h"
 #include "UObject/CoreRedirects.h"
 #include "UObject/FastReferenceCollector.h"
@@ -5308,7 +5309,8 @@ public:
 				}
 				else
 				{
-					Object->SetInternalFlags(UE::GC::GUnreachableObjectFlag);
+					FUObjectItem* ObjectItem = GUObjectArray.ObjectToObjectItem(Object);
+					UE::GC::Private::FGCFlags::ThisThreadAtomicallySetFlag_ForGC(ObjectItem, EInternalObjectFlags::Unreachable);
 				}
 			}
 		}
@@ -5361,7 +5363,8 @@ private:
 		}
 
 		// Mark it as reachable.
-		Object->ThisThreadAtomicallyClearedRFUnreachable();
+		FUObjectItem* ObjectItem = GUObjectArray.ObjectToObjectItem(Object);
+		UE::GC::Private::FGCFlags::ThisThreadAtomicallyClearedFlag_ForGC(ObjectItem, EInternalObjectFlags::Unreachable);
 
 		// Add it to the list of objects to serialize.
 		ObjectsToSerialize.Add( Object );
@@ -5388,7 +5391,8 @@ private:
 					CurrentReferenceInfo->TotalReferences++;
 				}
 				// Mark it as reachable.
-				InObject->ThisThreadAtomicallyClearedRFUnreachable();
+				FUObjectItem* ObjectItem = GUObjectArray.ObjectToObjectItem(InObject);
+				UE::GC::Private::FGCFlags::ThisThreadAtomicallyClearedFlag_ForGC(ObjectItem, EInternalObjectFlags::Unreachable);
 			}
 			else if (InObject->IsUnreachable())
 			{
diff --git a/Engine/Source/Runtime/CoreUObject/Private/UObject/UObjectHash.cpp b/Engine/Source/Runtime/CoreUObject/Private/UObject/UObjectHash.cpp
--- a/Engine/Source/Runtime/CoreUObject/Private/UObject/UObjectHash.cpp
+++ b/Engine/Source/Runtime/CoreUObject/Private/UObject/UObjectHash.cpp
@@ -573,7 +573,7 @@ static int32 GetObjectOuterHash(FName ObjName, PTRINT Outer)
 
 UObject* StaticFindObjectFastExplicitThreadSafe(FUObjectHashTables& ThreadHash, const UClass* ObjectClass, FName ObjectName, const FString& ObjectPathName, bool bExactClass, EObjectFlags ExcludeFlags/*=0*/)
 {
-	const EInternalObjectFlags ExclusiveInternalFlags = UE::GC::GUnreachableObjectFlag;
+	const EInternalObjectFlags ExclusiveInternalFlags = EInternalObjectFlags::Unreachable;
 
 	// Find an object with the specified name and (optional) class, in any package; if bAnyPackage is false, only matches top-level packages
 	int32 Hash = GetObjectHash(ObjectName);
@@ -703,7 +703,7 @@ struct FObjectSearchPath
 
 UObject* StaticFindObjectInPackageInternal(FUObjectHashTables& ThreadHash, const UClass* ObjectClass, const UPackage* ObjectPackage, FName ObjectName, bool bExactClass, EObjectFlags ExcludeFlags, EInternalObjectFlags ExclusiveInternalFlags)
 {
-	ExclusiveInternalFlags |= UE::GC::GUnreachableObjectFlag;
+	ExclusiveInternalFlags |= EInternalObjectFlags::Unreachable;
 	UObject* Result = nullptr;
 	if (FHashBucket* Inners = ThreadHash.PackageToObjectListMap.Find(ObjectPackage))
 	{
@@ -737,7 +737,7 @@ UObject* StaticFindObjectInPackageInternal(FUObjectHashTables& ThreadHash, const
 
 UObject* StaticFindObjectFastInternalThreadSafe(FUObjectHashTables& ThreadHash, const UClass* ObjectClass, const UObject* ObjectPackage, FName ObjectName, bool bExactClass, bool bAnyPackage, EObjectFlags ExcludeFlags, EInternalObjectFlags ExclusiveInternalFlags)
 {
-	ExclusiveInternalFlags |= UE::GC::GUnreachableObjectFlag;
+	ExclusiveInternalFlags |= EInternalObjectFlags::Unreachable;
 
 	// If they specified an outer use that during the hashing
 	UObject* Result = nullptr;
@@ -899,7 +899,7 @@ bool StaticFindAllObjectsFastInternal(TArray<UObject*>& OutFoundObjects, const U
 {
 	INC_DWORD_STAT(STAT_FindObjectFast);
 
-	ExclusiveInternalFlags |= UE::GC::GUnreachableObjectFlag;
+	ExclusiveInternalFlags |= EInternalObjectFlags::Unreachable;
 
 	FObjectSearchPath SearchPath(ObjectName);
 	const int32 Hash = GetObjectHash(SearchPath.Inner);
@@ -943,7 +943,7 @@ UObject* StaticFindFirstObjectFastInternal(const UClass* ObjectClass, FName Obje
 {
 	INC_DWORD_STAT(STAT_FindObjectFast);
 
-	ExclusiveInternalFlags |= UE::GC::GUnreachableObjectFlag;
+	ExclusiveInternalFlags |= EInternalObjectFlags::Unreachable;
 
 	UObject* Result = nullptr;
 	FObjectSearchPath SearchPath(ObjectName);
@@ -1188,7 +1188,7 @@ void GetObjectsWithOuter(const class UObjectBase* Outer, TArray<UObject *>& Resu
 #endif
 
 	// We don't want to return any objects that are currently being background loaded unless we're using the object iterator during async loading.
-	ExclusionInternalFlags |= UE::GC::GUnreachableObjectFlag;
+	ExclusionInternalFlags |= EInternalObjectFlags::Unreachable;
 	if (!IsInAsyncLoadingThread())
 	{
 		ExclusionInternalFlags |= EInternalObjectFlags::AsyncLoading;
@@ -1256,7 +1256,7 @@ void ForEachObjectWithOuterBreakable(const class UObjectBase* Outer, TFunctionRe
 #endif
 
 	// We don't want to return any objects that are currently being background loaded unless we're using the object iterator during async loading.
-	ExclusionInternalFlags |= UE::GC::GUnreachableObjectFlag;
+	ExclusionInternalFlags |= EInternalObjectFlags::Unreachable;
 	if (!IsInAsyncLoadingThread())
 	{
 		ExclusionInternalFlags |= EInternalObjectFlags::AsyncLoading;
@@ -1305,7 +1305,7 @@ UObjectBase* FindObjectWithOuter(const class UObjectBase* Outer, const class UCl
 	UObject* Result = nullptr;
 	check( Outer );
 	// We don't want to return any objects that are currently being background loaded unless we're using the object iterator during async loading.
-	EInternalObjectFlags ExclusionInternalFlags = UE::GC::GUnreachableObjectFlag;
+	EInternalObjectFlags ExclusionInternalFlags = EInternalObjectFlags::Unreachable;
 	if (!IsInAsyncLoadingThread())
 	{
 		ExclusionInternalFlags = EInternalObjectFlags::AsyncLoading;
@@ -1359,7 +1359,7 @@ void ForEachObjectWithPackage(const class UPackage* Package, TFunctionRef<bool(U
 	check(Package != nullptr);
 
 	// We don't want to return any objects that are currently being background loaded unless we're using the object iterator during async loading.
-	ExclusionInternalFlags |= UE::GC::GUnreachableObjectFlag;
+	ExclusionInternalFlags |= EInternalObjectFlags::Unreachable;
 	if (!IsInAsyncLoadingThread())
 	{
 		ExclusionInternalFlags |= EInternalObjectFlags::AsyncLoading;
@@ -1470,7 +1470,7 @@ FORCEINLINE void ForEachObjectOfClasses_Implementation(FUObjectHashTables& Threa
 	TRACE_CPUPROFILER_EVENT_SCOPE(ForEachObjectOfClasses_Implementation);
 
 	// We don't want to return any objects that are currently being background loaded unless we're using the object iterator during async loading.
-	ExclusionInternalFlags |= UE::GC::GUnreachableObjectFlag;
+	ExclusionInternalFlags |= EInternalObjectFlags::Unreachable;
 	if (!IsInAsyncLoadingThread())
 	{
 		ExclusionInternalFlags |= EInternalObjectFlags::AsyncLoading;
diff --git a/Engine/Source/Runtime/CoreUObject/Private/UObject/UObjectMarks.cpp b/Engine/Source/Runtime/CoreUObject/Private/UObject/UObjectMarks.cpp
--- a/Engine/Source/Runtime/CoreUObject/Private/UObject/UObjectMarks.cpp
+++ b/Engine/Source/Runtime/CoreUObject/Private/UObject/UObjectMarks.cpp
@@ -218,7 +218,7 @@ EObjectMark ObjectGetAllMarks(const class UObjectBase* Object)
 void GetObjectsWithAllMarks(TArray<UObject *>& Results, EObjectMark Marks)
 {
 	// We don't want to return any objects that are currently being background loaded unless we're using the object iterator during async loading.
-	EInternalObjectFlags ExclusionFlags = UE::GC::GUnreachableObjectFlag;
+	EInternalObjectFlags ExclusionFlags = EInternalObjectFlags::Unreachable;
 	if (!IsInAsyncLoadingThread())
 	{
 		ExclusionFlags |= EInternalObjectFlags::AsyncLoading;
@@ -241,7 +241,7 @@ void GetObjectsWithAllMarks(TArray<UObject *>& Results, EObjectMark Marks)
 void GetObjectsWithAnyMarks(TArray<UObject *>& Results, EObjectMark Marks)
 {
 	// We don't want to return any objects that are currently being background loaded unless we're using the object iterator during async loading.
-	EInternalObjectFlags ExclusionFlags = UE::GC::GUnreachableObjectFlag;
+	EInternalObjectFlags ExclusionFlags = EInternalObjectFlags::Unreachable;
 	if (!IsInAsyncLoadingThread())
 	{
 		ExclusionFlags |= EInternalObjectFlags::AsyncLoading;
diff --git a/Engine/Source/Runtime/CoreUObject/Public/UObject/GarbageCollectionGlobals.h b/Engine/Source/Runtime/CoreUObject/Public/UObject/GarbageCollectionGlobals.h
--- a/Engine/Source/Runtime/CoreUObject/Public/UObject/GarbageCollectionGlobals.h
+++ b/Engine/Source/Runtime/CoreUObject/Public/UObject/GarbageCollectionGlobals.h
@@ -11,12 +11,15 @@
 namespace UE::GC
 {
 	/** Current EInternalObjectFlags value representing a reachable object */
+	UE_DEPRECATED(5.5, "GUnreachableObjectFlag should no longer be used.")
 	extern COREUOBJECT_API EInternalObjectFlags GReachableObjectFlag;
 
 	/** Current EInternalObjectFlags value representing an unreachable object */
+	UE_DEPRECATED(5.5, "GUnreachableObjectFlag should no longer be used. Use Object->IsUnreachable() instead.")
 	extern COREUOBJECT_API EInternalObjectFlags GUnreachableObjectFlag;
 
 	/** Current EInternalObjectFlags value representing a maybe unreachable object */
+	UE_DEPRECATED(5.5, "GMaybeUnreachableObjectFlag should no longer be used.")
 	extern COREUOBJECT_API EInternalObjectFlags GMaybeUnreachableObjectFlag;
 
 	/** true if incremental reachability analysis is in progress (global for faster access in low level structs and functions otherwise use IsIncrementalReachabilityAnalisysPending()) */
diff --git a/Engine/Source/Runtime/CoreUObject/Public/UObject/ObjectMacros.h b/Engine/Source/Runtime/CoreUObject/Public/UObject/ObjectMacros.h
--- a/Engine/Source/Runtime/CoreUObject/Public/UObject/ObjectMacros.h
+++ b/Engine/Source/Runtime/CoreUObject/Public/UObject/ObjectMacros.h
@@ -615,7 +615,7 @@ enum class EInternalObjectFlags : int32
 	ReachabilityFlag1 = 1 << 1, ///< One of the flags used by Garbage Collector to determine UObject's reachability state
 	ReachabilityFlag2 = 1 << 2, ///< One of the flags used by Garbage Collector to determine UObject's reachability state
 
-	MaybeUnreachable UE_DEPRECATED(5.4, "MaybeUnreachable flag should no longer be used. Use FUObjectItem::IsMaybeUnreachable() or UE::GC::GMaybeUnreachableObjectFlag flag instead.") = 1 << 19, ///< Flag set on all non-root objects at the beginning of Reachability Analysis
+	MaybeUnreachable UE_DEPRECATED(5.4, "MaybeUnreachable flag should no longer be used.") = 1 << 19, ///< Flag set on all non-root objects at the beginning of Reachability Analysis
 	LoaderImport = 1 << 20, ///< Object is ready to be imported by another package during loading
 	Garbage = 1 << 21, ///< Garbage from logical point of view and should not be referenced. This flag is mirrored in EObjectFlags as RF_Garbage for performance
 	ReachableInCluster = 1 << 23, ///< External reference to object in cluster exists
@@ -623,7 +623,7 @@ enum class EInternalObjectFlags : int32
 	Native = 1 << 25, ///< Native (UClass only). 
 	Async = 1 << 26, ///< Object exists only on a different thread than the game thread.
 	AsyncLoading = 1 << 27, ///< Object is being asynchronously loaded.
-	Unreachable UE_DEPRECATED(5.4, "Unreachable flag should no longer be used. Use UObjectBaseUtility::IsUnreachable() or UE::GC::GUnreachableObjectFlag flag instead.") = 1 << 28, ///< Object is not reachable on the object graph.
+	Unreachable = 1 << 28, ///< Object is not reachable on the object graph.
 	RefCounted = 1 << 29, ///< Object currently has ref-counts associated with it.
 	RootSet = 1 << 30, ///< Object will not be garbage collected, even if unreferenced.
 	PendingConstruction = 1 << 31, ///< Object didn't have its class constructor called yet (only the UObjectBase one to initialize its most basic members)
@@ -644,8 +644,9 @@ ENUM_CLASS_FLAGS(EInternalObjectFlags);
 
 //~ Make sure these macros are up to date!
 #define EInternalObjectFlags_GarbageCollectionKeepFlags (EInternalObjectFlags::Native | EInternalObjectFlags::Async | EInternalObjectFlags::AsyncLoading | EInternalObjectFlags::LoaderImport | EInternalObjectFlags::RefCounted)
-#define EInternalObjectFlags_AllFlags (EInternalObjectFlags::ReachabilityFlag0 | EInternalObjectFlags::ReachabilityFlag1 | EInternalObjectFlags::ReachabilityFlag2 | EInternalObjectFlags::LoaderImport | EInternalObjectFlags::Garbage | EInternalObjectFlags::ReachableInCluster | EInternalObjectFlags::ClusterRoot | EInternalObjectFlags::Native | EInternalObjectFlags::RefCounted | EInternalObjectFlags::Async | EInternalObjectFlags::AsyncLoading | EInternalObjectFlags::RootSet | EInternalObjectFlags::PendingConstruction | (EInternalObjectFlags)(1 << 19 /*MaybeUnreachable*/) | (EInternalObjectFlags)(1 << 28 /*Unreachable*/))
+#define EInternalObjectFlags_AllFlags (EInternalObjectFlags::ReachabilityFlag0 | EInternalObjectFlags::ReachabilityFlag1 | EInternalObjectFlags::ReachabilityFlag2 | EInternalObjectFlags::LoaderImport | EInternalObjectFlags::Garbage | EInternalObjectFlags::ReachableInCluster | EInternalObjectFlags::ClusterRoot | EInternalObjectFlags::Native | EInternalObjectFlags::RefCounted | EInternalObjectFlags::Async | EInternalObjectFlags::AsyncLoading | EInternalObjectFlags::RootSet | EInternalObjectFlags::PendingConstruction | (EInternalObjectFlags)(1 << 19 /*MaybeUnreachable*/) | (EInternalObjectFlags::Unreachable))
 #define EInternalObjectFlags_RootFlags (EInternalObjectFlags::RootSet | EInternalObjectFlags_GarbageCollectionKeepFlags)
+#define EInternalObjectFlags_ReachabilityFlags (EInternalObjectFlags::ReachabilityFlag0 | EInternalObjectFlags::ReachabilityFlag1 | EInternalObjectFlags::ReachabilityFlag2 | EInternalObjectFlags::Unreachable)
 
 /**
  * Flags describing a UEnum 
diff --git a/Engine/Source/Runtime/CoreUObject/Public/UObject/UObjectArray.h b/Engine/Source/Runtime/CoreUObject/Public/UObject/UObjectArray.h
--- a/Engine/Source/Runtime/CoreUObject/Public/UObject/UObjectArray.h
+++ b/Engine/Source/Runtime/CoreUObject/Public/UObject/UObjectArray.h
@@ -26,6 +26,11 @@ namespace Verse
 #define UE_GC_TRACK_OBJ_AVAILABLE UE_DEPRECATED_MACRO(5.2, "The UE_GC_TRACK_OBJ_AVAILABLE macro has been deprecated because it is no longer necessary.") 1
 #endif
 
+namespace UE::GC::Private
+{
+	class FGCFlags;
+}
+
 /**
 * Single item in the UObject array.
 */
@@ -37,6 +42,7 @@ struct
 	FUObjectItem
 {
 	friend class FUObjectArray;
+	friend class UE::GC::Private::FGCFlags;
 
 	// Pointer to the allocated object
 	class UObjectBase* Object;
@@ -115,7 +121,6 @@ public:
 
 	FORCEINLINE void SetFlags(EInternalObjectFlags FlagsToSet)
 	{
-		check((int32(FlagsToSet) & ~int32(EInternalObjectFlags_AllFlags)) == 0);
 		ThisThreadAtomicallySetFlag(FlagsToSet);
 	}
 
@@ -126,7 +131,6 @@ public:
 
 	FORCEINLINE void ClearFlags(EInternalObjectFlags FlagsToClear)
 	{
-		check((int32(FlagsToClear) & ~int32(EInternalObjectFlags_AllFlags)) == 0);
 		ThisThreadAtomicallyClearedFlag(FlagsToClear);
 	}
 
@@ -135,25 +139,10 @@ public:
 	 * @param FlagsToClear
 	 * @return True if this call cleared the flag, false if it has been cleared by another thread.
 	 */
+	UE_DEPRECATED(5.5, "ThisThreadAtomicallyClearedFlag_ForGC() can only be used by the garbage collector. Use ThisThreadAtomicallyClearedFlag instead.")
 	FORCEINLINE bool ThisThreadAtomicallyClearedFlag_ForGC(EInternalObjectFlags FlagToClear)
 	{
-		static_assert(sizeof(int32) == sizeof(Flags), "Flags must be 32-bit for atomics.");
-		bool bIChangedIt = false;
-		while (1)
-		{
-			int32 StartValue = GetFlagsInternal();
-			if (!(StartValue & int32(FlagToClear)))
-			{
-				break;
-			}
-			int32 NewValue = StartValue & ~int32(FlagToClear);
-			if ((int32)FPlatformAtomics::InterlockedCompareExchange((int32*)&Flags, NewValue, StartValue) == StartValue)
-			{
-				bIChangedIt = true;
-				break;
-			}
-		}
-		return bIChangedIt;
+		return AtomicallyClearFlag_ForGC(FlagToClear);
 	}
 
 	/**
@@ -163,10 +152,11 @@ public:
 	 */
 	FORCEINLINE bool ThisThreadAtomicallyClearedFlag(EInternalObjectFlags FlagToClear)
 	{
+		checkf((int32(FlagToClear) & ~int32(EInternalObjectFlags_AllFlags)) == 0, TEXT("%d is not a valid internal flag value"), int32(FlagToClear));
 		bool Result = false;
 		UE_AUTORTFM_OPEN(
 		{
-			FlagToClear &= ~UE::GC::GReachableObjectFlag;     // reachability bit can only be cleared by GC through *_ForGC functions
+			FlagToClear &= ~EInternalObjectFlags_ReachabilityFlags; // reachability flags can only be cleared by GC through *_ForGC functions
 			FlagToClear &= ~EInternalObjectFlags::RefCounted; // refcounted flag is internal and must only be cleared internally by AddRef/ReleaseRef.
 			if (!!(FlagToClear & EInternalObjectFlags_RootFlags))
 			{
@@ -174,7 +164,7 @@ public:
 			}
 			else
 			{
-				Result = ThisThreadAtomicallyClearedFlag_ForGC(FlagToClear);
+				Result = AtomicallyClearFlag_ForGC(FlagToClear);
 			}
 		});
 
@@ -186,25 +176,10 @@ public:
 	 * @param FlagToSet
 	 * @return True if this call set the flag, false if it has been set by another thread.
 	 */
+	UE_DEPRECATED(5.5, "ThisThreadAtomicallySetFlag_ForGC() can only be used by the garbage collector. Use ThisThreadAtomicallySetFlag instead.")
 	FORCEINLINE bool ThisThreadAtomicallySetFlag_ForGC(EInternalObjectFlags FlagToSet)
 	{
-		static_assert(sizeof(int32) == sizeof(Flags), "Flags must be 32-bit for atomics.");
-		bool bIChangedIt = false;
-		while (1)
-		{
-			int32 StartValue = GetFlagsInternal();
-			if ((StartValue & int32(FlagToSet)) == int32(FlagToSet))
-			{
-				break;
-			}
-			int32 NewValue = StartValue | int32(FlagToSet);
-			if ((int32)FPlatformAtomics::InterlockedCompareExchange((int32*)&Flags, NewValue, StartValue) == StartValue)
-			{
-				bIChangedIt = true;
-				break;
-			}
-		}
-		return bIChangedIt;
+		return AtomicallySetFlag_ForGC(FlagToSet);
 	}
 
 	/**
@@ -214,9 +189,11 @@ public:
 	 */
 	FORCEINLINE bool ThisThreadAtomicallySetFlag(EInternalObjectFlags FlagToSet)
 	{
+		checkf((int32(FlagToSet) & ~int32(EInternalObjectFlags_AllFlags)) == 0, TEXT("%d is not a valid internal flag value"), int32(FlagToSet));
 		bool Result = false;
 		UE_AUTORTFM_OPEN(
 		{
+			FlagToSet &= ~EInternalObjectFlags_ReachabilityFlags; // reachability flags can only be cleared by GC through *_ForGC functions
 			FlagToSet &= ~EInternalObjectFlags::RefCounted; // refcounted flag is internal and must only be set by AddRef/ReleaseRef.
 			if (!!(FlagToSet & EInternalObjectFlags_RootFlags))
 			{
@@ -224,7 +201,7 @@ public:
 			}
 			else
 			{
-				Result = ThisThreadAtomicallySetFlag_ForGC(FlagToSet);
+				Result = AtomicallySetFlag_ForGC(FlagToSet);
 			}
 		});
 		return Result;
@@ -240,44 +217,39 @@ public:
 		return (GetFlagsInternal() & int32(InFlags)) == int32(InFlags);
 	}
 
+	UE_DEPRECATED(5.5, "SetUnreachable() can only be used by the garbage collector.")
 	FORCEINLINE void SetUnreachable()
 	{
-		ThisThreadAtomicallyClearedFlag_ForGC(UE::GC::GReachableObjectFlag);
-		ThisThreadAtomicallySetFlag_ForGC(UE::GC::GUnreachableObjectFlag);
-	}
-	FORCEINLINE void SetMaybeUnreachable()
-	{
-		ThisThreadAtomicallyClearedFlag_ForGC(UE::GC::GReachableObjectFlag);
-		ThisThreadAtomicallySetFlag_ForGC(UE::GC::GMaybeUnreachableObjectFlag);
+		AtomicallySetFlag_ForGC(EInternalObjectFlags::Unreachable);
 	}
+	UE_DEPRECATED(5.5, "SetMaybeUnreachable() can only be used by the garbage collector.")
+	COREUOBJECT_API void SetMaybeUnreachable();
+
+	UE_DEPRECATED(5.5, "ClearUnreachable() can only be used by the garbage collector.")
 	FORCEINLINE void ClearUnreachable()
 	{
-		ThisThreadAtomicallyClearedRFUnreachable();
+		AtomicallyClearFlag_ForGC(EInternalObjectFlags::Unreachable);
 	}
 	FORCEINLINE bool IsUnreachable() const
 	{
-		return !!(GetFlagsInternal() & int32(UE::GC::GUnreachableObjectFlag));
-	}
-	FORCEINLINE bool IsMaybeUnreachable() const
-	{
-		return !!(GetFlagsInternal() & int32(UE::GC::GMaybeUnreachableObjectFlag));
+		return !!(GetFlagsInternal() & int32(EInternalObjectFlags::Unreachable));
 	}
+
+	UE_DEPRECATED(5.5, "IsMaybeUnreachable() can only be used by the garbage collector.")
+	COREUOBJECT_API bool IsMaybeUnreachable() const;
+
+	UE_DEPRECATED(5.5, "ThisThreadAtomicallyClearedRFUnreachable() can only be used by the garbage collector.")
 	FORCEINLINE bool ThisThreadAtomicallyClearedRFUnreachable()
 	{
-		if (ThisThreadAtomicallyClearedFlag_ForGC(UE::GC::GUnreachableObjectFlag))
-		{
-			ThisThreadAtomicallySetFlag_ForGC(UE::GC::GReachableObjectFlag);
-			return true;
-		}
-		return false;
+		return AtomicallyClearFlag_ForGC(EInternalObjectFlags::Unreachable);
 	}
 	FORCEINLINE void SetGarbage()
 	{
-		ThisThreadAtomicallySetFlag_ForGC(EInternalObjectFlags::Garbage);
+		AtomicallySetFlag_ForGC(EInternalObjectFlags::Garbage);
 	}
 	FORCEINLINE void ClearGarbage()
 	{
-		ThisThreadAtomicallyClearedFlag_ForGC(EInternalObjectFlags::Garbage);
+		AtomicallyClearFlag_ForGC(EInternalObjectFlags::Garbage);
 	}
 	FORCEINLINE bool IsGarbage() const
 	{
@@ -350,38 +322,21 @@ public:
 #endif
 
 	// Mark this object item as Reachable and clear MaybeUnreachable flag. For GC use only.
-	FORCEINLINE void FastMarkAsReachableInterlocked_ForGC()
-	{
-		using namespace UE::GC;
-		FPlatformAtomics::InterlockedAnd(&Flags, ~int32(GMaybeUnreachableObjectFlag));
-		FPlatformAtomics::InterlockedOr(&Flags, int32(GReachableObjectFlag));
-	}
+	UE_DEPRECATED(5.5, "FastMarkAsReachableInterlocked_ForGC() can only be used by the garbage collector.")
+	COREUOBJECT_API void FastMarkAsReachableInterlocked_ForGC();
 
 	// Mark this object item as Reachable and clear ReachableInCluster and MaybeUnreachable flags. For GC use only.
-	FORCEINLINE void FastMarkAsReachableAndClearReachaleInClusterInterlocked_ForGC()
-	{
-		using namespace UE::GC;
-		FPlatformAtomics::InterlockedAnd(&Flags, ~int32(GMaybeUnreachableObjectFlag | EInternalObjectFlags::ReachableInCluster));
-		FPlatformAtomics::InterlockedOr(&Flags, int32(GReachableObjectFlag));
-	}
+	UE_DEPRECATED(5.5, "FastMarkAsReachableAndClearReachaleInClusterInterlocked_ForGC() can only be used by the garbage collector.")
+	COREUOBJECT_API void FastMarkAsReachableAndClearReachaleInClusterInterlocked_ForGC();
 
 	/**
 	 * Mark this object item as Reachable and clear MaybeUnreachable flag. Only thread-safe for concurrent clear, not concurrent set+clear. Don't use during mark phase. For GC use only.
 	 * @return True if this call cleared MaybeUnreachable flag, false if it has been cleared by another thread.
 	 */
-	FORCEINLINE bool MarkAsReachableInterlocked_ForGC()
-	{
-		using namespace UE::GC;
-		const int32 FlagToClear = int32(UE::GC::GMaybeUnreachableObjectFlag);
-		if (FPlatformAtomics::AtomicRead_Relaxed(&Flags) & FlagToClear)
-		{
-			int32 Old = FPlatformAtomics::InterlockedAnd(&Flags, ~FlagToClear);
-			FPlatformAtomics::InterlockedOr(&Flags, int32(GReachableObjectFlag));
-			return Old & FlagToClear;
-		}
-		return false;
-	}
+	UE_DEPRECATED(5.5, "MarkAsReachableInterlocked_ForGC() can only be used by the garbage collector.")
+	COREUOBJECT_API bool MarkAsReachableInterlocked_ForGC();
 
+	UE_DEPRECATED(5.5, "OffsetOfFlags() can only be used by the garbage collector.")
 	FORCEINLINE static constexpr ::size_t OffsetOfFlags()
 	{
 		return offsetof(FUObjectItem, Flags);
@@ -400,6 +355,58 @@ private:
 
 	COREUOBJECT_API bool SetRootFlags(EInternalObjectFlags FlagsToSet);
 	COREUOBJECT_API bool ClearRootFlags(EInternalObjectFlags FlagsToClear);
+
+	/**
+	 * Uses atomics to set the specified flag(s). GC internal version.
+	 * @param FlagToSet
+	 * @return True if this call set the flag, false if it has been set by another thread.
+	 */
+	FORCEINLINE bool AtomicallySetFlag_ForGC(EInternalObjectFlags FlagToSet)
+	{
+		static_assert(sizeof(int32) == sizeof(Flags), "Flags must be 32-bit for atomics.");
+		bool bIChangedIt = false;
+		while (1)
+		{
+			int32 StartValue = GetFlagsInternal();
+			if ((StartValue & int32(FlagToSet)) == int32(FlagToSet))
+			{
+				break;
+			}
+			int32 NewValue = StartValue | int32(FlagToSet);
+			if ((int32)FPlatformAtomics::InterlockedCompareExchange((int32*)&Flags, NewValue, StartValue) == StartValue)
+			{
+				bIChangedIt = true;
+				break;
+			}
+		}
+		return bIChangedIt;
+	}
+
+	/**
+	 * Uses atomics to clear the specified flag(s). GC internal version
+	 * @param FlagsToClear
+	 * @return True if this call cleared the flag, false if it has been cleared by another thread.
+	 */
+	FORCEINLINE bool AtomicallyClearFlag_ForGC(EInternalObjectFlags FlagToClear)
+	{
+		static_assert(sizeof(int32) == sizeof(Flags), "Flags must be 32-bit for atomics.");
+		bool bIChangedIt = false;
+		while (1)
+		{
+			int32 StartValue = GetFlagsInternal();
+			if (!(StartValue & int32(FlagToClear)))
+			{
+				break;
+			}
+			int32 NewValue = StartValue & ~int32(FlagToClear);
+			if ((int32)FPlatformAtomics::InterlockedCompareExchange((int32*)&Flags, NewValue, StartValue) == StartValue)
+			{
+				bIChangedIt = true;
+				break;
+			}
+		}
+		return bIChangedIt;
+	}
 };
 
 namespace UE::UObjectArrayPrivate
@@ -963,7 +970,7 @@ public:
 	{
 		if (ObjectItem)
 		{
-			return bEvenIfGarbage ? !ObjectItem->IsUnreachable() : !(ObjectItem->HasAnyFlags(UE::GC::GUnreachableObjectFlag | EInternalObjectFlags::Garbage));
+			return bEvenIfGarbage ? !ObjectItem->IsUnreachable() : !(ObjectItem->HasAnyFlags(EInternalObjectFlags::Unreachable | EInternalObjectFlags::Garbage));
 		}
 		return false;
 	}
@@ -984,7 +991,7 @@ public:
 	FORCEINLINE bool IsStale(FUObjectItem* ObjectItem, bool bIncludingGarbage)
 	{
 		// This method assumes ObjectItem is valid.
-		return bIncludingGarbage ? (ObjectItem->HasAnyFlags(UE::GC::GUnreachableObjectFlag | EInternalObjectFlags::Garbage)) : (ObjectItem->IsUnreachable());
+		return bIncludingGarbage ? (ObjectItem->HasAnyFlags(EInternalObjectFlags::Unreachable | EInternalObjectFlags::Garbage)) : (ObjectItem->IsUnreachable());
 	}
 
 	FORCEINLINE bool IsStale(int32 Index, bool bIncludingGarbage)
diff --git a/Engine/Source/Runtime/CoreUObject/Public/UObject/UObjectBaseUtility.h b/Engine/Source/Runtime/CoreUObject/Public/UObject/UObjectBaseUtility.h
--- a/Engine/Source/Runtime/CoreUObject/Public/UObject/UObjectBaseUtility.h
+++ b/Engine/Source/Runtime/CoreUObject/Public/UObject/UObjectBaseUtility.h
@@ -220,10 +220,8 @@ public:
 	 *
 	 * @return true if we are the thread that cleared RF_Unreachable
 	 */
-	FORCEINLINE bool ThisThreadAtomicallyClearedRFUnreachable()
-	{
-		return GUObjectArray.IndexToObject(InternalIndex)->ThisThreadAtomicallyClearedRFUnreachable();
-	}
+	UE_DEPRECATED(5.5, "ThisThreadAtomicallyClearedRFUnreachable() should no longer be used. Reachability flags can only be modified by the garbage collector.")
+	bool ThisThreadAtomicallyClearedRFUnreachable();
 
 	/** Checks if the object is unreachable. */
 	FORCEINLINE bool IsUnreachable() const
diff --git a/Engine/Source/Runtime/CoreUObject/Public/UObject/UObjectIterator.h b/Engine/Source/Runtime/CoreUObject/Public/UObject/UObjectIterator.h
--- a/Engine/Source/Runtime/CoreUObject/Public/UObject/UObjectIterator.h
+++ b/Engine/Source/Runtime/CoreUObject/Public/UObject/UObjectIterator.h
@@ -57,7 +57,7 @@ enum class EObjectIteratorThreadSafetyOptions : uint8
 
 inline EInternalObjectFlags GetObjectIteratorDefaultInternalExclusionFlags(EInternalObjectFlags InternalExclusionFlags)
 {
-	InternalExclusionFlags |= UE::GC::GUnreachableObjectFlag | EInternalObjectFlags::PendingConstruction;
+	InternalExclusionFlags |= EInternalObjectFlags::Unreachable | EInternalObjectFlags::PendingConstruction;
 	if (!IsInAsyncLoadingThread())
 	{
 		// We don't want to return any objects that are currently being background loaded unless we're using the object iterator during async loading.
diff --git a/Engine/Source/Runtime/Engine/Private/Components/SceneComponent.cpp b/Engine/Source/Runtime/Engine/Private/Components/SceneComponent.cpp
--- a/Engine/Source/Runtime/Engine/Private/Components/SceneComponent.cpp
+++ b/Engine/Source/Runtime/Engine/Private/Components/SceneComponent.cpp
@@ -1207,7 +1207,7 @@ void USceneComponent::OnComponentDestroyed(bool bDestroyingHierarchy)
 		AActor* MyOwner = GetOwner();
 
 		// Do not involve objects which will be destroyed in hierarchy fixups
-		const EInternalObjectFlags SkipFlags = EInternalObjectFlags::Garbage | UE::GC::GUnreachableObjectFlag;
+		const EInternalObjectFlags SkipFlags = EInternalObjectFlags::Garbage | EInternalObjectFlags::Unreachable;
 
 		if (bDestroyingHierarchy)
 		{
diff --git a/Engine/Source/Runtime/Engine/Private/ObjectCacheContext.cpp b/Engine/Source/Runtime/Engine/Private/ObjectCacheContext.cpp
--- a/Engine/Source/Runtime/Engine/Private/ObjectCacheContext.cpp
+++ b/Engine/Source/Runtime/Engine/Private/ObjectCacheContext.cpp
@@ -28,7 +28,7 @@ namespace ObjectCacheContextImpl {
 EInternalObjectFlags GetObjectCacheInternalFlagsExclusion()
 {
 	// We never want to return objects that are invalid or still being worked on by other threads
-	return UE::GC::GUnreachableObjectFlag | EInternalObjectFlags::Garbage | EInternalObjectFlags::AsyncLoading | EInternalObjectFlags::Async;
+	return EInternalObjectFlags::Unreachable | EInternalObjectFlags::Garbage | EInternalObjectFlags::AsyncLoading | EInternalObjectFlags::Async;
 }
 
 } // namespace ObjectCacheContextImpl
diff --git a/Engine/Source/Runtime/Engine/Public/ExternalPackageHelper.h b/Engine/Source/Runtime/Engine/Public/ExternalPackageHelper.h
--- a/Engine/Source/Runtime/Engine/Public/ExternalPackageHelper.h
+++ b/Engine/Source/Runtime/Engine/Public/ExternalPackageHelper.h
@@ -241,7 +241,7 @@ void FExternalPackageHelper::LoadObjectsFromExternalPackages(UObject* InOuter, T
 					return false;
 				}
 				return true;
-			}, true, RF_NoFlags, UE::GC::GUnreachableObjectFlag);
+			}, true, RF_NoFlags, EInternalObjectFlags::Unreachable);
 
 			if (ensure(LoadedObject))
 			{
