Removing deprecated EInternalObjectFlags

diff --git a/Engine/Source/Programs/Shared/EpicGames.Core/UnrealEngineTypes.cs b/Engine/Source/Programs/Shared/EpicGames.Core/UnrealEngineTypes.cs
--- a/Engine/Source/Programs/Shared/EpicGames.Core/UnrealEngineTypes.cs
+++ b/Engine/Source/Programs/Shared/EpicGames.Core/UnrealEngineTypes.cs
@@ -668,9 +668,9 @@ namespace EpicGames.Core
 		ReachabilityFlag1 = 1 << 1,
 
 		/// <summary>
-		/// Flag set on all non-root objects at the beginning of Reachability Analysis
+		/// One of the flags used by Garbage Collector to determine UObject's reachability state
 		/// </summary>
-		MaybeUnreachable = 1 << 19,
+		ReachabilityFlag2 = 1 << 2,
 
 		/// <summary>
 		/// Object is ready to be imported by another package during loading
@@ -730,7 +730,7 @@ namespace EpicGames.Core
 		GarbageCollectionKeepFlags = Native | Async | AsyncLoading | LoaderImport | RefCounted,
 
 		//~ Make sure this is up to date!
-		AllFlags = ReachabilityFlag0 | ReachabilityFlag1 | MaybeUnreachable | LoaderImport | Garbage | ReachableInCluster | ClusterRoot | Native | Async | AsyncLoading | Unreachable | RootSet | PendingConstruction | RefCounted
+		AllFlags = ReachabilityFlag0 | ReachabilityFlag1 | ReachabilityFlag2 | LoaderImport | Garbage | ReachableInCluster | ClusterRoot | Native | Async | AsyncLoading | Unreachable | RootSet | PendingConstruction | RefCounted
 	};
 
 	/// <summary>
diff --git a/Engine/Source/Runtime/CoreUObject/Public/UObject/ObjectMacros.h b/Engine/Source/Runtime/CoreUObject/Public/UObject/ObjectMacros.h
--- a/Engine/Source/Runtime/CoreUObject/Public/UObject/ObjectMacros.h
+++ b/Engine/Source/Runtime/CoreUObject/Public/UObject/ObjectMacros.h
@@ -615,7 +615,6 @@ enum class EInternalObjectFlags : int32
 	ReachabilityFlag1 = 1 << 1, ///< One of the flags used by Garbage Collector to determine UObject's reachability state
 	ReachabilityFlag2 = 1 << 2, ///< One of the flags used by Garbage Collector to determine UObject's reachability state
 
-	MaybeUnreachable UE_DEPRECATED(5.4, "MaybeUnreachable flag should no longer be used.") = 1 << 19, ///< Flag set on all non-root objects at the beginning of Reachability Analysis
 	LoaderImport = 1 << 20, ///< Object is ready to be imported by another package during loading
 	Garbage = 1 << 21, ///< Garbage from logical point of view and should not be referenced. This flag is mirrored in EObjectFlags as RF_Garbage for performance
 	ReachableInCluster = 1 << 23, ///< External reference to object in cluster exists
@@ -628,23 +627,15 @@ enum class EInternalObjectFlags : int32
 	RootSet = 1 << 30, ///< Object will not be garbage collected, even if unreferenced.
 	PendingConstruction = 1 << 31, ///< Object didn't have its class constructor called yet (only the UObjectBase one to initialize its most basic members)
 
-	// !!!!!!!!!!!!!!!!!!!!!!!!!!!!!!!!!!!!!!!!!!!!!!!!!!!!!!!!!!!!!!!!!!!!!!!!!!!!!!!!!!!!!!!!!!!!!!!!!!!!!!!!!!!!!!!!!!!!!!!!
-	// DO NOT ADD new compound flags to EInternalObjectFlags. The below flags are deprecated so that one day we can remove them.
-	// !!!!!!!!!!!!!!!!!!!!!!!!!!!!!!!!!!!!!!!!!!!!!!!!!!!!!!!!!!!!!!!!!!!!!!!!!!!!!!!!!!!!!!!!!!!!!!!!!!!!!!!!!!!!!!!!!!!!!!!!
-
-	GarbageCollectionKeepFlags UE_DEPRECATED(5.4, "GarbageCollectionKeepFlags should no longer be used. Use EInternalObjectFlags_GarbageCollectionKeepFlags instead.") = Native | Async | AsyncLoading | LoaderImport | RefCounted,
-	MirroredFlags UE_DEPRECATED(5.4, "MirroredFlags should no longer be used. Use Garbage instead.") = Garbage,
-
-	//~ Make sure this is up to date!
-	PRAGMA_DISABLE_DEPRECATION_WARNINGS
-	AllFlags UE_DEPRECATED(5.4, "AllFlags should no longer be used. Use EInternalObjectFlags_AllFlags instead.") = ReachabilityFlag0 | ReachabilityFlag1 | MaybeUnreachable | LoaderImport | Garbage | ReachableInCluster | ClusterRoot | Native | Async | AsyncLoading | Unreachable | RootSet | PendingConstruction | RefCounted
-	PRAGMA_ENABLE_DEPRECATION_WARNINGS
+	// !!!!!!!!!!!!!!!!!!!!!!!!!!!!!!!!!!!!!!!!!!!!!!!!!!!!!!!!!!!!!!!!!!!!!!!!!!!!!!!!!!!!!!!!!!!!!!!!!!!!!!!!!!!!!!!!!!!!!!!!!!!!!!!!!!!
+	// DO NOT ADD composite flags to EInternalObjectFlags. Debugger visualisations have trouble displaying composite flag values as text.
+	// !!!!!!!!!!!!!!!!!!!!!!!!!!!!!!!!!!!!!!!!!!!!!!!!!!!!!!!!!!!!!!!!!!!!!!!!!!!!!!!!!!!!!!!!!!!!!!!!!!!!!!!!!!!!!!!!!!!!!!!!!!!!!!!!!!!
 };
 ENUM_CLASS_FLAGS(EInternalObjectFlags);
 
 //~ Make sure these macros are up to date!
 #define EInternalObjectFlags_GarbageCollectionKeepFlags (EInternalObjectFlags::Native | EInternalObjectFlags::Async | EInternalObjectFlags::AsyncLoading | EInternalObjectFlags::LoaderImport | EInternalObjectFlags::RefCounted)
-#define EInternalObjectFlags_AllFlags (EInternalObjectFlags::ReachabilityFlag0 | EInternalObjectFlags::ReachabilityFlag1 | EInternalObjectFlags::ReachabilityFlag2 | EInternalObjectFlags::LoaderImport | EInternalObjectFlags::Garbage | EInternalObjectFlags::ReachableInCluster | EInternalObjectFlags::ClusterRoot | EInternalObjectFlags::Native | EInternalObjectFlags::RefCounted | EInternalObjectFlags::Async | EInternalObjectFlags::AsyncLoading | EInternalObjectFlags::RootSet | EInternalObjectFlags::PendingConstruction | (EInternalObjectFlags)(1 << 19 /*MaybeUnreachable*/) | (EInternalObjectFlags::Unreachable))
+#define EInternalObjectFlags_AllFlags (EInternalObjectFlags::ReachabilityFlag0 | EInternalObjectFlags::ReachabilityFlag1 | EInternalObjectFlags::ReachabilityFlag2 | EInternalObjectFlags::LoaderImport | EInternalObjectFlags::Garbage | EInternalObjectFlags::ReachableInCluster | EInternalObjectFlags::ClusterRoot | EInternalObjectFlags::Native | EInternalObjectFlags::RefCounted | EInternalObjectFlags::Async | EInternalObjectFlags::AsyncLoading | EInternalObjectFlags::RootSet | EInternalObjectFlags::PendingConstruction | (EInternalObjectFlags::Unreachable))
 #define EInternalObjectFlags_RootFlags (EInternalObjectFlags::RootSet | EInternalObjectFlags_GarbageCollectionKeepFlags)
 #define EInternalObjectFlags_ReachabilityFlags (EInternalObjectFlags::ReachabilityFlag0 | EInternalObjectFlags::ReachabilityFlag1 | EInternalObjectFlags::ReachabilityFlag2 | EInternalObjectFlags::Unreachable)
 
