Fixing the final purge not reporting the correct number of objects left after GC.

Before:
LogGarbage: GC purged 36639 objects (36639 -> 36639) in 32.150ms

After:
LogGarbage: GC purged 36639 objects (36639 -> 0) in 32.150ms

diff --git a/Engine/Source/Runtime/CoreUObject/Private/UObject/Obj.cpp b/Engine/Source/Runtime/CoreUObject/Private/UObject/Obj.cpp
--- a/Engine/Source/Runtime/CoreUObject/Private/UObject/Obj.cpp
+++ b/Engine/Source/Runtime/CoreUObject/Private/UObject/Obj.cpp
@@ -5474,9 +5474,13 @@ void StaticExit()
 	IncrementalPurgeGarbage(false);
 
 	{
+		GObjectCountDuringLastMarkPhase.Reset();
+
 		//Repeat GC for every object, including structures and properties.
 		for (FRawObjectIterator It; It; ++It)
 		{
+			GObjectCountDuringLastMarkPhase.Increment();
+
 			// Mark as unreachable so purge phase will kill it.
 			UE::GC::Private::FGCFlags::SetMaybeUnreachable_ForGC(*It);
 		}
