Accelerate CreateLinkerExports by replacing the FPublicExportMap's custom container with a TMAp so we may fetch ObjectIndicies via their ExportHash. This dramatically improves performance for both addition and removal when dealing with many exports. Iteration of exports is moderately improved when exports have been removed.

This fixes:
- The previous container relied on a fixed array of sorted indices. Each insertion would mean memmov'ing chunks indices and values
- Removing elements doesn't actually shrink the table (that's mostly fine) but instead adds a sentinel value. This means anytime a value is looked, for a sentinel value comparison is needed. When fetching all publicexports, we need to walk the whole table despite it not being uncommon in editor for exports to be invalidated and can be removed from iteration entirely.

When loading a UMap with 2M exports, we now take 25s compared to 44m. This usecase comes from a real customer umap.

diff --git a/Engine/Source/Runtime/CoreUObject/Private/Serialization/AsyncLoading2.cpp b/Engine/Source/Runtime/CoreUObject/Private/Serialization/AsyncLoading2.cpp
--- a/Engine/Source/Runtime/CoreUObject/Private/Serialization/AsyncLoading2.cpp
+++ b/Engine/Source/Runtime/CoreUObject/Private/Serialization/AsyncLoading2.cpp
@@ -686,131 +686,29 @@ private:
 	class FPublicExportMap
 	{
 	public:
-		FPublicExportMap() = default;
+		typedef TMap<uint64, int32> FObjectIndexMap;
 
-		FPublicExportMap(const FPublicExportMap&) = delete;
-
-		FPublicExportMap(FPublicExportMap&& Other)
-		{
-			Allocation = Other.Allocation;
-			Count = Other.Count;
-			SingleItemValue = Other.SingleItemValue;
-			Other.Allocation = nullptr;
-			Other.Count = 0;
-		};
-
-		FPublicExportMap& operator=(const FPublicExportMap&) = delete;
-
-		FPublicExportMap& operator=(FPublicExportMap&& Other)
-		{
-			if (Count > 1)
-			{
-				FMemory::Free(Allocation);
-			}
-			Allocation = Other.Allocation;
-			Count = Other.Count;
-			SingleItemValue = Other.SingleItemValue;
-			Other.Allocation = nullptr;
-			Other.Count = 0;
-			return *this;
-		}
-	
-		~FPublicExportMap()
+		void Reserve(int32 NewCount)
 		{
-			if (Count > 1)
-			{
-				FMemory::Free(Allocation);
-			}
-		}
-		void Grow(int32 NewCount)
-		{
-			if (NewCount <= Count)
-			{
-				return;
-			}
-			if (NewCount > 1)
-			{
-				TArrayView<uint64> OldKeys = GetKeys();
-				TArrayView<int32> OldValues = GetValues();
-				const uint64 OldKeysSize = Count * sizeof(uint64);
-				const uint64 NewKeysSize = NewCount * sizeof(uint64);
-				const uint64 OldValuesSize = Count * sizeof(int32);
-				const uint64 NewValuesSize = NewCount * sizeof(int32);
-				const uint64 KeysToAddSize = NewKeysSize - OldKeysSize;
-				const uint64 ValuesToAddSize = NewValuesSize - OldValuesSize;
-
-				uint8* NewAllocation = reinterpret_cast<uint8*>(FMemory::Malloc(NewKeysSize + NewValuesSize));
-				FMemory::Memzero(NewAllocation, KeysToAddSize); // Insert new keys initialized to zero
-				FMemory::Memcpy(NewAllocation + KeysToAddSize, OldKeys.GetData(), OldKeysSize); // Copy old keys
-				FMemory::Memset(NewAllocation + NewKeysSize, 0xFF, ValuesToAddSize); // Insert new values initialized to -1
-				FMemory::Memcpy(NewAllocation + NewKeysSize + ValuesToAddSize, OldValues.GetData(), OldValuesSize); // Copy old values
-				if (Count > 1)
-				{
-					FMemory::Free(Allocation);
-				}
-				Allocation = NewAllocation;
-			}
-			Count = NewCount;
+			Data.Reserve(NewCount);
 		}
 
 		void Store(uint64 ExportHash, UObject* Object)
 		{
-			TArrayView<uint64> Keys = GetKeys();
-			TArrayView<int32> Values = GetValues();
-			int32 Index = Algo::LowerBound(Keys, ExportHash);
-			if (Index < Count && Keys[Index] == ExportHash)
-			{
-				// Slot already exists so reuse it
-				Values[Index] = GUObjectArray.ObjectToIndex(Object);
-				return;
-			}
-			if (Count == 0 || Keys[0] != 0)
-			{
-				// No free slots so we need to add one (will be inserted at the beginning of the array)
-				Grow(Count + 1);
-				Keys = GetKeys();
-				Values = GetValues();
-			}
-			else
-			{
-				--Index; // Update insertion index to one before the lower bound item
-			}
-			if (Index > 0)
-			{
-				// Move items down
-				FMemory::Memmove(Keys.GetData(), Keys.GetData() + 1, Index * sizeof(uint64));
-				FMemory::Memmove(Values.GetData(), Values.GetData() + 1, Index * sizeof(int32));
-			}
-			Keys[Index] = ExportHash;
-			Values[Index] = GUObjectArray.ObjectToIndex(Object);
+			Data.Add(ExportHash, GUObjectArray.ObjectToIndex(Object));
 		}
 
 		bool Remove(uint64 ExportHash)
 		{
-			TArrayView<uint64> Keys = GetKeys();
-			int32 Index = Algo::LowerBound(Keys, ExportHash);
-			if (Index < Count && Keys[Index] == ExportHash)
-			{
-				TArrayView<int32> Values = GetValues();
-				Values[Index] = -1;
-				return true;
-			}
-
-			return false;
+			return !!Data.Remove(ExportHash);
 		}
 
 		UObject* Find(uint64 ExportHash)
 		{
-			TArrayView<uint64> Keys = GetKeys();
-			int32 Index = Algo::LowerBound(Keys, ExportHash);
-			if (Index < Count && Keys[Index] == ExportHash)
+			int32* ObjectIndex = Data.Find(ExportHash);
+			if (ObjectIndex)
 			{
-				TArrayView<int32> Values = GetValues();
-				int32 ObjectIndex = Values[Index];
-				if (ObjectIndex >= 0)
-				{
-					return static_cast<UObject*>(GUObjectArray.IndexToObject(ObjectIndex)->Object);
-				}
+				return static_cast<UObject*>(GUObjectArray.IndexToObject(*ObjectIndex)->Object);
 			}
 			return nullptr;
 		}
@@ -818,73 +716,43 @@ private:
 		[[nodiscard]] bool PinForGC(TArray<int32>& OutUnreachableObjectIndices)
 		{
 			OutUnreachableObjectIndices.Reset();
-			for (int32& ObjectIndex : GetValues())
+			for (FObjectIndexMap::TIterator It(Data); It; ++It)
 			{
-				if (ObjectIndex >= 0)
+				FUObjectItem* ObjectItem = GUObjectArray.IndexToObject(It->Value);
+				if (!ObjectItem->IsUnreachable())
 				{
-					FUObjectItem* ObjectItem = GUObjectArray.IndexToObject(ObjectIndex);
-					if (!ObjectItem->IsUnreachable())
-					{
-						UObject* Object = static_cast<UObject*>(ObjectItem->Object);
-						checkf(!ObjectItem->HasAnyFlags(EInternalObjectFlags::LoaderImport), TEXT("%s"), *Object->GetFullName());
-						ObjectItem->SetFlags(EInternalObjectFlags::LoaderImport);
-					}
-					else
-					{
-						OutUnreachableObjectIndices.Reserve(Count);
-						OutUnreachableObjectIndices.Add(ObjectIndex);
-						ObjectIndex = -1;
-					}
+					UObject* Object = static_cast<UObject*>(ObjectItem->Object);
+					checkf(!ObjectItem->HasAnyFlags(EInternalObjectFlags::LoaderImport), TEXT("%s"), *Object->GetFullName());
+					ObjectItem->SetFlags(EInternalObjectFlags::LoaderImport);
 				}
-			}
-			return OutUnreachableObjectIndices.Num() == 0;
-		}
-
-		void UnpinForGC()
-		{
-			for (int32 ObjectIndex : GetValues())
-			{
-				if (ObjectIndex >= 0)
+				else
 				{
-					UObject* Object = static_cast<UObject*>(GUObjectArray.IndexToObject(ObjectIndex)->Object);
-					checkf(Object->HasAnyInternalFlags(EInternalObjectFlags::LoaderImport), TEXT("%s"), *Object->GetFullName());
-					Object->AtomicallyClearInternalFlags(EInternalObjectFlags::LoaderImport);
+					OutUnreachableObjectIndices.Reserve(Data.Num());
+					OutUnreachableObjectIndices.Add(It->Value);
+					It.RemoveCurrent();
 				}
 			}
+
+			return OutUnreachableObjectIndices.Num() == 0;
 		}
 
-		TArrayView<uint64> GetKeys()
+		void UnpinForGC()
 		{
-			if (Count == 1)
-			{
-				return MakeArrayView(&SingleItemKey, 1);
-			}
-			else
+			for (FObjectIndexMap::TConstIterator It(Data); It; ++It)
 			{
-				return MakeArrayView(reinterpret_cast<uint64*>(Allocation), Count);
+				UObject* Object = static_cast<UObject*>(GUObjectArray.IndexToObject(It->Value)->Object);
+				checkf(Object->HasAnyInternalFlags(EInternalObjectFlags::LoaderImport), TEXT("%s"), *Object->GetFullName());
+				Object->AtomicallyClearInternalFlags(EInternalObjectFlags::LoaderImport);
 			}
 		}
 
-		TArrayView<int32> GetValues()
+		FObjectIndexMap::TConstIterator ConstIterator() const
 		{
-			if (Count == 1)
-			{
-				return MakeArrayView(&SingleItemValue, 1);
-			}
-			else
-			{
-				return MakeArrayView(reinterpret_cast<int32*>(Allocation + Count * sizeof(uint64)), Count);
-			}
+			return FObjectIndexMap::TConstIterator(Data);
 		}
 
 	private:
-		union
-		{
-			uint8* Allocation = nullptr;
-			uint64 SingleItemKey;
-		};
-		int32 Count = 0;
-		int32 SingleItemValue = -1;
+		FObjectIndexMap Data;
 	};
 
 	FPublicExportMap PublicExportMap;
@@ -1001,14 +869,14 @@ public:
 		bHasFailed = true;
 	}
 
-	TArrayView<int32> GetPublicExportObjectIndices()
+	FPublicExportMap::FObjectIndexMap::TConstIterator GetPublicExportObjectIndices() const
 	{
-		return PublicExportMap.GetValues();
+		return PublicExportMap.ConstIterator();
 	}
 
 	void ReserveSpaceForPublicExports(int32 PublicExportCount)
 	{
-		PublicExportMap.Grow(PublicExportCount);
+		PublicExportMap.Reserve(PublicExportCount);
 	}
 
 	void StorePublicExport(uint64 ExportHash, UObject* Object)
@@ -1297,13 +1165,12 @@ public:
 		{
 			VerifyPackageForRemoval(PackageRef);
 		}
-		for (int32 ObjectIndex : PackageRef.GetPublicExportObjectIndices())
+
+		for (auto It = PackageRef.PublicExportMap.ConstIterator(); It; ++It)
 		{
-			if (ObjectIndex >= 0)
-			{
-				ObjectIndexToPublicExport.Remove(ObjectIndex);
-			}
+			ObjectIndexToPublicExport.Remove(It->Value);
 		}
+
 		PackageRef.RemoveUnreferencedObsoletePackage();
 		// Reset PackageId to prevent a double remove from GC NotifyUnreachableObjects
 		OldPackage->SetPackageId(FPackageId());
@@ -1342,12 +1209,9 @@ public:
 		bool bRemoved = Packages.RemoveAndCopyValue(PackageId, PackageRef);
 		if (bRemoved)
 		{
-			for (int32 ObjectIndex : PackageRef.GetPublicExportObjectIndices())
+			for (auto It = PackageRef.PublicExportMap.ConstIterator(); It; ++It)
 			{
-				if (ObjectIndex >= 0)
-				{
-					ObjectIndexToPublicExport.Remove(ObjectIndex);
-				}
+				ObjectIndexToPublicExport.Remove(It->Value);
 			}
 		}
 	}
@@ -1437,37 +1301,35 @@ public:
 			PackageId.ValueForDebugging(),
 			PackageRef.GetRefCount());
 
-		for (int32 ObjectIndex : PackageRef.GetPublicExportObjectIndices())
+		for (auto It = PackageRef.GetPublicExportObjectIndices(); It; ++It)
 		{
-			if (ObjectIndex >= 0)
-			{
-				UObject* Object = static_cast<UObject*>(GUObjectArray.IndexToObject(ObjectIndex)->Object);
-				ensureMsgf(!Object->HasAnyInternalFlags(EInternalObjectFlags::LoaderImport) || GUObjectArray.IsDisregardForGC(Object),
-						TEXT("FGlobalImportStore::VerifyPackageForRemoval: The loaded public export object '%s' with flags (ObjectFlags=%x, InternalObjectFlags=%x) and id %s is probably still referenced by the loader."),
-						*Object->GetFullName(),
-						Object->GetFlags(),
-						Object->GetInternalFlags(),
-						*FormatPackageId(PackageId));
-
-				FPublicExportKey* PublicExportKey = ObjectIndexToPublicExport.Find(ObjectIndex);
-				UE_CLOG(!PublicExportKey, LogStreaming, Fatal,
-					TEXT("FGlobalImportStore::VerifyPackageForRemoval: %s (%s) - ")
-					TEXT("The loaded public export object '%s' is missing in GlobalImportStore."),
-					*Package->GetName(),
-					*FormatPackageId(PackageId),
-					*Object->GetFullName());
-
-				FPackageId ObjectPackageId = PublicExportKey->GetPackageId();
-				UE_CLOG(ObjectPackageId != PackageId, LogStreaming, Fatal,
-					TEXT("FGlobalImportStore::VerifyPackageForRemoval: %s (%s) - ")
-					TEXT("The loaded public export object '%s' has a mismatching package id %s in GlobalImportStore."),
-					*Package->GetName(),
-					*FormatPackageId(PackageId),
+			const int32 ObjectIndex = It->Value;
+			UObject* Object = static_cast<UObject*>(GUObjectArray.IndexToObject(ObjectIndex)->Object);
+			ensureMsgf(!Object->HasAnyInternalFlags(EInternalObjectFlags::LoaderImport) || GUObjectArray.IsDisregardForGC(Object),
+					TEXT("FGlobalImportStore::VerifyPackageForRemoval: The loaded public export object '%s' with flags (ObjectFlags=%x, InternalObjectFlags=%x) and id %s is probably still referenced by the loader."),
 					*Object->GetFullName(),
-					*FormatPackageId(ObjectPackageId));
+					Object->GetFlags(),
+					Object->GetInternalFlags(),
+					*FormatPackageId(PackageId));
+
+			FPublicExportKey* PublicExportKey = ObjectIndexToPublicExport.Find(ObjectIndex);
+			UE_CLOG(!PublicExportKey, LogStreaming, Fatal,
+				TEXT("FGlobalImportStore::VerifyPackageForRemoval: %s (%s) - ")
+				TEXT("The loaded public export object '%s' is missing in GlobalImportStore."),
+				*Package->GetName(),
+				*FormatPackageId(PackageId),
+				*Object->GetFullName());
 
-				VerifyObjectForRemoval(Object);
-			}
+			FPackageId ObjectPackageId = PublicExportKey->GetPackageId();
+			UE_CLOG(ObjectPackageId != PackageId, LogStreaming, Fatal,
+				TEXT("FGlobalImportStore::VerifyPackageForRemoval: %s (%s) - ")
+				TEXT("The loaded public export object '%s' has a mismatching package id %s in GlobalImportStore."),
+				*Package->GetName(),
+				*FormatPackageId(PackageId),
+				*Object->GetFullName(),
+				*FormatPackageId(ObjectPackageId));
+
+			VerifyObjectForRemoval(Object);
 		}
 	}
 
@@ -1513,7 +1375,8 @@ public:
 				TEXT("The existing object will be replaced since it or its package was most likely renamed after it was loaded the first time."),
 				Object ? *Object->GetFullName() : TEXT("null"),
 				ObjectIndex,
-				*FormatPackageId(Key.GetPackageId()), Key.GetExportHash(),
+				*FormatPackageId(Key.GetPackageId()), 
+				Key.GetExportHash(),
 				*ExistingObject->GetFullName(),
 				ExistingObject->GetFlags(),
 				int(ExistingObject->GetInternalFlags()),
@@ -5905,6 +5768,7 @@ bool FAsyncPackage2::CreateLinkerLoadExports(FAsyncLoadingThreadState2& ThreadSt
 
 	// Create exports
 	const int32 ExportCount = LinkerLoadState->Linker->ExportMap.Num();
+	FLoadedPackageRef& PackageRef = ImportStore.GlobalImportStore.FindPackageRefChecked(Desc.UPackageId, Desc.UPackageName);
 	while (LinkerLoadState->CreateExportIndex < ExportCount)
 	{
 		const int32 ExportIndex = LinkerLoadState->CreateExportIndex++;
@@ -5914,6 +5778,7 @@ bool FAsyncPackage2::CreateLinkerLoadExports(FAsyncLoadingThreadState2& ThreadSt
 			continue;
 		}
 #endif
+
 		FObjectExport& LinkerExport = LinkerLoadState->Linker->ExportMap[ExportIndex];
 		FExportObject& ExportObject = Data.Exports[ExportIndex];
 		if (UObject* Object = LinkerLoadState->Linker->CreateExport(ExportIndex))
@@ -9179,7 +9044,7 @@ EAsyncPackageState::Type FAsyncPackage2::PostLoadInstances(FAsyncLoadingThreadSt
 	{
 		const FExportObject& Export = Data.Exports[PostLoadInstanceIndex++];
 
-		if (!(Export.bFiltered | Export.bExportLoadFailed))
+		if (!(Export.bFiltered || Export.bExportLoadFailed))
 		{
 			UClass* ObjClass = Export.Object->GetClass();
 			ObjClass->PostLoadInstance(Export.Object);
