[BugFix] _ValuesSetBySerialization was made into a loose property on a transform component when switching levels and 'IDO.Enable true' is active

diff --git a/Engine/Source/Runtime/CoreUObject/Private/UObject/InstanceDataObjectUtils.cpp b/Engine/Source/Runtime/CoreUObject/Private/UObject/InstanceDataObjectUtils.cpp
--- a/Engine/Source/Runtime/CoreUObject/Private/UObject/InstanceDataObjectUtils.cpp
+++ b/Engine/Source/Runtime/CoreUObject/Private/UObject/InstanceDataObjectUtils.cpp
@@ -393,6 +393,12 @@ namespace UE
 			for (FPropertyPathNameTree::FConstIterator It = PropertyTree->CreateConstIterator(); It; ++It)
 			{
 				FName Name = It.GetName();
+				if (Name == NAME_ValuesSetBySerialization)
+				{
+					// in rare cases, the ValuesSetBySerialization property will get serialized/deserialized even though it's transient.
+					// Since this property is regenerated by this system and should never be loose, ignore it fully here.
+					continue;
+				}
 				FPropertyTypeName Type = It.GetType();
 				FPropertyPathName Path;
 				Path.Push({Name, Type});
