Always open around reporting a hitch

diff --git a/Engine/Source/Runtime/CoreUObject/Private/UObject/UObjectBaseUtility.cpp b/Engine/Source/Runtime/CoreUObject/Private/UObject/UObjectBaseUtility.cpp
--- a/Engine/Source/Runtime/CoreUObject/Private/UObject/UObjectBaseUtility.cpp
+++ b/Engine/Source/Runtime/CoreUObject/Private/UObject/UObjectBaseUtility.cpp
@@ -19,6 +19,7 @@
 #include "Containers/VersePath.h"
 #include "Misc/CoreDelegates.h"
 #include "Misc/DelayedAutoRegister.h"
+#include "AutoRTFM/AutoRTFM.h"
 
 /***********************/
 /******** Names ********/
@@ -647,6 +648,7 @@ UClass* GetParentNativeClass(UClass* Class)
 #include "HAL/ThreadHeartBeat.h"
 #include "HAL/ThreadManager.h"
 
+UE_AUTORTFM_ALWAYS_OPEN
 void FScopeCycleCounterUObject::ReportHitch()
 {
 	float Delta = float(FGameThreadHitchHeartBeat::Get().GetCurrentTime() - FGameThreadHitchHeartBeat::Get().GetFrameStartTime()) * 1000.0f;
