[CISFix] double declaration in unity builds

diff --git a/Engine/Source/Runtime/CoreUObject/Private/UObject/InstanceDataObjectUtils.cpp b/Engine/Source/Runtime/CoreUObject/Private/UObject/InstanceDataObjectUtils.cpp
--- a/Engine/Source/Runtime/CoreUObject/Private/UObject/InstanceDataObjectUtils.cpp
+++ b/Engine/Source/Runtime/CoreUObject/Private/UObject/InstanceDataObjectUtils.cpp
@@ -160,6 +160,15 @@ namespace UE
 		return bIsEnabled;
 	}
 
+	bool StructContainsLooseProperties(const UStruct* Struct)
+	{
+#if WITH_EDITORONLY_DATA
+		return Struct->GetBoolMetaData(NAME_ContainsLoosePropertiesMetadata);
+#else
+		return false;
+#endif
+	}
+
 	static UStruct* CreateInstanceDataObjectStructRec(const UClass* StructClass, UStruct* TemplateStruct, UObject* Outer, const FPropertyPathNameTree* PropertyTree);
 
 	template <typename StructType>
diff --git a/Engine/Source/Runtime/CoreUObject/Private/UObject/InstanceDataObjectUtils.h b/Engine/Source/Runtime/CoreUObject/Private/UObject/InstanceDataObjectUtils.h
--- a/Engine/Source/Runtime/CoreUObject/Private/UObject/InstanceDataObjectUtils.h
+++ b/Engine/Source/Runtime/CoreUObject/Private/UObject/InstanceDataObjectUtils.h
@@ -21,6 +21,7 @@ namespace UE
  * Pass nullptr to query if the system is enabled.
  */
 bool IsInstanceDataObjectSupportEnabled(UObject* Object = nullptr);
+bool StructContainsLooseProperties(const UStruct* Struct);
 
 /** Generate a UClass that contains the union of the properties of PropertyTree and OwnerClass. */
 UClass* CreateInstanceDataObjectClass(const FPropertyPathNameTree* PropertyTree, UClass* OwnerClass, UObject* Outer);
diff --git a/Engine/Source/Runtime/CoreUObject/Private/UObject/PropertyBagRepository.cpp b/Engine/Source/Runtime/CoreUObject/Private/UObject/PropertyBagRepository.cpp
--- a/Engine/Source/Runtime/CoreUObject/Private/UObject/PropertyBagRepository.cpp
+++ b/Engine/Source/Runtime/CoreUObject/Private/UObject/PropertyBagRepository.cpp
@@ -618,11 +618,8 @@ void FPropertyBagRepository::CreateInstanceDataObjectUnsafe(UObject* Owner, FPro
 	// construct InstanceDataObject class
 	// TODO: should we put the InstanceDataObject or it's class in a package?
 	const UClass* InstanceDataObjectClass = CreateInstanceDataObjectClass(PropertyTree, Owner->GetClass(), GetTransientPackage());
-	
-#if WITH_EDITOR
-	static const FName NAME_ContainsLoosePropertiesMetadata(ANSITEXTVIEW("ContainsLooseProperties"));
-	BagData.bNeedsFixup = InstanceDataObjectClass->GetBoolMetaData(NAME_ContainsLoosePropertiesMetadata);
-#endif
+
+	BagData.bNeedsFixup = UE::StructContainsLooseProperties(InstanceDataObjectClass);
 
 	TObjectPtr<UObject>* OuterPtr;
 	if (FPropertyBagAssociationData* OuterData = AssociatedData.Find(Owner->GetOuter()))
