Fixed constness of FProperty* in TPropertyValueIterator.

diff --git a/Engine/Plugins/Experimental/Avalanche/Source/AvalancheTransitionEditor/Private/Serialization/AvaTransitionStateSerializer.cpp b/Engine/Plugins/Experimental/Avalanche/Source/AvalancheTransitionEditor/Private/Serialization/AvaTransitionStateSerializer.cpp
--- a/Engine/Plugins/Experimental/Avalanche/Source/AvalancheTransitionEditor/Private/Serialization/AvaTransitionStateSerializer.cpp
+++ b/Engine/Plugins/Experimental/Avalanche/Source/AvalancheTransitionEditor/Private/Serialization/AvaTransitionStateSerializer.cpp
@@ -74,7 +74,7 @@ namespace UE::AvaTransitionEditor::Private
 		{
 			UScriptStruct* StateLinkStruct = TBaseStructure<FStateTreeStateLink>::Get();
 
-			for (const TPair<FStructProperty*, const void*>& Pair : TPropertyValueRange<FStructProperty>(InStruct, InStructMemory))
+			for (const TPair<const FStructProperty*, const void*>& Pair : TPropertyValueRange<FStructProperty>(InStruct, InStructMemory))
 			{
 				if (Pair.Key && Pair.Key->Struct == StateLinkStruct)
 				{
diff --git a/Engine/Source/Editor/AnimGraph/Private/AnimGraphNode_Base.cpp b/Engine/Source/Editor/AnimGraph/Private/AnimGraphNode_Base.cpp
--- a/Engine/Source/Editor/AnimGraph/Private/AnimGraphNode_Base.cpp
+++ b/Engine/Source/Editor/AnimGraph/Private/AnimGraphNode_Base.cpp
@@ -307,7 +307,7 @@ void UAnimGraphNode_Base::GetBoundFunctionsInfo(TArray<TPair<FName, FName>>& InO
 void UAnimGraphNode_Base::ValidateAnimNodeDuringCompilation(USkeleton* ForSkeleton, FCompilerResultsLog& MessageLog)
 {
 	// Validate any bone references we have
-	for(const TPair<FStructProperty*, const void*>& PropertyValuePair : TPropertyValueRange<FStructProperty>(GetClass(), this))
+	for(const TPair<const FStructProperty*, const void*>& PropertyValuePair : TPropertyValueRange<FStructProperty>(GetClass(), this))
 	{
 		if(PropertyValuePair.Key->Struct == FBoneReference::StaticStruct())
 		{
diff --git a/Engine/Source/Runtime/CoreUObject/Public/UObject/UnrealType.h b/Engine/Source/Runtime/CoreUObject/Public/UObject/UnrealType.h
--- a/Engine/Source/Runtime/CoreUObject/Public/UObject/UnrealType.h
+++ b/Engine/Source/Runtime/CoreUObject/Public/UObject/UnrealType.h
@@ -7025,7 +7025,7 @@ enum class EPropertyValueIteratorFlags : uint8
 class FPropertyValueIterator
 {
 public:
-	typedef TPair<const FProperty*, const void*> BasePairType;
+	using BasePairType = TPair<const FProperty*, const void*>;
 
 	/** 
 	 * Construct an iterator using a struct and struct value
@@ -7036,7 +7036,7 @@ public:
 	 * @param InRecursionFlags	Rather to recurse into container and struct properties
 	 * @param InDeprecatedPropertyFlags	Rather to iterate over deprecated properties
 	 */
-	COREUOBJECT_API FPropertyValueIterator(FFieldClass* InPropertyClass, const UStruct* InStruct, const void* InStructValue,
+	COREUOBJECT_API explicit FPropertyValueIterator(FFieldClass* InPropertyClass, const UStruct* InStruct, const void* InStructValue,
 		EPropertyValueIteratorFlags						InRecursionFlags = EPropertyValueIteratorFlags::FullRecursion,
 		EFieldIteratorFlags::DeprecatedPropertyFlags	InDeprecatedPropertyFlags = EFieldIteratorFlags::IncludeDeprecated);
 
@@ -7220,7 +7220,7 @@ template <class T>
 class TPropertyValueIterator : public FPropertyValueIterator
 {
 public:
-	typedef TPair<T*, const void*> PairType;
+	using PairType = TPair<const T*, const void*>;
 	
 	/** 
 	 * Construct an iterator using a struct and struct value
@@ -7230,7 +7230,7 @@ public:
 	 * @param InRecursionFlags	Rather to recurse into container and struct properties
 	 * @param InDeprecatedPropertyFlags	Rather to iterate over deprecated properties
 	 */
-	TPropertyValueIterator(const UStruct* InStruct, const void* InStructValue,
+	explicit TPropertyValueIterator(const UStruct* InStruct, const void* InStructValue,
 		EPropertyValueIteratorFlags						InRecursionFlags = EPropertyValueIteratorFlags::FullRecursion,
 		EFieldIteratorFlags::DeprecatedPropertyFlags	InDeprecatedPropertyFlags = EFieldIteratorFlags::IncludeDeprecated)
 		: FPropertyValueIterator(T::StaticClass(), InStruct, InStructValue, InRecursionFlags, InDeprecatedPropertyFlags)
@@ -7238,10 +7238,7 @@ public:
 	}
 
 	/** Invalid iterator, start with empty stack */
-	TPropertyValueIterator() 
-		: FPropertyValueIterator()
-	{
-	}
+	TPropertyValueIterator() = default;
 
 	/** Returns a TPair containing Property/Value currently being iterated */
 	FORCEINLINE const PairType& operator*() const
@@ -7255,7 +7252,7 @@ public:
 	}
 
 	/** Returns Property currently being iterated */
-	FORCEINLINE T* Key() const
+	FORCEINLINE const T* Key() const
 	{
 		return (*this)->Key;
 	}
