Add PlatformMemory::KernelAddressBit to enable an extra high-bit for pointer tagging

diff --git a/Engine/Source/Runtime/Core/Private/HAL/PlatformMemory.cpp b/Engine/Source/Runtime/Core/Private/HAL/PlatformMemory.cpp
--- a/Engine/Source/Runtime/Core/Private/HAL/PlatformMemory.cpp
+++ b/Engine/Source/Runtime/Core/Private/HAL/PlatformMemory.cpp
@@ -2,4 +2,4 @@
 
 #include "HAL/PlatformMemory.h"
 
-
+static_assert(FPlatformMemory::KernelAddressBit < 64, "Unsupported architecture, please declare which address bit distinguish user space from kernel space for non-x64/ARM64 platforms");
\ No newline at end of file
diff --git a/Engine/Source/Runtime/Core/Public/GenericPlatform/GenericPlatformMemory.h b/Engine/Source/Runtime/Core/Public/GenericPlatform/GenericPlatformMemory.h
--- a/Engine/Source/Runtime/Core/Public/GenericPlatform/GenericPlatformMemory.h
+++ b/Engine/Source/Runtime/Core/Public/GenericPlatform/GenericPlatformMemory.h
@@ -821,6 +821,14 @@ public:
 		return true;
 	}
 
+	// This bit is always zero in user mode addresses and most likely won't be used by current or future
+	// CPU features like ARM's PAC / Top-Byte Ignore or Intel's Linear Address Masking / 5-Level Paging
+#if defined(__x86_64__) || defined(_M_X64)
+	static constexpr uint32 KernelAddressBit = 63;
+#elif defined(__aarch64__) || defined(_M_ARM64)
+	static constexpr uint32 KernelAddressBit = 55;
+#endif
+
 protected:
 	friend struct FGenericStatsUpdater;
 
diff --git a/Engine/Source/Runtime/CoreUObject/Public/AssetRegistry/AssetDataTagMap.h b/Engine/Source/Runtime/CoreUObject/Public/AssetRegistry/AssetDataTagMap.h
--- a/Engine/Source/Runtime/CoreUObject/Public/AssetRegistry/AssetDataTagMap.h
+++ b/Engine/Source/Runtime/CoreUObject/Public/AssetRegistry/AssetDataTagMap.h
@@ -174,19 +174,6 @@ namespace FixedTagPrivate
 			}
 		}
 	};
-
-	// This bit is always zero in user mode addresses and most likely won't be used by current or future
-	// CPU features like ARM's PAC / Top-Byte Ignore or Intel's Linear Address Masking / 5-Level Paging
-#if defined(__x86_64__) || defined(_M_X64)
-	static constexpr uint32 KernelAddressBit = 63;
-#elif defined(__aarch64__) || defined(_M_ARM64)
-	static constexpr uint32 KernelAddressBit = 55;
-#elif defined(PLATFORM_KERNEL_ADDRESS_BIT)
-	static constexpr uint32 KernelAddressBit = PLATFORM_KERNEL_ADDRESS_BIT;
-#else
-	#error Unsupported architecture, please declare which address bit distinguish user space from kernel space
-#endif
-
 } // end namespace FixedTagPrivate
 
 /**
@@ -201,8 +188,8 @@ class FAssetTagValueRef
 
 	class FFixedTagValue
 	{
-		static constexpr uint64 FixedMask = uint64(1) << FixedTagPrivate::KernelAddressBit;
-		static_assert(FixedTagPrivate::FMapHandle::StoreIndexBits <= (FixedTagPrivate::KernelAddressBit - 32), 
+		static constexpr uint64 FixedMask = uint64(1) << FPlatformMemory::KernelAddressBit;
+		static_assert(FixedTagPrivate::FMapHandle::StoreIndexBits <= (FPlatformMemory::KernelAddressBit - 32), 
 			"Too few bits remain for the StoreIndex. Consider using other high bits but note that ARM64 use the top byte for HWASAN & MTE.)");
 
 		uint64 Bits;
