Core: Array Text import not working for empty arrays
#rnx

diff --git a/Engine/Source/Runtime/CoreUObject/Private/UObject/PropertyArray.cpp b/Engine/Source/Runtime/CoreUObject/Private/UObject/PropertyArray.cpp
--- a/Engine/Source/Runtime/CoreUObject/Private/UObject/PropertyArray.cpp
+++ b/Engine/Source/Runtime/CoreUObject/Private/UObject/PropertyArray.cpp
@@ -986,6 +986,7 @@ const TCHAR* FArrayProperty::ImportTextInnerItem(const TCHAR* Buffer, const FPro
 		{
 			ArrayHelper->EmptyValues();
 		}
+		++Buffer; // consume closing parenthesis
 		return Buffer;
 	}
 
