Added optional tracking and persistence of which properties have been initialized

#rnx

diff --git a/Engine/Plugins/Experimental/InstanceDataObjectFixupTool/Source/Private/Tests/InstanceDataObjectTests.cpp b/Engine/Plugins/Experimental/InstanceDataObjectFixupTool/Source/Private/Tests/InstanceDataObjectTests.cpp
--- a/Engine/Plugins/Experimental/InstanceDataObjectFixupTool/Source/Private/Tests/InstanceDataObjectTests.cpp
+++ b/Engine/Plugins/Experimental/InstanceDataObjectFixupTool/Source/Private/Tests/InstanceDataObjectTests.cpp
@@ -334,7 +334,8 @@ static UObject* GenerateTestInstanceDataObject(UObject* ObjectOld, UClass* NewCl
 
 		// scoped flag changes to load context to support unknown property tracking
 		TGuardValue<bool> ScopedTrackSerializedPropertyPath(LoadContext->bTrackSerializedPropertyPath, true);
-		TGuardValue<bool> ScopedSerializeUnknownProperty(LoadContext->bTrackUnknownProperties, true);
+		TGuardValue<bool> ScopedTrackInitializedProperties(LoadContext->bTrackInitializedProperties, true);
+		TGuardValue<bool> ScopedTrackUnknownProperties(LoadContext->bTrackUnknownProperties, true);
 		TGuardValue<bool> ScopedImpersonateProperties(LoadContext->bImpersonateProperties, true);
 		
 		UPackage* TempPackage = CreateTestPackage(NewClass);
diff --git a/Engine/Source/Runtime/CoreUObject/Internal/UObject/InstanceDataObjectUtils.h b/Engine/Source/Runtime/CoreUObject/Internal/UObject/InstanceDataObjectUtils.h
new file mode 100644
--- /dev/null
+++ b/Engine/Source/Runtime/CoreUObject/Internal/UObject/InstanceDataObjectUtils.h
@@ -0,0 +1,49 @@
+// Copyright Epic Games, Inc. All Rights Reserved.
+
+#pragma once
+
+#include "HAL/Platform.h"
+
+#define UE_API COREUOBJECT_API
+
+class FFieldVariant;
+class FProperty;
+class UClass;
+class UObject;
+class UStruct;
+
+namespace UE { class FPropertyPathNameTree; }
+
+namespace UE
+{
+
+/**
+ * Query if InstanceDataObject support is enabled for a specific object.
+ *
+ * Pass nullptr to query if the system is enabled.
+ */
+bool IsInstanceDataObjectSupportEnabled(UObject* Object = nullptr);
+bool StructContainsLooseProperties(const UStruct* Struct);
+
+/** Generate a UClass that contains the union of the properties of PropertyTree and OwnerClass. */
+UClass* CreateInstanceDataObjectClass(const FPropertyPathNameTree* PropertyTree, UClass* OwnerClass, UObject* Outer);
+
+/** Notify that a property in an struct was set when the struct was deserialized. */
+void MarkPropertySetBySerialization(const UStruct* Struct, void* StructData, const FProperty* Property, int32 ArrayIndex = 0);
+/** Query whether a property in the struct was set when the struct was deserialized. */
+bool WasPropertySetBySerialization(const UStruct* Struct, const void* StructData, const FProperty* Property, int32 ArrayIndex = 0);
+/** Copy whether each property was set by serialization from one IDO to another. */
+void CopyPropertySetBySerializationData(const FFieldVariant& OldField, void* OldDataPtr, const FFieldVariant& NewField, void* NewDataPtr);
+
+/** Query whether the property value is initialized for a property in the struct. */
+UE_API bool IsPropertyValueInitialized(const UStruct* Struct, void* StructData, const FProperty* Property, int32 ArrayIndex = 0);
+/** Set the flag that the property value is initialized for a property in the struct. */
+UE_API void SetPropertyValueInitialized(const UStruct* Struct, void* StructData, const FProperty* Property, int32 ArrayIndex = 0);
+/** Clear the flag that the property value is initialized for a property in the struct. */
+UE_API void ClearPropertyValueInitialized(const UStruct* Struct, void* StructData, const FProperty* Property, int32 ArrayIndex = 0);
+/** Reset the property value initialized flags for every property in the struct. */
+UE_API void ResetPropertyValueInitialized(const UStruct* Struct, void* StructData);
+
+} // UE
+
+#undef UE_API
diff --git a/Engine/Source/Runtime/CoreUObject/Private/UObject/Class.cpp b/Engine/Source/Runtime/CoreUObject/Private/UObject/Class.cpp
--- a/Engine/Source/Runtime/CoreUObject/Private/UObject/Class.cpp
+++ b/Engine/Source/Runtime/CoreUObject/Private/UObject/Class.cpp
@@ -1259,6 +1259,7 @@ void UStruct::SerializeBinEx( FStructuredArchive::FSlot Slot, void* Data, void c
 void UStruct::LoadTaggedPropertiesFromText(FStructuredArchive::FSlot Slot, uint8* Data, UStruct* DefaultsStruct, uint8* Defaults, const UObject* BreakRecursionIfFullyLoad) const
 {
 	FArchive& UnderlyingArchive = Slot.GetUnderlyingArchive();
+	FUObjectSerializeContext* SerializeContext = FUObjectThreadContext::Get().GetSerializeContext();
 	const bool bUseRedirects = !FPlatformProperties::RequiresCookedData() || UnderlyingArchive.IsSaveGame();
 	int32 NumProperties = 0;
 	FStructuredArchiveMap PropertiesMap = Slot.EnterMap(NumProperties);
@@ -1339,6 +1340,16 @@ void UStruct::LoadTaggedPropertiesFromText(FStructuredArchive::FSlot Slot, uint8
 				Tag.ArrayIndex = ItemIndex;
 				Tag.Name = PropertyName;
 
+				if (SerializeContext->bTrackInitializedProperties)
+				{
+					UE::SetPropertyValueInitialized(this, Data, Property, ItemIndex);
+				}
+
+				if (Tag.SerializeType == EPropertyTagSerializeType::Skipped)
+				{
+					continue;
+				}
+
 				if (bUseRedirects)
 				{
 					if (UE::FPropertyTypeName NewTypeName = ApplyRedirectsToPropertyType(Tag.GetType(), Property); !NewTypeName.IsEmpty())
@@ -1636,6 +1647,16 @@ void UStruct::SerializeVersionedTaggedProperties(FStructuredArchive::FSlot Slot,
 					Property = CustomFindProperty(Tag.Name);
 				}
 
+				if (SerializeContext->bTrackInitializedProperties && Property)
+				{
+					UE::SetPropertyValueInitialized(this, Data, Property, Tag.ArrayIndex);
+				}
+
+				if (Tag.SerializeType == EPropertyTagSerializeType::Skipped)
+				{
+					continue;
+				}
+
 				Tag.SetProperty(Property);
 
 				if (bUseRedirects)
@@ -1865,9 +1886,11 @@ void UStruct::SerializeVersionedTaggedProperties(FStructuredArchive::FSlot Slot,
 				{
 					uint8* DataPtr      = Property->ContainerPtrToValuePtr           <uint8>(Data, Idx);
 					uint8* DefaultValue = Property->ContainerPtrToValuePtrForDefaults<uint8>(DefaultsStruct, Defaults, Idx);
-					if (StaticArrayContainer.IsSet() || CustomPropertyNode || !bDoDeltaSerialization ||
+					const bool bSerializeValue = (StaticArrayContainer.IsSet() || CustomPropertyNode || !bDoDeltaSerialization ||
 						(FOverridableSerializationLogic::IsEnabled() && FOverridableSerializationLogic::GetOverriddenPropertyOperation(UnderlyingArchive, Property, DataPtr, DefaultValue) != EOverriddenPropertyOperation::None) ||
-						(!FOverridableSerializationLogic::IsEnabled() && !Property->Identical(DataPtr, DefaultValue, UnderlyingArchive.GetPortFlags())))
+						(!FOverridableSerializationLogic::IsEnabled() && !Property->Identical(DataPtr, DefaultValue, UnderlyingArchive.GetPortFlags())));
+					const bool bInitializedValue = !SerializeContext->bTrackInitializedProperties || UE::IsPropertyValueInitialized(this, Data, Property, Idx);
+					if (bInitializedValue && (bSerializeValue || SerializeContext->bTrackInitializedProperties))
 					{
 						if (bUseAtomicSerialization)
 						{
@@ -1888,12 +1911,23 @@ void UStruct::SerializeVersionedTaggedProperties(FStructuredArchive::FSlot Slot,
 							Tag.SetPropertyGuid(PropertyGuid);
 						}
 
+						if (!bSerializeValue)
+						{
+							Tag.SerializeType = EPropertyTagSerializeType::Skipped;
+						}
+
 						TStringBuilder<256> TagName;
 						Tag.Name.ToString(TagName);
 						FStructuredArchive::FSlot PropertySlot = StaticArrayContainer.IsSet() ? StaticArrayContainer->EnterElement() : PropertiesRecord.EnterField(TagName.ToString());
 
 						PropertySlot << Tag;
 
+						if (!bSerializeValue)
+						{
+							PropertySlot.EnterStream(); // Save an empty value for text format archives.
+							continue;
+						}
+
 						// need to know how much data this call to SerializeTaggedProperty consumes, so mark where we are
 						int64 DataOffset = UnderlyingArchive.Tell();
 
diff --git a/Engine/Source/Runtime/CoreUObject/Private/UObject/InstanceDataObjectUtils.cpp b/Engine/Source/Runtime/CoreUObject/Private/UObject/InstanceDataObjectUtils.cpp
--- a/Engine/Source/Runtime/CoreUObject/Private/UObject/InstanceDataObjectUtils.cpp
+++ b/Engine/Source/Runtime/CoreUObject/Private/UObject/InstanceDataObjectUtils.cpp
@@ -14,6 +14,7 @@
 #include "UObject/PropertyPathNameTree.h"
 #include "UObject/UnrealType.h"
 
+static const FName NAME_InitializedValues(ANSITEXTVIEW("_InitializedValues"));
 static const FName NAME_ValuesSetBySerialization(ANSITEXTVIEW("_ValuesSetBySerialization"));
 
 /** Type used for InstanceDataObject classes. */
@@ -22,6 +23,7 @@ class UInstanceDataObjectClass final : public UClass
 public:
 	DECLARE_CASTED_CLASS_INTRINSIC(UInstanceDataObjectClass, UClass, CLASS_Transient, TEXT("/Script/CoreUObject"), CASTCLASS_UClass)
 
+	FByteProperty* InitializedValuesProperty = nullptr;
 	FByteProperty* ValuesSetBySerializationProperty = nullptr;
 };
 
@@ -38,6 +40,7 @@ public:
 	uint32 GetStructTypeHash(const void* Src) const final;
 	FGuid GetCustomGuid() const final { return Guid; }
 
+	FByteProperty* InitializedValuesProperty = nullptr;
 	FByteProperty* ValuesSetBySerializationProperty = nullptr;
 	FGuid Guid;
 };
@@ -85,7 +88,7 @@ uint32 UInstanceDataObjectStruct::GetStructTypeHash(const void* Src) const
 	uint32 ValueHash = 0;
 	for (TFieldIterator<const FProperty> It(this); It; ++It)
 	{
-		if (It->GetFName() == NAME_ValuesSetBySerialization)
+		if (It->GetFName() == NAME_InitializedValues || It->GetFName() == NAME_ValuesSetBySerialization)
 		{
 			continue;
 		}
@@ -402,10 +405,10 @@ namespace UE
 			for (FPropertyPathNameTree::FConstIterator It = PropertyTree->CreateConstIterator(); It; ++It)
 			{
 				FName Name = It.GetName();
-				if (Name == NAME_ValuesSetBySerialization)
+				if (Name == NAME_InitializedValues || Name == NAME_ValuesSetBySerialization)
 				{
-					// in rare cases, the ValuesSetBySerialization property will get serialized/deserialized even though it's transient.
-					// Since this property is regenerated by this system and should never be loose, ignore it fully here.
+					// In rare cases, these hidden properties will get serialized even though they are transient.
+					// Ignore them here since they are generated below.
 					continue;
 				}
 				FPropertyTypeName Type = It.GetType();
@@ -427,20 +430,25 @@ namespace UE
 			}
 		}
 
-		// Add a hidden byte array property to record whether its sibling properties were set by serialization.
+		// Add hidden byte array properties to record whether its sibling properties were initialized or set by serialization.
+		FByteProperty* InitializedValuesProperty = CastFieldChecked<FByteProperty>(FByteProperty::Construct(Result, NAME_InitializedValues, RF_Transient | RF_MarkAsNative));
 		FByteProperty* ValuesSetBySerializationProperty = CastFieldChecked<FByteProperty>(FByteProperty::Construct(Result, NAME_ValuesSetBySerialization, RF_Transient | RF_MarkAsNative));
 		{
+			InitializedValuesProperty->SetPropertyFlags(CPF_Transient | CPF_EditorOnly | CPF_NativeAccessSpecifierPrivate);
 			ValuesSetBySerializationProperty->SetPropertyFlags(CPF_Transient | CPF_EditorOnly | CPF_NativeAccessSpecifierPrivate);
+			Result->AddCppProperty(InitializedValuesProperty);
 			Result->AddCppProperty(ValuesSetBySerializationProperty);
 		}
 
 		// Store generated properties to avoid scanning every property to find it when it is needed.
 		if (UInstanceDataObjectClass* IdoClass = Cast<UInstanceDataObjectClass>(Result))
 		{
+			IdoClass->InitializedValuesProperty = InitializedValuesProperty;
 			IdoClass->ValuesSetBySerializationProperty = ValuesSetBySerializationProperty;
 		}
 		else if (UInstanceDataObjectStruct* IdoStruct = Cast<UInstanceDataObjectStruct>(Result))
 		{
+			IdoStruct->InitializedValuesProperty = InitializedValuesProperty;
 			IdoStruct->ValuesSetBySerializationProperty = ValuesSetBySerializationProperty;
 		}
 
@@ -451,12 +459,14 @@ namespace UE
 		}
 
 		// Count properties and set the size of the array of flags.
-		int32 PropertyCount = -1; // Start at -1 to exclude ValuesSetBySerialization.
+		int32 PropertyCount = -2; // Start at -2 to exclude the two hidden properties.
 		for (TFieldIterator<FProperty> It(Result); It; ++It)
 		{
 			PropertyCount += It->ArrayDim;
 		}
-		ValuesSetBySerializationProperty->ArrayDim = FMath::DivideAndRoundUp(PropertyCount, 8);
+		const int32 PropertyCountBytes = FMath::Max(1, FMath::DivideAndRoundUp(PropertyCount, 8));
+		InitializedValuesProperty->ArrayDim = PropertyCountBytes;
+		ValuesSetBySerializationProperty->ArrayDim = PropertyCountBytes;
 
 		Result->Bind();
 		Result->StaticLink(/*bRelinkExistingProperties*/true);
@@ -666,4 +676,81 @@ namespace UE
 			}
 		}
 	}
+
+	static const FByteProperty* FindInitializedValuesProperty(const UStruct* Struct)
+	{
+		if (const UInstanceDataObjectClass* IdoClass = Cast<UInstanceDataObjectClass>(Struct))
+		{
+			return IdoClass->InitializedValuesProperty;
+		}
+		if (const UInstanceDataObjectStruct* IdoStruct = Cast<UInstanceDataObjectStruct>(Struct))
+		{
+			return IdoStruct->InitializedValuesProperty;
+		}
+		return CastField<FByteProperty>(Struct->FindPropertyByName(NAME_InitializedValues));
+	}
+
+	static void SetPropertyValueInitializedFlag(const UStruct* Struct, void* StructData, const FProperty* Property, int32 ArrayIndex, bool bValue)
+	{
+	#if WITH_EDITORONLY_DATA
+		if (const FByteProperty* InitializedValuesProperty = FindInitializedValuesProperty(Struct))
+		{
+			const int32 PropertyIndex = Property->GetIndexInOwner() + ArrayIndex;
+			const int32 ByteIndex = PropertyIndex / 8;
+			const int32 BitOffset = PropertyIndex % 8;
+			if (ByteIndex < InitializedValuesProperty->ArrayDim)
+			{
+				uint8* PropertyDataPtr = InitializedValuesProperty->ContainerPtrToValuePtr<uint8>(StructData, ByteIndex);
+				if (bValue)
+				{
+					*PropertyDataPtr |= (1 << BitOffset);
+				}
+				else
+				{
+					*PropertyDataPtr &= ~(1 << BitOffset);
+				}
+			}
+		}
+	#endif
+	}
+
+	bool IsPropertyValueInitialized(const UStruct* Struct, void* StructData, const FProperty* Property, int32 ArrayIndex)
+	{
+	#if WITH_EDITORONLY_DATA
+		if (const FByteProperty* InitializedValuesProperty = FindInitializedValuesProperty(Struct))
+		{
+			const int32 PropertyIndex = Property->GetIndexInOwner() + ArrayIndex;
+			const int32 ByteIndex = PropertyIndex / 8;
+			const int32 BitOffset = PropertyIndex % 8;
+			if (ByteIndex < InitializedValuesProperty->ArrayDim)
+			{
+				const uint8* PropertyDataPtr = InitializedValuesProperty->ContainerPtrToValuePtr<uint8>(StructData, ByteIndex);
+				return (*PropertyDataPtr & (1 << BitOffset)) != 0;
+			}
+		}
+	#endif
+		return false;
+	}
+
+	void SetPropertyValueInitialized(const UStruct* Struct, void* StructData, const FProperty* Property, int32 ArrayIndex)
+	{
+		SetPropertyValueInitializedFlag(Struct, StructData, Property, ArrayIndex, /*bValue*/ true);
+	}
+
+	void ClearPropertyValueInitialized(const UStruct* Struct, void* StructData, const FProperty* Property, int32 ArrayIndex)
+	{
+		SetPropertyValueInitializedFlag(Struct, StructData, Property, ArrayIndex, /*bValue*/ false);
+	}
+
+	void ResetPropertyValueInitialized(const UStruct* Struct, void* StructData)
+	{
+	#if WITH_EDITORONLY_DATA
+		if (const FByteProperty* InitializedValuesProperty = FindInitializedValuesProperty(Struct))
+		{
+			uint8* PropertyDataPtr = InitializedValuesProperty->ContainerPtrToValuePtr<uint8>(StructData);
+			FMemory::Memzero(PropertyDataPtr, InitializedValuesProperty->ArrayDim);
+		}
+	#endif
+	}
+
 } // UE
diff --git a/Engine/Source/Runtime/CoreUObject/Private/UObject/InstanceDataObjectUtils.h b/Engine/Source/Runtime/CoreUObject/Private/UObject/InstanceDataObjectUtils.h
deleted file mode 100644
diff --git a/Engine/Source/Runtime/CoreUObject/Private/UObject/LinkerLoad.cpp b/Engine/Source/Runtime/CoreUObject/Private/UObject/LinkerLoad.cpp
--- a/Engine/Source/Runtime/CoreUObject/Private/UObject/LinkerLoad.cpp
+++ b/Engine/Source/Runtime/CoreUObject/Private/UObject/LinkerLoad.cpp
@@ -4698,6 +4698,9 @@ void FLinkerLoad::Preload( UObject* Object )
 					// Enable property path tracking when IDO support is enabled. Both the property path name tree and IDO creation require the paths.
 					const bool bHasIDOSupport = UE::IsInstanceDataObjectSupportEnabled(Object);
 					TGuardValue<bool> ScopedTrackSerializedPropertyPath(LoadContext->bTrackSerializedPropertyPath, bHasIDOSupport);
+					// Enable tracking of initialized properties when loading an IDO, which is implied by impersonation being enabled.
+					const bool bLoadingIDO = bHasIDOSupport && LoadContext->bImpersonateProperties;
+					TGuardValue<bool> ScopedTrackInitializedProperties(LoadContext->bTrackInitializedProperties, bLoadingIDO);
 					// Enable creation of a property path name tree to track any property that does not match the current class schema,
 					// except when impersonation is enabled because that implies we are deserializing an IDO.
 					const bool bCreateIDO = bHasIDOSupport && !LoadContext->bImpersonateProperties;
diff --git a/Engine/Source/Runtime/CoreUObject/Private/UObject/PropertyTag.cpp b/Engine/Source/Runtime/CoreUObject/Private/UObject/PropertyTag.cpp
--- a/Engine/Source/Runtime/CoreUObject/Private/UObject/PropertyTag.cpp
+++ b/Engine/Source/Runtime/CoreUObject/Private/UObject/PropertyTag.cpp
@@ -22,6 +22,7 @@ enum class EPropertyTagFlags : uint8
 	HasPropertyExtensions		= 0x04,
 	HasBinaryOrNativeSerialize	= 0x08,
 	BoolTrue					= 0x10,
+	SkippedSerialize			= 0x20,
 };
 
 ENUM_CLASS_FLAGS(EPropertyTagFlags);
@@ -145,7 +146,14 @@ static EPropertyTagExtension CalculatePropertyExtensionFlags(FArchive& Underlyin
 static void SerializePropertyExtensions(FStructuredArchive::FSlot Slot, EPropertyTagExtension PropertyTagExtensions, FPropertyTag& Tag)
 {
 	// Serialize tag extensions, consider doing an init function and context as `EClassSerializationControlExtension` if we add more extensions
-	Slot << SA_ATTRIBUTE(TEXT("PropertyExtensions"), PropertyTagExtensions);
+	if (Slot.GetArchiveState().IsTextFormat())
+	{
+		Slot << SA_OPTIONAL_ATTRIBUTE(TEXT("PropertyExtensions"), PropertyTagExtensions, EPropertyTagExtension::NoExtension);
+	}
+	else
+	{
+		Slot << SA_ATTRIBUTE(TEXT("PropertyExtensions"), PropertyTagExtensions);
+	}
 
 	// OverridableInformation
 	if (EnumHasAnyFlags(PropertyTagExtensions, EPropertyTagExtension::OverridableInformation))
@@ -394,12 +402,16 @@ FORCENOINLINE void SerializePropertyTagAsText(FStructuredArchive::FSlot Slot, FP
 		Tag.SetType(Tag.TypeName);
 	}
 
-	Slot << SA_OPTIONAL_ATTRIBUTE(TEXT("PropertyGuid"), Tag.PropertyGuid, FGuid());
-	Tag.HasPropertyGuid = Tag.PropertyGuid.IsValid();
-
 	bool bHasBinaryOrNativeSerialize = (Tag.SerializeType == EPropertyTagSerializeType::BinaryOrNative);
+	bool bSkippedSerialize = (Tag.SerializeType == EPropertyTagSerializeType::Skipped);
+
+	Slot << SA_OPTIONAL_ATTRIBUTE(TEXT("PropertyGuid"), Tag.PropertyGuid, FGuid());
 	Slot << SA_OPTIONAL_ATTRIBUTE(TEXT("HasBinaryOrNativeSerialize"), bHasBinaryOrNativeSerialize, false);
-	Tag.SerializeType = bHasBinaryOrNativeSerialize ? EPropertyTagSerializeType::BinaryOrNative : EPropertyTagSerializeType::Property;
+	Slot << SA_OPTIONAL_ATTRIBUTE(TEXT("SkippedSerialize"), bSkippedSerialize, false);
+
+	Tag.HasPropertyGuid = Tag.PropertyGuid.IsValid();
+	Tag.SerializeType = bSkippedSerialize ? EPropertyTagSerializeType::Skipped
+		: (bHasBinaryOrNativeSerialize ? EPropertyTagSerializeType::BinaryOrNative : EPropertyTagSerializeType::Property);
 
 	EPropertyTagExtension PropertyTagExtensions = CalculatePropertyExtensionFlags(UnderlyingArchive, Tag);
 	SerializePropertyExtensions(Slot, PropertyTagExtensions, Tag);
@@ -427,6 +439,13 @@ void operator<<(FStructuredArchive::FSlot Slot, FPropertyTag& Tag)
 		return;
 	}
 
+	if (UnderlyingArchive.IsSaving())
+	{
+		Tag.SerializeType = Tag.SerializeType == EPropertyTagSerializeType::Skipped ? Tag.SerializeType
+			: (Tag.GetProperty()->UseBinaryOrNativeSerialization(UnderlyingArchive)
+				? EPropertyTagSerializeType::BinaryOrNative : EPropertyTagSerializeType::Property);
+	}
+
 	if (UNLIKELY(bIsTextFormat))
 	{
 		SerializePropertyTagAsText(Slot, Tag);
@@ -458,9 +477,6 @@ void operator<<(FStructuredArchive::FSlot Slot, FPropertyTag& Tag)
 
 	if (UnderlyingArchive.IsSaving())
 	{
-		Tag.SerializeType = Tag.GetProperty()->UseBinaryOrNativeSerialization(UnderlyingArchive)
-			? EPropertyTagSerializeType::BinaryOrNative : EPropertyTagSerializeType::Property;
-
 		if (Tag.ArrayIndex != 0)
 		{
 			PropertyTagFlags |= EPropertyTagFlags::HasArrayIndex;
@@ -477,6 +493,10 @@ void operator<<(FStructuredArchive::FSlot Slot, FPropertyTag& Tag)
 		{
 			PropertyTagFlags |= EPropertyTagFlags::HasBinaryOrNativeSerialize;
 		}
+		if (Tag.SerializeType == EPropertyTagSerializeType::Skipped)
+		{
+			PropertyTagFlags |= EPropertyTagFlags::SkippedSerialize;
+		}
 		if (Tag.BoolVal && Tag.Type == NAME_BoolProperty)
 		{
 			PropertyTagFlags |= EPropertyTagFlags::BoolTrue;
@@ -488,8 +508,10 @@ void operator<<(FStructuredArchive::FSlot Slot, FPropertyTag& Tag)
 	if (UnderlyingArchive.IsLoading())
 	{
 		Tag.HasPropertyGuid = EnumHasAnyFlags(PropertyTagFlags, EPropertyTagFlags::HasPropertyGuid);
-		Tag.SerializeType = EnumHasAnyFlags(PropertyTagFlags, EPropertyTagFlags::HasBinaryOrNativeSerialize)
-			? EPropertyTagSerializeType::BinaryOrNative : EPropertyTagSerializeType::Property;
+		Tag.SerializeType = EnumHasAnyFlags(PropertyTagFlags, EPropertyTagFlags::SkippedSerialize)
+			? EPropertyTagSerializeType::Skipped
+			: (EnumHasAnyFlags(PropertyTagFlags, EPropertyTagFlags::HasBinaryOrNativeSerialize)
+				? EPropertyTagSerializeType::BinaryOrNative : EPropertyTagSerializeType::Property);
 		Tag.BoolVal = EnumHasAnyFlags(PropertyTagFlags, EPropertyTagFlags::BoolTrue);
 	}
 
diff --git a/Engine/Source/Runtime/CoreUObject/Private/UObject/UObjectThreadContext.cpp b/Engine/Source/Runtime/CoreUObject/Private/UObject/UObjectThreadContext.cpp
--- a/Engine/Source/Runtime/CoreUObject/Private/UObject/UObjectThreadContext.cpp
+++ b/Engine/Source/Runtime/CoreUObject/Private/UObject/UObjectThreadContext.cpp
@@ -45,6 +45,7 @@ FUObjectSerializeContext::FUObjectSerializeContext()
 	, SerializedExportIndex(0)
 	, SerializedExportLinker(nullptr)
 	, bTrackSerializedPropertyPath(false)
+	, bTrackInitializedProperties(false)
 	, bTrackUnknownProperties(false)
 	, bImpersonateProperties(false)
 {}
diff --git a/Engine/Source/Runtime/CoreUObject/Public/UObject/PropertyTag.h b/Engine/Source/Runtime/CoreUObject/Public/UObject/PropertyTag.h
--- a/Engine/Source/Runtime/CoreUObject/Public/UObject/PropertyTag.h
+++ b/Engine/Source/Runtime/CoreUObject/Public/UObject/PropertyTag.h
@@ -23,6 +23,8 @@ enum class EPropertyTagSerializeType : uint8
 {
 	/** Tag was loaded from an older version or has not yet been saved. */
 	Unknown,
+	/** Serialization of the property value was skipped. Tag has no value. */
+	Skipped,
 	/** Serialized with tagged property serialization. */
 	Property,
 	/** Serialized with binary or native serialization. */
diff --git a/Engine/Source/Runtime/CoreUObject/Public/UObject/UObjectThreadContext.h b/Engine/Source/Runtime/CoreUObject/Public/UObject/UObjectThreadContext.h
--- a/Engine/Source/Runtime/CoreUObject/Public/UObject/UObjectThreadContext.h
+++ b/Engine/Source/Runtime/CoreUObject/Public/UObject/UObjectThreadContext.h
@@ -172,6 +172,8 @@ public:
 	UE_INTERNAL UE::FPropertyPathName SerializedPropertyPath;
 	/** True when SerializedPropertyPath is being tracked during serialization. */
 	UE_INTERNAL bool bTrackSerializedPropertyPath;
+	/** True when initialized properties will be tracked for the serialized object. */
+	UE_INTERNAL bool bTrackInitializedProperties;
 	/** True when unknown properties will tracked for the serialized object. */
 	UE_INTERNAL bool bTrackUnknownProperties;
 	/** True when the SerializedObject properties are being impersonated. */
diff --git a/Engine/Source/Runtime/CoreUObject/Tests/UObject/InstanceDataObjectUtilsTest.cpp b/Engine/Source/Runtime/CoreUObject/Tests/UObject/InstanceDataObjectUtilsTest.cpp
--- a/Engine/Source/Runtime/CoreUObject/Tests/UObject/InstanceDataObjectUtilsTest.cpp
+++ b/Engine/Source/Runtime/CoreUObject/Tests/UObject/InstanceDataObjectUtilsTest.cpp
@@ -4,14 +4,22 @@
 
 #if WITH_TESTS && WITH_EDITORONLY_DATA
 
+#include "Misc/ScopeExit.h"
+#include "Serialization/Formatters/BinaryArchiveFormatter.h"
+#include "Serialization/Formatters/JsonArchiveInputFormatter.h"
+#include "Serialization/Formatters/JsonArchiveOutputFormatter.h"
+#include "Serialization/MemoryReader.h"
+#include "Serialization/MemoryWriter.h"
+#include "Serialization/StructuredArchive.h"
 #include "Tests/TestHarnessAdapter.h"
 #include "UObject/InstanceDataObjectUtils.h"
 #include "UObject/Package.h"
+#include "UObject/UObjectThreadContext.h"
 
 namespace UE
 {
 
-TEST_CASE_NAMED(FInstanceDataObjectUtilsTest, "CoreUObject::InstanceDataObjectUtils", "[Core][UObject][EngineFilter]")
+TEST_CASE_NAMED(FInstanceDataObjectUtilsTest, "CoreUObject::Serialization::InstanceDataObjectUtils", "[Core][UObject][EngineFilter]")
 {
 	UTestInstanceDataObjectClass* BaseObject = NewObject<UTestInstanceDataObjectClass>();
 	UClass* TestClass = CreateInstanceDataObjectClass(nullptr, BaseObject->GetClass(), BaseObject->GetOuter());
@@ -31,7 +39,8 @@ TEST_CASE_NAMED(FInstanceDataObjectUtilsTest, "CoreUObject::InstanceDataObjectUt
 	REQUIRE(CProperty);
 	REQUIRE(DProperty);
 
-	UObject* TestObject = NewObject<UObject>(GetTransientPackage(), TestClass);
+	FName TestObjectName = MakeUniqueObjectName(nullptr, TestClass, FName(WriteToString<128>(TestClass->GetFName(), TEXT("_Instance"))));
+	UObject* TestObject = NewObject<UObject>(GetTransientPackage(), TestClass, TestObjectName);
 	void* StructData = StructProperty->ContainerPtrToValuePtr<void>(TestObject);
 
 	CHECK_FALSE(WasPropertySetBySerialization(TestClass, TestObject, StructProperty));
@@ -39,6 +48,12 @@ TEST_CASE_NAMED(FInstanceDataObjectUtilsTest, "CoreUObject::InstanceDataObjectUt
 	CHECK(WasPropertySetBySerialization(TestClass, TestObject, StructProperty));
 	CHECK_FALSE(WasPropertySetBySerialization(TestClass, TestObject, Int32Property));
 
+	CHECK_FALSE(IsPropertyValueInitialized(TestClass, TestObject, StructProperty));
+	SetPropertyValueInitialized(TestClass, TestObject, StructProperty);
+	CHECK(IsPropertyValueInitialized(TestClass, TestObject, StructProperty));
+	ClearPropertyValueInitialized(TestClass, TestObject, StructProperty);
+	CHECK_FALSE(IsPropertyValueInitialized(TestClass, TestObject, StructProperty));
+
 	CHECK_FALSE(WasPropertySetBySerialization(StructProperty->Struct, StructData, AProperty));
 	MarkPropertySetBySerialization(StructProperty->Struct, StructData, AProperty);
 	CHECK(WasPropertySetBySerialization(StructProperty->Struct, StructData, AProperty));
@@ -47,6 +62,139 @@ TEST_CASE_NAMED(FInstanceDataObjectUtilsTest, "CoreUObject::InstanceDataObjectUt
 	CHECK(WasPropertySetBySerialization(StructProperty->Struct, StructData, BProperty));
 	CHECK_FALSE(WasPropertySetBySerialization(StructProperty->Struct, StructData, CProperty));
 	CHECK_FALSE(WasPropertySetBySerialization(StructProperty->Struct, StructData, DProperty));
+
+	CHECK_FALSE(IsPropertyValueInitialized(StructProperty->Struct, StructData, AProperty));
+	SetPropertyValueInitialized(StructProperty->Struct, StructData, AProperty);
+	CHECK(IsPropertyValueInitialized(StructProperty->Struct, StructData, AProperty));
+	CHECK_FALSE(IsPropertyValueInitialized(StructProperty->Struct, StructData, BProperty));
+	SetPropertyValueInitialized(StructProperty->Struct, StructData, BProperty);
+	CHECK(IsPropertyValueInitialized(StructProperty->Struct, StructData, BProperty));
+	CHECK_FALSE(IsPropertyValueInitialized(StructProperty->Struct, StructData, CProperty));
+	CHECK_FALSE(IsPropertyValueInitialized(StructProperty->Struct, StructData, DProperty));
+	ClearPropertyValueInitialized(StructProperty->Struct, StructData, AProperty);
+	ClearPropertyValueInitialized(StructProperty->Struct, StructData, BProperty);
+	CHECK_FALSE(IsPropertyValueInitialized(StructProperty->Struct, StructData, AProperty));
+	CHECK_FALSE(IsPropertyValueInitialized(StructProperty->Struct, StructData, BProperty));
+
+	SetPropertyValueInitialized(StructProperty->Struct, StructData, AProperty);
+	SetPropertyValueInitialized(StructProperty->Struct, StructData, DProperty);
+	ResetPropertyValueInitialized(StructProperty->Struct, StructData);
+	CHECK_FALSE(IsPropertyValueInitialized(StructProperty->Struct, StructData, AProperty));
+	CHECK_FALSE(IsPropertyValueInitialized(StructProperty->Struct, StructData, BProperty));
+	CHECK_FALSE(IsPropertyValueInitialized(StructProperty->Struct, StructData, CProperty));
+	CHECK_FALSE(IsPropertyValueInitialized(StructProperty->Struct, StructData, DProperty));
+}
+
+TEST_CASE_NAMED(FITrackInitializedPropertiesTest, "CoreUObject::Serialization::TrackInitializedProperties", "[Core][UObject][EngineFilter]")
+{
+	UTestInstanceDataObjectClass* BaseObject = NewObject<UTestInstanceDataObjectClass>();
+	UClass* TestClass = CreateInstanceDataObjectClass(nullptr, BaseObject->GetClass(), BaseObject->GetOuter());
+
+	FStructProperty* StructProperty = FindFProperty<FStructProperty>(TestClass, TEXT("Struct"));
+	REQUIRE(StructProperty);
+	REQUIRE(StructProperty->Struct);
+
+	FIntProperty* AProperty = FindFProperty<FIntProperty>(StructProperty->Struct, TEXT("A"));
+	FIntProperty* BProperty = FindFProperty<FIntProperty>(StructProperty->Struct, TEXT("B"));
+	FIntProperty* CProperty = FindFProperty<FIntProperty>(StructProperty->Struct, TEXT("C"));
+	FIntProperty* DProperty = FindFProperty<FIntProperty>(StructProperty->Struct, TEXT("D"));
+	REQUIRE(AProperty);
+	REQUIRE(BProperty);
+	REQUIRE(CProperty);
+	REQUIRE(DProperty);
+
+	void* DefaultStructData = StructProperty->AllocateAndInitializeValue();
+	void* StructData = StructProperty->AllocateAndInitializeValue();
+	ON_SCOPE_EXIT
+	{
+		StructProperty->DestroyAndFreeValue(StructData);
+		StructProperty->DestroyAndFreeValue(DefaultStructData);
+	};
+
+	AProperty->SetPropertyValue_InContainer(DefaultStructData, -1);
+	BProperty->SetPropertyValue_InContainer(DefaultStructData, -1);
+	CProperty->SetPropertyValue_InContainer(DefaultStructData, -1);
+	DProperty->SetPropertyValue_InContainer(DefaultStructData, -1);
+
+	AProperty->SetPropertyValue_InContainer(StructData, 1);
+	BProperty->SetPropertyValue_InContainer(StructData, 2);
+	CProperty->SetPropertyValue_InContainer(StructData, 3);
+	DProperty->SetPropertyValue_InContainer(StructData, -1);
+
+	SetPropertyValueInitialized(StructProperty->Struct, StructData, AProperty);
+	SetPropertyValueInitialized(StructProperty->Struct, StructData, BProperty);
+	SetPropertyValueInitialized(StructProperty->Struct, StructData, DProperty);
+
+	FUObjectSerializeContext* SerializeContext = FUObjectThreadContext::Get().GetSerializeContext();
+	TGuardValue<bool> TrackInitializedPropertiesScope(SerializeContext->bTrackInitializedProperties, true);
+
+	TArray<uint8> BinaryData;
+	{
+		FMemoryWriter Ar(BinaryData, /*bIsPersistent*/ true);
+		FBinaryArchiveFormatter Formatter(Ar);
+		FStructuredArchive StructuredAr(Formatter);
+		StructProperty->Struct->SerializeTaggedProperties(StructuredAr.Open(), (uint8*)StructData, StructProperty->Struct, (uint8*)DefaultStructData);
+	}
+
+#if WITH_TEXT_ARCHIVE_SUPPORT
+	TArray<uint8> JsonData;
+	{
+		FMemoryWriter Ar(JsonData, /*bIsPersistent*/ true);
+		FJsonArchiveOutputFormatter Formatter(Ar);
+		FStructuredArchive StructuredAr(Formatter);
+		StructProperty->Struct->SerializeTaggedProperties(StructuredAr.Open(), (uint8*)StructData, StructProperty->Struct, (uint8*)DefaultStructData);
+	}
+#endif // WITH_TEXT_ARCHIVE_SUPPORT
+
+	AProperty->SetPropertyValue_InContainer(StructData, 4);
+	BProperty->SetPropertyValue_InContainer(StructData, 4);
+	CProperty->SetPropertyValue_InContainer(StructData, 4);
+	DProperty->SetPropertyValue_InContainer(StructData, 4);
+
+	ResetPropertyValueInitialized(StructProperty->Struct, StructData);
+
+	{
+		FMemoryReader Ar(BinaryData, /*bIsPersistent*/ true);
+		FBinaryArchiveFormatter Formatter(Ar);
+		FStructuredArchive StructuredAr(Formatter);
+		StructProperty->Struct->SerializeTaggedProperties(StructuredAr.Open(), (uint8*)StructData, StructProperty->Struct, (uint8*)DefaultStructData);
+	}
+
+	CHECK(AProperty->GetPropertyValue_InContainer(StructData) == 1);
+	CHECK(BProperty->GetPropertyValue_InContainer(StructData) == 2);
+	CHECK(CProperty->GetPropertyValue_InContainer(StructData) == 4); // C unchanged because it is not initialized
+	CHECK(DProperty->GetPropertyValue_InContainer(StructData) == 4); // D unchanged because it was serialized without its value
+
+	CHECK(IsPropertyValueInitialized(StructProperty->Struct, StructData, AProperty));
+	CHECK(IsPropertyValueInitialized(StructProperty->Struct, StructData, BProperty));
+	CHECK_FALSE(IsPropertyValueInitialized(StructProperty->Struct, StructData, CProperty));
+	CHECK(IsPropertyValueInitialized(StructProperty->Struct, StructData, DProperty));
+
+#if WITH_TEXT_ARCHIVE_SUPPORT
+	AProperty->SetPropertyValue_InContainer(StructData, 4);
+	BProperty->SetPropertyValue_InContainer(StructData, 4);
+	CProperty->SetPropertyValue_InContainer(StructData, 4);
+	DProperty->SetPropertyValue_InContainer(StructData, 4);
+
+	ResetPropertyValueInitialized(StructProperty->Struct, StructData);
+
+	{
+		FMemoryReader Ar(JsonData, /*bIsPersistent*/ true);
+		FJsonArchiveInputFormatter Formatter(Ar);
+		FStructuredArchive StructuredAr(Formatter);
+		StructProperty->Struct->SerializeTaggedProperties(StructuredAr.Open(), (uint8*)StructData, StructProperty->Struct, (uint8*)DefaultStructData);
+	}
+
+	CHECK(AProperty->GetPropertyValue_InContainer(StructData) == 1);
+	CHECK(BProperty->GetPropertyValue_InContainer(StructData) == 2);
+	CHECK(CProperty->GetPropertyValue_InContainer(StructData) == 4); // C unchanged because it is not initialized
+	CHECK(DProperty->GetPropertyValue_InContainer(StructData) == 4); // D unchanged because it was serialized without its value
+
+	CHECK(IsPropertyValueInitialized(StructProperty->Struct, StructData, AProperty));
+	CHECK(IsPropertyValueInitialized(StructProperty->Struct, StructData, BProperty));
+	CHECK_FALSE(IsPropertyValueInitialized(StructProperty->Struct, StructData, CProperty));
+	CHECK(IsPropertyValueInitialized(StructProperty->Struct, StructData, DProperty));
+#endif // WITH_TEXT_ARCHIVE_SUPPORT
 }
 
 } // UE
