Make Coreredirect warning referring users to the documentation a bit clearer by moving the direction to a new, indented line

#rnx

diff --git a/Engine/Source/Runtime/CoreUObject/Private/UObject/CoreRedirects.cpp b/Engine/Source/Runtime/CoreUObject/Private/UObject/CoreRedirects.cpp
--- a/Engine/Source/Runtime/CoreUObject/Private/UObject/CoreRedirects.cpp
+++ b/Engine/Source/Runtime/CoreUObject/Private/UObject/CoreRedirects.cpp
@@ -1580,7 +1580,7 @@ bool FCoreRedirects::ReadRedirectsFromIni(const FString& IniName)
 					if (bMatchWildcard || bMatchSubstring)
 					{
 						UE_CLOG(bMatchSubstring, LogCoreRedirects, Warning, TEXT("ReadRedirectsFromIni(%s) 'MatchSubstring=' is deprecated. "
-							"Please prefer `MatchWildcard=' instead for redirect %s. "
+							"Please prefer `MatchWildcard=' instead for redirect %s.\n\t"
 							"For more information refer to the documentation in Engine/Config/BaseEngine.ini."), *IniName, *ValueString);
 
 						constexpr FStringView WildcardMarker(TEXTVIEW("..."));
