IterativeCooker: Add build dependency on CoreRedirects that affect a package's runtime dependencies. The names of runtime dependencies are stored in the package's imports, so if the change names due to a redirect then we need to recook them.
#rnx

diff --git a/Engine/Source/Editor/UnrealEd/Private/TargetDomain/TargetDomainUtils.cpp b/Engine/Source/Editor/UnrealEd/Private/TargetDomain/TargetDomainUtils.cpp
--- a/Engine/Source/Editor/UnrealEd/Private/TargetDomain/TargetDomainUtils.cpp
+++ b/Engine/Source/Editor/UnrealEd/Private/TargetDomain/TargetDomainUtils.cpp
@@ -31,6 +31,7 @@
 #include "Serialization/CompactBinarySerialization.h"
 #include "Serialization/CompactBinaryWriter.h"
 #include "Serialization/PackageWriter.h"
+#include "UObject/CoreRedirects.h"
 #include "ZenStoreHttpClient.h"
 
 namespace UE::TargetDomain
@@ -169,6 +170,9 @@ bool FCookDependencies::TryCalculateCurrentKey(const FAssetPackageData* Override
 		KeyBuilder.Update(&PackageDigest.Hash, sizeof(PackageDigest.Hash));
 	}
 
+	FCoreRedirects::AppendHashOfRedirectsAffectingPackages(KeyBuilder, RuntimePackageDependencies);
+	FCoreRedirects::AppendHashOfRedirectsAffectingPackages(KeyBuilder, ScriptPackageDependencies);
+
 	if (!ConfigDependencies.IsEmpty())
 	{
 #if UE_WITH_CONFIG_TRACKING
@@ -230,6 +234,7 @@ void FCookDependencies::Reset()
 	BuildPackageDependencies.Reset();
 	ConfigDependencies.Reset();
 	RuntimePackageDependencies.Reset();
+	ScriptPackageDependencies.Reset();
 	PackageName = FName();
 	StoredKey = FIoHash::Zero;
 	CurrentKey = FIoHash::Zero;
@@ -244,34 +249,40 @@ void FCookDependencies::Empty()
 	RuntimePackageDependencies.Empty();
 }
 
+enum class EPackageMountPoint
+{
+	Transient,
+	Script,
+	Content,
+	GeneratedContent,
+};
+EPackageMountPoint GetPackageMountPoint(FName PackageName, FName TransientPackageName)
+{
+	if (PackageName == TransientPackageName)
+	{
+		return EPackageMountPoint::Transient;
+	}
+	TStringBuilder<256> StringBuffer;
+	PackageName.ToString(StringBuffer);
+	if (FPackageName::IsMemoryPackage(StringBuffer))
+	{
+		return EPackageMountPoint::Transient;
+	}
+	if (FPackageName::IsScriptPackage(StringBuffer))
+	{
+		return EPackageMountPoint::Script;
+	}
+	if (ICookPackageSplitter::IsUnderGeneratedPackageSubPath(StringBuffer))
+	{
+		return EPackageMountPoint::GeneratedContent;
+	}
+	return EPackageMountPoint::Content;
+}
+
 FCookDependencies FCookDependencies::Collect(UPackage* Package, const ITargetPlatform* TargetPlatform,
 	FSavePackageResultStruct* SaveResult, const FGeneratedPackageResultStruct* GeneratedResult, 
 	TArray<FName>&& RuntimeDependencies, FString* OutErrorMessage)
 {
-	TStringBuilder<256> StringBuffer;
-	FName TransientPackageName = GetTransientPackage()->GetFName();
-	auto IsDisallowedDependencyName = [&StringBuffer, TransientPackageName](FName InPackageName,
-		bool bAllowGeneratedPackages)
-		{
-			if (InPackageName == TransientPackageName)
-			{
-				return true;
-			}
-			InPackageName.ToString(StringBuffer);
-			FStringView PackageNameStr = StringBuffer.ToView();
-			return FPackageName::IsMemoryPackage(PackageNameStr) ||
-				FPackageName::IsScriptPackage(PackageNameStr) ||
-				(!bAllowGeneratedPackages && ICookPackageSplitter::IsUnderGeneratedPackageSubPath(PackageNameStr));
-		};
-	auto IsDisallowedBuildDependencyName = [&IsDisallowedDependencyName](FName InPackageName)
-		{
-			return IsDisallowedDependencyName(InPackageName, false /* bAllowGeneratedPackages */);
-		};
-	auto IsDisallowedRuntimeDependencyName = [&IsDisallowedDependencyName](FName InPackageName)
-		{
-			return IsDisallowedDependencyName(InPackageName, true /* bAllowGeneratedPackages */);
-		};
-
 	if (!Package)
 	{
 		if (OutErrorMessage) *OutErrorMessage = TEXT("Invalid null package.");
@@ -382,15 +393,35 @@ FCookDependencies FCookDependencies::Collect(UPackage* Package, const ITargetPla
 		Algo::Sort(Result.TransitiveBuildDependencies);
 	}
 
+	FName TransientPackageName = GetTransientPackage()->GetFName();
 	Result.BuildPackageDependencies = BuildDependenciesSet.Array();
-	Result.BuildPackageDependencies.RemoveAllSwap(IsDisallowedBuildDependencyName, EAllowShrinking::Yes);
+	Result.BuildPackageDependencies.RemoveAllSwap([TransientPackageName](FName InPackageName)
+		{
+			return GetPackageMountPoint(InPackageName, TransientPackageName) != EPackageMountPoint::Content;
+		}, EAllowShrinking::Yes);
 	Result.BuildPackageDependencies.Sort(FNameLexicalLess());
 
-	RuntimeDependencies.RemoveAllSwap(IsDisallowedRuntimeDependencyName, EAllowShrinking::No);
-	RuntimeDependencies.Sort(FNameLexicalLess());
-	RuntimeDependencies.SetNum(Algo::Unique(RuntimeDependencies), EAllowShrinking::Yes);
-
-	Result.RuntimePackageDependencies = MoveTemp(RuntimeDependencies);
+	for (TArray<FName>::TIterator Iter(RuntimeDependencies); Iter; ++Iter)
+	{
+		FName PackageName = *Iter;
+		EPackageMountPoint MountPoint = GetPackageMountPoint(PackageName, TransientPackageName);
+		switch (MountPoint)
+		{
+		case EPackageMountPoint::GeneratedContent:
+		case EPackageMountPoint::Content:
+			Result.RuntimePackageDependencies.Add(PackageName);
+			break;
+		case EPackageMountPoint::Script:
+			Result.ScriptPackageDependencies.Add(PackageName);
+			break;
+		default:
+			break;
+		}
+	}
+	Result.RuntimePackageDependencies.Sort(FNameLexicalLess());
+	Result.RuntimePackageDependencies.SetNum(Algo::Unique(Result.RuntimePackageDependencies), EAllowShrinking::Yes);
+	Result.ScriptPackageDependencies.Sort(FNameLexicalLess());
+	Result.ScriptPackageDependencies.SetNum(Algo::Unique(Result.ScriptPackageDependencies), EAllowShrinking::Yes);
 
 	const FAssetPackageData* AssetPackageData = GeneratedResult ? &GeneratedResult->AssetPackageData : nullptr;
 	if (!Result.TryCalculateCurrentKey(AssetPackageData, OutErrorMessage))
@@ -452,6 +483,13 @@ bool LoadFromCompactBinary(FCbObjectView ObjectView, UE::TargetDomain::FCookDepe
 				return false;
 			}
 		}
+		if (FieldView.GetName().Equals(UTF8TEXTVIEW("ScriptPackageDependencies")))
+		{
+			if (!LoadFromCompactBinary(FieldView++, Dependencies.ScriptPackageDependencies))
+			{
+				return false;
+			}
+		}
 		if (FieldView.GetName().Equals(UTF8TEXTVIEW("CookDependencies")))
 		{
 			if (!LoadFromCompactBinary(FieldView++, Dependencies.CookDependencies))
@@ -498,6 +536,10 @@ FCbWriter& operator<<(FCbWriter& Writer, const UE::TargetDomain::FCookDependenci
 	{
 		Writer << "RuntimePackageDependencies" << CookDependencies.RuntimePackageDependencies;
 	}
+	if (!CookDependencies.ScriptPackageDependencies.IsEmpty())
+	{
+		Writer << "ScriptPackageDependencies" << CookDependencies.ScriptPackageDependencies;
+	}
 	if (!CookDependencies.CookDependencies.IsEmpty())
 	{
 		Writer << "CookDependencies" << CookDependencies.CookDependencies;
diff --git a/Engine/Source/Editor/UnrealEd/Private/TargetDomain/TargetDomainUtils.h b/Engine/Source/Editor/UnrealEd/Private/TargetDomain/TargetDomainUtils.h
--- a/Engine/Source/Editor/UnrealEd/Private/TargetDomain/TargetDomainUtils.h
+++ b/Engine/Source/Editor/UnrealEd/Private/TargetDomain/TargetDomainUtils.h
@@ -72,6 +72,8 @@ public:
 
 	// Runtime Dependencies
 	const TArray<FName>& GetRuntimePackageDependencies() const;
+	// Script dependencies, needed for TryCalculateCurrentKey
+	const TArray<FName>& GetScriptPackageDependencies() const;
 
 	// Cook Dependencies
 	const TArray<UE::Cook::FCookDependency>& GetCookDependencies() const;
@@ -109,6 +111,7 @@ private:
 	TArray<FName> BuildPackageDependencies;
 	TArray<FString> ConfigDependencies;
 	TArray<FName> RuntimePackageDependencies;
+	TArray<FName> ScriptPackageDependencies;
 	TArray<UE::Cook::FCookDependency> CookDependencies;
 	TArray<UE::Cook::FCookDependency> TransitiveBuildDependencies;
 	FName PackageName;
@@ -207,6 +210,11 @@ inline const TArray<FName>& FCookDependencies::GetRuntimePackageDependencies() c
 	return RuntimePackageDependencies;
 }
 
+inline const TArray<FName>& FCookDependencies::GetScriptPackageDependencies() const
+{
+	return ScriptPackageDependencies;
+}
+
 inline const TArray<UE::Cook::FCookDependency>& FCookDependencies::GetCookDependencies() const
 {
 	return CookDependencies;
diff --git a/Engine/Source/Runtime/CoreUObject/Private/UObject/CoreRedirects.cpp b/Engine/Source/Runtime/CoreUObject/Private/UObject/CoreRedirects.cpp
--- a/Engine/Source/Runtime/CoreUObject/Private/UObject/CoreRedirects.cpp
+++ b/Engine/Source/Runtime/CoreUObject/Private/UObject/CoreRedirects.cpp
@@ -9,12 +9,14 @@
 #include "UObject/UnrealType.h"
 
 #include "GenericPlatform/GenericPlatformFile.h"
+#include "Hash/Blake3.h"
 #include "Misc/AutomationTest.h"
 #include "Misc/CommandLine.h"
 #include "Misc/ConfigCacheIni.h"
 #include "Misc/CoreDelegates.h"
 #include "Misc/PackageName.h"
 #include "Misc/ScopeRWLock.h"
+#include "RedirectionSummary.h"
 #include "Serialization/DeferredMessageLog.h"
 #include "Templates/Casts.h"
 
@@ -225,6 +227,10 @@ namespace UE::CoreRedirects::Private
 	}
 }
 
+#if WITH_EDITOR
+FRedirectionSummary GRedirectionSummary;
+#endif
+
 FCoreRedirectObjectName::FCoreRedirectObjectName(const FTopLevelAssetPath& TopLevelAssetPath)
 	: FCoreRedirectObjectName(TopLevelAssetPath.GetAssetName(), NAME_None, TopLevelAssetPath.GetPackageName())
 {
@@ -930,6 +936,151 @@ bool FCoreRedirect::IdenticalMatchRules(const FCoreRedirect& Other) const
 	return RedirectFlags == Other.RedirectFlags && OldName == Other.OldName;
 }
 
+void FCoreRedirect::AppendHash(FBlake3& Hasher) const
+{
+	Hasher.Update(&RedirectFlags, sizeof(RedirectFlags));
+	OldName.AppendHash(Hasher);
+	NewName.AppendHash(Hasher);
+	OverrideClassName.AppendHash(Hasher);
+	TArray<TPair<FString, FString>> ValueArray = ValueChanges.Array();
+	Algo::Sort(ValueArray);
+	for (const TPair<FString, FString>& Pair : ValueArray)
+	{
+		Hasher.Update(*Pair.Key, Pair.Key.Len() * sizeof((*Pair.Key)[0]));
+		Hasher.Update(*Pair.Value, Pair.Value.Len() * sizeof((*Pair.Value)[0]));
+	}
+}
+
+template <typename KeyType, typename ValueType>
+int DeterministicCompare(const TMap<KeyType, ValueType>& A, const TMap<KeyType, ValueType>& B)
+{
+	if (A.Num() != B.Num())
+	{
+		return A.Num() < B.Num() ? -1 : 1;
+	}
+	if (A.Num() == 0)
+	{
+		return 0;
+	}
+
+	const KeyType* MinKeyWithDifference = nullptr;
+	bool bMinKeyIsLessInA = false;
+	for (const TPair<KeyType, ValueType>& Pair : A)
+	{
+		const ValueType* BValue = B.Find(Pair.Key);
+		int Compare = 0;
+		if (!BValue)
+		{
+			Compare = -1;
+		}
+		else if (Pair.Value < *BValue)
+		{
+			Compare = -1;
+		}
+		else if (*BValue < Pair.Value)
+		{
+			Compare = 1;
+		}
+		if (Compare != 0)
+		{
+			if (!MinKeyWithDifference || *Pair.Key < *MinKeyWithDifference)
+			{
+				MinKeyWithDifference = &Pair.Key;
+				bMinKeyIsLessInA = Compare < 0;
+			}
+		}
+	}
+	for (const TPair<KeyType, ValueType>& Pair : B)
+	{
+		if (!A.Contains(Pair.Key))
+		{
+			if (!MinKeyWithDifference || *Pair.Key < *MinKeyWithDifference)
+			{
+				MinKeyWithDifference = &Pair.Key;
+				bMinKeyIsLessInA = false;
+			}
+		}
+	}
+	if (MinKeyWithDifference)
+	{
+		return bMinKeyIsLessInA ? -1 : 1;
+	}
+	return 0;
+}
+
+int FCoreRedirect::Compare(const FCoreRedirect& Other) const
+{
+	if (RedirectFlags != Other.RedirectFlags)
+	{
+		return RedirectFlags < Other.RedirectFlags ? -1 : 1;
+	}
+	int Compare = OldName.Compare(Other.OldName);
+	if (Compare != 0)
+	{
+		return Compare;
+	}
+	Compare = NewName.Compare(Other.NewName);
+	if (Compare != 0)
+	{
+		return Compare;
+	}
+	Compare = OverrideClassName.Compare(Other.OverrideClassName);
+	if (Compare != 0)
+	{
+		return Compare;
+	}
+	Compare = DeterministicCompare(ValueChanges, Other.ValueChanges);
+	if (Compare != 0)
+	{
+		return Compare;
+	}
+
+	return 0;
+}
+
+void FCoreRedirectObjectName::AppendHash(FBlake3& Hasher) const
+{
+	FNameBuilder NameStr;
+	int32 Marker = 0xabacadab;
+	if (!PackageName.IsNone())
+	{
+		NameStr << PackageName;
+		Hasher.Update(NameStr.GetData(), NameStr.Len() * sizeof(NameStr.GetData()[0]));
+	}
+	Hasher.Update(&Marker, sizeof(Marker));
+	if (!OuterName.IsNone())
+	{
+		NameStr.Reset();
+		NameStr << OuterName;
+		Hasher.Update(NameStr.GetData(), NameStr.Len() * sizeof(NameStr.GetData()[0]));
+	}
+	Hasher.Update(&Marker, sizeof(Marker));
+	if (!ObjectName.IsNone())
+	{
+		NameStr.Reset();
+		NameStr << ObjectName;
+		Hasher.Update(NameStr.GetData(), NameStr.Len() * sizeof(NameStr.GetData()[0]));
+	}
+	Hasher.Update(&Marker, sizeof(Marker));
+}
+
+int FCoreRedirectObjectName::Compare(const FCoreRedirectObjectName& Other) const
+{
+	if (PackageName != Other.PackageName)
+	{
+		return PackageName.Compare(Other.PackageName);
+	}
+	if (OuterName != Other.OuterName)
+	{
+		return OuterName.Compare(Other.OuterName);
+	}
+	if (ObjectName != Other.ObjectName)
+	{
+		return ObjectName.Compare(Other.ObjectName);
+	}
+	return 0;
+}
+
 bool FCoreRedirects::bInitialized = false;
 bool FCoreRedirects::bInDebugMode = false;
 bool FCoreRedirects::bValidatedOnce = false;
@@ -1314,6 +1465,18 @@ void FCoreRedirects::ClearKnownMissing(ECoreRedirectFlags Type, ECoreRedirectFla
 	}
 }
 
+#if WITH_EDITOR
+void FCoreRedirects::AppendHashOfRedirectsAffectingPackages(FBlake3& Hasher, TConstArrayView<FName> PackageNames)
+{
+	GRedirectionSummary.AppendHashAffectingPackages(Hasher, PackageNames);
+}
+
+void FCoreRedirects::AppendHashOfGlobalRedirects(FBlake3& Hasher)
+{
+	GRedirectionSummary.AppendHashGlobal(Hasher);
+}
+#endif
+
 bool FCoreRedirects::RunTests()
 {
 	bool bSuccess = true;
@@ -1324,6 +1487,10 @@ bool FCoreRedirects::RunTests()
 	bIsInMultithreadedPhase = false;
 #endif
 	RedirectTypeMap.Empty();
+#if WITH_EDITOR
+	FRedirectionSummary BackupSummary = MoveTemp(GRedirectionSummary);
+	GRedirectionSummary = FRedirectionSummary();
+#endif
 
 	TArray<FCoreRedirect> NewRedirects;
 
@@ -1514,6 +1681,9 @@ bool FCoreRedirects::RunTests()
 
 	// Restore old state
 	RedirectTypeMap = MoveTemp(BackupMap);
+#if WITH_EDITOR
+	GRedirectionSummary = MoveTemp(BackupSummary);
+#endif
 #if WITH_COREREDIRECTS_MULTITHREAD_WARNING
 	bIsInMultithreadedPhase = BackupIsInMultithreadedPhase;
 #endif
@@ -1825,6 +1995,9 @@ bool FCoreRedirects::AddSingleRedirect(const FCoreRedirect& NewRedirect, const F
 	{
 		ExistingRedirects->Add(NewRedirect);
 	}
+#if WITH_EDITOR
+	GRedirectionSummary.Add(NewRedirect, bIsWildcardMatch);
+#endif
 
 	return true;
 }
@@ -1926,11 +2099,17 @@ bool FCoreRedirects::RemoveSingleRedirect(const FCoreRedirect& RedirectToRemove,
 		}
 	}
 
-	if (bIsWildcardMatch && bRemovedRedirect)
+	if (bRemovedRedirect)
 	{
-		// We removed a substring redirect so we need to regenerate our prediction 
-		// tables to avoid unnecessary false positives
-		ExistingNameMap->Wildcards->Rebuild();
+		if (bIsWildcardMatch)
+		{
+			// We removed a substring redirect so we need to regenerate our prediction 
+			// tables to avoid unnecessary false positives
+			ExistingNameMap->Wildcards->Rebuild();
+		}
+#if WITH_EDITOR
+		GRedirectionSummary.Remove(RedirectToRemove, bIsWildcardMatch);
+#endif
 	}
 
 	return bRemovedRedirect;
diff --git a/Engine/Source/Runtime/CoreUObject/Private/UObject/RedirectionSummary.cpp b/Engine/Source/Runtime/CoreUObject/Private/UObject/RedirectionSummary.cpp
new file mode 100644
--- /dev/null
+++ b/Engine/Source/Runtime/CoreUObject/Private/UObject/RedirectionSummary.cpp
@@ -0,0 +1,240 @@
+// Copyright Epic Games, Inc. All Rights Reserved.
+
+#include "RedirectionSummary.h"
+
+#if WITH_EDITOR
+
+#include "Misc/ScopeRWLock.h"
+
+FRedirectionSummary::FRedirectionSummary(FRedirectionSummary&& Other)
+{
+	*this = MoveTemp(Other);
+}
+
+FRedirectionSummary& FRedirectionSummary::operator=(FRedirectionSummary&& Other)
+{
+	if (&Other == this)
+	{
+		return *this;
+	}
+
+	// Technical correctness for holding multiple critical sections: enter the locks in the same order
+	// from all possible callsites, in this case, from two calls to operator=.
+	TOptional<FWriteScopeLock> ScopeA;
+	TOptional<FWriteScopeLock> ScopeB;
+	if (&Other < this)
+	{
+		ScopeA.Emplace(Other.Lock);
+		ScopeB.Emplace(Lock);
+	}
+	else
+	{
+		ScopeA.Emplace(Lock);
+		ScopeB.Emplace(Other.Lock);
+	}
+
+	RedirectsForPackage = MoveTemp(Other.RedirectsForPackage);
+	GlobalRedirects = MoveTemp(Other.GlobalRedirects);
+	Other.RedirectsForPackage.Empty();
+	Other.GlobalRedirects.Empty();
+
+	return *this;
+}
+
+void FRedirectionSummary::Add(const FCoreRedirect& CoreRedirect, bool bIsWildcardMatch)
+{
+	FWriteScopeLock ScopeLock(Lock);
+
+	if (bIsWildcardMatch || CoreRedirect.OldName.PackageName.IsNone())
+	{
+		GlobalRedirects.Add(FCoreRedirect(CoreRedirect));
+	}
+	else
+	{
+		for (FName PackageName : GetAffectedPackages(CoreRedirect))
+		{
+			RedirectsForPackage.FindOrAdd(PackageName).Add(FCoreRedirect(CoreRedirect));
+		}
+	}
+}
+
+void FRedirectionSummary::Remove(const FCoreRedirect& CoreRedirect, bool bIsWildcardMatch)
+{
+	FWriteScopeLock ScopeLock(Lock);
+
+	if (bIsWildcardMatch || CoreRedirect.OldName.PackageName.IsNone())
+	{
+		GlobalRedirects.Remove(CoreRedirect);
+	}
+	else
+	{
+		for (FName PackageName : GetAffectedPackages(CoreRedirect))
+		{
+			FRedirectContainer* Container = RedirectsForPackage.Find(PackageName);
+			if (Container)
+			{
+				Container->Remove(CoreRedirect);
+				if (Container->IsEmpty())
+				{
+					RedirectsForPackage.Remove(PackageName);
+				}
+			}
+		}
+	}
+}
+
+void FRedirectionSummary::AppendHashAffectingPackages(FBlake3& Hasher, TConstArrayView<FName> PackageNames)
+{
+	if (PackageNames.IsEmpty())
+	{
+		return;
+	}
+
+	int32 Index = -1;
+	const FName* PackageNamesData = PackageNames.GetData();
+	int32 PackageNamesNum = PackageNames.Num();
+	bool bNeedWriteLock = false;
+
+	{
+		FReadScopeLock ScopeLock(Lock);
+		if (GlobalRedirects.TryAppendHashInReadLock(Hasher))
+		{
+			Index = 0;
+		}
+		else
+		{
+			bNeedWriteLock = true;
+		}
+		while (!bNeedWriteLock && Index < PackageNamesNum)
+		{
+			FRedirectContainer* Container = RedirectsForPackage.Find(PackageNamesData[Index]);
+			if (!Container || Container->TryAppendHashInReadLock(Hasher))
+			{
+				++Index;
+			}
+			else
+			{
+				bNeedWriteLock = true;
+			}
+		}
+	}
+	if (!bNeedWriteLock)
+	{
+		return;
+	}
+
+	FWriteScopeLock WriteLock(Lock);
+	if (Index < 0)
+	{
+		GlobalRedirects.AppendHashInWriteLock(Hasher);
+		Index = 0;
+	}
+	while (Index < PackageNamesNum)
+	{
+		FRedirectContainer* Container = RedirectsForPackage.Find(PackageNamesData[Index]);
+		if (Container)
+		{
+			Container->AppendHashInWriteLock(Hasher);
+		}
+		++Index;
+	}
+}
+
+void FRedirectionSummary::AppendHashGlobal(FBlake3& Hasher)
+{
+	{
+		FReadScopeLock ScopeLock(Lock);
+		if (GlobalRedirects.TryAppendHashInReadLock(Hasher))
+		{
+			return;
+		}
+	}
+	FWriteScopeLock WriteLock(Lock);
+	GlobalRedirects.AppendHashInWriteLock(Hasher);
+}
+
+void FRedirectionSummary::FRedirectContainer::Add(FCoreRedirect&& Redirect)
+{
+	Redirects.Add(MoveTemp(Redirect));
+	bHashDirty = true;
+}
+
+void FRedirectionSummary::FRedirectContainer::Remove(const FCoreRedirect& Redirect)
+{
+	Redirects.Remove(Redirect);
+	bHashDirty = true;
+}
+
+bool FRedirectionSummary::FRedirectContainer::IsEmpty() const
+{
+	return Redirects.IsEmpty();
+}
+
+void FRedirectionSummary::FRedirectContainer::Empty()
+{
+	Redirects.Empty();
+	Hash = FBlake3Hash();
+	bHashDirty = false;
+}
+
+bool FRedirectionSummary::FRedirectContainer::TryAppendHashInReadLock(FBlake3& Hasher) const
+{
+	if (bHashDirty)
+	{
+		return false;
+	}
+	AppendHashWithoutDirtyCheck(Hasher);
+	return true;
+}
+
+void FRedirectionSummary::FRedirectContainer::AppendHashInWriteLock(FBlake3& Hasher)
+{
+	if (bHashDirty)
+	{
+		CalculateHash();
+		bHashDirty = false;
+	}
+	AppendHashWithoutDirtyCheck(Hasher);
+}
+
+void FRedirectionSummary::FRedirectContainer::CalculateHash()
+{
+	FBlake3 Hasher;
+	for (const TPair<FCoreRedirect,bool>& Pair: Redirects)
+	{
+		const FCoreRedirect& Redirect = Pair.Key;
+		Redirect.AppendHash(Hasher);
+	}
+	Hash = Hasher.Finalize();
+}
+
+void FRedirectionSummary::FRedirectContainer::AppendHashWithoutDirtyCheck(FBlake3& Hasher) const
+{
+	Hasher.Update(&Hash.GetBytes(), sizeof(Hash.GetBytes()));
+}
+
+TArray<FName, TInlineAllocator<2>> FRedirectionSummary::GetAffectedPackages(const FCoreRedirect& Redirect)
+{
+	// Note on why we need the redirects that redirect both from and to the package:
+	// When a redirector from X to Y changes to redirect from X to Z, package A that references X will now need to
+	// write Z instead of Y into its saved imports.
+	// When a redirector from X to Y changes to redirect from W to Y, or is deleted, package A that references X will
+	// now need to write X instead of Y into its saved imports.
+	TArray<FName, TInlineAllocator<2>> Result;
+	if (!Redirect.OldName.PackageName.IsNone())
+	{
+		Result.Add(Redirect.OldName.PackageName);
+	}
+	if (!Redirect.NewName.PackageName.IsNone() && Redirect.NewName.PackageName != Redirect.OldName.PackageName)
+	{
+		Result.Add(Redirect.NewName.PackageName);
+	}
+	return Result;
+}
+
+bool FRedirectionSummary::FCompareRedirect::operator()(const FCoreRedirect& A, const FCoreRedirect& B)
+{
+	return A.Compare(B) < 0;
+}
+
+#endif
\ No newline at end of file
diff --git a/Engine/Source/Runtime/CoreUObject/Private/UObject/RedirectionSummary.h b/Engine/Source/Runtime/CoreUObject/Private/UObject/RedirectionSummary.h
new file mode 100644
--- /dev/null
+++ b/Engine/Source/Runtime/CoreUObject/Private/UObject/RedirectionSummary.h
@@ -0,0 +1,70 @@
+// Copyright Epic Games, Inc. All Rights Reserved.
+
+#pragma once
+
+#if WITH_EDITOR
+
+#include "Containers/Array.h"
+#include "Containers/ArrayView.h"
+#include "Containers/Map.h"
+#include "Containers/SortedMap.h"
+#include "HAL/CriticalSection.h"
+#include "Hash/Blake3.h"
+#include "UObject/CoreRedirects.h"
+#include "UObject/NameTypes.h"
+
+/**
+ * Container for FCoreRedirects that can affect a package. Used by class FCoreRedirects to implement
+ * AppendHashOfRedirectsAffectingPackages.
+ */
+class FRedirectionSummary
+{
+public:
+	FRedirectionSummary() = default;
+	FRedirectionSummary(FRedirectionSummary&& Other);
+	FRedirectionSummary& operator=(FRedirectionSummary&& Other);
+
+	void Add(const FCoreRedirect& CoreRedirect, bool bIsWildcardMatch);
+	void Remove(const FCoreRedirect& CoreRedirect, bool bIsWildcardMatch);
+	void AppendHashAffectingPackages(FBlake3& Hasher, TConstArrayView<FName> PackageNames);
+	void AppendHashGlobal(FBlake3& Hasher);
+
+private:
+	struct FCompareRedirect
+	{
+		bool operator()(const FCoreRedirect& A, const FCoreRedirect& B);
+	};
+	struct FRedirectContainer
+	{
+	public:
+		void Add(FCoreRedirect&& Redirect);
+		void Remove(const FCoreRedirect& Redirect);
+		bool IsEmpty() const;
+		void Empty();
+		bool TryAppendHashInReadLock(FBlake3& Hasher) const;
+		void AppendHashInWriteLock(FBlake3& Hasher);
+	private:
+		void CalculateHash();
+		void AppendHashWithoutDirtyCheck(FBlake3& Hasher) const;
+	private:
+		TSortedMap<FCoreRedirect, bool, FDefaultAllocator, FCompareRedirect> Redirects;
+		FBlake3Hash Hash;
+		bool bHashDirty = false;
+	};
+
+private:
+	static TArray<FName, TInlineAllocator<2>> GetAffectedPackages(const FCoreRedirect& Redirect);
+
+private:
+	TMap<FName, FRedirectContainer> RedirectsForPackage;
+	FRedirectContainer GlobalRedirects;
+	/**
+	 * CoreRedirects are written when the engine is single threaded, but we do not write the hashes for the
+	 * global and per-package containers until they are requested on demand. The requests on demand can
+	 * occur on multiple threads, so we need a lock to ensure that the first thread that requests them and finds
+	 * them dirty does not cause a data race with the second thread requests them.
+	 */
+	FRWLock Lock;
+};
+
+#endif
\ No newline at end of file
diff --git a/Engine/Source/Runtime/CoreUObject/Public/UObject/CoreRedirects.h b/Engine/Source/Runtime/CoreUObject/Public/UObject/CoreRedirects.h
--- a/Engine/Source/Runtime/CoreUObject/Public/UObject/CoreRedirects.h
+++ b/Engine/Source/Runtime/CoreUObject/Public/UObject/CoreRedirects.h
@@ -19,6 +19,7 @@
 #include "UObject/UnrealNames.h"
 #include "UObject/CoreRedirects/PM-k.h"
 
+class FBlake3;
 class IPakFile;
 class UClass;
 struct FTopLevelAssetPath;
@@ -118,6 +119,9 @@ struct FCoreRedirectObjectName
 		return !(*this == Other);
 	}
 
+	/** Compares the two names lexically, returning -,0,+ */
+	COREUOBJECT_API int Compare(const FCoreRedirectObjectName& Other) const;
+
 	/** Flags for the Matches function. These flags overlap but are lower-level than ECoreRedirectMatchFlags. */
 	enum class EMatchFlags
 	{
@@ -177,6 +181,9 @@ struct FCoreRedirectObjectName
 	/** Returns true if all names have valid characters */
 	COREUOBJECT_API bool HasValidCharacters(ECoreRedirectFlags Type) const;
 
+	/** Update Hasher with all fields from this */
+	COREUOBJECT_API void AppendHash(FBlake3& Hasher) const;
+
 	/** Expand OldName/NewName as needed */
 	static COREUOBJECT_API bool ExpandNames(const FStringView FullString, FName& OutName, FName& OutOuter, FName& OutPackage);
 
@@ -250,6 +257,11 @@ struct FCoreRedirect
 		return OldName.GetSearchKey(RedirectFlags);
 	}
 
+	/** Update Hasher with all fields from this */
+	COREUOBJECT_API void AppendHash(FBlake3& Hasher) const;
+	/** Returns -,0,+ based on a full lexical-fnames compare of all fields on the two CoreRedirects. */
+	COREUOBJECT_API int Compare(const FCoreRedirect& Other) const;
+
 private:
 	friend struct FCoreRedirects;
 
@@ -281,7 +293,10 @@ private:
  */
 struct FCoreRedirects
 {
-	/** Run initialization steps that are needed before any data can be stored in FCoreRedirects. Reads can occur before this, but no redirects will exist and redirect queries will all return empty. */
+	/**
+	 * Run initialization steps that are needed before any data can be stored in FCoreRedirects.
+	 * Reads can occur before this, but no redirects will exist and redirect queries will all return empty.
+	 */
 	static COREUOBJECT_API void Initialize();
 
 	/** Returns a redirected version of the object name. If there are no valid redirects, it will return the original name */
@@ -357,6 +372,21 @@ struct FCoreRedirects
 	/** Goes from UClass Type to the type flag */
 	static COREUOBJECT_API ECoreRedirectFlags GetFlagsForTypeClass(UClass *TypeClass);
 
+#if WITH_EDITOR
+	/**
+	 * Iterate the list of PackageNames and append the hash of all redirects that affect the package, either
+	 * redirecting from or to the package. Used in iterative cooking to invalidate the cooked version of packages when
+	 * CoreRedirects change. 
+	 */
+	static COREUOBJECT_API void AppendHashOfRedirectsAffectingPackages(FBlake3& Hasher,
+		TConstArrayView<FName> PackageNames);
+	/**
+	 * Append the hash of all redirects that can affect multiple packages, or for which the affected packages are unknown.
+	 * Used in iterative cooking to invalidate the cooked version of packages when CoreRedirects change.
+	 */
+	static COREUOBJECT_API void AppendHashOfGlobalRedirects(FBlake3& Hasher);
+#endif
+
 	/** Runs set of redirector tests, returns false on failure */
 	static COREUOBJECT_API bool RunTests();
 
