Don't crash when resolving a bogus FFieldPath

#rnx

diff --git a/Engine/Source/Runtime/CoreUObject/Private/UObject/FieldPath.cpp b/Engine/Source/Runtime/CoreUObject/Private/UObject/FieldPath.cpp
--- a/Engine/Source/Runtime/CoreUObject/Private/UObject/FieldPath.cpp
+++ b/Engine/Source/Runtime/CoreUObject/Private/UObject/FieldPath.cpp
@@ -263,18 +263,14 @@ FField* FFieldPath::TryToResolvePath(UStruct* InCurrentStruct, FFieldPath::EPath
 		Owner = TryToResolveOwnerFromStruct(InCurrentStruct, InResolveType);
 	}
 	// At this point the owner should've been fully resolved
-	if (Owner && Path.Num())
+	if (Owner && Path.Num() > 0)
 	{
 		int32 PathIndex = Path.Num() - 1;
-		check(PathIndex <= 1);
-		Result = FindFProperty<FField>(Owner, Path[PathIndex]);
-		if (Result)
+		Result = FindFProperty<FField>(Owner, Path[PathIndex--]);
+		while (Result && PathIndex > 0)
 		{
-			if (PathIndex > 0)
-			{
-				// Nested property
-				Result = Result->GetInnerFieldByName(Path[0]);
-			}
+			// Nested property
+			Result = Result->GetInnerFieldByName(Path[PathIndex--]);
 		}
 	}
 	return Result;
