Remove a short lived allocation of 2336 bytes when calling GetDefaultObjectName()
Previously, we would reserve space for 1024 characters, which in the internal FString builder in bumped up to 1168, and with TCHAR being a wide character that doubles to the 2336 total bytes requests. This allocation is then deleted at the end of the function.

diff --git a/Engine/Source/Runtime/CoreUObject/Private/UObject/Class.cpp b/Engine/Source/Runtime/CoreUObject/Private/UObject/Class.cpp
--- a/Engine/Source/Runtime/CoreUObject/Private/UObject/Class.cpp
+++ b/Engine/Source/Runtime/CoreUObject/Private/UObject/Class.cpp
@@ -4824,11 +4824,11 @@ FFeedbackContext& UClass::GetDefaultPropertiesFeedbackContext()
 */
 FName UClass::GetDefaultObjectName() const
 {
-	FString DefaultName;
-	DefaultName.Reserve(NAME_SIZE);
-	DefaultName += DEFAULT_OBJECT_PREFIX;
-	AppendName(DefaultName);
-	return FName(*DefaultName);
+	FNameBuilder DefaultNameBuilder;
+	DefaultNameBuilder.Append(DEFAULT_OBJECT_PREFIX);
+	GetFName().AppendString(DefaultNameBuilder);
+
+	return FName(*DefaultNameBuilder);
 }
 
 //
