Trivial TSAN fix to make a thread status bool atomic.

#rnx

diff --git a/Engine/Source/Runtime/CoreUObject/Private/Serialization/AsyncLoading2.cpp b/Engine/Source/Runtime/CoreUObject/Private/Serialization/AsyncLoading2.cpp
--- a/Engine/Source/Runtime/CoreUObject/Private/Serialization/AsyncLoading2.cpp
+++ b/Engine/Source/Runtime/CoreUObject/Private/Serialization/AsyncLoading2.cpp
@@ -3108,7 +3108,7 @@ private:
 	std::atomic<int32> SuspendRequestedCount = 0;
 	bool bHasRegisteredAllScriptObjects = false;
 	/** [ASYNC/GAME THREAD] true if the async thread is actually started. We don't start it until after we boot because the boot process on the game thread can create objects that are also being created by the loader */
-	bool bThreadStarted = false;
+	std::atomic<bool> bThreadStarted = false;
 
 #if !UE_BUILD_SHIPPING
 	FPlatformFileOpenLog* FileOpenLogWrapper = nullptr;
@@ -3311,7 +3311,7 @@ public:
 	/** True if multithreaded async loading is currently being used. */
 	inline virtual bool IsMultithreaded() override
 	{
-		return bThreadStarted;
+		return bThreadStarted.load(std::memory_order_acquire);
 	}
 
 	/** Sets the current state of async loading */
@@ -4202,7 +4202,7 @@ FAsyncPackage2* FAsyncLoadingThread2::FindOrInsertPackage(FAsyncLoadingThreadSta
 	// highly recursive InitDefaultMaterials function on some platforms.
 	// Since postload groups are there to protect against race conditions between postloads and serialize,
 	// no such race can exists until the loading thread is started.
-	const bool bIsPostLoadGroupFeatureActive = bThreadStarted;
+	const bool bIsPostLoadGroupFeatureActive = IsMultithreaded();
 #endif
 
 	if (bInserted)
@@ -8163,7 +8163,7 @@ void FAsyncLoadingThread2::StartThread()
 		AsyncLoadingThreadState->bCanAccessAsyncLoadingThreadData = true;
 		GameThreadState->bCanAccessAsyncLoadingThreadData = false;
 		UE_LOG(LogStreaming, Log, TEXT("Starting Async Loading Thread."));
-		bThreadStarted = true;
+		bThreadStarted.store(true, std::memory_order_release);
 		FPlatformMisc::MemoryBarrier();
 		UE::Trace::ThreadGroupBegin(TEXT("AsyncLoading"));
 		Thread = FRunnableThread::Create(this, TEXT("FAsyncLoadingThread"), 0, TPri_Normal);
