FCookDependency::LoadFromCompactBinary: fix assertion when compact binary has invalid value for the ECookDependency Type. LoadFromCompactBinary by contract should return false for invalid data rather than asserting.
#rnx

diff --git a/Engine/Source/Runtime/CoreUObject/Private/Cooker/CookDependency.cpp b/Engine/Source/Runtime/CoreUObject/Private/Cooker/CookDependency.cpp
--- a/Engine/Source/Runtime/CoreUObject/Private/Cooker/CookDependency.cpp
+++ b/Engine/Source/Runtime/CoreUObject/Private/Cooker/CookDependency.cpp
@@ -334,11 +334,12 @@ bool FCookDependency::Load(FCbFieldView Value)
 		return false;
 	}
 	FCbFieldViewIterator Field(Value.CreateViewIterator());
-	ECookDependency LocalType = static_cast<ECookDependency>(Field.AsUInt8());
-	if ((Field++).HasError())
+	int32 LocalTypeAsInt = Field.AsUInt8();
+	if ((Field++).HasError() || LocalTypeAsInt >= static_cast<uint8>(ECookDependency::Count))
 	{
 		return false;
 	}
+	ECookDependency LocalType = static_cast<ECookDependency>(LocalTypeAsInt);
 
 	switch (LocalType)
 	{
diff --git a/Engine/Source/Runtime/CoreUObject/Public/Cooker/CookDependency.h b/Engine/Source/Runtime/CoreUObject/Public/Cooker/CookDependency.h
--- a/Engine/Source/Runtime/CoreUObject/Public/Cooker/CookDependency.h
+++ b/Engine/Source/Runtime/CoreUObject/Public/Cooker/CookDependency.h
@@ -85,7 +85,9 @@ enum class ECookDependency : uint8
 	Function = 0x02,
 	TransitiveBuild = 0x03,
 	Package = 0x4,
-	ConsoleVariable = 0x5
+	ConsoleVariable = 0x5,
+
+	Count,
 };
 
 /**
