Disabling warning about too many console variable lookups when creating console variable cook dependencies.

diff --git a/Engine/Source/Runtime/CoreUObject/Private/Cooker/CookDependency.cpp b/Engine/Source/Runtime/CoreUObject/Private/Cooker/CookDependency.cpp
--- a/Engine/Source/Runtime/CoreUObject/Private/Cooker/CookDependency.cpp
+++ b/Engine/Source/Runtime/CoreUObject/Private/Cooker/CookDependency.cpp
@@ -224,7 +224,7 @@ void FCookDependency::UpdateHash(FCookDependencyContext& Context) const
 		return;
 	case ECookDependency::ConsoleVariable:
 	{
-		IConsoleVariable* VariableInstance = IConsoleManager::Get().FindConsoleVariable(*StringData);
+		IConsoleVariable* VariableInstance = IConsoleManager::Get().FindConsoleVariable(*StringData, false /* bTrackFrequentCalls */);
 		if (VariableInstance == nullptr)
 		{
 			Context.LogError(FString::Printf(
