Moving the final purge (GC) of UObjects to its own function inside garbage collection source file to hide GC internals from Obj.cpp

diff --git a/Engine/Source/Runtime/CoreUObject/Private/UObject/GarbageCollection.cpp b/Engine/Source/Runtime/CoreUObject/Private/UObject/GarbageCollection.cpp
--- a/Engine/Source/Runtime/CoreUObject/Private/UObject/GarbageCollection.cpp
+++ b/Engine/Source/Runtime/CoreUObject/Private/UObject/GarbageCollection.cpp
@@ -6221,6 +6221,100 @@ bool TryCollectGarbage(EObjectFlags KeepFlags, bool bPerformFullPurge)
 	return true;
 }
 
+/**
+ * Destroys all existing UObjects when the engine shuts down. Does NOT perform reachability analysis to do so.
+ */
+void PurgeAllUObjectsOnExit()
+{
+	// This can happen when we run into an error early in the init process
+	if (GUObjectArray.IsOpenForDisregardForGC())
+	{
+		GUObjectArray.CloseDisregardForGC();
+	}
+
+	// Complete any pending incremental reachability analysis
+	FinalizeIncrementalReachabilityAnalysis();
+
+	// Complete any pending incremental purge
+	if (IsIncrementalPurgePending())
+	{
+		IncrementalPurgeGarbage(false);
+	}
+
+	// From now on we'll be destroying objects without time limit during exit purge
+	// so doing it on a separate thread doesn't make anything faster,
+	// also the exit purge is not a standard GC pass so no need to overcompilcate things
+	GMultithreadedDestructionEnabled = false;
+
+	// Make sure no other threads manipulate UObjects
+	AcquireGCLock();
+
+	// Dissolve all clusters before the final GC pass
+	GUObjectClusters.DissolveClusters(true);
+
+	// Keep track of how many objects there are for GC stats as we simulate a mark pass.
+	extern FThreadSafeCounter GObjectCountDuringLastMarkPhase;
+	GObjectCountDuringLastMarkPhase.Reset();
+
+	// Tag all non template & class objects as unreachable. We can't use object iterators for this as they ignore certain objects.
+	//
+	// Excluding class default, archetype and class objects allows us to not have to worry about fixing issues with initialization 
+	// and certain CDO objects like UNetConnection and UChildConnection having members with arrays that point to the same data and 
+	// will be double freed if destroyed. Hacky, but much cleaner and lower risk than trying to fix the root cause behind it all. 
+	// We need the exit purge for closing network connections and such and only operating on instances of objects is sufficient for 
+	// this purpose.
+	for (FRawObjectIterator It; It; ++It)
+	{
+		// Valid object.
+		GObjectCountDuringLastMarkPhase.Increment();
+
+		FUObjectItem* ObjItem = *It;
+		checkSlow(ObjItem);
+		UObject* Obj = static_cast<UObject*>(ObjItem->Object);
+		if (Obj)
+		{
+			// Skip Structures, properties, etc.. They could be still necessary while GC.
+			if (!Obj->IsA<UField>())
+			{
+				// Mark as unreachable so purge phase will kill it.
+				UE::GC::Private::FGCFlags::SetMaybeUnreachable_ForGC(ObjItem);
+			}
+			else
+			{
+				UE::GC::Private::FGCFlags::ClearMaybeUnreachable_ForGC(ObjItem);
+			}
+		}
+	}
+
+	// Fully purge all objects, not using time limit.
+	GExitPurge = true;
+
+	// Route BeginDestroy. This needs to be a separate pass from marking as RF_Unreachable as code might rely on RF_Unreachable to be 
+	// set on all objects that are about to be deleted. One example is FLinkerLoad detaching textures - the SetLinker call needs to 
+	// not kick off texture streaming.
+	//
+	GatherUnreachableObjects(false);
+	IncrementalPurgeGarbage(false);
+
+	{
+		GObjectCountDuringLastMarkPhase.Reset();
+
+		//Repeat GC for every object, including structures and properties.
+		for (FRawObjectIterator It; It; ++It)
+		{
+			GObjectCountDuringLastMarkPhase.Increment();
+
+			// Mark as unreachable so purge phase will kill it.
+			UE::GC::Private::FGCFlags::SetMaybeUnreachable_ForGC(*It);
+		}
+
+		GatherUnreachableObjects(false);
+		IncrementalPurgeGarbage(false);
+	}
+
+	ReleaseGCLock();
+}
+
 void UObject::CallAddReferencedObjects(FReferenceCollector& Collector)
 {
 	GetClass()->CallAddReferencedObjects(this, Collector);
diff --git a/Engine/Source/Runtime/CoreUObject/Private/UObject/Obj.cpp b/Engine/Source/Runtime/CoreUObject/Private/UObject/Obj.cpp
--- a/Engine/Source/Runtime/CoreUObject/Private/UObject/Obj.cpp
+++ b/Engine/Source/Runtime/CoreUObject/Private/UObject/Obj.cpp
@@ -25,7 +25,6 @@
 #include "UObject/Object.h"
 #include "Serialization/ArchiveUObject.h"
 #include "UObject/GarbageCollectionHistory.h"
-#include "UObject/GarbageCollectionInternalFlags.h"
 #include "UObject/Class.h"
 #include "UObject/EnumProperty.h"
 #include "UObject/UObjectIterator.h"
@@ -5379,9 +5378,7 @@ void StaticUObjectInit()
 void ShutdownGarbageCollection();
 void CleanupLinkerAnnotations();
 void CleanupCachedArchetypes();
-void AcquireGCLock();
-void ReleaseGCLock();
-extern int32 GMultithreadedDestructionEnabled;
+void PurgeAllUObjectsOnExit();
 
 //
 // Shut down the object manager.
@@ -5403,93 +5400,7 @@ void StaticExit()
 		GObjTransientPkg = NULL;
 	}
 	
-	// This can happen when we run into an error early in the init process
-	if (GUObjectArray.IsOpenForDisregardForGC())
-	{
-		GUObjectArray.CloseDisregardForGC();
-	}
-
-	// Complete any pending incremental reachability analysis
-	FinalizeIncrementalReachabilityAnalysis();
-
-	// Complete any pending incremental purge
-	if (IsIncrementalPurgePending())
-	{
-		IncrementalPurgeGarbage(false);
-	}
-
-	// From now on we'll be destroying objects without time limit during exit purge
-	// so doing it on a separate thread doesn't make anything faster,
-	// also the exit purge is not a standard GC pass so no need to overcompilcate things
-	GMultithreadedDestructionEnabled = false;
-
-	// Make sure no other threads manipulate UObjects
-	AcquireGCLock();
-
-	// Dissolve all clusters before the final GC pass
-	GUObjectClusters.DissolveClusters(true);
-
-	// Keep track of how many objects there are for GC stats as we simulate a mark pass.
-	extern FThreadSafeCounter GObjectCountDuringLastMarkPhase;
-	GObjectCountDuringLastMarkPhase.Reset();
-
-	// Tag all non template & class objects as unreachable. We can't use object iterators for this as they ignore certain objects.
-	//
-	// Excluding class default, archetype and class objects allows us to not have to worry about fixing issues with initialization 
-	// and certain CDO objects like UNetConnection and UChildConnection having members with arrays that point to the same data and 
-	// will be double freed if destroyed. Hacky, but much cleaner and lower risk than trying to fix the root cause behind it all. 
-	// We need the exit purge for closing network connections and such and only operating on instances of objects is sufficient for 
-	// this purpose.
-	for ( FRawObjectIterator It; It; ++It )
-	{
-		// Valid object.
-		GObjectCountDuringLastMarkPhase.Increment();
-
-		FUObjectItem* ObjItem = *It;
-		checkSlow(ObjItem);
-		UObject* Obj = static_cast<UObject*>(ObjItem->Object);
-		if (Obj) 
-		{
-			// Skip Structures, properties, etc.. They could be still necessary while GC.
-			if (!Obj->IsA<UField>())
-			{
-				// Mark as unreachable so purge phase will kill it.
-				UE::GC::Private::FGCFlags::SetMaybeUnreachable_ForGC(ObjItem);
-			}
-			else
-			{
-				UE::GC::Private::FGCFlags::ClearMaybeUnreachable_ForGC(ObjItem);
-			}
-		}
-	}
-
-	// Fully purge all objects, not using time limit.
-	GExitPurge					= true;
-
-	// Route BeginDestroy. This needs to be a separate pass from marking as RF_Unreachable as code might rely on RF_Unreachable to be 
-	// set on all objects that are about to be deleted. One example is FLinkerLoad detaching textures - the SetLinker call needs to 
-	// not kick off texture streaming.
-	//
-	GatherUnreachableObjects(false);
-	IncrementalPurgeGarbage(false);
-
-	{
-		GObjectCountDuringLastMarkPhase.Reset();
-
-		//Repeat GC for every object, including structures and properties.
-		for (FRawObjectIterator It; It; ++It)
-		{
-			GObjectCountDuringLastMarkPhase.Increment();
-
-			// Mark as unreachable so purge phase will kill it.
-			UE::GC::Private::FGCFlags::SetMaybeUnreachable_ForGC(*It);
-		}
-
-		GatherUnreachableObjects(false);
-		IncrementalPurgeGarbage(false);
-	}
-
-	ReleaseGCLock();
+	PurgeAllUObjectsOnExit();
 
 	ShutdownGarbageCollection();
 	UObjectBaseShutdown();
