Open around loading package

diff --git a/Engine/Source/Runtime/CoreUObject/Private/UObject/SoftObjectPath.cpp b/Engine/Source/Runtime/CoreUObject/Private/UObject/SoftObjectPath.cpp
--- a/Engine/Source/Runtime/CoreUObject/Private/UObject/SoftObjectPath.cpp
+++ b/Engine/Source/Runtime/CoreUObject/Private/UObject/SoftObjectPath.cpp
@@ -13,6 +13,7 @@
 #include "Misc/RedirectCollector.h"
 #include "Misc/AutomationTest.h"
 #include "String/Find.h"
+#include "AutoRTFM/AutoRTFM.h"
 
 // Deprecated constructor
 FSoftObjectPath::FSoftObjectPath(FName InAssetPathName, FString InSubPathString)
@@ -535,7 +536,10 @@ UObject* FSoftObjectPath::TryLoad(FUObjectSerializeContext* InLoadContext) const
 			}
 #endif
 
-			LoadedObject = StaticLoadObject(UObject::StaticClass(), nullptr, *PathString, nullptr, LOAD_None, nullptr, true);
+			UE_AUTORTFM_OPEN(
+			{
+				LoadedObject = StaticLoadObject(UObject::StaticClass(), nullptr, *PathString, nullptr, LOAD_None, nullptr, true);
+			});
 
 #if WITH_EDITOR
 			// Look at core redirects if we didn't find the object
@@ -1056,4 +1060,4 @@ bool FSoftObjectPathFixupForPIETests::RunTest(const FString& Parameters)
 
 #endif 
 
-#endif
\ No newline at end of file
+#endif
