SavePackage: Add a cache for IsEditorOnlyObject so an object with a long outer chain above it and many objects below it will not cause a performance problem due to repeatedly searching the outer chain.
Refactor the TransientAssignments cache used by IsTransient and IsUnsaveable so that we can share storage and lookup time with the new cache for IsEditorOnlyObject.
#rnx

diff --git a/Engine/Source/Runtime/CoreUObject/Private/UObject/SavePackage/PackageHarvester.cpp b/Engine/Source/Runtime/CoreUObject/Private/UObject/SavePackage/PackageHarvester.cpp
--- a/Engine/Source/Runtime/CoreUObject/Private/UObject/SavePackage/PackageHarvester.cpp
+++ b/Engine/Source/Runtime/CoreUObject/Private/UObject/SavePackage/PackageHarvester.cpp
@@ -11,7 +11,7 @@
 #include "UObject/UObjectGlobals.h"
 #include "UObject/UObjectHash.h"
 
-EObjectMark GenerateMarksForObject(const UObject* InObject, const ITargetPlatform* TargetPlatform)
+EObjectMark GenerateMarksForObject(const UObject* InObject, FSaveContext& SaveContext)
 {
 	using namespace UE::SavePackageUtilities;
 
@@ -33,6 +33,7 @@ EObjectMark GenerateMarksForObject(const UObject* InObject, const ITargetPlatfor
 		Marks = (EObjectMark)(Marks | OBJECTMARK_NotForServer);
 	}
 #if WITH_ENGINE
+	const ITargetPlatform* TargetPlatform = SaveContext.GetTargetPlatform();
 	bool bCheckTargetPlatform = false;
 	if (TargetPlatform != nullptr)
 	{
@@ -50,15 +51,22 @@ EObjectMark GenerateMarksForObject(const UObject* InObject, const ITargetPlatfor
 	}
 #endif
 	
-	// If doing an editor save, HasNonEditorOnlyReferences=true overrides NotForClient, NotForServer, and virtual IsEditorOnly and marks it as UsedInGame
-	bool bApplyHasNonEditorOnlyReferences = TargetPlatform == nullptr;
-	if (IsStrippedEditorOnlyObject(InObject,
-		EEditorOnlyObjectFlags::CheckRecursive |
-		(bApplyHasNonEditorOnlyReferences ? EEditorOnlyObjectFlags::ApplyHasNonEditorOnlyReferences : EEditorOnlyObjectFlags::None)))
+	EEditorOnlyObjectFlags EditorOnlyObjectFlags = SaveContext.GetEditorOnlyObjectFlags();
+	bool bApplyHasNonEditorOnlyReferences = EnumHasAnyFlags(EditorOnlyObjectFlags,
+		EEditorOnlyObjectFlags::ApplyHasNonEditorOnlyReferences);
+#if WITH_EDITORONLY_DATA
+	bool bIsEditorOnlyObject = UE::SavePackageUtilities::IsEditorOnlyObjectInternal(InObject,
+		EditorOnlyObjectFlags,
+		SaveContext.GetFunctorReadCachedEditorOnlyObject(),
+		SaveContext.GetFunctorWriteCachedEditorOnlyObject());
+	bool bStrippableEditorOnlyObject = bIsEditorOnlyObject
+		&& UE::SavePackageUtilities::CanStripEditorOnlyImportsAndExports();
+	if (bStrippableEditorOnlyObject)
 	{
 		Marks = (EObjectMark)(Marks | OBJECTMARK_EditorOnly);
 	}
 	else
+#endif
 	// If NotForClient and NotForServer, it is implicitly editor only
 	if ((Marks & OBJECTMARK_NotForClient) && (Marks & OBJECTMARK_NotForServer) &&
 		(!bApplyHasNonEditorOnlyReferences || !InObject->HasNonEditorOnlyReferences()))
@@ -92,7 +100,7 @@ bool ConditionallyExcludeObjectForRealm(FSaveContext& SaveContext, TObjectPtr<UO
 
 	const EObjectMark ExcludedObjectMarks = SaveContext.GetExcludedObjectMarks(HarvestingContext);
 	const ITargetPlatform* TargetPlatform = SaveContext.GetTargetPlatform();
-	EObjectMark ObjectMarks = GenerateMarksForObject(Obj, TargetPlatform);
+	EObjectMark ObjectMarks = GenerateMarksForObject(Obj, SaveContext);
 	if (!!(ObjectMarks & ExcludedObjectMarks))
 	{
 		RealmBeingChecked.AddExcluded(Obj);
@@ -100,7 +108,8 @@ bool ConditionallyExcludeObjectForRealm(FSaveContext& SaveContext, TObjectPtr<UO
 	}
 
 	// If the object class is excluded, the object must be excluded too
-	bool bApplyHasNonEditorOnlyReferences = TargetPlatform == nullptr;
+	bool bApplyHasNonEditorOnlyReferences = EnumHasAnyFlags(SaveContext.GetEditorOnlyObjectFlags(),
+		UE::SavePackageUtilities::EEditorOnlyObjectFlags::ApplyHasNonEditorOnlyReferences);
 	bool bIgnoreEditorOnlyClass = bApplyHasNonEditorOnlyReferences && Obj->HasNonEditorOnlyReferences();
 	if (!bIgnoreEditorOnlyClass && ConditionallyExcludeObjectForRealm(SaveContext, Obj->GetClass(), HarvestingContext))
 	{
@@ -415,6 +424,10 @@ FPackageHarvester::FPackageHarvester(FSaveContext& InContext)
 	this->SetUseUnversionedPropertySerialization(SaveContext.IsSaveUnversionedProperties());
 
 	ResolveOverrides();
+	// Clear the SaveContext's Saveable cache. It was used earlier during RoutePreSave, and the PreSave functions
+	// may have modified RF_Transient or other flags on UObjects and therefore invalidated the cached result. It
+	// can also be invalidated by SaveOverrides with bForceTransient.
+	SaveContext.ClearSaveableCache();
 }
 
 FPackageHarvester::FExportWithContext FPackageHarvester::PopExportToProcess()
@@ -599,10 +612,14 @@ void FPackageHarvester::TryHarvestExportInternal(UObject* InObject)
 
 #if WITH_EDITORONLY_DATA
 	// Remove the Game realm if the object is editoronly and does not override it with HasNonEditorOnlyReferences
-	bool bIsEditorOnlyObject = UE::SavePackageUtilities::IsStrippedEditorOnlyObject(InObject,
-		UE::SavePackageUtilities::EEditorOnlyObjectFlags::CheckRecursive)
-		&& !InObject->HasNonEditorOnlyReferences();
-	FHarvestScope EditorOnlyScope = EnterConditionalEditorOnlyScope(bIsEditorOnlyObject);
+	bool bIsEditorOnlyObject = UE::SavePackageUtilities::IsEditorOnlyObjectInternal(InObject,
+		SaveContext.GetEditorOnlyObjectFlags(),
+		SaveContext.GetFunctorReadCachedEditorOnlyObject(),
+		SaveContext.GetFunctorWriteCachedEditorOnlyObject());
+	bool bStrippableEditorOnlyObject = bIsEditorOnlyObject
+		&& !InObject->HasNonEditorOnlyReferences()
+		&& UE::SavePackageUtilities::CanStripEditorOnlyImportsAndExports();
+	FHarvestScope EditorOnlyScope = EnterConditionalEditorOnlyScope(bStrippableEditorOnlyObject);
 #endif
 
 	// Remove realms for which we have already harvested the export
@@ -613,8 +630,8 @@ void FPackageHarvester::TryHarvestExportInternal(UObject* InObject)
 	}
 
 	// Check whether the object is unsaveable and skip adding it as an export to any realm if so
-	SaveContext.MarkUnsaveable(InObject);
-	if (SaveContext.IsTransient(InObject))
+	SaveContext.GetCachedObjectStatus(InObject).bAttemptedExport = true;
+	if (SaveContext.IsUnsaveable(InObject))
 	{
 		return;
 	}
@@ -975,9 +992,9 @@ void FPackageHarvester::ResolveOverrides()
 		{
 			TransientPropertyOverrides.Add(PairObjectOverrides.Key, MoveTemp(Props));
 		}
-		if (!SaveContext.IsTransient(PairObjectOverrides.Key) && PairObjectOverrides.Value.bForceTransient)
+		if (!PairObjectOverrides.Key->HasAnyFlags(RF_Transient) && PairObjectOverrides.Value.bForceTransient)
 		{
-			SaveContext.TransientAssignments.Add(PairObjectOverrides.Key, EMarkedTransientReason::TransientOverride);
+			SaveContext.GetCachedObjectStatus(PairObjectOverrides.Key).bSaveOverrideForcedTransient = true;
 		}
 	}
 }
@@ -1124,27 +1141,28 @@ ESaveableStatus FPackageHarvester::GetSaveableStatusForRealm(UObject* Obj, ESave
 		return ESaveableStatus::Success;
 	}
 
-	ESaveableStatus CulpritStatus;
-	ESaveableStatus Status = SaveContext.GetSaveableStatus(Obj, &OutCulprit, &CulpritStatus, EIgnoreMarkUnsaveable::Yes);
-	if (Status != ESaveableStatus::Success)
+	UE::SavePackageUtilities::FObjectStatus& ObjectStatus = SaveContext.UpdateSaveableStatus(Obj);
+	if (ObjectStatus.SaveableStatus != ESaveableStatus::Success)
 	{
-		if (Status == ESaveableStatus::OuterUnsaveable)
+		if (ObjectStatus.SaveableStatus == ESaveableStatus::OuterUnsaveable)
 		{
-			check(OutCulprit);
-			OutReason = FString::Printf(TEXT("has outer %s which %s"), *OutCulprit->GetPathName(),
-				LexToString(CulpritStatus));
+			check(ObjectStatus.SaveableStatusCulprit != nullptr &&
+				ObjectStatus.SaveableStatusCulpritStatus != ESaveableStatus::Success);
+			OutCulprit = ObjectStatus.SaveableStatusCulprit;
+			OutReason = FString::Printf(TEXT("has outer %s which %s"), *OutCulprit.GetPathName(),
+				LexToString(ObjectStatus.SaveableStatusCulpritStatus));
 		}
 		else
 		{
 			OutCulprit = Obj;
-			OutReason = LexToString(Status);
+			OutReason = LexToString(ObjectStatus.SaveableStatus);
 		}
-		return Status;
+		return ObjectStatus.SaveableStatus;
 	}
 
 	const EObjectMark ExcludedObjectMarks = SaveContext.GetExcludedObjectMarks(RealmInWhichItIsUnsaveable);
 	const ITargetPlatform* TargetPlatform = SaveContext.GetTargetPlatform();
-	EObjectMark ObjectMarks = GenerateMarksForObject(Obj, TargetPlatform);
+	EObjectMark ObjectMarks = GenerateMarksForObject(Obj, SaveContext);
 	if (!!(ObjectMarks & ExcludedObjectMarks))
 	{
 		OutCulprit = Obj;
diff --git a/Engine/Source/Runtime/CoreUObject/Private/UObject/SavePackage/SaveContext.cpp b/Engine/Source/Runtime/CoreUObject/Private/UObject/SavePackage/SaveContext.cpp
--- a/Engine/Source/Runtime/CoreUObject/Private/UObject/SavePackage/SaveContext.cpp
+++ b/Engine/Source/Runtime/CoreUObject/Private/UObject/SavePackage/SaveContext.cpp
@@ -28,106 +28,118 @@ TArray<ESaveRealm> FSaveContext::GetHarvestedRealmsToSave()
 	return HarvestedContextsToSave;
 }
 
-void FSaveContext::MarkUnsaveable(UObject* InObject)
+bool FSaveContext::IsUnsaveable(TObjectPtr<UObject> InObject, bool bEmitWarning)
 {
-	if (IsUnsaveable(InObject))
+	if (!InObject)
 	{
-		EMarkedTransientReason& Reason = TransientAssignments.FindOrAdd(InObject);
-		if (Reason == EMarkedTransientReason::Uninitialized)
-		{
-			Reason = EMarkedTransientReason::Unsaveable;
-		}
+		return false;
 	}
+	UE::SavePackageUtilities::FObjectStatus& ObjectStatus = UpdateSaveableStatus(InObject);
+	check(ObjectStatus.bSaveableStatusValid);
 
-	// if this is the class default object, make sure it's not
-	// marked transient for any reason, as we need it to be saved
-	// to disk (unless it's associated with a transient generated class)
+	if (bEmitWarning && ObjectStatus.SaveableStatus != ESaveableStatus::Success)
+	{
+		// if this is a class default object being exported, make sure it's not unsaveable for any reason,
+		// as we need it to be saved to disk (unless it's associated with a transient generated class)
 #if WITH_EDITORONLY_DATA
-	ensureAlways(!InObject->HasAllFlags(RF_ClassDefaultObject) || !IsTransient(InObject) || (InObject->GetClass()->ClassGeneratedBy != nullptr && InObject->GetClass()->HasAnyFlags(RF_Transient)));
+		ensureAlways(!ObjectStatus.bAttemptedExport || !InObject->HasAllFlags(RF_ClassDefaultObject) ||
+			(InObject->GetClass()->ClassGeneratedBy != nullptr && InObject->GetClass()->HasAnyFlags(RF_Transient)));
 #endif
+
+		if (ObjectStatus.SaveableStatus == ESaveableStatus::OuterUnsaveable
+			&& (ObjectStatus.SaveableStatusCulpritStatus == ESaveableStatus::AbstractClass
+				|| ObjectStatus.SaveableStatusCulpritStatus == ESaveableStatus::DeprecatedClass
+				|| ObjectStatus.SaveableStatusCulpritStatus == ESaveableStatus::NewerVersionExistsClass)
+			&& InObject.GetPackage() == GetPackage())
+		{
+			check(ObjectStatus.SaveableStatusCulprit);
+			UE_LOG(LogSavePackage, Warning, TEXT("%s has unsaveable outer %s (outer is %s), so it will not be saved."),
+				*InObject.GetFullName(), *ObjectStatus.SaveableStatusCulprit->GetFullName(),
+				LexToString(ObjectStatus.SaveableStatusCulpritStatus));
+		}
+	}
+
+	return ObjectStatus.SaveableStatus != ESaveableStatus::Success;
 }
 
-bool FSaveContext::IsUnsaveable(TObjectPtr<UObject> InObject, bool bEmitWarning) const
+UE::SavePackageUtilities::FObjectStatus& FSaveContext::UpdateSaveableStatus(TObjectPtr<UObject> InObject)
 {
-	TObjectPtr<UObject> Culprit;
-	ESaveableStatus CulpritStatus;
-	ESaveableStatus Status = GetSaveableStatus(InObject, &Culprit, &CulpritStatus);
-	if (Status == ESaveableStatus::Success)
+	UE::SavePackageUtilities::FObjectStatus* ObjectStatus = &ObjectStatusCache.FindOrAdd(InObject);
+	if (ObjectStatus->bSaveableStatusValid)
 	{
-		return false;
+		return *ObjectStatus;
 	}
-	if (Status == ESaveableStatus::OuterUnsaveable && bEmitWarning &&
-		(CulpritStatus == ESaveableStatus::AbstractClass || CulpritStatus == ESaveableStatus::DeprecatedClass || CulpritStatus == ESaveableStatus::NewerVersionExistsClass))
+
+	ObjectStatus->bSaveableStatusValid = true;
+	ObjectStatus->SaveableStatus = ESaveableStatus::Success;
+
+	ESaveableStatus StatusNoOuter = GetSaveableStatusNoOuter(InObject, *ObjectStatus);
+	if (StatusNoOuter != ESaveableStatus::Success)
 	{
-		check(Culprit);
-		// Only warn if the base object is fine but the outer is invalid. If an object is itself unsaveable, the old behavior is to ignore it
-		if (InObject.GetPackage() == GetPackage())
-		{
-			UE_LOG(LogSavePackage, Warning, TEXT("%s has a deprecated or abstract class outer %s, so it will not be saved"),
-				*InObject.GetFullName(), *Culprit.GetFullName());
-		}
+		check(StatusNoOuter != ESaveableStatus::OuterUnsaveable &&
+			StatusNoOuter != ESaveableStatus::ClassUnsaveable);
+		ObjectStatus->SaveableStatus = StatusNoOuter;
+		return *ObjectStatus;
 	}
-	return true;
-}
 
-ESaveableStatus FSaveContext::GetSaveableStatus(TObjectPtr<UObject> InObject, TObjectPtr<UObject>* OutCulprit, ESaveableStatus* OutCulpritStatus, EIgnoreMarkUnsaveable IgnoreMarkUnsaveable) const
-{
-	TObjectPtr<UObject> Obj = InObject;
-	while (Obj)
+	if (!InObject.IsResolved())
+	{
+		// We do not test the saveability of the outer of unresolved objects because we cannot get their
+		// outer without resolving them.
+		return *ObjectStatus;
+	}
+
+	UObject* Outer = InObject->GetOuter();
+	if (Outer)
 	{
-		ESaveableStatus Status = GetSaveableStatusNoOuter(Obj, IgnoreMarkUnsaveable);
-		if (Status != ESaveableStatus::Success)
+		UE::SavePackageUtilities::FObjectStatus& OuterStatus = UpdateSaveableStatus(Outer);
+		// Calling UpdateSaveableStatus on the outer might have modified ObjectStatusCache, so
+		// look up our pointer to ObjectStatus again.
+		ObjectStatus = ObjectStatusCache.Find(InObject);
+		check(ObjectStatus);
+
+		ObjectStatus->bSaveableStatusValid = true;
+		ObjectStatus->SaveableStatus = ESaveableStatus::Success;
+		if (OuterStatus.SaveableStatus != ESaveableStatus::Success)
 		{
-			if (OutCulprit)
+			ObjectStatus->SaveableStatus = ESaveableStatus::OuterUnsaveable;
+			if (OuterStatus.SaveableStatus == ESaveableStatus::OuterUnsaveable)
 			{
-				*OutCulprit = Obj;
+				check(OuterStatus.SaveableStatusCulprit);
+				check(OuterStatus.SaveableStatusCulpritStatus != ESaveableStatus::Success);
+				ObjectStatus->SaveableStatusCulprit = OuterStatus.SaveableStatusCulprit;
+				ObjectStatus->SaveableStatusCulpritStatus = OuterStatus.SaveableStatusCulpritStatus;
 			}
-			if (OutCulpritStatus)
+			else
 			{
-				*OutCulpritStatus = Status;
+				ObjectStatus->SaveableStatusCulprit = Outer;
+				ObjectStatus->SaveableStatusCulpritStatus = OuterStatus.SaveableStatus;
 			}
-			return Obj == InObject ? Status : ESaveableStatus::OuterUnsaveable;
 		}
-		Obj = Obj.GetOuter();
-	}
-	if (OutCulprit)
-	{
-		*OutCulprit = InObject;
-	}
-	if (OutCulpritStatus)
-	{
-		*OutCulpritStatus = ESaveableStatus::Success;
 	}
-	return ESaveableStatus::Success;
+
+	return *ObjectStatus;
 }
 
-ESaveableStatus FSaveContext::GetSaveableStatusNoOuter(TObjectPtr<UObject> Obj, EIgnoreMarkUnsaveable IgnoreMarkUnsaveable) const
+ESaveableStatus FSaveContext::GetSaveableStatusNoOuter(TObjectPtr<UObject> Obj,
+	UE::SavePackageUtilities::FObjectStatus& ObjectStatus) const
 {
-	// pending kill object are unsaveable
+	// pending kill objects are unsaveable
 	if (Obj.IsResolved() && !IsValidChecked(Obj))
 	{
 		return ESaveableStatus::PendingKill;
 	}
 
-	// transient object are considered unsaveable if non native
-	if (Obj.IsResolved() && IsTransient(Obj, IgnoreMarkUnsaveable) && !Obj->IsNative())
+	// transient objects are unsaveable if non-native
+	if (Obj.IsResolved() && !Obj->IsNative())
 	{
-		if (Obj->HasAnyFlags(RF_Transient))
+		if (ObjectStatus.HasTransientFlag(Obj))
 		{
 			return ESaveableStatus::TransientFlag;
 		}
-
-		if (const EMarkedTransientReason* Reason = TransientAssignments.Find(Obj))
+		if (ObjectStatus.bSaveOverrideForcedTransient)
 		{
-			switch (*Reason)
-			{
-			case EMarkedTransientReason::TransientOverride:
-				return ESaveableStatus::TransientOverride;
-			case EMarkedTransientReason::Unsaveable:
-				return ESaveableStatus::MarkedUnsaveable;
-			default:
-				return ESaveableStatus::MarkedUnsaveable;
-			}
+			return ESaveableStatus::TransientOverride;
 		}
 	}
 
@@ -147,7 +159,7 @@ ESaveableStatus FSaveContext::GetSaveableStatusNoOuter(TObjectPtr<UObject> Obj,
 	return ESaveableStatus::Success;
 }
 
-bool FSaveContext::IsTransient(const TObjectPtr<UObject> InObject, EIgnoreMarkUnsaveable IgnoreMarkUnsaveable) const
+bool FSaveContext::IsTransient(TObjectPtr<UObject> InObject)
 {
 	if (!InObject)
 	{
@@ -159,12 +171,15 @@ bool FSaveContext::IsTransient(const TObjectPtr<UObject> InObject, EIgnoreMarkUn
 		return true;
 	}
 
-	if (const EMarkedTransientReason* Reason = TransientAssignments.Find(InObject))
+	UE::SavePackageUtilities::FObjectStatus& Status = GetCachedObjectStatus(InObject);
+	if (Status.bSaveOverrideForcedTransient)
 	{
-		if (*Reason == EMarkedTransientReason::Unsaveable && IgnoreMarkUnsaveable == EIgnoreMarkUnsaveable::Yes)
-		{
-			return false;
-		}
+		return true;
+	}
+	if (Status.bAttemptedExport && Status.bSaveableStatusValid && Status.SaveableStatus != ESaveableStatus::Success)
+	{
+		// Exports found to be unsaveable are treated the same as transient objects for all the calls to IsTransient
+		// in SavePackage.
 		return true;
 	}
 
@@ -208,6 +223,20 @@ FSavePackageResultStruct FSaveContext::GetFinalResult()
 	return ResultData;
 }
 
+UE::SavePackageUtilities::EEditorOnlyObjectFlags FSaveContext::GetEditorOnlyObjectFlags() const
+{
+	using namespace UE::SavePackageUtilities;
+
+	// If doing an editor save, HasNonEditorOnlyReferences=true overrides NotForClient, NotForServer,
+	// and virtual IsEditorOnly and marks it as UsedInGame
+	bool bApplyHasNonEditorOnlyReferences = GetTargetPlatform() == nullptr;
+	return
+		EEditorOnlyObjectFlags::CheckRecursive |
+		(bApplyHasNonEditorOnlyReferences
+			? EEditorOnlyObjectFlags::ApplyHasNonEditorOnlyReferences
+			: EEditorOnlyObjectFlags::None);
+}
+
 namespace
 {
 	TArray<UClass*> AutomaticOptionalInclusionAssetTypeList;
@@ -267,14 +296,13 @@ EObjectMark FSaveContext::GetExcludedObjectMarksForGameRealm(const ITargetPlatfo
 
 const TCHAR* LexToString(ESaveableStatus Status)
 {
-	static_assert(static_cast<int32>(ESaveableStatus::__Count) == 11);
+	static_assert(static_cast<int32>(ESaveableStatus::__Count) == 10);
 	switch (Status)
 	{
 	case ESaveableStatus::Success: return TEXT("is saveable");
 	case ESaveableStatus::PendingKill: return TEXT("is pendingkill");
 	case ESaveableStatus::TransientFlag: return TEXT("is transient");
 	case ESaveableStatus::TransientOverride: return TEXT("is Overriden as transient");
-	case ESaveableStatus::MarkedUnsaveable: return TEXT("has been marked unsaveable");
 	case ESaveableStatus::AbstractClass: return TEXT("has a Class with CLASS_Abstract");
 	case ESaveableStatus::DeprecatedClass: return TEXT("has a Class with CLASS_Deprecated");
 	case ESaveableStatus::NewerVersionExistsClass: return TEXT("has a Class with CLASS_NewerVersionExists");
diff --git a/Engine/Source/Runtime/CoreUObject/Private/UObject/SavePackage/SaveContext.h b/Engine/Source/Runtime/CoreUObject/Private/UObject/SavePackage/SaveContext.h
--- a/Engine/Source/Runtime/CoreUObject/Private/UObject/SavePackage/SaveContext.h
+++ b/Engine/Source/Runtime/CoreUObject/Private/UObject/SavePackage/SaveContext.h
@@ -109,7 +109,6 @@ enum class ESaveableStatus
 	PendingKill,
 	TransientFlag,
 	TransientOverride,
-	MarkedUnsaveable,
 	AbstractClass,
 	DeprecatedClass,
 	NewerVersionExistsClass,
@@ -119,19 +118,44 @@ enum class ESaveableStatus
 	__Count,
 };
 
-enum class EIgnoreMarkUnsaveable
+namespace UE::SavePackageUtilities
 {
-	No = 0,
-	Yes,
-};
 
-enum class EMarkedTransientReason
+/** Calculated flags about saveable status and other state for UObjects encountered during save. */
+struct FObjectStatus
 {
-	Uninitialized = 0,
-	Unsaveable,
-	TransientOverride,
+	FObjectStatus()
+		: bSaveOverrideForcedTransient(false)
+		, bSaveableStatusValid(false), bEditorOnlyValid(false), bEditorOnly(false), bAttemptedExport(false)
+	{
+	}
+	FObjectStatus(const FObjectStatus&) = default;
+	FObjectStatus(FObjectStatus&&) = default;
+	FObjectStatus& operator=(const FObjectStatus&) = default;
+	FObjectStatus& operator=(FObjectStatus&&) = default;
+	bool HasTransientFlag(const UObject* InObject)
+	{
+		check(InObject);
+		return InObject->HasAnyFlags(RF_Transient);
+	}
+	void ClearSaveableStatus()
+	{
+		SaveableStatus = ESaveableStatus::Success;
+		bSaveableStatusValid = false;
+		SaveableStatusCulprit = nullptr;
+	}
+
+	UObject* SaveableStatusCulprit = nullptr;
+	ESaveableStatus SaveableStatus = ESaveableStatus::Success;
+	ESaveableStatus SaveableStatusCulpritStatus = ESaveableStatus::Success;
+	bool bSaveOverrideForcedTransient : 1;
+	bool bSaveableStatusValid : 1;
+	bool bEditorOnlyValid : 1;
+	bool bEditorOnly : 1;
+	bool bAttemptedExport : 1;
 };
 
+}
 /** Hold the harvested exports and imports for a realm */
 struct FHarvestedRealm
 {
@@ -793,13 +817,13 @@ public:
 	/** Returns which save context should be saved. */
 	TArray<ESaveRealm> GetHarvestedRealmsToSave();
 
-	void MarkUnsaveable(UObject* InObject);
-
-	bool IsUnsaveable(TObjectPtr<UObject> InObject, bool bEmitWarning = true) const;
-	ESaveableStatus GetSaveableStatus(TObjectPtr<UObject> InObject, TObjectPtr<UObject>* OutCulprit = nullptr, ESaveableStatus* OutCulpritStatus = nullptr, EIgnoreMarkUnsaveable IgnoreMarkUnsaveable = EIgnoreMarkUnsaveable::No) const;
-	ESaveableStatus GetSaveableStatusNoOuter(TObjectPtr<UObject> InObject, EIgnoreMarkUnsaveable IgnoreMarkUnsaveable = EIgnoreMarkUnsaveable::No) const;
-
-	bool IsTransient(TObjectPtr<UObject> InObject, EIgnoreMarkUnsaveable IgnoreMarkUnsaveable = EIgnoreMarkUnsaveable::No) const;
+	bool IsTransient(TObjectPtr<UObject> InObject);
+	bool IsUnsaveable(TObjectPtr<UObject> InObject, bool bEmitWarning = true);
+	UE::SavePackageUtilities::FObjectStatus& UpdateSaveableStatus(TObjectPtr<UObject> InObject);
+	UE::SavePackageUtilities::FObjectStatus& GetCachedObjectStatus(TObjectPtr<UObject> InObject)
+	{
+		return ObjectStatusCache.FindOrAdd(InObject);
+	}
 
 	void RecordIllegalReference(UObject* InFrom, UObject* InTo, EIllegalRefReason InReason, FString&& InOptionalReasonText = FString())
 	{
@@ -1113,6 +1137,37 @@ public:
 		TransientPropertyOverrides = MoveTemp(InTransientPropertyOverrides);
 	}
 
+	void ClearSaveableCache()
+	{
+		for (TPair<TObjectPtr<UObject>, UE::SavePackageUtilities::FObjectStatus>& Pair : ObjectStatusCache)
+		{
+			Pair.Value.ClearSaveableStatus();
+		}
+	}
+
+	auto GetFunctorReadCachedEditorOnlyObject()
+	{
+		using namespace UE::SavePackageUtilities;
+		return [this](const UObject* Obj)
+			{
+				FObjectStatus& Status = ObjectStatusCache.FindOrAdd(const_cast<UObject*>(Obj));
+				return !Status.bEditorOnlyValid ? EEditorOnlyObjectResult::Uninitialized :
+					Status.bEditorOnly ? EEditorOnlyObjectResult::EditorOnly :
+					EEditorOnlyObjectResult::NonEditorOnly;
+			};
+	}
+	auto GetFunctorWriteCachedEditorOnlyObject()
+	{
+		using namespace UE::SavePackageUtilities;
+		return [this](const UObject* Obj, bool bEditorOnly)
+			{
+				FObjectStatus& Status = ObjectStatusCache.FindOrAdd(const_cast<UObject*>(Obj));
+				Status.bEditorOnlyValid = true;
+				Status.bEditorOnly = bEditorOnly;
+			};
+	}
+	UE::SavePackageUtilities::EEditorOnlyObjectFlags GetEditorOnlyObjectFlags() const;
+
 public:
 	ESavePackageResult Result;
 
@@ -1132,6 +1187,8 @@ private:
 
 	// Create the harvesting contexts and automatic optional context gathering options
 	void SetupHarvestingRealms();
+	ESaveableStatus GetSaveableStatusNoOuter(TObjectPtr<UObject> Obj,
+		UE::SavePackageUtilities::FObjectStatus& ObjectStatus) const;
 	static EObjectMark GetExcludedObjectMarksForGameRealm(const ITargetPlatform* TargetPlatform);
 		
 	friend class FPackageHarvester;
@@ -1190,8 +1247,8 @@ private:
 	// Overridden properties for each export that should be treated as transient, and nulled out when serializing
 	TMap<UObject*, TSet<FProperty*>> TransientPropertyOverrides;
 
-	// Overridden objects that should be treated as transient, and skipped when serializing
-	TMap<UObject*, EMarkedTransientReason> TransientAssignments;
+	// Cache of FObjectStatus for every object encountered during the save
+	TMap<TObjectPtr<UObject>, UE::SavePackageUtilities::FObjectStatus> ObjectStatusCache;
 };
 
 const TCHAR* LexToString(ESaveableStatus Status);
diff --git a/Engine/Source/Runtime/CoreUObject/Private/UObject/SavePackage/SavePackageUtilities.cpp b/Engine/Source/Runtime/CoreUObject/Private/UObject/SavePackage/SavePackageUtilities.cpp
--- a/Engine/Source/Runtime/CoreUObject/Private/UObject/SavePackage/SavePackageUtilities.cpp
+++ b/Engine/Source/Runtime/CoreUObject/Private/UObject/SavePackage/SavePackageUtilities.cpp
@@ -612,10 +612,10 @@ void GetCDOSubobjects(UObject* CDO, TArray<UObject*>& Subobjects)
 		}
 	}
 }
-	
-bool IsStrippedEditorOnlyObject(const UObject* InObject, EEditorOnlyObjectFlags Flags)
+
+#if WITH_EDITORONLY_DATA
+bool CanStripEditorOnlyImportsAndExports()
 {
-#if WITH_EDITOR
 	// Configurable via ini setting
 	static struct FCanStripEditorOnlyExportsAndImports
 	{
@@ -626,17 +626,10 @@ bool IsStrippedEditorOnlyObject(const UObject* InObject, EEditorOnlyObjectFlags
 			GConfig->GetBool(TEXT("Core.System"), TEXT("CanStripEditorOnlyExportsAndImports"), bCanStripEditorOnlyObjects, GEngineIni);
 		}
 		FORCEINLINE operator bool() const { return bCanStripEditorOnlyObjects; }
-	} CanStripEditorOnlyExportsAndImports;
-	if (!CanStripEditorOnlyExportsAndImports)
-	{
-		return false;
-	}
-	
-	return IsEditorOnlyObjectInternal(InObject, Flags);
-#else
-	return true;
-#endif
+	} CanStripEditorOnlyExportsAndImportsData;
+	return CanStripEditorOnlyExportsAndImportsData;
 }
+#endif
 
 bool IsUpdatingLoadedPath(bool bIsCooking, const FPackagePath& TargetPackagePath, uint32 SaveFlags)
 {
@@ -851,8 +844,12 @@ bool IsEditorOnlyObject(const UObject* InObject, bool bCheckRecursive, bool bChe
 namespace UE::SavePackageUtilities
 {
 
-bool IsEditorOnlyObjectInternal(const UObject* InObject, EEditorOnlyObjectFlags Flags)
+bool IsEditorOnlyObjectWithoutWritingCache(const UObject* InObject, EEditorOnlyObjectFlags Flags,
+	TFunctionRef<EEditorOnlyObjectResult(const UObject*)> LookupInCache,
+	TFunctionRef<void(const UObject*, bool)> AddToCache)
 {
+	check(InObject);
+
 	bool bCheckRecursive = EnumHasAnyFlags(Flags, EEditorOnlyObjectFlags::CheckRecursive);
 	bool bIgnoreEditorOnlyClass = EnumHasAnyFlags(Flags, EEditorOnlyObjectFlags::ApplyHasNonEditorOnlyReferences) &&
 		InObject->HasNonEditorOnlyReferences();
@@ -861,7 +858,6 @@ bool IsEditorOnlyObjectInternal(const UObject* InObject, EEditorOnlyObjectFlags
 	PRAGMA_ENABLE_DEPRECATION_WARNINGS;
 
 	DECLARE_SCOPE_CYCLE_COUNTER(TEXT("IsEditorOnlyObject"), STAT_IsEditorOnlyObject, STATGROUP_LoadTime);
-	check(InObject);
 
 	// CDOs must be included if their class and archetype and outer are included.
 	// Ignore their value of IsEditorOnly
@@ -882,7 +878,8 @@ bool IsEditorOnlyObjectInternal(const UObject* InObject, EEditorOnlyObjectFlags
 	{
 		if (InObject->HasAnyFlags(RF_ClassDefaultObject))
 		{
-			// The default package is not editor-only, and it is part of a cycle that would cause infinite recursion: DefaultPackage -> GetOuter() -> Package:/Script/CoreUObject -> GetArchetype() -> DefaultPackage
+			// The default package is not editor-only, and it is part of a cycle that would cause infinite recursion:
+			// DefaultPackage -> GetOuter() -> Package:/Script/CoreUObject -> GetArchetype() -> DefaultPackage
 			return false;
 		}
 		Package = static_cast<const UPackage*>(InObject);
@@ -907,7 +904,7 @@ bool IsEditorOnlyObjectInternal(const UObject* InObject, EEditorOnlyObjectFlags
 		UObject* Outer = InObject->GetOuter();
 		if (Outer && Outer != Package)
 		{
-			if (IsEditorOnlyObjectInternal(Outer, Flags))
+			if (IsEditorOnlyObjectInternal(Outer, Flags, LookupInCache, AddToCache))
 			{
 				return true;
 			}
@@ -918,20 +915,20 @@ bool IsEditorOnlyObjectInternal(const UObject* InObject, EEditorOnlyObjectFlags
 			if (InStruct)
 			{
 				const UStruct* SuperStruct = InStruct->GetSuperStruct();
-				if (SuperStruct && IsEditorOnlyObjectInternal(SuperStruct, Flags))
+				if (SuperStruct && IsEditorOnlyObjectInternal(SuperStruct, Flags, LookupInCache, AddToCache))
 				{
 					return true;
 				}
 			}
 			else
 			{
-				if (IsEditorOnlyObjectInternal(InObject->GetClass(), Flags))
+				if (IsEditorOnlyObjectInternal(InObject->GetClass(), Flags, LookupInCache, AddToCache))
 				{
 					return true;
 				}
 
 				UObject* Archetype = InObject->GetArchetype();
-				if (Archetype && IsEditorOnlyObjectInternal(Archetype, Flags))
+				if (Archetype && IsEditorOnlyObjectInternal(Archetype, Flags, LookupInCache, AddToCache))
 				{
 					return true;
 				}
@@ -941,6 +938,33 @@ bool IsEditorOnlyObjectInternal(const UObject* InObject, EEditorOnlyObjectFlags
 	return false;
 }
 
+bool IsEditorOnlyObjectInternal(const UObject* InObject, EEditorOnlyObjectFlags Flags)
+{
+	return IsEditorOnlyObjectInternal(InObject, Flags,
+		[](const UObject* Object)
+		{
+			return EEditorOnlyObjectResult::Uninitialized;
+		},
+		[](const UObject* Object, bool bEditorOnly)
+		{
+		});
+}
+
+bool IsEditorOnlyObjectInternal(const UObject* InObject, EEditorOnlyObjectFlags Flags,
+	TFunctionRef<EEditorOnlyObjectResult(const UObject*)> LookupInCache,
+	TFunctionRef<void(const UObject*, bool)> AddToCache)
+{
+	EEditorOnlyObjectResult Result = LookupInCache(InObject);
+	if (Result != EEditorOnlyObjectResult::Uninitialized)
+	{
+		return Result == EEditorOnlyObjectResult::EditorOnly;
+	}
+
+	bool bResult = IsEditorOnlyObjectWithoutWritingCache(InObject, Flags, LookupInCache, AddToCache);
+	AddToCache(InObject, bResult);
+	return bResult;
+}
+
 } // namespace UE::SavePackageUtilities
 
 void FObjectImportSortHelper::SortImports(FLinkerSave* Linker)
diff --git a/Engine/Source/Runtime/CoreUObject/Private/UObject/SavePackage/SavePackageUtilities.h b/Engine/Source/Runtime/CoreUObject/Private/UObject/SavePackage/SavePackageUtilities.h
--- a/Engine/Source/Runtime/CoreUObject/Private/UObject/SavePackage/SavePackageUtilities.h
+++ b/Engine/Source/Runtime/CoreUObject/Private/UObject/SavePackage/SavePackageUtilities.h
@@ -376,9 +376,20 @@ enum class EEditorOnlyObjectFlags
 };
 ENUM_CLASS_FLAGS(EEditorOnlyObjectFlags);
 
-/** Returns result of IsEditorOnlyObjectInternal if Engine:[Core.System]:CanStripEditorOnlyExportsAndImports (ini) is set to true */
-bool IsStrippedEditorOnlyObject(const UObject* InObject, EEditorOnlyObjectFlags Flags);
+enum class EEditorOnlyObjectResult
+{
+	Uninitialized,
+	EditorOnly,
+	NonEditorOnly,
+};
+#if WITH_EDITORONLY_DATA
+bool CanStripEditorOnlyImportsAndExports();
+#endif
 
+/** Returns result of IsEditorOnlyObjectInternal if Engine:[Core.System]:CanStripEditorOnlyExportsAndImports (ini) is set to true */
+bool IsEditorOnlyObjectInternal(const UObject* InObject, EEditorOnlyObjectFlags Flags,
+	TFunctionRef<EEditorOnlyObjectResult(const UObject*)> LookupInCache,
+	TFunctionRef<void(const UObject*, bool)> AddToCache);
 bool IsEditorOnlyObjectInternal(const UObject* InObject, EEditorOnlyObjectFlags Flags);
 
 }
