Match properties by both name and type when loading into an IDO

#rnx

diff --git a/Engine/Source/Runtime/CoreUObject/Private/UObject/Class.cpp b/Engine/Source/Runtime/CoreUObject/Private/UObject/Class.cpp
--- a/Engine/Source/Runtime/CoreUObject/Private/UObject/Class.cpp
+++ b/Engine/Source/Runtime/CoreUObject/Private/UObject/Class.cpp
@@ -1667,6 +1667,19 @@ void UStruct::SerializeVersionedTaggedProperties(FStructuredArchive::FSlot Slot,
 					}
 				}
 
+				// Try to match the type when impersonating because there can be multiple properties with the same name.
+				if (UNLIKELY(SerializeContext->bImpersonateProperties && Property && !Property->CanSerializeFromTypeName(Tag.GetType())))
+				{
+					for (FProperty* PropertyMatch = PropertyLink; PropertyMatch; PropertyMatch = PropertyMatch->PropertyLinkNext)
+					{
+						if (PropertyMatch->GetFName() == Tag.Name && PropertyMatch->CanSerializeFromTypeName(Tag.GetType()))
+						{
+							Property = PropertyMatch;
+							break;
+						}
+					}
+				}
+
 				TOptional<UE::FSerializedPropertyPathScope> SerializedPropertyPath;
 				if (SerializeContext->bTrackSerializedPropertyPath)
 				{
