Extend assert in SetLinker to provide more information about which objects are involved when there is a fatal mismatch between the object and its expected entry in the linker export table.

#rnx

diff --git a/Engine/Source/Runtime/CoreUObject/Private/UObject/UObjectLinker.cpp b/Engine/Source/Runtime/CoreUObject/Private/UObject/UObjectLinker.cpp
--- a/Engine/Source/Runtime/CoreUObject/Private/UObject/UObjectLinker.cpp
+++ b/Engine/Source/Runtime/CoreUObject/Private/UObject/UObjectLinker.cpp
@@ -121,9 +121,11 @@ void UObject::SetLinker( FLinkerLoad* LinkerLoad, int32 LinkerIndex, bool bShoul
 			*GetFullName(),
 			*LexToString(GetFlags()),
 			LinkerLoad ? *LinkerLoad->GetDebugName() : TEXT("nullptr"));
-		check(Existing.Linker->ExportMap[Existing.LinkerIndex].Object!=nullptr);
-		check(Existing.Linker->ExportMap[Existing.LinkerIndex].Object==this);
-		Existing.Linker->ExportMap[Existing.LinkerIndex].ResetObject();
+
+		FObjectExport& ExportObject = Existing.Linker->ExportMap[Existing.LinkerIndex];
+		checkf(ExportObject.Object != nullptr, TEXT("Expected ExportMap[%d].Object to not be null for this ('%s')"), Existing.LinkerIndex, *GetFName().ToString());
+		checkf(ExportObject.Object == this, TEXT("Expected ExportMap[%d].Object ('%s') to equal this ('%s')"), Existing.LinkerIndex, *ExportObject.Object->GetFName().ToString(), *GetFName().ToString());
+		ExportObject.ResetObject();
 	}
 
 	if (Existing.Linker == LinkerLoad)
