[CrashFix] removed erronous code that was causing all transient objects to have IDOs, and also disabled IDOs for class default objects
[CrashFix] fixed a bug in CPFUO and in StaticDuplicate that was causing several objects to not get IDOs on reinstantiation
[Feature] added helper structs for enabling/disabling IDO flags on load and save

diff --git a/Engine/Source/Runtime/CoreUObject/Internal/UObject/InstanceDataObjectUtils.h b/Engine/Source/Runtime/CoreUObject/Internal/UObject/InstanceDataObjectUtils.h
--- a/Engine/Source/Runtime/CoreUObject/Internal/UObject/InstanceDataObjectUtils.h
+++ b/Engine/Source/Runtime/CoreUObject/Internal/UObject/InstanceDataObjectUtils.h
@@ -22,7 +22,7 @@ namespace UE
  *
  * Pass nullptr to query if the system is enabled.
  */
-bool IsInstanceDataObjectSupportEnabled(UObject* Object = nullptr);
+bool IsInstanceDataObjectSupportEnabled(const UObject* Object = nullptr);
 bool StructContainsLooseProperties(const UStruct* Struct);
 
 /** Generate a UClass that contains the union of the properties of PropertyTree and OwnerClass. */
diff --git a/Engine/Source/Runtime/CoreUObject/Internal/UObject/PropertyBagRepository.h b/Engine/Source/Runtime/CoreUObject/Internal/UObject/PropertyBagRepository.h
--- a/Engine/Source/Runtime/CoreUObject/Internal/UObject/PropertyBagRepository.h
+++ b/Engine/Source/Runtime/CoreUObject/Internal/UObject/PropertyBagRepository.h
@@ -8,6 +8,7 @@
 #include "UObject/GCObject.h"
 
 class UObject;
+class FMemoryArchive;
 
 namespace UE
 {
@@ -71,6 +72,10 @@ public:
 	 * @return				- Custom InstanceDataObject object, UClass derived from associated bag.
 	 */
 	COREUOBJECT_API UObject* CreateInstanceDataObject(UObject* Owner, FArchive* Archive = nullptr);
+	COREUOBJECT_API UObject* DuplicateInstanceDataObject(UObject* SourceOwner, UObject* DestOwner);
+
+	// called at the end of postload to copy data from Owner to its IDO
+	COREUOBJECT_API void PostLoadInstanceDataObject(const UObject* Owner);
 
 	// TODO: Restrict property bag  destruction to within UObject::BeginDestroy() & FPropertyBagProperty destructor.
 	// Removes bag, InstanceDataObject, and all associated data for this object.
@@ -154,6 +159,33 @@ private:
 
 	// Instantiate InstanceDataObject within BagData. Returns InstanceDataObject object. 
 	void CreateInstanceDataObjectUnsafe(UObject* Owner, FPropertyBagAssociationData& BagData, FArchive* Archive = nullptr);
+
+	static void CopyTaggedProperties (const UObject* Source, UObject* Dest);
+};
+
+// construct this context in the same scope as an object is being serialized to support InstanceDataObjects.
+// if saving, this will simply set flags. If loading, an IDO will be constructed at the end of the scope when needed
+struct FScopedIDOSerializationContext
+{
+	COREUOBJECT_API FScopedIDOSerializationContext(UObject* InObject, FArchive& Archive);
+	explicit COREUOBJECT_API FScopedIDOSerializationContext(UObject* InObject); // assumes save
+	COREUOBJECT_API ~FScopedIDOSerializationContext();
+
+	bool bHasIDOSupport;
+	bool bCreateIDO;
+	FArchive* Archive;
+	UObject* const Object;
+	const int64 PreSerializeOffset;
+	TOptional<TGuardValue<bool>> ScopedTrackSerializedPropertyPath;
+	TOptional<TGuardValue<bool>> ScopedSerializeUnknownProperty;
+	TOptional<TGuardValue<bool>> ScopedImpersonateProperties;
+	TOptional<TGuardValue<UObject*>> ScopedSerializedObject;
+private:
+	// if we're loading and IDO should be created, this will be called when the context falls out of scope
+	void FinishCreatingInstanceDataObject() const;
+	
+	FScopedIDOSerializationContext(const FScopedIDOSerializationContext&) = delete;
+	FScopedIDOSerializationContext& operator=(const FScopedIDOSerializationContext&) = delete;
 };
 
 } // UE
diff --git a/Engine/Source/Runtime/CoreUObject/Private/UObject/InstanceDataObjectUtils.cpp b/Engine/Source/Runtime/CoreUObject/Private/UObject/InstanceDataObjectUtils.cpp
--- a/Engine/Source/Runtime/CoreUObject/Private/UObject/InstanceDataObjectUtils.cpp
+++ b/Engine/Source/Runtime/CoreUObject/Private/UObject/InstanceDataObjectUtils.cpp
@@ -138,11 +138,11 @@ namespace UE
 		TEXT("Allows property bags and IDOs to be created for supported classes.")
 	);
 
-	bool IsInstanceDataObjectSupportEnabled(UObject* InObject)
+	bool IsInstanceDataObjectSupportEnabled(const UObject* InObject)
 	{
 		// Note: NULL is a valid (default) input here; in that case we just return the enable flag.
 		bool bIsEnabled = bEnableIDOSupport;
-		if (bIsEnabled && InObject && !InObject->IsInPackage(GetTransientPackage()))
+		if (bIsEnabled && InObject)
 		{
 			// Property bag placeholder objects are always enabled for IDO support
 			if (UE::FPropertyBagRepository::IsPropertyBagPlaceholderObject(InObject))
@@ -150,6 +150,12 @@ namespace UE
 				return true;
 			}
 
+			// reinst classes shouldn't use IDOs
+			if (InObject->GetClass()->HasAnyClassFlags(CLASS_NewerVersionExists))
+			{
+				return false;
+			}
+
 			//@todo FH: change to check trait when available or use config object
 			const UClass* ObjClass = InObject->GetClass();
 			while (ObjClass && ObjClass->GetClass()->GetFName() != NAME_VerseClass)
diff --git a/Engine/Source/Runtime/CoreUObject/Private/UObject/LinkerLoad.cpp b/Engine/Source/Runtime/CoreUObject/Private/UObject/LinkerLoad.cpp
--- a/Engine/Source/Runtime/CoreUObject/Private/UObject/LinkerLoad.cpp
+++ b/Engine/Source/Runtime/CoreUObject/Private/UObject/LinkerLoad.cpp
@@ -4694,17 +4694,7 @@ void FLinkerLoad::Preload( UObject* Object )
 #endif
 
 					// Toggle support for IDOs
-					FUObjectSerializeContext* LoadContext = FUObjectThreadContext::Get().GetSerializeContext();
-					// Enable property path tracking when IDO support is enabled. Both the property path name tree and IDO creation require the paths.
-					const bool bHasIDOSupport = UE::IsInstanceDataObjectSupportEnabled(Object);
-					TGuardValue<bool> ScopedTrackSerializedPropertyPath(LoadContext->bTrackSerializedPropertyPath, bHasIDOSupport);
-					// Enable tracking of initialized properties when loading an IDO, which is implied by impersonation being enabled.
-					const bool bLoadingIDO = bHasIDOSupport && LoadContext->bImpersonateProperties;
-					TGuardValue<bool> ScopedTrackInitializedProperties(LoadContext->bTrackInitializedProperties, bLoadingIDO);
-					// Enable creation of a property path name tree to track any property that does not match the current class schema,
-					// except when impersonation is enabled because that implies we are deserializing an IDO.
-					const bool bCreateIDO = bHasIDOSupport && !LoadContext->bImpersonateProperties;
-					TGuardValue<bool> ScopedTrackUnknownProperties(LoadContext->bTrackUnknownProperties, bCreateIDO);
+					UE::FScopedIDOSerializationContext IDOLoadContext(Object, *this);
 
 					if (Object->HasAnyFlags(RF_ClassDefaultObject))
 					{
diff --git a/Engine/Source/Runtime/CoreUObject/Private/UObject/Obj.cpp b/Engine/Source/Runtime/CoreUObject/Private/UObject/Obj.cpp
--- a/Engine/Source/Runtime/CoreUObject/Private/UObject/Obj.cpp
+++ b/Engine/Source/Runtime/CoreUObject/Private/UObject/Obj.cpp
@@ -1304,7 +1304,7 @@ void UObject::ConditionalPostLoad()
 		const bool bIDOEnabled = PropertyBagRepository.IsInstanceDataObjectSupportEnabled(this) && !LoadContext->bImpersonateProperties;
 		if (bIDOEnabled)
 		{
-			PropertyBagRepository.CreateInstanceDataObject(this);
+			PropertyBagRepository.PostLoadInstanceDataObject(this);
 		}
 
 		{
@@ -1875,7 +1875,6 @@ void UObject::SerializeScriptProperties( FStructuredArchive::FSlot Slot ) const
 		const UObject* ThisObject = this;
 		if (const UObject* Impersonator = UE::Private::GetDataImpersonator(ThisObject))
 		{
-			ensureAlwaysMsgf(!HasAnyFlags(RF_ClassDefaultObject), TEXT("CDO '%s' shoudn't be impersonated"), *ThisObject->GetPathName());
 			ThisObject = Impersonator;
 			ObjClass = ThisObject->GetClass();
 
diff --git a/Engine/Source/Runtime/CoreUObject/Private/UObject/PropertyBagRepository.cpp b/Engine/Source/Runtime/CoreUObject/Private/UObject/PropertyBagRepository.cpp
--- a/Engine/Source/Runtime/CoreUObject/Private/UObject/PropertyBagRepository.cpp
+++ b/Engine/Source/Runtime/CoreUObject/Private/UObject/PropertyBagRepository.cpp
@@ -519,6 +519,59 @@ UObject* FPropertyBagRepository::CreateInstanceDataObject(UObject* Owner, FArchi
 	return BagData.InstanceDataObject;
 }
 
+UObject* FPropertyBagRepository::DuplicateInstanceDataObject(UObject* SourceOwner, UObject* DestOwner)
+{
+	if (FPropertyBagAssociationData* SourceData = AssociatedData.Find(SourceOwner))
+	{
+		check(SourceData->InstanceDataObject);
+		FPropertyBagAssociationData& DestData = AssociatedData.FindOrAdd(DestOwner);
+		ensure(!DestData.InstanceDataObject);
+
+		// get outer pointer for new IDO
+		TObjectPtr<UObject>* OuterPtr = nullptr;
+		if (FPropertyBagAssociationData* OuterData = AssociatedData.Find(DestOwner->GetOuter()))
+		{
+			OuterPtr = &OuterData->InstanceDataObject;
+		}
+		if (!OuterPtr || !*OuterPtr)
+		{
+			OuterPtr = &Namespaces.FindOrAdd(DestOwner->GetOuter());
+			if (*OuterPtr == nullptr)
+			{
+				*OuterPtr = CreatePackage(nullptr); // TODO: replace with dummy object
+			}
+		}
+
+		// construct InstanceDataObject
+		FStaticConstructObjectParameters Params(SourceData->InstanceDataObject->GetClass());
+		Params.SetFlags |= EObjectFlags::RF_Transactional;
+		Params.Name = DestOwner->GetFName();
+		Params.Outer = *OuterPtr;
+		DestData.InstanceDataObject = StaticConstructObject_Internal(Params);
+		InstanceDataObjectToOwner.Add(DestData.InstanceDataObject, DestOwner);
+
+		CopyTaggedProperties(SourceData->InstanceDataObject, DestData.InstanceDataObject);
+
+		DestData.bNeedsFixup = SourceData->bNeedsFixup;
+		return DestData.InstanceDataObject;
+	}
+	return nullptr;
+}
+
+void FPropertyBagRepository::PostLoadInstanceDataObject(const UObject* Owner)
+{
+	// fixups may have been applied to the instance during PostLoad and they need to be copied to its IDO
+	FPropertyBagRepositoryLock LockRepo(this);
+	if (FPropertyBagAssociationData* BagData = AssociatedData.Find(Owner))
+	{
+		if (BagData->InstanceDataObject)
+		{
+			// copy data from owner to IDO
+			CopyTaggedProperties(Owner, BagData->InstanceDataObject);
+		}
+	}
+}
+
 // TODO: Remove this? Bag destruction to be handled entirely via UObject::BeginDestroy() (+ FPropertyBagProperty destructor)?
 void FPropertyBagRepository::DestroyOuterBag(const UObject* Owner)
 {
@@ -621,12 +674,12 @@ void FPropertyBagRepository::CreateInstanceDataObjectUnsafe(UObject* Owner, FPro
 
 	BagData.bNeedsFixup = UE::StructContainsLooseProperties(InstanceDataObjectClass);
 
-	TObjectPtr<UObject>* OuterPtr;
+	TObjectPtr<UObject>* OuterPtr = nullptr;
 	if (FPropertyBagAssociationData* OuterData = AssociatedData.Find(Owner->GetOuter()))
 	{
 		OuterPtr = &OuterData->InstanceDataObject;
 	}
-	else
+	if (!OuterPtr || !*OuterPtr)
 	{
 		OuterPtr = &Namespaces.FindOrAdd(Owner->GetOuter());
 		if (*OuterPtr == nullptr)
@@ -653,25 +706,14 @@ void FPropertyBagRepository::CreateInstanceDataObjectUnsafe(UObject* Owner, FPro
 	// setup load context to mark properties the that were set by serialization
 	FUObjectSerializeContext* LoadContext = FUObjectThreadContext::Get().GetSerializeContext();
 	TGuardValue<bool> ScopedImpersonateProperties(LoadContext->bImpersonateProperties, true);
-	
-	auto CopyTaggedProperties = [](const UObject* Source, UObject* Dest)
-	{
-		TArray<uint8> Buffer;
-		Buffer.Reserve(Source->GetClass()->GetStructureSize());
-		FObjectWriter Writer(Buffer);
-		Source->GetClass()->SerializeTaggedProperties(Writer, (uint8*)Source, Source->GetClass(), (uint8*)Source->GetArchetype());
-		
-		FObjectReader Reader(Buffer);
-		Reader.ArMergeOverrides = true;
-		Dest->GetClass()->SerializeTaggedProperties(Reader, (uint8*)Dest, Dest->GetClass(), (uint8*)Dest->GetArchetype());
-	};
-	
-	if (Archive)
+
+	FLinkerLoad* Linker = Owner->GetLinker();
+	if (Archive && Archive != Linker)
 	{
 		// re-deserialize Owner but redirect it into the IDO instead using impersonation
 		Owner->Serialize(*Archive);
 	}
-	else if (FLinkerLoad* Linker = Owner->GetLinker())
+	else if (Linker)
 	{
 		// TODO: @jordan.hoffmann - this is very inefficient! We should remove this call to Preload. To do so, we'd need to change MarkPropertySetBySerialization
 		// to cache the serialized property list in the property bag instead of the structs. We'd also need to copy the property bag values to the IDO
@@ -695,6 +737,93 @@ void FPropertyBagRepository::CreateInstanceDataObjectUnsafe(UObject* Owner, FPro
 	}
 }
 
+void FPropertyBagRepository::CopyTaggedProperties(const UObject* Source, UObject* Dest)
+{
+	TArray<uint8> Buffer;
+	Buffer.Reserve(Source->GetClass()->GetStructureSize());
+	FObjectWriter Writer(Buffer);
+	Source->GetClass()->SerializeTaggedProperties(Writer, (uint8*)Source, Source->GetClass(), nullptr);
+		
+	FObjectReader Reader(Buffer);
+	Reader.ArMergeOverrides = true;
+	Dest->GetClass()->SerializeTaggedProperties(Reader, (uint8*)Dest, Dest->GetClass(), nullptr);
+}
+
+FScopedIDOSerializationContext::FScopedIDOSerializationContext(UObject* InObject, FArchive& InArchive)
+	: Archive(&InArchive)
+    , Object(InObject)
+	, PreSerializeOffset(InArchive.Tell())
+{
+	FUObjectSerializeContext* SerializeContext = FUObjectThreadContext::Get().GetSerializeContext();
+	bHasIDOSupport = FPropertyBagRepository::IsInstanceDataObjectSupportEnabled(Object);
+	bCreateIDO = bHasIDOSupport && !SerializeContext->bImpersonateProperties && Archive->IsLoading();
+	
+	if (bHasIDOSupport)
+	{
+		if (Archive->IsLoading())
+		{
+			// Enable creation of a property path name tree to track any property that does not match the current class schema,
+			// except when impersonation is enabled because that implies we are deserializing an IDO.
+			ScopedTrackSerializedPropertyPath.Emplace(SerializeContext->bTrackSerializedPropertyPath, bCreateIDO);
+			ScopedSerializeUnknownProperty.Emplace(SerializeContext->bTrackUnknownProperties, bCreateIDO);
+			ScopedSerializedObject.Emplace(SerializeContext->SerializedObject, Object);
+			
+			// Enable tracking of initialized properties when loading an IDO, which is implied by impersonation being enabled.
+			const bool bLoadingIDO = bHasIDOSupport && SerializeContext->bImpersonateProperties;
+			TGuardValue<bool> ScopedTrackInitializedProperties(SerializeContext->bTrackInitializedProperties, bLoadingIDO);
+		}
+		else
+		{
+			ScopedImpersonateProperties.Emplace(SerializeContext->bImpersonateProperties, bHasIDOSupport);
+		}
+	}
+}
+
+FScopedIDOSerializationContext::FScopedIDOSerializationContext(UObject* InObject)
+	: bCreateIDO(false)
+	, Archive(nullptr)
+	, Object(InObject)
+	, PreSerializeOffset(0)
+{
+	FUObjectSerializeContext* SerializeContext = FUObjectThreadContext::Get().GetSerializeContext();
+	bHasIDOSupport = FPropertyBagRepository::IsInstanceDataObjectSupportEnabled(Object);
+	ScopedImpersonateProperties.Emplace(SerializeContext->bImpersonateProperties, true);
+}
+
+FScopedIDOSerializationContext::~FScopedIDOSerializationContext()
+{
+	if (bCreateIDO)
+	{
+		FinishCreatingInstanceDataObject();
+	}
+}
+
+void FScopedIDOSerializationContext::FinishCreatingInstanceDataObject() const
+{
+	if (Archive == Object->GetLinker())
+	{
+		// when using the linker, the repository will handle offsets
+		FPropertyBagRepository::Get().CreateInstanceDataObject(Object, Archive);
+	}
+	else
+	{
+		check(Archive);
+		const int64 PostSerializeOffset = Archive->Tell();
+	
+		// CreateInstanceDataObject will re-call DstObject->Serialize(*this) so set the seek pointer back before DestObject in the archive
+		Archive->Seek(PreSerializeOffset);
+			
+		FPropertyBagRepository::Get().CreateInstanceDataObject(Object, Archive);
+
+		// make sure seek pointer is back to where it should be
+		if (!ensure(Archive->Tell() == PostSerializeOffset))
+		{
+			// for some reason CreateInstanceDataObject read a different amount of data than expected... reset seek pointer back to where it should be
+			Archive->Seek(PostSerializeOffset);
+		}
+	}
+}
+
 // Not sure this is necessary.
 void FPropertyBagRepository::ShrinkMaps()
 {
diff --git a/Engine/Source/Runtime/CoreUObject/Private/UObject/SavePackage2.cpp b/Engine/Source/Runtime/CoreUObject/Private/UObject/SavePackage2.cpp
--- a/Engine/Source/Runtime/CoreUObject/Private/UObject/SavePackage2.cpp
+++ b/Engine/Source/Runtime/CoreUObject/Private/UObject/SavePackage2.cpp
@@ -3,6 +3,7 @@
 
 #include "UObject/SavePackage.h"
 #include "UObject/LinkerLoad.h"
+#include "UObject/PropertyBagRepository.h"
 
 #if UE_WITH_SAVEPACKAGE
 #include "AssetRegistry/AssetData.h"
@@ -2950,11 +2951,7 @@ ESavePackageResult InnerSave(FSaveContext& SaveContext)
 	SaveContext.SetSerializeContext(SerializeContext);
 	SaveContext.SetEDLCookChecker(&FEDLCookCheckerThreadState::Get());
 
-	TOptional<TGuardValue<bool>> IDOImpersonationScope;
-	if (UE::IsInstanceDataObjectSupportEnabled())
-	{
-		IDOImpersonationScope.Emplace(SerializeContext->bImpersonateProperties, true);
-	}
+	UE::FScopedIDOSerializationContext IDOSaveContext(nullptr);
 
 	// Create slow task dialog if needed
 	const int32 TotalSaveSteps = 3;
diff --git a/Engine/Source/Runtime/CoreUObject/Private/UObject/UObjectGlobals.cpp b/Engine/Source/Runtime/CoreUObject/Private/UObject/UObjectGlobals.cpp
--- a/Engine/Source/Runtime/CoreUObject/Private/UObject/UObjectGlobals.cpp
+++ b/Engine/Source/Runtime/CoreUObject/Private/UObject/UObjectGlobals.cpp
@@ -3254,6 +3254,14 @@ UObject* StaticDuplicateObjectEx( FObjectDuplicationParameters& Parameters )
 			}
 		}
 	}
+
+	// if Source has an IDO, make one for Dest and copy it
+	UE::FPropertyBagRepository& PropertyBagRepository = UE::FPropertyBagRepository::Get();
+	if (UE::FPropertyBagRepository::IsInstanceDataObjectSupportEnabled(Parameters.SourceObject))
+	{
+		PropertyBagRepository.DuplicateInstanceDataObject(Parameters.SourceObject, DupRootObject);
+	}
+	
 	return DupRootObject;
 }
 
diff --git a/Engine/Source/Runtime/Engine/Private/UnrealEngine.cpp b/Engine/Source/Runtime/Engine/Private/UnrealEngine.cpp
--- a/Engine/Source/Runtime/Engine/Private/UnrealEngine.cpp
+++ b/Engine/Source/Runtime/Engine/Private/UnrealEngine.cpp
@@ -17075,12 +17075,7 @@ public:
 		}
 #endif // USE_STABLE_LOCALIZATION_KEYS
 
-		TOptional<TGuardValue<bool>> ScopedImpersonateProperties;
-		if (UE::FPropertyBagRepository::IsInstanceDataObjectSupportEnabled(SrcObject))
-		{
-			FUObjectSerializeContext* SerializeContext = FUObjectThreadContext::Get().GetSerializeContext();
-			ScopedImpersonateProperties.Emplace(SerializeContext->bImpersonateProperties, true);
-		}
+		UE::FScopedIDOSerializationContext IDOSaveContext(SrcObject, *this);
 
 		SrcObject->Serialize(*this);
 	}
@@ -17148,6 +17143,9 @@ public:
 		}
 #endif // USE_STABLE_LOCALIZATION_KEYS
 
+		// Enable IDO, if needed
+		UE::FScopedIDOSerializationContext IDOLoadContext(Obj, *this);
+		
 		Obj->Serialize(*this);
 	}
 
@@ -17246,36 +17244,9 @@ public:
 #endif // USE_STABLE_LOCALIZATION_KEYS
 
 		// Enable IDO, if needed
-		FUObjectSerializeContext* LoadContext = FUObjectThreadContext::Get().GetSerializeContext();
-		TOptional<TGuardValue<bool>> ScopedTrackSerializedPropertyPath;
-		TOptional<TGuardValue<bool>> ScopedSerializeUnknownProperty;
-		TOptional<TGuardValue<UObject*>> ScopedSerializedObject;
-
-		// Do not enable IDO when impersonation is enabled as we are probably deserializing an IDO already at that point
-		bool bIDOEnabled = UE::FPropertyBagRepository::IsInstanceDataObjectSupportEnabled(DstObject) && !LoadContext->bImpersonateProperties;
-		if (bIDOEnabled)
-		{
-			// this will had property path tracking and create a property bag to hold data not matching the current class schema
-			ScopedTrackSerializedPropertyPath.Emplace(LoadContext->bTrackSerializedPropertyPath, true);
-			ScopedSerializeUnknownProperty.Emplace(LoadContext->bTrackUnknownProperties, true);
-			ScopedSerializedObject.Emplace(LoadContext->SerializedObject, DstObject);
-		}
-
-		const int64 PreSerializeOffset = Tell();
+		UE::FScopedIDOSerializationContext IDOLoadContext(DstObject, *this);
+		
 		DstObject->Serialize(*this);
-		const int64 PostSerializeOffset = Tell();
-
-		if (bIDOEnabled)
-		{
-			// CreateInstanceDataObject will re-call DstObject->Serialize(*this) so set the seek pointer back before DestObject in the archive
-			Seek(PreSerializeOffset);
-			UE::FPropertyBagRepository::Get().CreateInstanceDataObject(DstObject, this);
-			if (!ensure(Tell() == PostSerializeOffset))
-			{
-				// for some reason CreateInstanceDataObject read a different amount of data than expected... reset seek pointer back to where it should be
-				Seek(PostSerializeOffset);
-			}
-		}
 	}
 
 	//~ Begin FArchive Interface
@@ -17456,6 +17427,9 @@ void UEngine::CopyPropertiesForUnrelatedObjects(UObject* OldObject, UObject* New
 			pRecord->OldInstance = OldInstance;
 			OldInstanceMap.Add(OldInstance->GetPathName(OldObject), SavedInstances.Num() - 1);
 			const uint32 AdditionalPortFlags = Params.bCopyDeprecatedProperties ? PPF_UseDeprecatedProperties : PPF_None;
+			
+			UE::FScopedIDOSerializationContext IDOSaveContext(OldInstance);
+			
 			FObjectWriter SubObjWriter(OldInstance, pRecord->SavedProperties, true, true, Params.bDoDelta, AdditionalPortFlags);
 		}
 	}
