Removing support for multithreaded destruction of UObjects

diff --git a/Engine/Config/BaseEngine.ini b/Engine/Config/BaseEngine.ini
--- a/Engine/Config/BaseEngine.ini
+++ b/Engine/Config/BaseEngine.ini
@@ -1542,7 +1542,6 @@ gc.CreateGCClusters=True
 gc.MinGCClusterSize=5
 gc.AssetClustreringEnabled=False
 gc.ActorClusteringEnabled=False
-gc.MultithreadedDestructionEnabled=True
 gc.VerifyUObjectsAreNotFGCObjects=False
 gc.GarbageEliminationEnabled=True
 
diff --git a/Engine/Plugins/Experimental/Mutable/Source/CustomizableObject/Private/MuCO/CustomizableObjectSystem.cpp b/Engine/Plugins/Experimental/Mutable/Source/CustomizableObject/Private/MuCO/CustomizableObjectSystem.cpp
--- a/Engine/Plugins/Experimental/Mutable/Source/CustomizableObject/Private/MuCO/CustomizableObjectSystem.cpp
+++ b/Engine/Plugins/Experimental/Mutable/Source/CustomizableObject/Private/MuCO/CustomizableObjectSystem.cpp
@@ -179,7 +179,7 @@ FUpdateContextPrivate::FUpdateContextPrivate(UCustomizableObjectInstance& InInst
 
 FUpdateContextPrivate::~FUpdateContextPrivate()
 {
-	check(IsInGameThread() || IsInGarbageCollectorThread());
+	check(IsInGameThread());
 
 	if (UCustomizableObjectSystem::IsCreated())
 	{
diff --git a/Engine/Source/Runtime/CoreUObject/Private/UObject/GarbageCollection.cpp b/Engine/Source/Runtime/CoreUObject/Private/UObject/GarbageCollection.cpp
--- a/Engine/Source/Runtime/CoreUObject/Private/UObject/GarbageCollection.cpp
+++ b/Engine/Source/Runtime/CoreUObject/Private/UObject/GarbageCollection.cpp
@@ -207,14 +207,6 @@ static FAutoConsoleVariableRef CIncrementalBeginDestroyEnabled(
 	ECVF_Default
 );
 
-int32 GMultithreadedDestructionEnabled = 0;
-static FAutoConsoleVariableRef CMultithreadedDestructionEnabled(
-	TEXT("gc.MultithreadedDestructionEnabled"),
-	GMultithreadedDestructionEnabled,
-	TEXT("If true, the engine will free objects' memory from a worker thread"),
-	ECVF_Default
-);
-
 #if WITH_VERSE_VM || defined(__INTELLISENSE__)
 bool GEnableFrankenGC = true;
 static FAutoConsoleVariableRef CEnableFrankenGC(
@@ -312,6 +304,14 @@ static FAutoConsoleVariableRef CVarVerifyNoUnreachableObjects(
 	TEXT("Enables or disables no unreachable objects are reachable verification"),
 	ECVF_Default
 );
+
+static int32 GVerifyObjectsDestroyed = UE_BUILD_DEBUG;
+static FAutoConsoleVariableRef CVarVerifyObjectsDestroyed(
+	TEXT("gc.VerifyObjectsDestroyed"),
+	GVerifyObjectsDestroyed,
+	TEXT("Enables or disables verification if all objects have been destroyed"),
+	ECVF_Default
+);
 #endif // VERIFY_DISREGARD_GC_ASSUMPTIONS
 
 static bool GForceEnableDebugGCProcessor = 0;
@@ -690,7 +690,7 @@ namespace UE::GC::Test
 		FUObjectItem* ObjectItem = GUObjectArray.IndexToObjectUnsafeForGC(GUObjectArray.ObjectToIndex(Object));
 		UE::GC::Private::FGCFlags::SetMaybeUnreachable_ForGC(ObjectItem);
 	}
-} // namespace UE::GC::Privates
+} // namespace UE::GC::Test
 // END Hidden functions
 
 void OnDisregardForGCSetDisabled(int32 NumObjects)
@@ -730,156 +730,67 @@ void OnDisregardForGCSetDisabled(int32 NumObjects)
 }
 
 /**
- * Helper class for destroying UObjects on a worker thread
+ * Helper class for destroying UObjects
  */
-class FAsyncPurge : public FRunnable
-{
-	/** Thread to run the worker FRunnable on. Destroys objects. */
-	volatile FRunnableThread* Thread;
-	/** Id of the worker thread */
-	uint32 AsyncPurgeThreadId;
-	/** Stops this thread */
-	FThreadSafeCounter StopTaskCounter;
-	/** Event that triggers the UObject destruction */
-	FEvent* BeginPurgeEvent;
-	/** Event that signales the UObject destruction is finished */
-	FEvent* FinishedPurgeEvent;
+class FObjectPurge
+{
 	/** Current index into the global unreachable objects array (GUnreachableObjects) of the object being destroyed */
-	int32 ObjCurrentPurgeObjectIndex;
-	/** Number of objects deferred to the game thread to destroy */
-	std::atomic<int32> NumObjectsToDestroyOnGameThread;
-	/** Number of objectsalready destroyed on the game thread */
-	int32 NumObjectsDestroyedOnGameThread;
-	/** Current index into the global unreachable objects array (GUnreachableObjects) of the object being destroyed on the game thread */
-	int32 ObjCurrentPurgeObjectIndexOnGameThread;
-	/** Number of unreachable objects the last time single-threaded tick was called */
-	int32 LastUnreachableObjectsCount;
+	int32 ObjCurrentPurgeObjectIndex = 0;
 	/** Stats for the number of objects destroyed */
-	int32 ObjectsDestroyedSinceLastMarkPhase;
-
-	/** [PURGE/GAME THREAD] Destroys objects that are unreachable */
-	template <bool bMultithreaded> // Having this template argument lets the compiler strip unnecessary checks
-	bool TickDestroyObjects(bool bUseTimeLimit, double TimeLimit, double StartTime)
-	{
-		TRACE_CPUPROFILER_EVENT_SCOPE(FAsyncPurge::TickDestroyObjects);
-		const int32 TimeLimitEnforcementGranularityForDeletion = 100;
-		int32 ProcessedObjectsCount = 0;
-		bool bFinishedDestroyingObjects = true;
+	int32 ObjectsDestroyedSinceLastMarkPhase = 0;
+	/** True if all objects from the last purge phase have been destroyed */
+	bool bFinishedDestroyingObjects = true;
 
-		// Avoid fetch_add synchronization for each object sent to the game-thread and instead load the value
-		// once into a local variable and replace by a simple store to publish the value to the other thread.
-		// This is safe because we know only this thread is going to modify it.
-		int32 LocalNumObjectsToDestroyOnGameThread = NumObjectsToDestroyOnGameThread.load(std::memory_order_acquire);
-		
-		while (ObjCurrentPurgeObjectIndex < GUnreachableObjects.Num())
-		{
-			FUObjectItem* ObjectItem = GUnreachableObjects[ObjCurrentPurgeObjectIndex];
-			check(ObjectItem->IsUnreachable());
-
-			UObject* Object = (UObject*)ObjectItem->Object;
-			check(Object->HasAllFlags(RF_FinishDestroyed | RF_BeginDestroyed));
-			if (!bMultithreaded || Object->IsDestructionThreadSafe())
-			{
-				// Can't lock once for the entire batch here as it could hold the lock for too long
-				GUObjectArray.LockInternalArray();
-				Object->~UObject();
-				GUObjectArray.UnlockInternalArray();
-				GUObjectAllocator.FreeUObject(Object);
-				GUnreachableObjects[ObjCurrentPurgeObjectIndex] = nullptr;
-			}
-			else
-			{
-				NumObjectsToDestroyOnGameThread.store(++LocalNumObjectsToDestroyOnGameThread, std::memory_order_release);
-			}
-			++ProcessedObjectsCount;
-			++ObjectsDestroyedSinceLastMarkPhase;
-			++ObjCurrentPurgeObjectIndex;
+public:
 
-			// Time slicing when running on the game thread
-			if (!bMultithreaded && bUseTimeLimit && (ProcessedObjectsCount == TimeLimitEnforcementGranularityForDeletion) && (ObjCurrentPurgeObjectIndex < GUnreachableObjects.Num()))
-			{
-				ProcessedObjectsCount = 0;
-				if ((FPlatformTime::Seconds() - StartTime) > TimeLimit)
-				{
-					bFinishedDestroyingObjects = false;
-					break;
-				}				
-			}
-		}
+	/** Returns true if the destruction process is finished */
+	FORCEINLINE bool IsFinished() const
+	{
 		return bFinishedDestroyingObjects;
 	}
 
-	/** [GAME THREAD] Destroys objects that are unreachable and couldn't be destroyed on the worker thread */
-	bool TickDestroyGameThreadObjects(bool bUseTimeLimit, double TimeLimit, double StartTime)
+	/** Begins the UObject purge */
+	void Begin()
 	{
-		TRACE_CPUPROFILER_EVENT_SCOPE(FAsyncPurge::TickDestroyGameThreadObjects);
-		const int32 TimeLimitEnforcementGranularityForDeletion = 100;
-		int32 ProcessedObjectsCount = 0;
-		bool bFinishedDestroyingObjects = true;
-
-		// Lock once for the entire batch
-		GUObjectArray.LockInternalArray();
-
-		// Cache the number of objects to destroy locally. The number may grow later but that's ok, we'll catch up to it in the next tick
-		const int32 LocalNumObjectsToDestroyOnGameThread = NumObjectsToDestroyOnGameThread.load(std::memory_order_acquire);
-
-		while (NumObjectsDestroyedOnGameThread < LocalNumObjectsToDestroyOnGameThread && ObjCurrentPurgeObjectIndexOnGameThread < GUnreachableObjects.Num())
-		{
-			FUObjectItem* ObjectItem = GUnreachableObjects[ObjCurrentPurgeObjectIndexOnGameThread];
-			if (ObjectItem)
-			{
-				GUnreachableObjects[ObjCurrentPurgeObjectIndexOnGameThread] = nullptr;
-				UObject* Object = (UObject*)ObjectItem->Object;
-				Object->~UObject();
-				GUObjectAllocator.FreeUObject(Object);
-				++ProcessedObjectsCount;
-				++NumObjectsDestroyedOnGameThread;
-
-				if (bUseTimeLimit && (ProcessedObjectsCount == TimeLimitEnforcementGranularityForDeletion) && NumObjectsDestroyedOnGameThread < LocalNumObjectsToDestroyOnGameThread)
-				{
-					ProcessedObjectsCount = 0;
-					if ((FPlatformTime::Seconds() - StartTime) > TimeLimit)
-					{
-						bFinishedDestroyingObjects = false;
-						break;
-					}
-				}
-			}
-			++ObjCurrentPurgeObjectIndexOnGameThread;
-		}
-
-		GUObjectArray.UnlockInternalArray();
-
-		// Make sure that when we reach the end of GUnreachableObjects array, there's no objects to destroy left
-		check(!bFinishedDestroyingObjects || NumObjectsDestroyedOnGameThread == LocalNumObjectsToDestroyOnGameThread);
-
-		// Note that even though NumObjectsToDestroyOnGameThread may have been incremented by now or still hasn't but it will be 
-		// after we report we're done with all objects, it doesn't matter since we don't care about the result of this function in MT mode
-		return bFinishedDestroyingObjects;
+		check(IsFinished()); // In single-threaded mode we need to be finished or the condition below will hang
+		ObjCurrentPurgeObjectIndex = 0;
+		ObjectsDestroyedSinceLastMarkPhase = 0;
 	}
 
-	/** [GAME THREAD] Destroys objects that are unreachable for the first tick */
-	bool TickDestroyGameThreadObjectsFirst(bool bUseTimeLimit, double TimeLimit, double StartTime)
+	/** 
+	 * Destroys UObjects on the game thread 
+	 * @param bUseTimeLimit true if TimeLimit should be respected
+	 * @param TimeLimit Time limit (in seconds) for the purge
+	 * @param StartTime Time incremental destruction has begun
+	 * @return true if all objects have been purged, false otherwise
+	 */
+	FORCENOINLINE bool DestroyObjects(bool bUseTimeLimit, double TimeLimit, double StartTime)
 	{
-		TRACE_CPUPROFILER_EVENT_SCOPE(FAsyncPurge::TickDestroyGameThreadObjects);
-		const int32 TimeLimitEnforcementGranularityForDeletion = 100;
+		TRACE_CPUPROFILER_EVENT_SCOPE(FObjectPurge::DestroyObjects);
+		constexpr int32 TimeLimitEnforcementGranularityForDeletion = 100;
 		int32 ProcessedObjectsCount = 0;
-		bool bFinishedDestroyingObjects = true;
+		bFinishedDestroyingObjects = true;
 
-		// Lock once for the entire batch
 		GUObjectArray.LockInternalArray();
 
-		while (ObjCurrentPurgeObjectIndexOnGameThread < GUnreachableObjects.Num())
+		while (ObjCurrentPurgeObjectIndex < GUnreachableObjects.Num())
 		{
-			FUObjectItem* ObjectItem = GUnreachableObjects[ObjCurrentPurgeObjectIndexOnGameThread];
-			GUnreachableObjects[ObjCurrentPurgeObjectIndexOnGameThread] = nullptr;
+			FUObjectItem* ObjectItem = GUnreachableObjects[ObjCurrentPurgeObjectIndex];
+			check(ObjectItem->IsUnreachable());
+
 			UObject* Object = (UObject*)ObjectItem->Object;
+			check(Object->HasAllFlags(RF_FinishDestroyed | RF_BeginDestroyed));
+
 			Object->~UObject();
 			GUObjectAllocator.FreeUObject(Object);
+			GUnreachableObjects[ObjCurrentPurgeObjectIndex] = nullptr;
+
 			++ProcessedObjectsCount;
-			++ObjCurrentPurgeObjectIndexOnGameThread;
+			++ObjectsDestroyedSinceLastMarkPhase;
+			++ObjCurrentPurgeObjectIndex;
 
-			if (bUseTimeLimit && (ProcessedObjectsCount == TimeLimitEnforcementGranularityForDeletion))
+			// Time slicing when running on the game thread
+			if (bUseTimeLimit && (ProcessedObjectsCount == TimeLimitEnforcementGranularityForDeletion) && (ObjCurrentPurgeObjectIndex < GUnreachableObjects.Num()))
 			{
 				ProcessedObjectsCount = 0;
 				if ((FPlatformTime::Seconds() - StartTime) > TimeLimit)
@@ -889,140 +800,16 @@ class FAsyncPurge : public FRunnable
 				}
 			}
 		}
-		NumObjectsDestroyedOnGameThread = ObjCurrentPurgeObjectIndexOnGameThread;
-		NumObjectsToDestroyOnGameThread.store(ObjCurrentPurgeObjectIndexOnGameThread, std::memory_order_release);
-		ObjCurrentPurgeObjectIndex = ObjCurrentPurgeObjectIndexOnGameThread;
-		ObjectsDestroyedSinceLastMarkPhase = ObjCurrentPurgeObjectIndexOnGameThread;
 
 		GUObjectArray.UnlockInternalArray();
 
 		return bFinishedDestroyingObjects;
 	}
 
-	/** Waits for the worker thread to finish destroying objects */
-	void WaitForAsyncDestructionToFinish()
-	{
-		FinishedPurgeEvent->Wait();
-	}
-
-public:
-
-	/** 
-	 * Constructor
-	 * @param bMultithreaded if true, the destruction of objects will happen on a worker thread
-	 */
-	FAsyncPurge(bool bMultithreaded)
-		: Thread(nullptr)
-		, AsyncPurgeThreadId(0)
-		, BeginPurgeEvent(nullptr)
-		, FinishedPurgeEvent(nullptr)
-		, ObjCurrentPurgeObjectIndex(0)
-		, NumObjectsToDestroyOnGameThread(0)
-		, NumObjectsDestroyedOnGameThread(0)
-		, ObjCurrentPurgeObjectIndexOnGameThread(0)
-		, LastUnreachableObjectsCount(0)
-		, ObjectsDestroyedSinceLastMarkPhase(0)
-	{
-		BeginPurgeEvent = FPlatformProcess::GetSynchEventFromPool(true);
-		FinishedPurgeEvent = FPlatformProcess::GetSynchEventFromPool(true);
-		FinishedPurgeEvent->Trigger();
-		if (bMultithreaded)
-		{
-			check(FPlatformProcess::SupportsMultithreading());
-			FPlatformAtomics::InterlockedExchangePtr((void**)&Thread, FRunnableThread::Create(this, TEXT("FAsyncPurge"), 0, TPri_BelowNormal));			
-		}
-		else
-		{
-			AsyncPurgeThreadId = GGameThreadId;
-		}
-	}
-
-	virtual ~FAsyncPurge()
-	{
-		check(IsFinished());
-		delete Thread;
-		Thread = nullptr;
-		FPlatformProcess::ReturnSynchEventToPool(BeginPurgeEvent);
-		FPlatformProcess::ReturnSynchEventToPool(FinishedPurgeEvent);
-		BeginPurgeEvent = nullptr;
-		FinishedPurgeEvent = nullptr;
-	}
-
-	/** Returns true if the destruction process is finished */
-	FORCEINLINE bool IsFinished() const
-	{
-		if (Thread)
-		{
-			return FinishedPurgeEvent->Wait(0, true) && NumObjectsToDestroyOnGameThread == NumObjectsDestroyedOnGameThread;
-		}
-		else
-		{
-			return (ObjCurrentPurgeObjectIndex >= LastUnreachableObjectsCount && NumObjectsToDestroyOnGameThread == NumObjectsDestroyedOnGameThread);
-		}
-	}
-
-	/** [MAIN THREAD] Adds objects to the purge queue */
-	void BeginPurge(bool bUseTimeLimit, double TimeLimit, double StartTime)
-	{
-		check(IsFinished()); // In single-threaded mode we need to be finished or the condition below will hang
-		if (FinishedPurgeEvent->Wait())
-		{
-			FinishedPurgeEvent->Reset();
-
-			ObjCurrentPurgeObjectIndex = 0;
-			ObjectsDestroyedSinceLastMarkPhase = 0;
-			NumObjectsToDestroyOnGameThread = 0;
-			NumObjectsDestroyedOnGameThread = 0;
-			ObjCurrentPurgeObjectIndexOnGameThread = 0;
-
-			// In MT case we need to consume all available time limit before triggering Async Purge.
-			// Otherwise Async thread almost always will not process enought objects for us to destroy on GT.
-			// That happens because it takes significant amount time to unpark the thread and switch context before
-			// it can process and increment NumObjectsToDestroyOnGameThread.
-			if (!!Thread)
-			{
-				if (!TickDestroyGameThreadObjectsFirst(bUseTimeLimit, TimeLimit, StartTime))
-				{
-					// Start async helper when the timelimit reached and there is still stuff to do
-					BeginPurgeEvent->Trigger();
-				}
-				else
-				{
-					FinishedPurgeEvent->Trigger();
-				}
-			}
-		}
-	}
-
-	/** [GAME THREAD] Ticks the purge process on the game thread */
-	void TickPurge(bool bUseTimeLimit, double TimeLimit, double StartTime)
-	{
-		bool bCanStartDestroyingGameThreadObjects = true;
-		if (!Thread)
-		{
-			// If we're running single-threaded we need to tick the main loop here too
-			LastUnreachableObjectsCount = GUnreachableObjects.Num();
-			bCanStartDestroyingGameThreadObjects = TickDestroyObjects<false>(bUseTimeLimit, TimeLimit, StartTime);
-		}
-		if (bCanStartDestroyingGameThreadObjects)
-		{
-			do
-			{
-				// Deal with objects that couldn't be destroyed on the worker thread. This will do nothing when running single-threaded
-				bool bFinishedDestroyingObjectsOnGameThread = TickDestroyGameThreadObjects(bUseTimeLimit, TimeLimit, StartTime);
-				if (!Thread && bFinishedDestroyingObjectsOnGameThread)
-				{
-					// This only gets triggered here in single-threaded mode
-					FinishedPurgeEvent->Trigger();
-				}
-			} while (!bUseTimeLimit && !IsFinished());
-		}
-	}
-
 	/** Returns the number of objects already destroyed */
 	int32 GetObjectsDestroyedSinceLastMarkPhase() const
 	{
-		return ObjectsDestroyedSinceLastMarkPhase - NumObjectsToDestroyOnGameThread + NumObjectsDestroyedOnGameThread;
+		return ObjectsDestroyedSinceLastMarkPhase;
 	}
 
 	/** Resets the number of objects already destroyed */
@@ -1031,67 +818,22 @@ public:
 		ObjectsDestroyedSinceLastMarkPhase = 0;
 	}
 
-	/** 
-	  * Returns true if this function is called from the async destruction thread. 
-	  * It will also return true if we're running single-threaded and this function is called on the game thread
-	  */
-	bool IsInAsyncPurgeThread() const
-	{
-		return AsyncPurgeThreadId == FPlatformTLS::GetCurrentThreadId();
-	}
-
-	/* Returns true if it can run multi-threaded destruction */
-	bool IsMultithreaded() const
-	{
-		return !!Thread;
-	}
-
-	//~ Begin FRunnable Interface.
-	virtual bool Init()
-	{
-		return true;
-	}
-
-	virtual uint32 Run()
-	{
-		AsyncPurgeThreadId = FPlatformTLS::GetCurrentThreadId();
-		
-		while (StopTaskCounter.GetValue() == 0)
-		{
-			if (BeginPurgeEvent->Wait(15, true))
-			{
-				BeginPurgeEvent->Reset();
-				TickDestroyObjects<true>(/* bUseTimeLimit = */ false, /* TimeLimit = */ 0.0f, /* StartTime = */ 0.0);
-				FinishedPurgeEvent->Trigger();
-			}
-		}
-		FinishedPurgeEvent->Trigger();
-		return 0;
-	}
-
-	virtual void Stop()
-	{
-		StopTaskCounter.Increment();
-	}
-	//~ End FRunnable Interface
-
-	void VerifyAllObjectsDestroyed()
+	FORCENOINLINE void VerifyAllObjectsDestroyed() const
 	{
 		for (FUObjectItem* ObjectItem : GUnreachableObjects)
 		{
 			UE_CLOG(ObjectItem, LogGarbage, Fatal, TEXT("Object 0x%016llx has not been destroyed during async purge"), (int64)(PTRINT)ObjectItem->Object);
 		}
 	}
-};
-static FAsyncPurge* GAsyncPurge = nullptr;
+} GUObjectPurge;
 
 /**
-  * Returns true if this function is called from the async destruction thread.
+  * (Deprecated) Returns true if this function is called from the async destruction thread.
   * It will also return true if we're running single-threaded and this function is called on the game thread
   */
 bool IsInGarbageCollectorThread()
 {
-	return GAsyncPurge ? GAsyncPurge->IsInAsyncPurgeThread() : IsInGameThread();
+	return IsInGameThread();
 }
 
 //////////////////////////////////////////////////////////////////////////
@@ -3936,8 +3678,6 @@ void TFastReferenceCollector<ProcessorType, CollectorType>::ProcessStructs(Dispa
 void ShutdownGarbageCollection()
 {
 	UE::GC::FContextPoolScope().Cleanup();
-	delete GAsyncPurge;
-	GAsyncPurge = nullptr;
 #if ENABLE_GC_HISTORY
 	FGCHistory::Get().Cleanup();
 #endif
@@ -4783,17 +4523,6 @@ bool IncrementalDestroyGarbage(bool bUseTimeLimit, double TimeLimit)
 	using namespace UE::GC::Private;
 
 	TRACE_CPUPROFILER_EVENT_SCOPE(IncrementalDestroyGarbage);
-	const bool bMultithreadedPurge = !ShouldForceSingleThreadedGC() && GMultithreadedDestructionEnabled;
-	if (!GAsyncPurge)
-	{
-		GAsyncPurge = new FAsyncPurge(bMultithreadedPurge);
-	}
-	else if (GAsyncPurge->IsMultithreaded() != bMultithreadedPurge)
-	{
-		check(GAsyncPurge->IsFinished());
-		delete GAsyncPurge;
-		GAsyncPurge = new FAsyncPurge(bMultithreadedPurge);
-	}
 
 	bool bCompleted = false;
 	bool bTimeLimitReached = false;
@@ -5047,18 +4776,21 @@ bool IncrementalDestroyGarbage(bool bUseTimeLimit, double TimeLimit)
 		int32 ProcessCount = 0;
 		if (GObjCurrentPurgeObjectIndexNeedsReset)
 		{
-			GAsyncPurge->BeginPurge(bUseTimeLimit, TimeLimit, GCStartTime);
+			GUObjectPurge.Begin();
 			// Reset the reset flag but don't reset the actual index yet for stat purposes
 			GObjCurrentPurgeObjectIndexNeedsReset = false;
 		}
 
-		GAsyncPurge->TickPurge(bUseTimeLimit, TimeLimit, GCStartTime);
+		GUObjectPurge.DestroyObjects(bUseTimeLimit, TimeLimit, GCStartTime);
 
-		if (GAsyncPurge->IsFinished())
+		if (GUObjectPurge.IsFinished())
 		{
-#if UE_BUILD_DEBUG
-			GAsyncPurge->VerifyAllObjectsDestroyed();
-#endif
+#if VERIFY_DISREGARD_GC_ASSUMPTIONS
+			if (GVerifyObjectsDestroyed)
+			{
+				GUObjectPurge.VerifyAllObjectsDestroyed();
+			}
+#endif // VERIFY_DISREGARD_GC_ASSUMPTIONS
 
 			bCompleted = true;
 			// Incremental purge is finished, time to reset variables.
@@ -5067,13 +4799,13 @@ bool IncrementalDestroyGarbage(bool bUseTimeLimit, double TimeLimit)
 			GObjCurrentPurgeObjectIndexNeedsReset			= true;
 
 			// Log status information.
-			const int32 PurgedObjectCountSinceLastMarkPhase = GAsyncPurge->GetObjectsDestroyedSinceLastMarkPhase();
+			const int32 PurgedObjectCountSinceLastMarkPhase = GUObjectPurge.GetObjectsDestroyedSinceLastMarkPhase();
 			UE_LOG(LogGarbage, Log, TEXT("GC purged %i objects (%i -> %i) in %.3fms"), PurgedObjectCountSinceLastMarkPhase, 
 				GObjectCountDuringLastMarkPhase.GetValue(), 
 				GObjectCountDuringLastMarkPhase.GetValue() - PurgedObjectCountSinceLastMarkPhase,
 				(FPlatformTime::Seconds() - IncrementalDestroyGarbageStartTime) * 1000);
 			UE::GC::GDetailedStats.LogPurgeStats(PurgedObjectCountSinceLastMarkPhase);
-			GAsyncPurge->ResetObjectsDestroyedSinceLastMarkPhase();
+			GUObjectPurge.ResetObjectsDestroyedSinceLastMarkPhase();
 		}
 	}
 
@@ -5086,7 +4818,7 @@ bool IncrementalDestroyGarbage(bool bUseTimeLimit, double TimeLimit)
 		UE_LOG(LogGarbage, Log, TEXT("%.3f ms for incrementally purging unreachable objects (FinishDestroyed: %d, Destroyed: %d / %d)"),
 			ElapsedTime * 1000,
 			GObjCurrentPurgeObjectIndex,
-			GAsyncPurge->GetObjectsDestroyedSinceLastMarkPhase(),
+			GUObjectPurge.GetObjectsDestroyedSinceLastMarkPhase(),
 			GUnreachableObjects.Num());
 	}
 
@@ -6240,11 +5972,6 @@ void PurgeAllUObjectsOnExit()
 		IncrementalPurgeGarbage(false);
 	}
 
-	// From now on we'll be destroying objects without time limit during exit purge
-	// so doing it on a separate thread doesn't make anything faster,
-	// also the exit purge is not a standard GC pass so no need to overcompilcate things
-	GMultithreadedDestructionEnabled = false;
-
 	// Make sure no other threads manipulate UObjects
 	AcquireGCLock();
 
diff --git a/Engine/Source/Runtime/CoreUObject/Private/UObject/UObjectArray.cpp b/Engine/Source/Runtime/CoreUObject/Private/UObject/UObjectArray.cpp
--- a/Engine/Source/Runtime/CoreUObject/Private/UObject/UObjectArray.cpp
+++ b/Engine/Source/Runtime/CoreUObject/Private/UObject/UObjectArray.cpp
@@ -335,7 +335,7 @@ void FUObjectArray::FreeUObjectIndex(UObjectBase* Object)
 	LLM_SCOPE(ELLMTag::UObject);
 
 	// This should only be happening on the game thread (GC runs only on game thread when it's freeing objects)
-	check(IsInGameThread() || IsInGarbageCollectorThread());
+	check(IsInGameThread());
 
 	// No need to call LockInternalArray(); here as it should already be locked by GC
 
diff --git a/Engine/Source/Runtime/CoreUObject/Public/UObject/Object.h b/Engine/Source/Runtime/CoreUObject/Public/UObject/Object.h
--- a/Engine/Source/Runtime/CoreUObject/Public/UObject/Object.h
+++ b/Engine/Source/Runtime/CoreUObject/Public/UObject/Object.h
@@ -606,6 +606,7 @@ public:
 	*
 	* @return	true if this object's destructor is thread safe
 	*/
+	UE_DEPRECATED(5.5, "Multithreaded destruction of UObjects is no longer supported.")
 	COREUOBJECT_API virtual bool IsDestructionThreadSafe() const;
 
 	/**
diff --git a/Engine/Source/Runtime/CoreUObject/Public/UObject/Package.h b/Engine/Source/Runtime/CoreUObject/Public/UObject/Package.h
--- a/Engine/Source/Runtime/CoreUObject/Public/UObject/Package.h
+++ b/Engine/Source/Runtime/CoreUObject/Public/UObject/Package.h
@@ -347,7 +347,6 @@ public:
 	virtual bool NeedsLoadForClient() const override { return true; }
 	virtual bool NeedsLoadForServer() const override { return true; }
 	virtual bool IsPostLoadThreadSafe() const override;
-	virtual bool IsDestructionThreadSafe() const override { return true; }
 
 #if WITH_EDITORONLY_DATA
 	/** Sets the bLoadedByEditorPropertiesOnly flag */
diff --git a/Engine/Source/Runtime/CoreUObject/Public/UObject/UObjectGlobals.h b/Engine/Source/Runtime/CoreUObject/Public/UObject/UObjectGlobals.h
--- a/Engine/Source/Runtime/CoreUObject/Public/UObject/UObjectGlobals.h
+++ b/Engine/Source/Runtime/CoreUObject/Public/UObject/UObjectGlobals.h
@@ -869,6 +869,7 @@ COREUOBJECT_API float GetAsyncLoadPercentage( const FName& PackageName );
 /**
 * Whether we are running on the Garbage Collector Thread
 */
+UE_DEPRECATED(5.5, "IsInGarbageCollectorThread() is now equivalent to calling IsInGameThread()")
 COREUOBJECT_API bool IsInGarbageCollectorThread();
 
 /** 
diff --git a/Engine/Source/Runtime/Engine/Classes/Engine/CoreSettings.h b/Engine/Source/Runtime/Engine/Classes/Engine/CoreSettings.h
--- a/Engine/Source/Runtime/Engine/Classes/Engine/CoreSettings.h
+++ b/Engine/Source/Runtime/Engine/Classes/Engine/CoreSettings.h
@@ -194,11 +194,6 @@ protected:
 		ToolTip = "If true, the engine will destroy objects incrementally using time limit each frame."))
 	uint32 IncrementalBeginDestroyEnabled : 1;
 
-	UPROPERTY(EditAnywhere, config, Category = Optimization, meta = (
-		ConsoleVariable = "gc.MultithreadedDestructionEnabled", DisplayName = "Multithreaded Destruction Enabled",
-		ToolTip = "If true, the engine will free objects' memory on a worker thread."))
-	uint32 MultithreadedDestructionEnabled : 1;
-
 	UPROPERTY(EditAnywhere, config, Category = Optimization, meta = (
 		ConsoleVariable = "gc.CreateGCClusters", DisplayName = "Create Garbage Collector UObject Clusters",
 		ToolTip = "If true, the engine will attempt to create clusters of objects for better garbage collection performance."))
diff --git a/Engine/Source/Runtime/Engine/Private/CoreSettings.cpp b/Engine/Source/Runtime/Engine/Private/CoreSettings.cpp
--- a/Engine/Source/Runtime/Engine/Private/CoreSettings.cpp
+++ b/Engine/Source/Runtime/Engine/Private/CoreSettings.cpp
@@ -201,7 +201,6 @@ UGarbageCollectionSettings::UGarbageCollectionSettings()
 	FlushStreamingOnGC = false;
 	AllowParallelGC = true;
 	IncrementalBeginDestroyEnabled = true;
-	MultithreadedDestructionEnabled = true;
 	NumRetriesBeforeForcingGC = 0;
 	MaxObjectsNotConsideredByGC = 0;
 	MaxObjectsInEditor = 12 * 1024 * 1024;
diff --git a/Engine/Source/Runtime/Engine/Public/LightMap.h b/Engine/Source/Runtime/Engine/Public/LightMap.h
--- a/Engine/Source/Runtime/Engine/Public/LightMap.h
+++ b/Engine/Source/Runtime/Engine/Public/LightMap.h
@@ -92,7 +92,7 @@ public:
 	}
 	void Release()
 	{
-		check(IsInGameThread() || IsInAsyncLoadingThread() || IsInGarbageCollectorThread());
+		check(IsInGameThread() || IsInAsyncLoadingThread());
 		checkSlow(NumRefs > 0);
 		if(--NumRefs == 0)
 		{
diff --git a/Engine/Source/Runtime/Engine/Public/ShadowMap.h b/Engine/Source/Runtime/Engine/Public/ShadowMap.h
--- a/Engine/Source/Runtime/Engine/Public/ShadowMap.h
+++ b/Engine/Source/Runtime/Engine/Public/ShadowMap.h
@@ -258,7 +258,7 @@ public:
 	}
 	void Release()
 	{
-		check(IsInGameThread() || IsInAsyncLoadingThread() || IsInGarbageCollectorThread());
+		check(IsInGameThread() || IsInAsyncLoadingThread());
 		checkSlow(NumRefs > 0);
 		if (--NumRefs == 0)
 		{
