[OFPA] Don't consider objects with external packages as assets when choosing if SetEditorProperty can execute on external objects.

diff --git a/Engine/Source/Runtime/CoreUObject/Private/UObject/PropertyAccessUtil.cpp b/Engine/Source/Runtime/CoreUObject/Private/UObject/PropertyAccessUtil.cpp
--- a/Engine/Source/Runtime/CoreUObject/Private/UObject/PropertyAccessUtil.cpp
+++ b/Engine/Source/Runtime/CoreUObject/Private/UObject/PropertyAccessUtil.cpp
@@ -597,8 +597,8 @@ bool IsObjectTemplate(const UObject* InObject)
 		return true;
 	}
 	
-	// Assets can edit default properties
-	if (InObject->IsAsset())
+	// Assets can edit default properties, filter out objects with external packages as they are treated as assets but aren't
+	if (!InObject->IsPackageExternal() && InObject->IsAsset())
 	{
 		return true;
 	}
