Object.h: Update formatting of LoadConfig and SaveConfig to match current default formatting. Whitespace and NULL->nullptr changes only, no code changes.
#rnx

diff --git a/Engine/Source/Runtime/CoreUObject/Private/UObject/Obj.cpp b/Engine/Source/Runtime/CoreUObject/Private/UObject/Obj.cpp
--- a/Engine/Source/Runtime/CoreUObject/Private/UObject/Obj.cpp
+++ b/Engine/Source/Runtime/CoreUObject/Private/UObject/Obj.cpp
@@ -2677,21 +2677,22 @@ void CheckMissingSection(const FString& SectionName, const FString& IniFilename)
 }
 #endif
 
-void UObject::LoadConfig( UClass* ConfigClass/*=NULL*/, const TCHAR* InFilename/*=NULL*/, uint32 PropagationFlags/*=LCPF_None*/, FProperty* PropertyToLoad/*=NULL*/ )
+void UObject::LoadConfig(UClass* ConfigClass/*=nullptr*/, const TCHAR* InFilename/*=nullptr*/,
+	uint32 PropagationFlags/*=LCPF_None*/, FProperty* PropertyToLoad/*=nullptr*/)
 {
 	SCOPE_CYCLE_COUNTER(STAT_LoadConfig);
 
 	// OriginalClass is the class that LoadConfig() was originally called on
 	static thread_local UClass* OriginalClass = nullptr;
 
-	if( !ConfigClass )
+	if (!ConfigClass)
 	{
 		// if no class was specified in the call, this is the OriginalClass
 		ConfigClass = GetClass();
 		OriginalClass = ConfigClass;
 	}
 
-	if( !ConfigClass->HasAnyClassFlags(CLASS_Config) )
+	if (!ConfigClass->HasAnyClassFlags(CLASS_Config))
 	{
 		return;
 	}
@@ -2701,7 +2702,7 @@ void UObject::LoadConfig( UClass* ConfigClass/*=NULL*/, const TCHAR* InFilename/
 	{
 		TFieldIterator<FProperty> It1(Struct1);
 		TFieldIterator<FProperty> It2(Struct2);
-		for (;;++It1,++It2)
+		for (;; ++It1, ++It2)
 		{
 			bool bAtEnd1 = !It1;
 			bool bAtEnd2 = !It2;
@@ -2734,22 +2735,22 @@ void UObject::LoadConfig( UClass* ConfigClass/*=NULL*/, const TCHAR* InFilename/
 #endif
 
 	UClass* ParentClass = ConfigClass->GetSuperClass();
-	if ( ParentClass != NULL )
+	if (ParentClass != nullptr)
 	{
-		if ( ParentClass->HasAnyClassFlags(CLASS_Config) )
+		if (ParentClass->HasAnyClassFlags(CLASS_Config))
 		{
-			if ( (PropagationFlags&UE::LCPF_ReadParentSections) != 0 )
+			if ((PropagationFlags & UE::LCPF_ReadParentSections) != 0)
 			{
 				// call LoadConfig on the parent class
-				LoadConfig( ParentClass, NULL, PropagationFlags, PropertyToLoad );
+				LoadConfig(ParentClass, nullptr, PropagationFlags, PropertyToLoad);
 
 				// if we are also notifying child classes or instances, stop here as this object's properties will be imported as a result of notifying the others
-				if ( (PropagationFlags & (UE::LCPF_PropagateToChildDefaultObjects|UE::LCPF_PropagateToInstances)) != 0 )
+				if ((PropagationFlags & (UE::LCPF_PropagateToChildDefaultObjects | UE::LCPF_PropagateToInstances)) != 0)
 				{
 					return;
 				}
 			}
-			else if ( (PropagationFlags&UE::LCPF_PropagateToChildDefaultObjects) != 0 )
+			else if ((PropagationFlags & UE::LCPF_PropagateToChildDefaultObjects) != 0)
 			{
 				// not propagating the call upwards, but we are propagating the call to all child classes
 				for (TObjectIterator<UClass> It; It; ++It)
@@ -2757,14 +2758,16 @@ void UObject::LoadConfig( UClass* ConfigClass/*=NULL*/, const TCHAR* InFilename/
 					if (It->IsChildOf(ConfigClass))
 					{
 						// mask out the PropgateToParent and PropagateToChildren values
-						It->GetDefaultObject()->LoadConfig(*It, NULL, (PropagationFlags&(UE::LCPF_PersistentFlags|UE::LCPF_PropagateToInstances)), PropertyToLoad);
+						It->GetDefaultObject()->LoadConfig(*It, nullptr,
+							(PropagationFlags & (UE::LCPF_PersistentFlags | UE::LCPF_PropagateToInstances)),
+							PropertyToLoad);
 					}
 				}
 
 				// LoadConfig() was called on this object during iteration, so stop here 
 				return;
 			}
-			else if ( (PropagationFlags&UE::LCPF_PropagateToInstances) != 0 )
+			else if ((PropagationFlags & UE::LCPF_PropagateToInstances) != 0)
 			{
 				// call LoadConfig() on all instances of this class (except the CDO)
 				// Do not propagate this call to parents, and do not propagate to children or instances (would be redundant) 
@@ -2772,18 +2775,20 @@ void UObject::LoadConfig( UClass* ConfigClass/*=NULL*/, const TCHAR* InFilename/
 				{
 					if (It->IsA(ConfigClass))
 					{
-						if ( !GIsEditor )
+						if (!GIsEditor)
 						{
 							// make sure to pass in the class so that OriginalClass isn't reset
-							It->LoadConfig(It->GetClass(), NULL, (PropagationFlags&UE::LCPF_PersistentFlags), PropertyToLoad);
+							It->LoadConfig(It->GetClass(), nullptr, (PropagationFlags & UE::LCPF_PersistentFlags),
+								PropertyToLoad);
 						}
 #if WITH_EDITOR
 						else
 						{
-							It->PreEditChange(NULL);
+							It->PreEditChange(nullptr);
 
 							// make sure to pass in the class so that OriginalClass isn't reset
-							It->LoadConfig(It->GetClass(), NULL, (PropagationFlags&UE::LCPF_PersistentFlags), PropertyToLoad);
+							It->LoadConfig(It->GetClass(), nullptr, (PropagationFlags & UE::LCPF_PersistentFlags),
+								PropertyToLoad);
 
 							It->PostEditChange();
 						}
@@ -2792,25 +2797,29 @@ void UObject::LoadConfig( UClass* ConfigClass/*=NULL*/, const TCHAR* InFilename/
 				}
 			}
 		}
-		else if ( (PropagationFlags&UE::LCPF_PropagateToChildDefaultObjects) != 0 )
+		else if ((PropagationFlags & UE::LCPF_PropagateToChildDefaultObjects) != 0)
 		{
 			// we're at the base-most config class
-			for ( TObjectIterator<UClass> It; It; ++It )
+			for (TObjectIterator<UClass> It; It; ++It)
 			{
-				if ( It->IsChildOf(ConfigClass) )
+				if (It->IsChildOf(ConfigClass))
 				{
-					if ( !GIsEditor )
+					if (!GIsEditor)
 					{
 						// make sure to pass in the class so that OriginalClass isn't reset
-						It->GetDefaultObject()->LoadConfig( *It, NULL, (PropagationFlags&(UE::LCPF_PersistentFlags|UE::LCPF_PropagateToInstances)), PropertyToLoad );
+						It->GetDefaultObject()->LoadConfig(*It, nullptr,
+							(PropagationFlags & (UE::LCPF_PersistentFlags | UE::LCPF_PropagateToInstances)),
+							PropertyToLoad);
 					}
 #if WITH_EDITOR
 					else
 					{
-						It->PreEditChange(NULL);
+						It->PreEditChange(nullptr);
 
 						// make sure to pass in the class so that OriginalClass isn't reset
-						It->GetDefaultObject()->LoadConfig( *It, NULL, (PropagationFlags&(UE::LCPF_PersistentFlags|UE::LCPF_PropagateToInstances)), PropertyToLoad );
+						It->GetDefaultObject()->LoadConfig(*It, nullptr,
+							(PropagationFlags & (UE::LCPF_PersistentFlags | UE::LCPF_PropagateToInstances)),
+							PropertyToLoad);
 
 						It->PostEditChange();
 					}
@@ -2820,24 +2829,26 @@ void UObject::LoadConfig( UClass* ConfigClass/*=NULL*/, const TCHAR* InFilename/
 
 			return;
 		}
-		else if ( (PropagationFlags&UE::LCPF_PropagateToInstances) != 0 )
+		else if ((PropagationFlags & UE::LCPF_PropagateToInstances) != 0)
 		{
-			for ( TObjectIterator<UObject> It; It; ++It )
+			for (TObjectIterator<UObject> It; It; ++It)
 			{
-				if ( It->GetClass() == ConfigClass )
+				if (It->GetClass() == ConfigClass)
 				{
-					if ( !GIsEditor )
+					if (!GIsEditor)
 					{
 						// make sure to pass in the class so that OriginalClass isn't reset
-						It->LoadConfig(It->GetClass(), NULL, (PropagationFlags&UE::LCPF_PersistentFlags), PropertyToLoad);
+						It->LoadConfig(It->GetClass(), nullptr,	(PropagationFlags & UE::LCPF_PersistentFlags),
+							PropertyToLoad);
 					}
 #if WITH_EDITOR
 					else
 					{
-						It->PreEditChange(NULL);
+						It->PreEditChange(nullptr);
 
 						// make sure to pass in the class so that OriginalClass isn't reset
-						It->LoadConfig(It->GetClass(), NULL, (PropagationFlags&UE::LCPF_PersistentFlags), PropertyToLoad);
+						It->LoadConfig(It->GetClass(), nullptr, (PropagationFlags & UE::LCPF_PersistentFlags),
+							PropertyToLoad);
 						It->PostEditChange();
 					}
 #endif // WITH_EDITOR
@@ -2846,11 +2857,8 @@ void UObject::LoadConfig( UClass* ConfigClass/*=NULL*/, const TCHAR* InFilename/
 		}
 	}
 
-	const FString Filename
-	// if a filename was specified, always load from that file
-	=	InFilename
-		? InFilename
-		: GetConfigFilename(this);
+	// if a filename was specified, load all properties from that file no matter which class they come from
+	const FString Filename = InFilename ? InFilename : GetConfigFilename(this);
 
 	const bool bPerObject = UsesPerObjectConfig(this);
 
@@ -2874,7 +2882,7 @@ void UObject::LoadConfig( UClass* ConfigClass/*=NULL*/, const TCHAR* InFilename/
 		{
 			OverrideConfigFile = FConfigCacheIni::FindOrLoadPlatformConfig(LocalOverrideConfig, *GetClass()->ClassConfigName.ToString(), *PreviewPlatform.ToString());
 			bUseConfigOverride = true;
-	}
+		}
 	}
 #endif
 
@@ -2896,8 +2904,8 @@ void UObject::LoadConfig( UClass* ConfigClass/*=NULL*/, const TCHAR* InFilename/
 			GetPathName(Outermost, PathNameString);
 			LongCommitName = Outermost->GetFName();
 		}
-		
-		ClassSection = PathNameString + TEXT(" ") + GetClass()->GetName();		
+
+		ClassSection = PathNameString + TEXT(" ") + GetClass()->GetName();
 
 		FString OverrideClassSection;
 		OverridePerObjectConfigSection(OverrideClassSection);
@@ -2915,21 +2923,24 @@ void UObject::LoadConfig( UClass* ConfigClass/*=NULL*/, const TCHAR* InFilename/
 	}
 
 	// If any of my properties are class variables, then LoadConfig() would also be called for each one of those classes.
-	// Since OrigClass is a static variable, if the value of a class variable is a class different from the current class, 
+	// Since OriginalClass is a static variable, if the value of a class variable is a class different from the current class, 
 	// we'll lose our nice reference to the original class - and cause any variables which were declared after this class variable to fail 
 	// the 'if (OriginalClass != Class)' check....better store it in a temporary place while we do the actual loading of our properties 
 	UClass* MyOrigClass = OriginalClass;
 
-	if ( PropertyToLoad == NULL )
+	if ( PropertyToLoad == nullptr )
 	{
-		UE_LOG(LogConfig, VeryVerbose, TEXT("(%s) '%s' loading configuration from %s"), *ConfigClass->GetName(), *GetName(), *Filename);
+		UE_LOG(LogConfig, VeryVerbose, TEXT("(%s) '%s' loading configuration from %s"),
+			*ConfigClass->GetName(), *GetName(), *Filename);
 	}
 	else
 	{
-		UE_LOG(LogConfig, VeryVerbose, TEXT("(%s) '%s' loading configuration for property %s from %s"), *ConfigClass->GetName(), *GetName(), *PropertyToLoad->GetName(), *Filename);
+		UE_LOG(LogConfig, VeryVerbose, TEXT("(%s) '%s' loading configuration for property %s from %s"),
+			*ConfigClass->GetName(), *GetName(), *PropertyToLoad->GetName(), *Filename);
 	}
 
-	auto GetConfigValue = [&OverrideConfigFile, &bUseConfigOverride](const TCHAR* ClassSection, const TCHAR* Key, const TCHAR* ConfigName, FString& OutValue)
+	auto GetConfigValue = [&OverrideConfigFile, &bUseConfigOverride]
+	(const TCHAR* ClassSection, const TCHAR* Key, const TCHAR* ConfigName, FString& OutValue)
 	{
 		if (bUseConfigOverride)
 		{
@@ -2941,7 +2952,8 @@ void UObject::LoadConfig( UClass* ConfigClass/*=NULL*/, const TCHAR* InFilename/
 		}
 	};
 
-	auto GetConfigSection = [&OverrideConfigFile, &bUseConfigOverride](const TCHAR* SectionName, const TCHAR* ConfigFilename)
+	auto GetConfigSection = [&OverrideConfigFile, &bUseConfigOverride]
+	(const TCHAR* SectionName, const TCHAR* ConfigFilename)
 	{
 		if (bUseConfigOverride)
 		{
@@ -2953,19 +2965,21 @@ void UObject::LoadConfig( UClass* ConfigClass/*=NULL*/, const TCHAR* InFilename/
 		}
 	};
 
-	for ( FProperty* Property = ConfigClass->PropertyLink; Property; Property = Property->PropertyLinkNext )
+	for (FProperty* Property = ConfigClass->PropertyLink; Property; Property = Property->PropertyLinkNext)
 	{
 #if WITH_EDITOR
-		FSoftObjectPathSerializationScope SerializationScope(NAME_None, Property->GetFName(), Property->IsEditorOnlyProperty() ? ESoftObjectPathCollectType::EditorOnlyCollect : ESoftObjectPathCollectType::AlwaysCollect, ESoftObjectPathSerializeType::AlwaysSerialize);
+		FSoftObjectPathSerializationScope SerializationScope(NAME_None, Property->GetFName(),
+			Property->IsEditorOnlyProperty() ? ESoftObjectPathCollectType::EditorOnlyCollect : ESoftObjectPathCollectType::AlwaysCollect,
+			ESoftObjectPathSerializeType::AlwaysSerialize);
 #endif
 
-		if ( !Property->HasAnyPropertyFlags(CPF_Config) )
+		if (!Property->HasAnyPropertyFlags(CPF_Config))
 		{
 			continue;
 		}
 
 		// if we're only supposed to load the value for a specific property, skip all others
-		if ( PropertyToLoad != NULL && PropertyToLoad != Property )
+		if (PropertyToLoad != nullptr && PropertyToLoad != Property)
 		{
 			continue;
 		}
@@ -2980,7 +2994,7 @@ void UObject::LoadConfig( UClass* ConfigClass/*=NULL*/, const TCHAR* InFilename/
 		UClass* OwnerClass = Property->GetOwnerClass();
 
 		UClass* BaseClass = bGlobalConfig ? OwnerClass : ConfigClass;
-		if ( !bPerObject )
+		if (!bPerObject)
 		{
 			ClassSection = BaseClass->GetPathName();
 			LongCommitName = BaseClass->GetOutermost()->GetFName();
@@ -2993,7 +3007,7 @@ void UObject::LoadConfig( UClass* ConfigClass/*=NULL*/, const TCHAR* InFilename/
 		// specifying a value for InFilename will override this behavior (as it does with normal properties)
 		const FString* PropFileNamePtr = &Filename;
 		FString Temp;
-		if (bGlobalConfig && InFilename == NULL)
+		if (bGlobalConfig && InFilename == nullptr)
 		{
 			Temp = OwnerClass->GetConfigName();
 			PropFileNamePtr = &Temp;
@@ -3023,7 +3037,7 @@ void UObject::LoadConfig( UClass* ConfigClass/*=NULL*/, const TCHAR* InFilename/
 		FArrayProperty* Array = CastField<FArrayProperty>(Property);
 		FSetProperty* SetProperty = CastField<FSetProperty>(Property);
 
-		if (Array || SetProperty)
+		if ((Array || SetProperty))
 		{
 			const FConfigSection* Sec = GetConfigSection(*ClassSection, *PropFileName);
 			if (!Sec && bPerObject && ClassPathSection.Len())
@@ -3170,7 +3184,7 @@ void UObject::LoadConfig( UClass* ConfigClass/*=NULL*/, const TCHAR* InFilename/
 							{
 								// If we failed to import the value as an entire set, try to import it as a single
 								// element
-								const TCHAR* ElementImportResult = 
+								const TCHAR* ElementImportResult =
 									SetProperty->ElementProp->ImportText_Direct(*ListString, TempElement, this, PortFlags);
 
 								if (ElementImportResult && ElementImportResult != *ListString)
@@ -3188,12 +3202,12 @@ void UObject::LoadConfig( UClass* ConfigClass/*=NULL*/, const TCHAR* InFilename/
 				}
 			}
 		}
-		
+
 		if (!bProcessedProperty)
 		{
-			for( int32 i=0; i<Property->ArrayDim; i++ )
+			for (int32 i = 0; i < Property->ArrayDim; i++)
 			{
-				if( Property->ArrayDim!=1 )
+				if (Property->ArrayDim != 1)
 				{
 					Key = FString::Printf(TEXT("%s[%i]"), *Property->GetName(), i);
 				}
@@ -3208,7 +3222,7 @@ void UObject::LoadConfig( UClass* ConfigClass/*=NULL*/, const TCHAR* InFilename/
 
 				if (bFoundValue)
 				{
-					if (Property->ImportText_Direct(*Value, Property->ContainerPtrToValuePtr<uint8>(this, i), this, PortFlags) == NULL)
+					if (Property->ImportText_Direct(*Value, Property->ContainerPtrToValuePtr<uint8>(this, i), this, PortFlags) == nullptr)
 					{
 						// this should be an error as the properties from the .ini / .int file are not correctly being read in and probably are affecting things in subtle ways
 						UE_LOG(LogObj, Error, TEXT("LoadConfig (%s): import failed for %s in: %s"), *GetPathName(), *Property->GetName(), *Value);
@@ -3226,24 +3240,22 @@ void UObject::LoadConfig( UClass* ConfigClass/*=NULL*/, const TCHAR* InFilename/
 	}
 
 	// if we are reloading config data after the initial class load, fire the callback now
-	if ( (PropagationFlags&UE::LCPF_ReloadingConfigData) != 0 )
+	if ((PropagationFlags & UE::LCPF_ReloadingConfigData) != 0)
 	{
 		PostReloadConfig(PropertyToLoad);
 	}
 }
 
-void UObject::SaveConfig( uint64 Flags, const TCHAR* InFilename, FConfigCacheIni* Config/*=GConfig*/, bool bAllowCopyToDefaultObject/*=true*/)
+void UObject::SaveConfig(uint64 RequiredPropertyFlags, const TCHAR* InFilename, FConfigCacheIni* Config/*=GConfig*/,
+	bool bAllowCopyToDefaultObject/*=true*/)
 {
-	if( !GetClass()->HasAnyClassFlags(CLASS_Config) )
+	if (!GetClass()->HasAnyClassFlags(CLASS_Config))
 	{
 		return;
 	}
 
-	const FString Filename
-	// if a filename was specified, always load from that file
-	=	InFilename
-		? InFilename
-		: GetConfigFilename(this);
+	// if a filename was specified, save all properties to that file no matter which class they come from
+	const FString Filename = InFilename ? InFilename : GetConfigFilename(this);
 
 	// Determine whether the file we are writing is a default file config.
 	const bool bIsADefaultIniWrite = Filename == GetDefaultConfigFilename() || Filename == GetGlobalUserConfigFilename();
@@ -3276,14 +3288,14 @@ void UObject::SaveConfig( uint64 Flags, const TCHAR* InFilename, FConfigCacheIni
 	// only copy the values to the CDO if this is GConfig and we're not saving the CDO
 	const bool bCopyValues = (bAllowCopyToDefaultObject && this != CDO && Config == GConfig);
 
-	for ( FProperty* Property = GetClass()->PropertyLink; Property; Property = Property->PropertyLinkNext )
+	for (FProperty* Property = GetClass()->PropertyLink; Property; Property = Property->PropertyLinkNext)
 	{
-		if ( !Property->HasAnyPropertyFlags(CPF_Config) )
+		if (!Property->HasAnyPropertyFlags(CPF_Config))
 		{
 			continue;
 		}
 
-		if( (Property->PropertyFlags & Flags) == Flags )
+		if ((Property->PropertyFlags & RequiredPropertyFlags) == RequiredPropertyFlags)
 		{
 			UClass* BaseClass = GetClass();
 
@@ -3306,7 +3318,7 @@ void UObject::SaveConfig( uint64 Flags, const TCHAR* InFilename, FConfigCacheIni
 			}
 #endif // #if WITH_EDITOR
 
-			if ( !bPerObject )
+			if (!bPerObject)
 			{
 				Section = BaseClass->GetPathName();
 				// allow the class to override the expected section name
@@ -3315,7 +3327,7 @@ void UObject::SaveConfig( uint64 Flags, const TCHAR* InFilename, FConfigCacheIni
 
 			// globalconfig properties should always use the owning class's config file
 			// specifying a value for InFilename will override this behavior (as it does with normal properties)
-			const FString& PropFileName = ((Property->PropertyFlags & CPF_GlobalConfig) && InFilename == NULL) ? Property->GetOwnerClass()->GetConfigName() : Filename;
+			const FString& PropFileName = ((Property->PropertyFlags & CPF_GlobalConfig) && InFilename == nullptr) ? Property->GetOwnerClass()->GetConfigName() : Filename;
 
 			// Properties that are the same as the parent class' defaults should not be saved to ini
 			// Before modifying any key in the section, first check to see if it is different from the parent.
@@ -3346,10 +3358,10 @@ void UObject::SaveConfig( uint64 Flags, const TCHAR* InFilename, FConfigCacheIni
 					if (Array)
 					{
 						FScriptArrayHelper_InContainer ArrayHelper(Array, this);
-						for( int32 i=0; i<ArrayHelper.Num(); i++ )
+						for (int32 i=0; i<ArrayHelper.Num(); i++)
 						{
 							FString	Buffer;
-							Array->Inner->ExportTextItem_Direct( Buffer, ArrayHelper.GetRawPtr(i), ArrayHelper.GetRawPtr(i), this, PortFlags );
+							Array->Inner->ExportTextItem_Direct(Buffer, ArrayHelper.GetRawPtr(i), ArrayHelper.GetRawPtr(i), this, PortFlags);
 							Config->AddToSection(*Section, *CompleteKey, *Buffer, PropFileName);
 						}
 						if (ArrayHelper.Num() == 0 && bIsADefaultIniWrite)
@@ -3382,9 +3394,9 @@ void UObject::SaveConfig( uint64 Flags, const TCHAR* InFilename, FConfigCacheIni
 			if (!bProcessedProperty)
 			{
 				TCHAR TempKey[MAX_SPRINTF] = {};
-				for( int32 Index=0; Index<Property->ArrayDim; Index++ )
+				for (int32 Index=0; Index<Property->ArrayDim; Index++)
 				{
-					if( Property->ArrayDim!=1 )
+					if (Property->ArrayDim!=1)
 					{
 						FCString::Sprintf( TempKey, TEXT("%s[%i]"), *Property->GetName(), Index );
 						Key = TempKey;
@@ -3393,13 +3405,13 @@ void UObject::SaveConfig( uint64 Flags, const TCHAR* InFilename, FConfigCacheIni
 					if (!bPropDeprecated && (!bShouldCheckIfIdenticalBeforeAdding || !Property->Identical_InContainer(this, SuperClassDefaultObject, Index)))
 					{
 						FString	Value;
-						Property->ExportText_InContainer( Index, Value, this, this, this, PortFlags );
-						Config->SetString( *Section, *Key, *Value, PropFileName );
+						Property->ExportText_InContainer(Index, Value, this, this, this, PortFlags);
+						Config->SetString(*Section, *Key, *Value, PropFileName);
 					}
 					else
 					{
 						// If we are not writing it to config above, we should make sure that this property isn't stagnant in the cache.
-						Config->ResetKeyInSection( *Section, *Key, PropFileName );
+						Config->ResetKeyInSection(*Section, *Key, PropFileName);
 					}
 				}
 			}
@@ -3409,7 +3421,7 @@ void UObject::SaveConfig( uint64 Flags, const TCHAR* InFilename, FConfigCacheIni
 				void* ThisPropertyAddress = Property->ContainerPtrToValuePtr<void>( this );
 				void* CDOPropertyAddr = Property->ContainerPtrToValuePtr<void>( CDO );
 
-				Property->CopyCompleteValue( CDOPropertyAddr, ThisPropertyAddress );
+				Property->CopyCompleteValue(CDOPropertyAddr, ThisPropertyAddress);
 			}
 		}
 	}
diff --git a/Engine/Source/Runtime/CoreUObject/Public/UObject/Object.h b/Engine/Source/Runtime/CoreUObject/Public/UObject/Object.h
--- a/Engine/Source/Runtime/CoreUObject/Public/UObject/Object.h
+++ b/Engine/Source/Runtime/CoreUObject/Public/UObject/Object.h
@@ -1277,7 +1277,8 @@ public:
 	 * Save configuration out to ini files
 	 * @warning Must be safe to call on class-default object
 	 */
-	COREUOBJECT_API void SaveConfig( uint64 Flags=CPF_Config, const TCHAR* Filename=NULL, FConfigCacheIni* Config=GConfig, bool bAllowCopyToDefaultObject=true );
+	COREUOBJECT_API void SaveConfig(uint64 RequiredPropertyFlags = CPF_Config, const TCHAR* Filename = nullptr,
+		FConfigCacheIni* Config=GConfig, bool bAllowCopyToDefaultObject=true);
 
 	/**
 	 * Saves just the section(s) for this class into the default ini file for the class (with just the changes from base)
@@ -1373,7 +1374,8 @@ public:
 	 * @param	PropagationFlags	indicates how this call to LoadConfig should be propagated; expects a bitmask of UE::ELoadConfigPropagationFlags values.
 	 * @param	PropertyToLoad		if specified, only the ini value for the specified property will be imported.
 	 */
-	COREUOBJECT_API void LoadConfig( UClass* ConfigClass=NULL, const TCHAR* Filename=NULL, uint32 PropagationFlags=UE::LCPF_None, class FProperty* PropertyToLoad=NULL );
+	COREUOBJECT_API void LoadConfig(UClass* ConfigClass = nullptr, const TCHAR* Filename = nullptr,
+		uint32 PropagationFlags = UE::LCPF_None, class FProperty* PropertyToLoad = nullptr);
 
 	/**
 	 * Wrapper method for LoadConfig that is used when reloading the config data for objects at runtime which have already loaded their config data at least once.
