ConfigAccessTracking: Modify UObject::LoadConfig to take an OutAccessedValues parameter to record which fields would be accessed; used to track config dependencies of settings objects.
#rnx

diff --git a/Engine/Source/Runtime/CoreUObject/Private/UObject/Obj.cpp b/Engine/Source/Runtime/CoreUObject/Private/UObject/Obj.cpp
--- a/Engine/Source/Runtime/CoreUObject/Private/UObject/Obj.cpp
+++ b/Engine/Source/Runtime/CoreUObject/Private/UObject/Obj.cpp
@@ -11,6 +11,7 @@
 #include "Misc/Paths.h"
 #include "Logging/LogScopedCategoryAndVerbosityOverride.h"
 #include "Stats/Stats.h"
+#include "Misc/ConfigAccessData.h"
 #include "Misc/ConfigCacheIni.h"
 #include "Misc/ConfigContext.h"
 #include "Misc/CoreDelegates.h"
@@ -2677,8 +2678,13 @@ void CheckMissingSection(const FString& SectionName, const FString& IniFilename)
 }
 #endif
 
+#if WITH_EDITOR
+static FName GConsoleVariableFName(TEXT("ConsoleVariable"));
+#endif
+
 void UObject::LoadConfig(UClass* ConfigClass/*=nullptr*/, const TCHAR* InFilename/*=nullptr*/,
-	uint32 PropagationFlags/*=LCPF_None*/, FProperty* PropertyToLoad/*=nullptr*/)
+	uint32 PropagationFlags/*=LCPF_None*/, FProperty* PropertyToLoad/*=nullptr*/,
+	TArray<UE::ConfigAccessTracking::FConfigAccessData>* OutAccessedValues/*=nullptr*/)
 {
 	SCOPE_CYCLE_COUNTER(STAT_LoadConfig);
 
@@ -2697,6 +2703,14 @@ void UObject::LoadConfig(UClass* ConfigClass/*=nullptr*/, const TCHAR* InFilenam
 		return;
 	}
 
+#if !WITH_EDITOR
+	if (OutAccessedValues)
+	{
+		return;
+	}
+#endif
+	const bool bModifyingObjects = OutAccessedValues == nullptr;
+
 #if !IS_PROGRAM
 	auto HaveSameProperties = [](const UStruct* Struct1, const UStruct* Struct2) -> bool
 	{
@@ -2742,7 +2756,7 @@ void UObject::LoadConfig(UClass* ConfigClass/*=nullptr*/, const TCHAR* InFilenam
 			if ((PropagationFlags & UE::LCPF_ReadParentSections) != 0)
 			{
 				// call LoadConfig on the parent class
-				LoadConfig(ParentClass, nullptr, PropagationFlags, PropertyToLoad);
+				LoadConfig(ParentClass, nullptr, PropagationFlags, PropertyToLoad, OutAccessedValues);
 
 				// if we are also notifying child classes or instances, stop here as this object's properties will be imported as a result of notifying the others
 				if ((PropagationFlags & (UE::LCPF_PropagateToChildDefaultObjects | UE::LCPF_PropagateToInstances)) != 0)
@@ -2760,7 +2774,7 @@ void UObject::LoadConfig(UClass* ConfigClass/*=nullptr*/, const TCHAR* InFilenam
 						// mask out the PropgateToParent and PropagateToChildren values
 						It->GetDefaultObject()->LoadConfig(*It, nullptr,
 							(PropagationFlags & (UE::LCPF_PersistentFlags | UE::LCPF_PropagateToInstances)),
-							PropertyToLoad);
+							PropertyToLoad, OutAccessedValues);
 					}
 				}
 
@@ -2779,18 +2793,24 @@ void UObject::LoadConfig(UClass* ConfigClass/*=nullptr*/, const TCHAR* InFilenam
 						{
 							// make sure to pass in the class so that OriginalClass isn't reset
 							It->LoadConfig(It->GetClass(), nullptr, (PropagationFlags & UE::LCPF_PersistentFlags),
-								PropertyToLoad);
+								PropertyToLoad, OutAccessedValues);
 						}
 #if WITH_EDITOR
 						else
 						{
-							It->PreEditChange(nullptr);
+							if (bModifyingObjects)
+							{
+								It->PreEditChange(nullptr);
+							}
 
 							// make sure to pass in the class so that OriginalClass isn't reset
 							It->LoadConfig(It->GetClass(), nullptr, (PropagationFlags & UE::LCPF_PersistentFlags),
-								PropertyToLoad);
+								PropertyToLoad, OutAccessedValues);
 
-							It->PostEditChange();
+							if (bModifyingObjects)
+							{
+								It->PostEditChange();
+							}
 						}
 #endif // WITH_EDITOR
 					}
@@ -2809,19 +2829,25 @@ void UObject::LoadConfig(UClass* ConfigClass/*=nullptr*/, const TCHAR* InFilenam
 						// make sure to pass in the class so that OriginalClass isn't reset
 						It->GetDefaultObject()->LoadConfig(*It, nullptr,
 							(PropagationFlags & (UE::LCPF_PersistentFlags | UE::LCPF_PropagateToInstances)),
-							PropertyToLoad);
+							PropertyToLoad, OutAccessedValues);
 					}
 #if WITH_EDITOR
 					else
 					{
-						It->PreEditChange(nullptr);
+						if (bModifyingObjects)
+						{
+							It->PreEditChange(nullptr);
+						}
 
 						// make sure to pass in the class so that OriginalClass isn't reset
 						It->GetDefaultObject()->LoadConfig(*It, nullptr,
 							(PropagationFlags & (UE::LCPF_PersistentFlags | UE::LCPF_PropagateToInstances)),
-							PropertyToLoad);
+							PropertyToLoad, OutAccessedValues);
 
-						It->PostEditChange();
+						if (bModifyingObjects)
+						{
+							It->PostEditChange();
+						}
 					}
 #endif // WITH_EDITOR
 				}
@@ -2839,17 +2865,24 @@ void UObject::LoadConfig(UClass* ConfigClass/*=nullptr*/, const TCHAR* InFilenam
 					{
 						// make sure to pass in the class so that OriginalClass isn't reset
 						It->LoadConfig(It->GetClass(), nullptr,	(PropagationFlags & UE::LCPF_PersistentFlags),
-							PropertyToLoad);
+							PropertyToLoad, OutAccessedValues);
 					}
 #if WITH_EDITOR
 					else
 					{
-						It->PreEditChange(nullptr);
+						if (bModifyingObjects)
+						{
+							It->PreEditChange(nullptr);
+						}
 
 						// make sure to pass in the class so that OriginalClass isn't reset
 						It->LoadConfig(It->GetClass(), nullptr, (PropagationFlags & UE::LCPF_PersistentFlags),
-							PropertyToLoad);
-						It->PostEditChange();
+							PropertyToLoad, OutAccessedValues);
+
+						if (bModifyingObjects)
+						{
+							It->PostEditChange();
+						}
 					}
 #endif // WITH_EDITOR
 				}
@@ -2867,12 +2900,23 @@ void UObject::LoadConfig(UClass* ConfigClass/*=nullptr*/, const TCHAR* InFilenam
 	bool bUseConfigOverride = false;
 	FConfigFile* OverrideConfigFile = nullptr;
 	FConfigFile LocalOverrideConfig;
+	const TCHAR* ConfigOverridePlatform = GetConfigOverridePlatform();
+#if WITH_EDITOR
+	FName ConfigPlatform;
+#endif
 
-	if (InFilename == nullptr && GetConfigOverridePlatform() != nullptr &&
-		FCString::Stricmp(GetConfigOverridePlatform(), ANSI_TO_TCHAR(FPlatformProperties::IniPlatformName())) != 0)
+	if (InFilename == nullptr && ConfigOverridePlatform != nullptr &&
+		FCString::Stricmp(ConfigOverridePlatform, ANSI_TO_TCHAR(FPlatformProperties::IniPlatformName())) != 0)
 	{
-		OverrideConfigFile = FConfigCacheIni::FindOrLoadPlatformConfig(LocalOverrideConfig, *GetClass()->ClassConfigName.ToString(), GetConfigOverridePlatform());
-		bUseConfigOverride = true;
+#if WITH_EDITOR
+		ConfigPlatform = FName(ConfigOverridePlatform);
+#endif
+		if (bModifyingObjects)
+		{
+			OverrideConfigFile = FConfigCacheIni::FindOrLoadPlatformConfig(LocalOverrideConfig,
+				*GetClass()->ClassConfigName.ToString(), ConfigOverridePlatform);
+			bUseConfigOverride = true;
+		}
 	}
 #if WITH_EDITOR
 	else if (GetClass()->HasAnyClassFlags(CLASS_PerPlatformConfig) && UObject::OnGetPreviewPlatform.IsBound())
@@ -2880,8 +2924,14 @@ void UObject::LoadConfig(UClass* ConfigClass/*=nullptr*/, const TCHAR* InFilenam
 		FName PreviewPlatform;
 		if (UObject::OnGetPreviewPlatform.Execute(PreviewPlatform))
 		{
-			OverrideConfigFile = FConfigCacheIni::FindOrLoadPlatformConfig(LocalOverrideConfig, *GetClass()->ClassConfigName.ToString(), *PreviewPlatform.ToString());
-			bUseConfigOverride = true;
+			FString PreviewPlatformStr = PreviewPlatform.ToString();
+			ConfigPlatform = FName(*PreviewPlatformStr);
+			if (bModifyingObjects)
+			{
+				OverrideConfigFile = FConfigCacheIni::FindOrLoadPlatformConfig(LocalOverrideConfig,
+					*GetClass()->ClassConfigName.ToString(), *PreviewPlatformStr);
+				bUseConfigOverride = true;
+			}
 		}
 	}
 #endif
@@ -2928,7 +2978,7 @@ void UObject::LoadConfig(UClass* ConfigClass/*=nullptr*/, const TCHAR* InFilenam
 	// the 'if (OriginalClass != Class)' check....better store it in a temporary place while we do the actual loading of our properties 
 	UClass* MyOrigClass = OriginalClass;
 
-	if ( PropertyToLoad == nullptr )
+	if (PropertyToLoad == nullptr)
 	{
 		UE_LOG(LogConfig, VeryVerbose, TEXT("(%s) '%s' loading configuration from %s"),
 			*ConfigClass->GetName(), *GetName(), *Filename);
@@ -2990,7 +3040,7 @@ void UObject::LoadConfig(UClass* ConfigClass/*=nullptr*/, const TCHAR* InFilenam
 			continue;
 		}
 
-		const bool bGlobalConfig = (Property->PropertyFlags&CPF_GlobalConfig) != 0;
+		const bool bGlobalConfig = EnumHasAnyFlags(Property->PropertyFlags, CPF_GlobalConfig);
 		UClass* OwnerClass = Property->GetOwnerClass();
 
 		UClass* BaseClass = bGlobalConfig ? OwnerClass : ConfigClass;
@@ -3018,8 +3068,7 @@ void UObject::LoadConfig(UClass* ConfigClass/*=nullptr*/, const TCHAR* InFilenam
 		int32 PortFlags = 0;
 
 #if WITH_EDITOR
-		static FName ConsoleVariableFName(TEXT("ConsoleVariable"));
-		const FString& CVarName = Property->GetMetaData(ConsoleVariableFName);
+		const FString& CVarName = Property->GetMetaData(GConsoleVariableFName);
 		if (!CVarName.IsEmpty())
 		{
 			Key = CVarName;
@@ -3029,6 +3078,19 @@ void UObject::LoadConfig(UClass* ConfigClass/*=nullptr*/, const TCHAR* InFilenam
 		const bool bIsPerPlatformConfig = GetClass()->HasAnyClassFlags(CLASS_PerPlatformConfig);
 #endif // #if WITH_EDITOR
 
+		if (OutAccessedValues)
+		{
+#if WITH_EDITOR
+			OutAccessedValues->Add(UE::ConfigAccessTracking::FConfigAccessData(
+				UE::ConfigAccessTracking::ELoadType::ConfigSystem, ConfigPlatform, FName(*PropFileName),
+				FName(*ClassSection), FName(*Key), nullptr /* InRequestingPlatform */));
+#endif
+		}
+		if (!bModifyingObjects)
+		{
+			continue;
+		}
+
 		// Track if we loaded this config value using special handling (e.g. array or set)
 		bool bProcessedProperty = false;
 
@@ -3240,7 +3302,7 @@ void UObject::LoadConfig(UClass* ConfigClass/*=nullptr*/, const TCHAR* InFilenam
 	}
 
 	// if we are reloading config data after the initial class load, fire the callback now
-	if ((PropagationFlags & UE::LCPF_ReloadingConfigData) != 0)
+	if (bModifyingObjects && (PropagationFlags & UE::LCPF_ReloadingConfigData) != 0)
 	{
 		PostReloadConfig(PropertyToLoad);
 	}
@@ -3299,7 +3361,7 @@ void UObject::SaveConfig(uint64 RequiredPropertyFlags, const TCHAR* InFilename,
 		{
 			UClass* BaseClass = GetClass();
 
-			if (Property->PropertyFlags & CPF_GlobalConfig)
+			if (EnumHasAnyFlags(Property->PropertyFlags, CPF_GlobalConfig))
 			{
 				// call LoadConfig() on child classes if any of the properties were global config
 				BaseClass = Property->GetOwnerClass();
@@ -3309,8 +3371,7 @@ void UObject::SaveConfig(uint64 RequiredPropertyFlags, const TCHAR* InFilename,
 			int32 PortFlags			= 0;
 
 #if WITH_EDITOR
-			static FName ConsoleVariableFName(TEXT("ConsoleVariable"));
-			const FString& CVarName = Property->GetMetaData(ConsoleVariableFName);
+			const FString& CVarName = Property->GetMetaData(GConsoleVariableFName);
 			if (!CVarName.IsEmpty())
 			{
 				Key = CVarName;
@@ -3578,8 +3639,7 @@ void UObject::UpdateSinglePropertyInConfigFile(const FProperty* InProperty, cons
 		FString PropertyKey = InProperty->GetFName().ToString();
 
 #if WITH_EDITOR
-		static FName ConsoleVariableFName(TEXT("ConsoleVariable"));
-		const FString& CVarName = InProperty->GetMetaData(ConsoleVariableFName);
+		const FString& CVarName = InProperty->GetMetaData(GConsoleVariableFName);
 		if (!CVarName.IsEmpty())
 		{
 			PropertyKey = CVarName;
diff --git a/Engine/Source/Runtime/CoreUObject/Public/UObject/Object.h b/Engine/Source/Runtime/CoreUObject/Public/UObject/Object.h
--- a/Engine/Source/Runtime/CoreUObject/Public/UObject/Object.h
+++ b/Engine/Source/Runtime/CoreUObject/Public/UObject/Object.h
@@ -17,6 +17,7 @@
 struct FAssetData;
 class FAssetRegistryTagsContext;
 namespace UE::AssetRegistry { class FAssetRegistryImpl; }
+namespace UE::ConfigAccessTracking { struct FConfigAccessData; }
 class FConfigCacheIni;
 class FCustomPropertyConditionState;
 class FEditPropertyChain;
@@ -1373,9 +1374,12 @@ public:
 	 * @param	Filename			indicates the filename to load values from; if not specified, uses ConfigClass's ClassConfigName
 	 * @param	PropagationFlags	indicates how this call to LoadConfig should be propagated; expects a bitmask of UE::ELoadConfigPropagationFlags values.
 	 * @param	PropertyToLoad		if specified, only the ini value for the specified property will be imported.
+	 * @param 	OutAccessedKeys     if specified the object is not modified and (EditorOnly) all the config keys it
+	 *                              would read are added to OutAccessedKeys. In non-editor, function returns with no action.
 	 */
 	COREUOBJECT_API void LoadConfig(UClass* ConfigClass = nullptr, const TCHAR* Filename = nullptr,
-		uint32 PropagationFlags = UE::LCPF_None, class FProperty* PropertyToLoad = nullptr);
+		uint32 PropagationFlags = UE::LCPF_None, class FProperty* PropertyToLoad = nullptr,
+		TArray<UE::ConfigAccessTracking::FConfigAccessData>* OutAccessedKeys = nullptr);
 
 	/**
 	 * Wrapper method for LoadConfig that is used when reloading the config data for objects at runtime which have already loaded their config data at least once.
