[CrashFix] crash when loose property is modified in the fixup tool
technically this shouldn't be possible since the string property is marked CPF_EditConst. There's a but that causes VerseStringProperties to ignore CPF_EditConst so users can edit loose string properties and that causes a crash. Regardless, this makes it so that even if the user somehow edits a loose property, it won't crash.

diff --git a/Engine/Source/Runtime/CoreUObject/Private/UObject/PropertyBagRepository.cpp b/Engine/Source/Runtime/CoreUObject/Private/UObject/PropertyBagRepository.cpp
--- a/Engine/Source/Runtime/CoreUObject/Private/UObject/PropertyBagRepository.cpp
+++ b/Engine/Source/Runtime/CoreUObject/Private/UObject/PropertyBagRepository.cpp
@@ -194,14 +194,19 @@ static FProperty* FindPropertyByNameAndType(const UStruct* Struct, FName InName,
 	return nullptr;
 }
 
-static void ConstructRemappedPropertyChain(const FEditPropertyChain& Chain, FEditPropertyChain& NewChain, const UObject* Destination)
+static bool ConstructRemappedPropertyChain(const FEditPropertyChain& Chain, FEditPropertyChain& NewChain, const UObject* Destination)
 {
 	UStruct* Struct = Destination->GetClass();
 	for (FEditPropertyChain::TDoubleLinkedListNode* Itr = Chain.GetHead(); Itr; Itr = Itr->GetNextNode())
 	{
 		FProperty* Property = Itr->GetValue();
-		NewChain.AddTail(FindPropertyByNameAndType(Struct, Property->GetFName(), Property->GetID()));
-		Property = NewChain.GetTail()->GetValue();
+		Property = FindPropertyByNameAndType(Struct, Property->GetFName(), Property->GetID());
+		if (!Property)
+		{
+			NewChain.Empty();
+			return false;
+		}
+		NewChain.AddTail(Property);
 
 		// iterate the struct to look in
 		if (FOptionalProperty* AsOptionalProperty = CastField<FOptionalProperty>(Property))
@@ -240,6 +245,7 @@ static void ConstructRemappedPropertyChain(const FEditPropertyChain& Chain, FEdi
 			NewChain.SetActiveMemberPropertyNode(NewChain.GetTail()->GetValue());
 		}
 	}
+	return true;
 }
 
 static void* ResolveChangePath(const void* StructData, FPropertyChangedChainEvent& ChangeEvent, bool bGrowContainersWhenNeeded = false)
@@ -452,28 +458,34 @@ void FPropertyBagRepository::PostEditChangeChainProperty(const UObject* Object,
 	auto CopyChanges = [&PropertyChangedEvent](const UObject* Source, UObject* Dest) 
 	{
 		FEditPropertyChain RemappedChain;
-		ConstructRemappedPropertyChain(PropertyChangedEvent.PropertyChain, RemappedChain, Dest);
-		Dest->PreEditChange(RemappedChain);
-
-		FPropertyChangedChainEvent RemappedChangeEvent(RemappedChain, PropertyChangedEvent);
-		const void* SourceData = ResolveChangePath(Source, PropertyChangedEvent);
-		void* DestData = ResolveChangePath(Dest, RemappedChangeEvent, true);
-		FProperty* SourceProperty = PropertyChangedEvent.PropertyChain.GetTail()->GetValue();
-		FProperty* DestProperty = RemappedChangeEvent.PropertyChain.GetTail()->GetValue();
-
-		if (PropertyChangedEvent.ChangeType == EPropertyChangeType::ArrayAdd)
-		{
-			int32 ArrayIndex = PropertyChangedEvent.GetArrayIndex(SourceProperty->GetName());
-			check(ArrayIndex != INDEX_NONE);
-			AddProperty(SourceProperty, SourceData, DestProperty, DestData, ArrayIndex);
+		if (ConstructRemappedPropertyChain(PropertyChangedEvent.PropertyChain, RemappedChain, Dest))
+		{
+			Dest->PreEditChange(RemappedChain);
+            
+            FPropertyChangedChainEvent RemappedChangeEvent(RemappedChain, PropertyChangedEvent);
+            const void* SourceData = ResolveChangePath(Source, PropertyChangedEvent);
+            void* DestData = ResolveChangePath(Dest, RemappedChangeEvent, true);
+            FProperty* SourceProperty = PropertyChangedEvent.PropertyChain.GetTail()->GetValue();
+            FProperty* DestProperty = RemappedChangeEvent.PropertyChain.GetTail()->GetValue();
+			
+            if (PropertyChangedEvent.ChangeType == EPropertyChangeType::ArrayAdd)
+            {
+            	int32 ArrayIndex = PropertyChangedEvent.GetArrayIndex(SourceProperty->GetName());
+            	check(ArrayIndex != INDEX_NONE);
+            	AddProperty(SourceProperty, SourceData, DestProperty, DestData, ArrayIndex);
+            }
+            else
+            {
+            	CopyProperty(SourceProperty, SourceData, DestProperty, DestData);
+            }
+            
+            Dest->PostEditChangeChainProperty(RemappedChangeEvent);
 		}
 		else
 		{
-			CopyProperty(SourceProperty, SourceData, DestProperty, DestData);
+			ensureMsgf(false, TEXT("A const loose property was modified on an instance data object"));
 		}
 		
-		
-		Dest->PostEditChangeChainProperty(RemappedChangeEvent);
 	};
 	
 	if (UObject* Ido = Get().FindInstanceDataObject(Object))
