Made IDO functionality conditional on WITH_EDITORONLY_DATA

#rnx

diff --git a/Engine/Plugins/Experimental/InstanceDataObjectFixupTool/Source/Private/Tests/InstanceDataObjectTests.cpp b/Engine/Plugins/Experimental/InstanceDataObjectFixupTool/Source/Private/Tests/InstanceDataObjectTests.cpp
--- a/Engine/Plugins/Experimental/InstanceDataObjectFixupTool/Source/Private/Tests/InstanceDataObjectTests.cpp
+++ b/Engine/Plugins/Experimental/InstanceDataObjectFixupTool/Source/Private/Tests/InstanceDataObjectTests.cpp
@@ -2,6 +2,8 @@
 
 #include "InstanceDataObjectTests.h"
 
+#if WITH_EDITORONLY_DATA
+
 #include "CoreTypes.h"
 #include "InstanceDataObjectFixupToolModule.h"
 #include "WorkspaceMenuStructure.h"
@@ -454,3 +456,5 @@ FAutoConsoleCommand TestInstanceDataObjectFixupCommand(
 );
 
 #undef LOCTEXT_NAMESPACE
+
+#endif // WITH_EDITORONLY_DATA
diff --git a/Engine/Source/Runtime/CoreUObject/Internal/UObject/InstanceDataObjectUtils.h b/Engine/Source/Runtime/CoreUObject/Internal/UObject/InstanceDataObjectUtils.h
--- a/Engine/Source/Runtime/CoreUObject/Internal/UObject/InstanceDataObjectUtils.h
+++ b/Engine/Source/Runtime/CoreUObject/Internal/UObject/InstanceDataObjectUtils.h
@@ -2,7 +2,9 @@
 
 #pragma once
 
-#include "HAL/Platform.h"
+#include "CoreTypes.h"
+
+#if WITH_EDITORONLY_DATA
 
 #define UE_API COREUOBJECT_API
 
@@ -47,3 +49,5 @@ UE_API void ResetPropertyValueInitialized(const UStruct* Struct, void* StructDat
 } // UE
 
 #undef UE_API
+
+#endif // WITH_EDITORONLY_DATA
diff --git a/Engine/Source/Runtime/CoreUObject/Internal/UObject/PropertyBagRepository.h b/Engine/Source/Runtime/CoreUObject/Internal/UObject/PropertyBagRepository.h
--- a/Engine/Source/Runtime/CoreUObject/Internal/UObject/PropertyBagRepository.h
+++ b/Engine/Source/Runtime/CoreUObject/Internal/UObject/PropertyBagRepository.h
@@ -13,6 +13,8 @@ class FMemoryArchive;
 namespace UE
 {
 
+#if WITH_EDITORONLY_DATA
+
 class FPropertyPathNameTree;
 
 // Singleton class tracking property bag association with objects
@@ -163,10 +165,13 @@ private:
 	static void CopyTaggedProperties (const UObject* Source, UObject* Dest);
 };
 
+#endif // WITH_EDITORONLY_DATA
+
 // construct this context in the same scope as an object is being serialized to support InstanceDataObjects.
 // if saving, this will simply set flags. If loading, an IDO will be constructed at the end of the scope when needed
 struct FScopedIDOSerializationContext
 {
+#if WITH_EDITORONLY_DATA
 	COREUOBJECT_API FScopedIDOSerializationContext(UObject* InObject, FArchive& Archive);
 	explicit COREUOBJECT_API FScopedIDOSerializationContext(UObject* InObject); // assumes save
 	COREUOBJECT_API ~FScopedIDOSerializationContext();
@@ -180,10 +185,15 @@ struct FScopedIDOSerializationContext
 	TOptional<TGuardValue<bool>> ScopedSerializeUnknownProperty;
 	TOptional<TGuardValue<bool>> ScopedImpersonateProperties;
 	TOptional<TGuardValue<UObject*>> ScopedSerializedObject;
+
 private:
 	// if we're loading and IDO should be created, this will be called when the context falls out of scope
 	void FinishCreatingInstanceDataObject() const;
-	
+#else
+	inline FScopedIDOSerializationContext(UObject* InObject, FArchive& Archive) {}
+	inline explicit FScopedIDOSerializationContext(UObject* InObject) {}
+#endif
+
 	FScopedIDOSerializationContext(const FScopedIDOSerializationContext&) = delete;
 	FScopedIDOSerializationContext& operator=(const FScopedIDOSerializationContext&) = delete;
 };
diff --git a/Engine/Source/Runtime/CoreUObject/Private/UObject/Class.cpp b/Engine/Source/Runtime/CoreUObject/Private/UObject/Class.cpp
--- a/Engine/Source/Runtime/CoreUObject/Private/UObject/Class.cpp
+++ b/Engine/Source/Runtime/CoreUObject/Private/UObject/Class.cpp
@@ -1166,6 +1166,7 @@ void UStruct::DestroyStruct(void* Dest, int32 ArrayDim) const
 //
 void UStruct::SerializeBin( FStructuredArchive::FSlot Slot, void* Data ) const
 {
+#if WITH_EDITORONLY_DATA
 	FUObjectSerializeContext* SerializeContext = FUObjectThreadContext::Get().GetSerializeContext();
 	const bool bSaveSerializedPropertyPath = IsA<UClass>() && SerializeContext && !SerializeContext->SerializedPropertyPath.IsEmpty();
 	UE::FPropertyPathName PrevSerializedPropertyPath;
@@ -1175,6 +1176,7 @@ void UStruct::SerializeBin( FStructuredArchive::FSlot Slot, void* Data ) const
 		PrevSerializedPropertyPath = MoveTemp(SerializeContext->SerializedPropertyPath);
 		SerializeContext->SerializedPropertyPath.Reset();
 	}
+#endif
 
 	FArchive& UnderlyingArchive = Slot.GetUnderlyingArchive();
 
@@ -1221,10 +1223,12 @@ void UStruct::SerializeBin( FStructuredArchive::FSlot Slot, void* Data ) const
 		}
 	}
 
+#if WITH_EDITORONLY_DATA
 	if (bSaveSerializedPropertyPath)
 	{
 		SerializeContext->SerializedPropertyPath = MoveTemp(PrevSerializedPropertyPath);
 	}
+#endif
 }
 
 void UStruct::SerializeBinEx( FStructuredArchive::FSlot Slot, void* Data, void const* DefaultData, UStruct* DefaultStruct ) const
@@ -1235,6 +1239,7 @@ void UStruct::SerializeBinEx( FStructuredArchive::FSlot Slot, void* Data, void c
 		return;
 	}
 
+#if WITH_EDITORONLY_DATA
 	FUObjectSerializeContext* SerializeContext = FUObjectThreadContext::Get().GetSerializeContext();
 	const bool bSaveSerializedPropertyPath = IsA<UClass>() && SerializeContext && !SerializeContext->SerializedPropertyPath.IsEmpty();
 	UE::FPropertyPathName PrevSerializedPropertyPath;
@@ -1244,16 +1249,19 @@ void UStruct::SerializeBinEx( FStructuredArchive::FSlot Slot, void* Data, void c
 		PrevSerializedPropertyPath = MoveTemp(SerializeContext->SerializedPropertyPath);
 		SerializeContext->SerializedPropertyPath.Reset();
 	}
+#endif
 
 	for( TFieldIterator<FProperty> It(this); It; ++It )
 	{
 		It->SerializeNonMatchingBinProperty(Slot, Data, DefaultData, DefaultStruct);
 	}
 
+#if WITH_EDITORONLY_DATA
 	if (bSaveSerializedPropertyPath)
 	{
 		SerializeContext->SerializedPropertyPath = MoveTemp(PrevSerializedPropertyPath);
 	}
+#endif
 }
 
 void UStruct::LoadTaggedPropertiesFromText(FStructuredArchive::FSlot Slot, uint8* Data, UStruct* DefaultsStruct, uint8* Defaults, const UObject* BreakRecursionIfFullyLoad) const
@@ -1340,10 +1348,12 @@ void UStruct::LoadTaggedPropertiesFromText(FStructuredArchive::FSlot Slot, uint8
 				Tag.ArrayIndex = ItemIndex;
 				Tag.Name = PropertyName;
 
+			#if WITH_EDITORONLY_DATA
 				if (SerializeContext->bTrackInitializedProperties)
 				{
 					UE::SetPropertyValueInitialized(this, Data, Property, ItemIndex);
 				}
+			#endif
 
 				if (Tag.SerializeType == EPropertyTagSerializeType::Skipped)
 				{
@@ -1400,6 +1410,7 @@ void UStruct::LoadTaggedPropertiesFromText(FStructuredArchive::FSlot Slot, uint8
 
 void UStruct::SerializeTaggedProperties(FStructuredArchive::FSlot Slot, uint8* Data, UStruct* DefaultsStruct, uint8* Defaults, const UObject* BreakRecursionIfFullyLoad) const
 {
+#if WITH_EDITORONLY_DATA
 	FUObjectSerializeContext* SerializeContext = FUObjectThreadContext::Get().GetSerializeContext();
 	const bool bSaveSerializedPropertyPath = IsA<UClass>() && SerializeContext && !SerializeContext->SerializedPropertyPath.IsEmpty();
 	UE::FPropertyPathName PrevSerializedPropertyPath;
@@ -1409,6 +1420,7 @@ void UStruct::SerializeTaggedProperties(FStructuredArchive::FSlot Slot, uint8* D
 		PrevSerializedPropertyPath = MoveTemp(SerializeContext->SerializedPropertyPath);
 		SerializeContext->SerializedPropertyPath.Reset();
 	}
+#endif
 
 	if (Slot.GetArchiveState().UseUnversionedPropertySerialization())
 	{
@@ -1419,10 +1431,12 @@ void UStruct::SerializeTaggedProperties(FStructuredArchive::FSlot Slot, uint8* D
 		SerializeVersionedTaggedProperties(Slot, Data, DefaultsStruct, Defaults, BreakRecursionIfFullyLoad);
 	}
 
+#if WITH_EDITORONLY_DATA
 	if (bSaveSerializedPropertyPath)
 	{
 		SerializeContext->SerializedPropertyPath = MoveTemp(PrevSerializedPropertyPath);
 	}
+#endif
 }
 
 #if WITH_EDITORONLY_DATA
@@ -1647,10 +1661,12 @@ void UStruct::SerializeVersionedTaggedProperties(FStructuredArchive::FSlot Slot,
 					Property = CustomFindProperty(Tag.Name);
 				}
 
+			#if WITH_EDITORONLY_DATA
 				if (SerializeContext->bTrackInitializedProperties && Property)
 				{
 					UE::SetPropertyValueInitialized(this, Data, Property, Tag.ArrayIndex);
 				}
+			#endif
 
 				if (Tag.SerializeType == EPropertyTagSerializeType::Skipped)
 				{
@@ -1667,6 +1683,7 @@ void UStruct::SerializeVersionedTaggedProperties(FStructuredArchive::FSlot Slot,
 					}
 				}
 
+			#if WITH_EDITORONLY_DATA
 				// Try to match the type when impersonating because there can be multiple properties with the same name.
 				if (UNLIKELY(SerializeContext->bImpersonateProperties && Property && !Property->CanSerializeFromTypeName(Tag.GetType())))
 				{
@@ -1688,6 +1705,7 @@ void UStruct::SerializeVersionedTaggedProperties(FStructuredArchive::FSlot Slot,
 					const UE::FPropertyPathNameSegment Segment{Name, Tag.GetType(), Index};
 					SerializedPropertyPath.Emplace(SerializeContext, Segment);
 				}
+			#endif
 
 				bool bTryStoreUnknownPropertyPath = false;
 
@@ -1772,10 +1790,12 @@ void UStruct::SerializeVersionedTaggedProperties(FStructuredArchive::FSlot Slot,
 								break;
 						}
 
+					#if WITH_EDITORONLY_DATA
 						if (bAdvanceProperty && SerializeContext->bImpersonateProperties)
 						{
 							UE::MarkPropertySetBySerialization(this, Data, Property, Tag.ArrayIndex);
 						}
+					#endif
 					}
 				}
 				else
@@ -1783,6 +1803,7 @@ void UStruct::SerializeVersionedTaggedProperties(FStructuredArchive::FSlot Slot,
 					bTryStoreUnknownPropertyPath = true;
 				}
 
+			#if WITH_EDITORONLY_DATA
 				// Track the path for an unknown property and serialize it to track any unknown property within it.
 				if (UNLIKELY(bTryStoreUnknownPropertyPath))
 				{
@@ -1838,6 +1859,7 @@ void UStruct::SerializeVersionedTaggedProperties(FStructuredArchive::FSlot Slot,
 
 				// Broadcast that a property was serialized if tracking the serialized property path.
 				SerializedPropertyPath.Reset();
+			#endif // WITH_EDITORONLY_DATA
 
 				int64 Loaded = UnderlyingArchive.Tell() - StartOfProperty;
 
@@ -1902,8 +1924,12 @@ void UStruct::SerializeVersionedTaggedProperties(FStructuredArchive::FSlot Slot,
 					const bool bSerializeValue = (StaticArrayContainer.IsSet() || CustomPropertyNode || !bDoDeltaSerialization ||
 						(FOverridableSerializationLogic::IsEnabled() && FOverridableSerializationLogic::GetOverriddenPropertyOperation(UnderlyingArchive, Property, DataPtr, DefaultValue) != EOverriddenPropertyOperation::None) ||
 						(!FOverridableSerializationLogic::IsEnabled() && !Property->Identical(DataPtr, DefaultValue, UnderlyingArchive.GetPortFlags())));
+				#if WITH_EDITORONLY_DATA
 					const bool bInitializedValue = !SerializeContext->bTrackInitializedProperties || UE::IsPropertyValueInitialized(this, Data, Property, Idx);
 					if (bInitializedValue && (bSerializeValue || SerializeContext->bTrackInitializedProperties))
+				#else
+					if (bSerializeValue)
+				#endif
 					{
 						if (bUseAtomicSerialization)
 						{
diff --git a/Engine/Source/Runtime/CoreUObject/Private/UObject/InstanceDataObjectUtils.cpp b/Engine/Source/Runtime/CoreUObject/Private/UObject/InstanceDataObjectUtils.cpp
--- a/Engine/Source/Runtime/CoreUObject/Private/UObject/InstanceDataObjectUtils.cpp
+++ b/Engine/Source/Runtime/CoreUObject/Private/UObject/InstanceDataObjectUtils.cpp
@@ -2,6 +2,8 @@
 
 #include "UObject/InstanceDataObjectUtils.h"
 
+#if WITH_EDITORONLY_DATA
+
 #include "HAL/IConsoleManager.h"
 #include "Misc/ReverseIterate.h"
 #include "UObject/Class.h"
@@ -171,11 +173,7 @@ namespace UE
 
 	bool StructContainsLooseProperties(const UStruct* Struct)
 	{
-#if WITH_EDITORONLY_DATA
 		return Struct->GetBoolMetaData(NAME_ContainsLoosePropertiesMetadata);
-#else
-		return false;
-#endif
 	}
 
 	static UStruct* CreateInstanceDataObjectStructRec(const UClass* StructClass, UStruct* TemplateStruct, UObject* Outer, const FPropertyPathNameTree* PropertyTree);
@@ -205,7 +203,6 @@ namespace UE
 	// recursively re-instances all structs contained by this property to include loose properties
 	static void ConvertToInstanceDataObjectProperty(FProperty* Property, FPropertyTypeName PropertyType, UObject* Outer, const FPropertyPathNameTree* PropertyTree)
 	{
-#if WITH_EDITORONLY_DATA
 		if (!Property->HasMetaData(NAME_DisplayName))
 		{
 			bool bNeedsDisplayName = false;
@@ -215,23 +212,19 @@ namespace UE
 				Property->SetMetaData(NAME_DisplayName, MoveTemp(DisplayName));
 			}
 		}
-#endif
 
 		const auto TrySetContainsLooseProperties = [](FProperty* Property, const FFieldVariant& Inner)
 		{
-#if WITH_EDITORONLY_DATA
 			if (Inner.HasMetaData(NAME_ContainsLoosePropertiesMetadata))
 			{
 				Property->SetMetaData(NAME_ContainsLoosePropertiesMetadata, TEXT("True"));
 			}
-#endif
 		};
 
 		if (FStructProperty* AsStructProperty = CastField<FStructProperty>(Property))
 		{
 			if (!AsStructProperty->Struct->UseNativeSerialization())
 			{
-#if WITH_EDITORONLY_DATA
 				//@note: Transfer existing metadata over as we build the InstanceDataObject from the struct or it owner, if any, this is useful for testing purposes
 				FString OriginalName;
 				if (const FString* OriginalType = AsStructProperty->FindMetaData(NAME_OriginalType))
@@ -253,20 +246,18 @@ namespace UE
 					OriginalNameBuilder.AddPath(AsStructProperty->Struct);
 					OriginalName = WriteToString<256>(OriginalNameBuilder.Build()).ToView();
 				}
-#endif
+
 				UInstanceDataObjectStruct* Struct = CreateInstanceDataObjectStructRec<UInstanceDataObjectStruct>(AsStructProperty->Struct, Outer, PropertyTree);
 				if (const FName StructGuidName = PropertyType.GetParameterName(1); !StructGuidName.IsNone())
 				{
 					FGuid::Parse(StructGuidName.ToString(), Struct->Guid);
 				}
 				AsStructProperty->Struct = Struct;
-#if WITH_EDITORONLY_DATA
 				AsStructProperty->SetMetaData(NAME_OriginalType, *OriginalName);
 				AsStructProperty->SetMetaData(NAME_PresentAsTypeMetadata, *OriginalName);
 				AsStructProperty->Struct->SetMetaData(NAME_PresentAsTypeMetadata, *OriginalName);
 
 				TrySetContainsLooseProperties(AsStructProperty, AsStructProperty->Struct);
-#endif
 			}
 		}
 		else if (FArrayProperty* AsArrayProperty = CastField<FArrayProperty>(Property))
@@ -305,20 +296,16 @@ namespace UE
 			TrySetContainsLooseProperties(AsOptionalProperty, AsOptionalProperty->GetValueProperty());
 		}
 
-#if WITH_EDITORONLY_DATA
 		if (Property->GetBoolMetaData(NAME_IsLooseMetadata) || Property->GetBoolMetaData(NAME_ContainsLoosePropertiesMetadata))
 		{
 			Property->GetOwnerStruct()->SetMetaData(NAME_ContainsLoosePropertiesMetadata, TEXT("True"));
 		}
-#endif
 	}
 
 	// recursively gives a property the metadata and flags of a loose property
 	static void MarkPropertyAsLoose(FProperty* Property)
 	{
-#if WITH_EDITORONLY_DATA
 		Property->SetMetaData(NAME_IsLooseMetadata, TEXT("True"));
-#endif
 		Property->SetPropertyFlags(CPF_Edit | CPF_EditConst);
 		if (const FArrayProperty* AsArrayProperty = CastField<FArrayProperty>(Property))
 		{
@@ -364,9 +351,7 @@ namespace UE
 				FProperty* SuperProperty = CastFieldChecked<FProperty>(FField::Duplicate(TemplateProperty, Super));
 				SuperProperties.Add(SuperProperty);
 
-			#if WITH_EDITORONLY_DATA
 				FField::CopyMetaData(TemplateProperty, SuperProperty);
-			#endif
 
 				FPropertyTypeName Type;
 				{
@@ -517,12 +502,10 @@ namespace UE
 	UClass* CreateInstanceDataObjectClass(const FPropertyPathNameTree* PropertyTree, UClass* OwnerClass, UObject* Outer)
 	{
 		UClass* Result = CreateInstanceDataObjectStructRec<UInstanceDataObjectClass>(OwnerClass, Outer, PropertyTree);
-#if WITH_EDITORONLY_DATA
 		if (const FString& DisplayName = OwnerClass->GetMetaData(NAME_DisplayName); !DisplayName.IsEmpty())
 		{
 			Result->SetMetaData(NAME_DisplayName, *DisplayName);
 		}
-#endif
 
 		SetClassFlags(Result, OwnerClass);
 
@@ -550,7 +533,6 @@ namespace UE
 
 	void MarkPropertySetBySerialization(const UStruct* Struct, void* StructData, const FProperty* Property, int32 ArrayIndex)
 	{
-	#if WITH_EDITORONLY_DATA
 		if (const FByteProperty* ValuesSetBySerializationProperty = FindValuesSetBySerializationProperty(Struct))
 		{
 			const int32 PropertyIndex = Property->GetIndexInOwner() + ArrayIndex;
@@ -562,12 +544,10 @@ namespace UE
 				*PropertyDataPtr |= (1 << BitOffset);
 			}
 		}
-	#endif
 	}
 
 	bool WasPropertySetBySerialization(const UStruct* Struct, const void* StructData, const FProperty* Property, int32 ArrayIndex)
 	{
-	#if WITH_EDITORONLY_DATA
 		if (const FByteProperty* ValuesSetBySerializationProperty = FindValuesSetBySerializationProperty(Struct))
 		{
 			const int32 PropertyIndex = Property->GetIndexInOwner() + ArrayIndex;
@@ -579,7 +559,6 @@ namespace UE
 				return (*PropertyDataPtr & (1 << BitOffset)) != 0;
 			}
 		}
-	#endif
 		return false;
 	}
 
@@ -698,7 +677,6 @@ namespace UE
 
 	static void SetPropertyValueInitializedFlag(const UStruct* Struct, void* StructData, const FProperty* Property, int32 ArrayIndex, bool bValue)
 	{
-	#if WITH_EDITORONLY_DATA
 		if (const FByteProperty* InitializedValuesProperty = FindInitializedValuesProperty(Struct))
 		{
 			const int32 PropertyIndex = Property->GetIndexInOwner() + ArrayIndex;
@@ -717,12 +695,10 @@ namespace UE
 				}
 			}
 		}
-	#endif
 	}
 
 	bool IsPropertyValueInitialized(const UStruct* Struct, void* StructData, const FProperty* Property, int32 ArrayIndex)
 	{
-	#if WITH_EDITORONLY_DATA
 		if (const FByteProperty* InitializedValuesProperty = FindInitializedValuesProperty(Struct))
 		{
 			const int32 PropertyIndex = Property->GetIndexInOwner() + ArrayIndex;
@@ -734,7 +710,6 @@ namespace UE
 				return (*PropertyDataPtr & (1 << BitOffset)) != 0;
 			}
 		}
-	#endif
 		return false;
 	}
 
@@ -750,13 +725,13 @@ namespace UE
 
 	void ResetPropertyValueInitialized(const UStruct* Struct, void* StructData)
 	{
-	#if WITH_EDITORONLY_DATA
 		if (const FByteProperty* InitializedValuesProperty = FindInitializedValuesProperty(Struct))
 		{
 			uint8* PropertyDataPtr = InitializedValuesProperty->ContainerPtrToValuePtr<uint8>(StructData);
 			FMemory::Memzero(PropertyDataPtr, InitializedValuesProperty->ArrayDim);
 		}
-	#endif
 	}
 
 } // UE
+
+#endif // WITH_EDITORONLY_DATA
diff --git a/Engine/Source/Runtime/CoreUObject/Private/UObject/Obj.cpp b/Engine/Source/Runtime/CoreUObject/Private/UObject/Obj.cpp
--- a/Engine/Source/Runtime/CoreUObject/Private/UObject/Obj.cpp
+++ b/Engine/Source/Runtime/CoreUObject/Private/UObject/Obj.cpp
@@ -1023,8 +1023,10 @@ void UObject::BeginDestroy()
 	// Remove any associated external package, at this point
 	SetExternalPackage(nullptr);
 
+#if WITH_EDITORONLY_DATA
 	// Destroy any associated property bag.
 	UE::FPropertyBagRepository::Get().DestroyOuterBag(this);
+#endif
 	
 	// ensure BeginDestroy has been routed back to UObject::BeginDestroy.
 #if !(UE_BUILD_SHIPPING || UE_BUILD_TEST)
@@ -1297,8 +1299,8 @@ void UObject::ConditionalPostLoad()
 		}
 
 		ConditionalPostLoadSubobjects();
-		
-		
+
+#if WITH_EDITORONLY_DATA
 		// Object has been deserialized, if IDO is enabled, generate it
 		UE::FPropertyBagRepository& PropertyBagRepository = UE::FPropertyBagRepository::Get();
 		const FUObjectSerializeContext* LoadContext = FUObjectThreadContext::Get().GetSerializeContext();
@@ -1307,6 +1309,7 @@ void UObject::ConditionalPostLoad()
 		{
 			PropertyBagRepository.PostLoadInstanceDataObject(this);
 		}
+#endif
 
 		{
 			FExclusiveLoadPackageTimeTracker::FScopedPostLoadTracker Tracker(this);
@@ -1829,6 +1832,7 @@ void UObject::SerializeScriptProperties(FArchive& Ar) const
 	SerializeScriptProperties(FStructuredArchiveFromArchive(Ar).GetSlot());
 }
 
+#if WITH_EDITORONLY_DATA
 namespace UE::Private
 {
 	/**
@@ -1849,6 +1853,7 @@ namespace UE::Private
 		return FPropertyBagRepository::Get().FindInstanceDataObject(ThisObject);
 	}
 }
+#endif
 
 void UObject::SerializeScriptProperties( FStructuredArchive::FSlot Slot ) const
 {
@@ -1874,6 +1879,7 @@ void UObject::SerializeScriptProperties( FStructuredArchive::FSlot Slot ) const
 
 		// Query if this object data is being impersonated 
 		const UObject* ThisObject = this;
+#if WITH_EDITORONLY_DATA
 		if (const UObject* Impersonator = UE::Private::GetDataImpersonator(ThisObject))
 		{
 			ThisObject = Impersonator;
@@ -1889,6 +1895,7 @@ void UObject::SerializeScriptProperties( FStructuredArchive::FSlot Slot ) const
 			ensureAlwaysMsgf(DiffClass == ObjClass, TEXT("Impersonation of '%s' using a different default class not appropriately supported at the moment. Class: '%s', DefaultClass: '%s'")
 				, *ThisObject->GetPathName(), *ObjClass->GetPathName(), *DiffClass->GetPathName());
 		}
+#endif
 
 #if WITH_EDITOR
 		static const FBoolConfigValueHelper BreakSerializationRecursion(TEXT("StructSerialization"), TEXT("BreakSerializationRecursion"));
diff --git a/Engine/Source/Runtime/CoreUObject/Private/UObject/PropertyArray.cpp b/Engine/Source/Runtime/CoreUObject/Private/UObject/PropertyArray.cpp
--- a/Engine/Source/Runtime/CoreUObject/Private/UObject/PropertyArray.cpp
+++ b/Engine/Source/Runtime/CoreUObject/Private/UObject/PropertyArray.cpp
@@ -623,6 +623,7 @@ void FArrayProperty::SerializeItem(FStructuredArchive::FSlot Slot, void* Value,
 			}
 		}
 
+	#if WITH_EDITORONLY_DATA
 		if (Context->bTrackSerializedPropertyPath)
 		{
 			// Update the path with types from the inner tag if the outer tag is incomplete.
@@ -645,6 +646,7 @@ void FArrayProperty::SerializeItem(FStructuredArchive::FSlot Slot, void* Value,
 				}
 			}
 		}
+	#endif
 	}
 
 	// need to know how much data this call to SerializeItem consumes, so mark where we are
diff --git a/Engine/Source/Runtime/CoreUObject/Private/UObject/PropertyBagRepository.cpp b/Engine/Source/Runtime/CoreUObject/Private/UObject/PropertyBagRepository.cpp
--- a/Engine/Source/Runtime/CoreUObject/Private/UObject/PropertyBagRepository.cpp
+++ b/Engine/Source/Runtime/CoreUObject/Private/UObject/PropertyBagRepository.cpp
@@ -2,6 +2,8 @@
 
 #include "UObject/PropertyBagRepository.h"
 
+#if WITH_EDITORONLY_DATA
+
 #include "Containers/Queue.h"
 #include "Serialization/ObjectReader.h"
 #include "Serialization/ObjectWriter.h"
@@ -763,7 +765,7 @@ void FPropertyBagRepository::CopyTaggedProperties(const UObject* Source, UObject
 
 FScopedIDOSerializationContext::FScopedIDOSerializationContext(UObject* InObject, FArchive& InArchive)
 	: Archive(&InArchive)
-    , Object(InObject)
+	, Object(InObject)
 	, PreSerializeOffset(InArchive.Tell())
 {
 	FUObjectSerializeContext* SerializeContext = FUObjectThreadContext::Get().GetSerializeContext();
@@ -929,3 +931,5 @@ UStruct* FPropertyBagRepository::CreatePropertyBagPlaceholderType(UObject* Outer
 }
 
 } // UE
+
+#endif // WITH_EDITORONLY_DATA
diff --git a/Engine/Source/Runtime/CoreUObject/Private/UObject/PropertyHelper.h b/Engine/Source/Runtime/CoreUObject/Private/UObject/PropertyHelper.h
--- a/Engine/Source/Runtime/CoreUObject/Private/UObject/PropertyHelper.h
+++ b/Engine/Source/Runtime/CoreUObject/Private/UObject/PropertyHelper.h
@@ -106,19 +106,23 @@ class FSerializedPropertyPathScope
 public:
 	[[nodiscard]] FSerializedPropertyPathScope(FUObjectSerializeContext* InContext, const FPropertyPathNameSegment& InSegment)
 	{
+	#if WITH_EDITORONLY_DATA
 		if (InContext->bTrackSerializedPropertyPath)
 		{
 			Context = InContext;
 			Context->SerializedPropertyPath.Push(InSegment);
 		}
+	#endif
 	}
 
 	~FSerializedPropertyPathScope()
 	{
+	#if WITH_EDITORONLY_DATA
 		if (Context)
 		{
 			Context->SerializedPropertyPath.Pop();
 		}
+	#endif
 	}
 
 	FSerializedPropertyPathScope(const FSerializedPropertyPathScope&) = delete;
diff --git a/Engine/Source/Runtime/CoreUObject/Private/UObject/UObjectGlobals.cpp b/Engine/Source/Runtime/CoreUObject/Private/UObject/UObjectGlobals.cpp
--- a/Engine/Source/Runtime/CoreUObject/Private/UObject/UObjectGlobals.cpp
+++ b/Engine/Source/Runtime/CoreUObject/Private/UObject/UObjectGlobals.cpp
@@ -3255,12 +3255,14 @@ UObject* StaticDuplicateObjectEx( FObjectDuplicationParameters& Parameters )
 		}
 	}
 
+#if WITH_EDITORONLY_DATA
 	// if Source has an IDO, make one for Dest and copy it
 	UE::FPropertyBagRepository& PropertyBagRepository = UE::FPropertyBagRepository::Get();
 	if (UE::FPropertyBagRepository::IsInstanceDataObjectSupportEnabled(Parameters.SourceObject))
 	{
 		PropertyBagRepository.DuplicateInstanceDataObject(Parameters.SourceObject, DupRootObject);
 	}
+#endif
 	
 	return DupRootObject;
 }
diff --git a/Engine/Source/Runtime/CoreUObject/Private/UObject/UObjectThreadContext.cpp b/Engine/Source/Runtime/CoreUObject/Private/UObject/UObjectThreadContext.cpp
--- a/Engine/Source/Runtime/CoreUObject/Private/UObject/UObjectThreadContext.cpp
+++ b/Engine/Source/Runtime/CoreUObject/Private/UObject/UObjectThreadContext.cpp
@@ -44,11 +44,14 @@ FUObjectSerializeContext::FUObjectSerializeContext()
 	, SerializedImportLinker(nullptr)
 	, SerializedExportIndex(0)
 	, SerializedExportLinker(nullptr)
+#if WITH_EDITORONLY_DATA
 	, bTrackSerializedPropertyPath(false)
 	, bTrackInitializedProperties(false)
 	, bTrackUnknownProperties(false)
 	, bImpersonateProperties(false)
-{}
+#endif
+{
+}
 
 FUObjectSerializeContext::~FUObjectSerializeContext()
 {
diff --git a/Engine/Source/Runtime/CoreUObject/Public/UObject/UObjectThreadContext.h b/Engine/Source/Runtime/CoreUObject/Public/UObject/UObjectThreadContext.h
--- a/Engine/Source/Runtime/CoreUObject/Public/UObject/UObjectThreadContext.h
+++ b/Engine/Source/Runtime/CoreUObject/Public/UObject/UObjectThreadContext.h
@@ -168,6 +168,7 @@ public:
 	int32 SerializedExportIndex;
 	/** Points to the most recently used Linker for serialization by CreateExport() */
 	FLinkerLoad* SerializedExportLinker;
+#if WITH_EDITORONLY_DATA
 	/** Path to the property currently being serialized */
 	UE_INTERNAL UE::FPropertyPathName SerializedPropertyPath;
 	/** True when SerializedPropertyPath is being tracked during serialization. */
@@ -178,6 +179,7 @@ public:
 	UE_INTERNAL bool bTrackUnknownProperties;
 	/** True when the SerializedObject properties are being impersonated. */
 	UE_INTERNAL bool bImpersonateProperties;
+#endif
 
 	/** Adds a new loaded object */
 	COREUOBJECT_API void AddLoadedObject(UObject* InObject);
diff --git a/Engine/Source/Runtime/Engine/Private/Level.cpp b/Engine/Source/Runtime/Engine/Private/Level.cpp
--- a/Engine/Source/Runtime/Engine/Private/Level.cpp
+++ b/Engine/Source/Runtime/Engine/Private/Level.cpp
@@ -458,7 +458,10 @@ void ULevel::AddReferencedObjects(UObject* InThis, FReferenceCollector& Collecto
 void ULevel::CleanupLevel(bool bCleanupResources, bool bUnloadFromEditor)
 {
 	OnCleanupLevel.Broadcast();
+
+#if WITH_EDITORONLY_DATA
 	UE::FPropertyBagRepository::Get().CleanupLevel(this);
+#endif
 
 	if (bCleanupResources)
 	{
