Add a way to safely use UE::AssetRegistry::FFiltering::ShouldSkipAsset(const FTopLevelAssetPath& AssetClass, uint32 PackageFlags)
in parrallelfor that block the GameThread.

diff --git a/Engine/Source/Runtime/CoreUObject/Private/Misc/AssetRegistryInterface.cpp b/Engine/Source/Runtime/CoreUObject/Private/Misc/AssetRegistryInterface.cpp
--- a/Engine/Source/Runtime/CoreUObject/Private/Misc/AssetRegistryInterface.cpp
+++ b/Engine/Source/Runtime/CoreUObject/Private/Misc/AssetRegistryInterface.cpp
@@ -38,6 +38,22 @@ namespace Private
 		SkipUncookedClasses = InSkipUncookedClasses;
 		SkipCookedClasses = InSkipCookedClasses;
 	}
+
+	void FFiltering::InitializeShouldSkipAsset()
+	{
+		if (!bInitializedSkipClasses)
+		{
+			// Since we only collect these the first on-demand time, it is possible we will miss subclasses
+			// from plugins that load later. This flaw is a rare edge case, though, and this solution will
+			// be replaced eventually, so leaving it for now.
+			if (GIsEditor && (!IsRunningCommandlet() || IsRunningCookCommandlet()))
+			{
+				Utils::PopulateSkipClasses(SkipUncookedClasses, SkipCookedClasses);
+			}
+
+			bInitializedSkipClasses = true;
+		}
+	}
 #endif
 
 #if WITH_ENGINE && WITH_EDITOR
@@ -140,18 +156,8 @@ namespace Utils
 		// an asset; the content browser does not handle the multiple assets correctly and displays this
 		// class asset as if it is in a separate package. Revisit when we have removed the UBlueprint as an asset
 		// or when we support multiple assets.
-		if (!bInitializedSkipClasses)
-		{
-			// Since we only collect these the first on-demand time, it is possible we will miss subclasses
-			// from plugins that load later. This flaw is a rare edge case, though, and this solution will
-			// be replaced eventually, so leaving it for now.
-			if (GIsEditor && (!IsRunningCommandlet() || IsRunningCookCommandlet()))
-			{
-				Utils::PopulateSkipClasses(SkipUncookedClasses, SkipCookedClasses);
-			}
+		InitializeShouldSkipAsset();
 
-			bInitializedSkipClasses = true;
-		}
 		return Utils::ShouldSkipAsset(AssetClass, PackageFlags, SkipUncookedClasses, SkipCookedClasses);
 #else
 		return false;
diff --git a/Engine/Source/Runtime/CoreUObject/Public/Misc/AssetRegistryInterface.h b/Engine/Source/Runtime/CoreUObject/Public/Misc/AssetRegistryInterface.h
--- a/Engine/Source/Runtime/CoreUObject/Public/Misc/AssetRegistryInterface.h
+++ b/Engine/Source/Runtime/CoreUObject/Public/Misc/AssetRegistryInterface.h
@@ -236,6 +236,12 @@ namespace UE::AssetRegistry
 #if WITH_ENGINE && WITH_EDITOR
 		/** Copy the global skip classes set from the given external sets that were already populated. */
 		static COREUOBJECT_API void SetSkipClasses(const TSet<FTopLevelAssetPath>& InSkipUncookedClasses, const TSet<FTopLevelAssetPath>& InSkipCookedClasses);
+
+		/**
+		 * Prepare the data structure needed for a call to Should skip asset. 
+		 * This make the function bool ShouldSkipAsset(const FTopLevelAssetPath& AssetClass, uint32 PackageFlags) thread safe if the Game Thread is on hold during those calls.
+		 */
+		static COREUOBJECT_API void InitializeShouldSkipAsset();
 #endif
 	};
 
