Preserve PKG_ForDiffing when using Zen or async loading
[RN] PKG_ForDiffing is now preserved when using Zen or AsyncLoading a package

diff --git a/Engine/Source/Runtime/CoreUObject/Private/UObject/LinkerLoad.cpp b/Engine/Source/Runtime/CoreUObject/Private/UObject/LinkerLoad.cpp
--- a/Engine/Source/Runtime/CoreUObject/Private/UObject/LinkerLoad.cpp
+++ b/Engine/Source/Runtime/CoreUObject/Private/UObject/LinkerLoad.cpp
@@ -1674,12 +1674,10 @@ FLinkerLoad::ELinkerStatus FLinkerLoad::UpdateFromPackageFileSummary()
 
 	if (UPackage* LinkerRootPackage = LinkerRoot)
 	{
-		// Preserve PIE package flag
-		uint32 NewPackageFlags = Summary.GetPackageFlags();
-		if (LinkerRootPackage->HasAnyPackageFlags(PKG_PlayInEditor))
-		{
-			NewPackageFlags |= PKG_PlayInEditor;
-		}
+		const uint32 NewPackageFlags = Summary.GetPackageFlags() |
+			// Preserve PKG_PlayInEditor and PKG_ForDiffing, they have been provided 
+			// by the caller, but are not identified as truly transient:
+			(LinkerRootPackage->GetPackageFlags() & (PKG_PlayInEditor|PKG_ForDiffing));
 
 		// Propagate package flags
 		LinkerRootPackage->SetPackageFlagsTo(NewPackageFlags);
