Global function IsReferenced will no longer be using the Unreachable flag. This is to prevent it from leaving objects marked as unreachable outside of normal GC cycle.

diff --git a/Engine/Source/Runtime/CoreUObject/Private/UObject/UObjectGlobals.cpp b/Engine/Source/Runtime/CoreUObject/Private/UObject/UObjectGlobals.cpp
--- a/Engine/Source/Runtime/CoreUObject/Private/UObject/UObjectGlobals.cpp
+++ b/Engine/Source/Runtime/CoreUObject/Private/UObject/UObjectGlobals.cpp
@@ -5,6 +5,7 @@
 =============================================================================*/
 
 #include "UObject/UObjectGlobals.h"
+#include "Containers/BitArray.h"
 #include "HAL/PlatformFileManager.h"
 #include "HAL/FileManager.h"
 #include "Logging/StructuredLog.h"
@@ -29,7 +30,6 @@
 #include "UObject/UObjectHashPrivate.h"
 #include "UObject/Object.h"
 #include "UObject/GarbageCollection.h"
-#include "UObject/GarbageCollectionInternalFlags.h"
 #include "UObject/Class.h"
 #include "UObject/CoreRedirects.h"
 #include "UObject/FastReferenceCollector.h"
@@ -5262,6 +5262,11 @@ public:
 	{
 	}
 
+	FORCEINLINE bool IsUnreachable(const UObject* Object) const
+	{
+		return !ReachabilityBits[GUObjectArray.ObjectToIndex(Object)];
+	}
+
 	// FReferenceCollector interface
 	virtual bool IsIgnoringArchetypeRef() const override
 	{
@@ -5285,22 +5290,22 @@ public:
 	 */
 	void PerformReachabilityAnalysis( EObjectFlags KeepFlags, EInternalObjectFlags InternalKeepFlags, EObjectFlags SearchFlags = RF_NoFlags, FReferencerInformationList* FoundReferences = NULL)
 	{
-		// Reset object count.
-		extern FThreadSafeCounter GObjectCountDuringLastMarkPhase;
-		GObjectCountDuringLastMarkPhase.Reset();
 		ReferenceSearchFlags = SearchFlags;
 		FoundReferencesList = FoundReferences;
 
+		ReachabilityBits.Init(false, GUObjectArray.GetObjectArrayNum());
+
 		// Iterate over all objects.
 		for( FThreadSafeObjectIterator It; It; ++It )
 		{
 			UObject* Object	= *It;
 			checkSlow(Object->IsValidLowLevel());
-			GObjectCountDuringLastMarkPhase.Increment();
 
 			// Special case handling for objects that are part of the root set.
 			if( Object->IsRooted() )
 			{
+				SetReachable(Object);
+
 				checkSlow( Object->IsValidLowLevel() );
 				// We cannot use RF_PendingKill on objects that are part of the root set.
 				checkCode( if( !IsValidChecked(Object) ) { UE_LOG(LogUObjectGlobals, Fatal, TEXT("Object %s is part of root set though is invalid!"), *Object->GetFullName() ); } );
@@ -5310,18 +5315,14 @@ public:
 			// Regular objects.
 			else
 			{
-				// Mark objects as unreachable unless they have any of the passed in KeepFlags set and none of the passed in Search.
+				// Mark objects as reachable when they have any of the passed in KeepFlags set and none of the passed in Search.
 				if (!Object->HasAnyFlags(SearchFlags) &&
 					((KeepFlags == RF_NoFlags && InternalKeepFlags == EInternalObjectFlags::None) || Object->HasAnyFlags(KeepFlags) || Object->HasAnyInternalFlags(InternalKeepFlags))
 					)
 				{
+					SetReachable(Object);
 					ObjectsToSerialize.Add(Object);
 				}
-				else
-				{
-					FUObjectItem* ObjectItem = GUObjectArray.ObjectToObjectItem(Object);
-					UE::GC::Private::FGCFlags::ThisThreadAtomicallySetFlag_ForGC(ObjectItem, EInternalObjectFlags::Unreachable);
-				}
 			}
 		}
 
@@ -5373,8 +5374,7 @@ private:
 		}
 
 		// Mark it as reachable.
-		FUObjectItem* ObjectItem = GUObjectArray.ObjectToObjectItem(Object);
-		UE::GC::Private::FGCFlags::ThisThreadAtomicallyClearedFlag_ForGC(ObjectItem, EInternalObjectFlags::Unreachable);
+		SetReachable(Object);
 
 		// Add it to the list of objects to serialize.
 		ObjectsToSerialize.Add( Object );
@@ -5401,10 +5401,9 @@ private:
 					CurrentReferenceInfo->TotalReferences++;
 				}
 				// Mark it as reachable.
-				FUObjectItem* ObjectItem = GUObjectArray.ObjectToObjectItem(InObject);
-				UE::GC::Private::FGCFlags::ThisThreadAtomicallyClearedFlag_ForGC(ObjectItem, EInternalObjectFlags::Unreachable);
+				SetReachable(InObject);
 			}
-			else if (InObject->IsUnreachable())
+			else if (IsUnreachable(InObject))
 			{
 				// Add encountered object reference to list of to be serialized objects if it hasn't already been added.
 				AddToObjectList(InReferencingObject, InReferencingProperty, InObject);
@@ -5412,6 +5411,13 @@ private:
 		}
 	}
 
+	FORCEINLINE void SetReachable(const UObject* Object)
+	{
+		ReachabilityBits[GUObjectArray.ObjectToIndex(Object)] = true;
+	}
+
+	/** Bitset containing reachability bits for each of the existing objects */
+	TBitArray<> 		ReachabilityBits;
 	/** Object we're currently serializing */
 	UObject*			CurrentObject;
 	/** Growing array of objects that require serialization */
@@ -5483,12 +5489,12 @@ bool IsReferenced(UObject*& Obj, EObjectFlags KeepFlags, EInternalObjectFlags In
 				i--;
 			}
 		}
-		bIsReferenced = FoundReferences->ExternalReferences.Num() > 0 || bReferencedByOuters || !Obj->IsUnreachable();
+		bIsReferenced = FoundReferences->ExternalReferences.Num() > 0 || bReferencedByOuters || !ObjectReferenceTagger.IsUnreachable(Obj);
 	}
 	else
 	{
-		// Return whether the object was referenced and restore original state.
-		bIsReferenced = !Obj->IsUnreachable();
+		// Return whether the object was referenced
+		bIsReferenced = !ObjectReferenceTagger.IsUnreachable(Obj);
 	}
 	
 	if (bTempReferenceList)
