Further revise TObjectPtr to avoid indirection by keeping references to placeholder-typed objects unresolved.

Change summary:
- Reduced IsObjectHandleTypeSafe() to just a simple bit test. The LSB must also be set to require that it's a packed ref.
- Modified ResolveObjectHandleNoRead()/ResolveObjectHandleNoReadNoCheck() to update the handle only if type safe.
- Modified ObjectPtr_Private::Get() to return NULL without first resolving the handle (for pointers bound to UObject subtypes).
- Modified FLinkerLoad::operator<<() to serialize placeholder-typed object references as an "unsafe" packed reference only if the pointer type is unsafe.
- Modified TObjectPtr::GetInternalRef() to assert if the handle is not resolved. This will catch an attempt to return an unsafe pointer as a packed reference.
- Modified ToRawPtrArrayUnsafe() to assert if any item in the array is not resolved. This will catch an attempt to treat a packed reference as the element type.
- Modified TContainerElementTypeCompatibility::ReinterpretRange()/ReinterpretRangeContiguous() to assert if any element of the pointer array is not resolved.
- Modified ToRawPtrTArrayUnsafe() to reinterpret the pointer array through the type compatibility layer to catch any attempt to treat a packed reference as the element type.
- Revised TObjectPtr natvis to reduce the placeholder intrinsic query down to just a single bit test.
- Removed 'GObjectIndexToPackedObjectRefDebug' as it is no longer required for natvis support.
- Revised FObjectHandle/TObjectPtr type safety unit tests to validate the unresolved state of placeholder-typed references as well as include invalid address inputs to ensure we're only using pointer semantics.

diff --git a/Engine/Source/Runtime/Core/Private/Misc/CoreGlobals.cpp b/Engine/Source/Runtime/Core/Private/Misc/CoreGlobals.cpp
--- a/Engine/Source/Runtime/Core/Private/Misc/CoreGlobals.cpp
+++ b/Engine/Source/Runtime/Core/Private/Misc/CoreGlobals.cpp
@@ -253,7 +253,6 @@ float				GNearClippingPlane_RenderThread		= 10.0f;				/* Near clipping plane (Re
 
 bool					GExitPurge						= false;
 
-const UPTRINT* GCoreObjectIndexToPackedObjectRefDebug = nullptr;
 FChunkedFixedUObjectArray* GCoreObjectArrayForDebugVisualizers = nullptr;
 
 namespace UE::CoreUObject::Private
@@ -265,7 +264,6 @@ UE::CoreUObject::Private::FStoredObjectPathDebug* GCoreComplexObjectPathDebug =
 UE::CoreUObject::Private::FObjectHandlePackageDebugData* GCoreObjectHandlePackageDebug = nullptr;
 #if PLATFORM_UNIX
 uint8** CORE_API GNameBlocksDebug = FNameDebugVisualizer::GetBlocks();
-const UPTRINT*& CORE_API GObjectIndexToPackedObjectRefDebug = GCoreObjectIndexToPackedObjectRefDebug;
 FChunkedFixedUObjectArray*& CORE_API GObjectArrayForDebugVisualizers = GCoreObjectArrayForDebugVisualizers;
 UE::CoreUObject::Private::FStoredObjectPathDebug*& GComplexObjectPathDebug = GCoreComplexObjectPathDebug;
 UE::CoreUObject::Private::FObjectHandlePackageDebugData*& CORE_API GObjectHandlePackageDebug = GCoreObjectHandlePackageDebug;
diff --git a/Engine/Source/Runtime/Core/Public/CoreGlobals.h b/Engine/Source/Runtime/Core/Public/CoreGlobals.h
--- a/Engine/Source/Runtime/Core/Public/CoreGlobals.h
+++ b/Engine/Source/Runtime/Core/Public/CoreGlobals.h
@@ -809,9 +809,6 @@ struct FScopedLoadingState
 bool CORE_API GetEmitDrawEvents();
 void CORE_API SetEmitDrawEvents(bool EmitDrawEvents);
 
-/** Array to help visualize packed object references */
-extern CORE_API const UPTRINT* GCoreObjectIndexToPackedObjectRefDebug;
-
 /** Array to help visualize weak pointers in the debugger */
 class FChunkedFixedUObjectArray;
 
diff --git a/Engine/Source/Runtime/Core/Public/Modules/Boilerplate/ModuleBoilerplate.h b/Engine/Source/Runtime/Core/Public/Modules/Boilerplate/ModuleBoilerplate.h
--- a/Engine/Source/Runtime/Core/Public/Modules/Boilerplate/ModuleBoilerplate.h
+++ b/Engine/Source/Runtime/Core/Public/Modules/Boilerplate/ModuleBoilerplate.h
@@ -91,7 +91,6 @@ class FChunkedFixedUObjectArray;
 #else
 	#define UE_VISUALIZERS_HELPERS \
 		uint8** GNameBlocksDebug = FNameDebugVisualizer::GetBlocks(); \
-		const UPTRINT*& GObjectIndexToPackedObjectRefDebug = GCoreObjectIndexToPackedObjectRefDebug; \
 		FChunkedFixedUObjectArray*& GObjectArrayForDebugVisualizers = GCoreObjectArrayForDebugVisualizers; \
 		UE::CoreUObject::Private::FStoredObjectPathDebug*& GComplexObjectPathDebug = GCoreComplexObjectPathDebug; \
 		UE::CoreUObject::Private::FObjectHandlePackageDebugData*& GObjectHandlePackageDebug = GCoreObjectHandlePackageDebug;
diff --git a/Engine/Source/Runtime/CoreUObject/Private/UObject/LinkerLoad.cpp b/Engine/Source/Runtime/CoreUObject/Private/UObject/LinkerLoad.cpp
--- a/Engine/Source/Runtime/CoreUObject/Private/UObject/LinkerLoad.cpp
+++ b/Engine/Source/Runtime/CoreUObject/Private/UObject/LinkerLoad.cpp
@@ -6230,23 +6230,21 @@ FArchive& FLinkerLoad::operator<<(FObjectPtr& ObjectPtr)
 		// Note: References to placeholder objects cannot resolve to it if the underlying pointer type is unsafe.
 		if (ResolvedObject && UE::FPropertyBagRepository::IsPropertyBagPlaceholderObject(ResolvedObject))
 		{
-#if UE_WITH_OBJECT_HANDLE_LATE_RESOLVE && UE_WITH_OBJECT_HANDLE_TYPE_SAFETY
-			// If type safety features are enabled, create a packed reference mapping for the placeholder-typed object.
-			// This resolves to the placeholder object for UObject-typed pointers. For other types it resolves to NULL.
-			// However, the underlying value when serialized will always resolve to the object reference (e.g. for GC).
-			// Note: We could return an FObjectPtr that wraps the packed reference result, but that implies a lazy load.
-			// We need to register the packed object reference here, but we don't need to also defer the pointer resolve.
-			UE::CoreUObject::Private::MakePackedObjectRef(ResolvedObject);
-#else
-			// If type safety features are disabled, resolve unsafe references to placeholder-typed objects now to NULL.
-			// Note: Similar to hard references above, this means we won't find it for replacement at reinstancing time.
 			const FObjectProperty* ObjectProperty = CastField<FObjectProperty>(GetSerializedProperty());
 			if (!ObjectProperty || !ResolvedObject->GetClass()->IsChildOf(ObjectProperty->PropertyClass))
 			{
-				UE_LOG(LogLinker, Warning, TEXT("Serializing reference to \"%s\" as NULL to ensure type safety."), *ResolvedObject->GetPathName());
+#if UE_WITH_OBJECT_HANDLE_LATE_RESOLVE && UE_WITH_OBJECT_HANDLE_TYPE_SAFETY
+				// If type safety features are enabled, create a packed reference mapping for the placeholder-typed object.
+				// This resolves to NULL on access since the reference can't be cast to a pointer bound to its original type.
+				// However, the underlying value (when serialized) will always resolve to the placeholder object (e.g. for GC).
+				return FObjectPtr({ UE::CoreUObject::Private::MakePackedObjectRef(ResolvedObject).EncodedRef });
+#else
+				// If type safety features are disabled, serialize it as an unsafe reference to a placeholder-typed object.
+				// Note: Similar to hard references above, this means we won't find it for replacement at reinstancing time.
+				UE_LOG(LogLinker, Warning, TEXT("Serializing reference to \"%s\" as NULL to ensure type safety. This will lead to data loss if the referencing object is saved."), *ResolvedObject->GetPathName());
 				ResolvedObject = nullptr;
-			}
 #endif
+			}
 		}
 #endif
 		return FObjectPtr(ResolvedObject);
diff --git a/Engine/Source/Runtime/CoreUObject/Private/UObject/ObjectHandleIndex.cpp b/Engine/Source/Runtime/CoreUObject/Private/UObject/ObjectHandleIndex.cpp
--- a/Engine/Source/Runtime/CoreUObject/Private/UObject/ObjectHandleIndex.cpp
+++ b/Engine/Source/Runtime/CoreUObject/Private/UObject/ObjectHandleIndex.cpp
@@ -130,9 +130,6 @@ namespace UE::CoreUObject::Private
 	void InitObjectHandles(int32 MaxObjects)
 	{
 		GObjectHandleIndex.ObjectIndexToPackedObjectRef.SetNumZeroed(MaxObjects);
-
-		// This allows the debug visualizer to evaluate a resolved handle to see if it has a packed ref stored in the global table.
-		GCoreObjectIndexToPackedObjectRefDebug = reinterpret_cast<const UPTRINT*>(GObjectHandleIndex.ObjectIndexToPackedObjectRef.GetData());
 	}
 
 	static FAutoConsoleCommand CmdPrintUnresolvedObjects(
diff --git a/Engine/Source/Runtime/CoreUObject/Public/UObject/ObjectHandle.h b/Engine/Source/Runtime/CoreUObject/Public/UObject/ObjectHandle.h
--- a/Engine/Source/Runtime/CoreUObject/Public/UObject/ObjectHandle.h
+++ b/Engine/Source/Runtime/CoreUObject/Public/UObject/ObjectHandle.h
@@ -88,11 +88,6 @@ namespace UE::CoreUObject::Private
 
 	/** Read the handle as a pointer without checking if it is resolved. Invalid to call for unresolved handles. */
 	inline UObject* ReadObjectHandlePointerNoCheck(FObjectHandle Handle);
-
-#if UE_WITH_OBJECT_HANDLE_TYPE_SAFETY
-	/** Determines if the handle represents a reference that is considered to be safe for object subtype pointers. */
-	inline bool IsObjectHandleTypeSafeNoReadNoCheck(FObjectHandle Handle);
-#endif
 }
 
 /////////////////////////////////////////////////////////////////////////////////////////////////////////////////////////
@@ -121,12 +116,12 @@ inline bool IsObjectHandleResolved(FObjectHandle Handle)
 }
 
 /* return true if a handle is type safe.
- * null handles are considered safe, otherwise does a bit test
+ * null and resolved handles are considered type safe
  */ 
 inline bool IsObjectHandleTypeSafe(FObjectHandle Handle)
 {
-#if UE_WITH_OBJECT_HANDLE_TYPE_SAFETY
-	return UE::CoreUObject::Private::IsObjectHandleTypeSafeNoReadNoCheck(Handle);
+#if UE_WITH_OBJECT_HANDLE_LATE_RESOLVE && UE_WITH_OBJECT_HANDLE_TYPE_SAFETY
+	return !((Handle.PointerOrRef & 3) == 3);
 #else
 	return true;
 #endif
@@ -281,7 +276,12 @@ namespace UE::CoreUObject::Private
 			FPackedObjectRef PackedObjectRef = ReadObjectHandlePackedObjectRefNoCheck(LocalHandle);
 			FObjectRef ObjectRef = MakeObjectRef(PackedObjectRef);
 			UObject* ResolvedObject = ObjectRef.Resolve();
-			Handle = MakeObjectHandle(ResolvedObject);
+#if UE_WITH_OBJECT_HANDLE_TYPE_SAFETY
+			if (IsObjectHandleTypeSafe(LocalHandle))
+#endif
+			{
+				Handle = MakeObjectHandle(ResolvedObject);
+			}
 			return ResolvedObject;
 		}
 #else
@@ -297,8 +297,13 @@ namespace UE::CoreUObject::Private
 		FPackedObjectRef PackedObjectRef = ReadObjectHandlePackedObjectRefNoCheck(LocalHandle);
 		FObjectRef ObjectRef = MakeObjectRef(PackedObjectRef);
 		UObject* ResolvedObject = ObjectRef.Resolve();
-		LocalHandle = MakeObjectHandle(ResolvedObject);
-		Handle = LocalHandle;
+#if UE_WITH_OBJECT_HANDLE_TYPE_SAFETY
+		if (IsObjectHandleTypeSafe(LocalHandle))
+#endif
+		{
+			LocalHandle = MakeObjectHandle(ResolvedObject);
+			Handle = LocalHandle;
+		}
 		return ResolvedObject;
 #else
 		return ReadObjectHandlePointerNoCheck(Handle);
@@ -360,18 +365,6 @@ namespace UE::CoreUObject::Private
 	inline constexpr uint32 PackageIdShift = 34;
 	inline constexpr uint32 PackageIdMask = 0x3FFF'FFFF;
 
-#if UE_WITH_OBJECT_HANDLE_TYPE_SAFETY
-	inline bool IsObjectHandleTypeSafeNoReadNoCheck(FObjectHandle Handle)
-	{
-#if UE_WITH_OBJECT_HANDLE_LATE_RESOLVE
-		constexpr uint32 UnsafeTypeMask = (1 << TypeIdShift);
-		return IsObjectHandleNull(Handle) || ((Handle.PointerOrRef & (UnsafeTypeMask | 1)) == 1) || (IsObjectHandleResolved(Handle) && !(FindExistingPackedObjectRef(ReadObjectHandlePointerNoCheck(Handle)).EncodedRef & UnsafeTypeMask));
-#else
-		return true;
-#endif
-	}
-#endif
-
 #if UE_WITH_OBJECT_HANDLE_LATE_RESOLVE
 	//forward declarations
 	void InitObjectHandles(int32 Size);
diff --git a/Engine/Source/Runtime/CoreUObject/Public/UObject/ObjectPtr.h b/Engine/Source/Runtime/CoreUObject/Public/UObject/ObjectPtr.h
--- a/Engine/Source/Runtime/CoreUObject/Public/UObject/ObjectPtr.h
+++ b/Engine/Source/Runtime/CoreUObject/Public/UObject/ObjectPtr.h
@@ -433,14 +433,13 @@ namespace ObjectPtr_Private
 	>
 	FORCEINLINE T* Get(const FObjectPtr& ObjectPtr)
 	{
-		// Always resolve the pointer first (for late resolve), even if the handle is not considered type safe.
-		UObject* Result = ObjectPtr.Get();
 		if (!IsObjectHandleTypeSafe(ObjectPtr.GetHandle()))
 		{
-			Result = nullptr;
+			// Type is unsafe; return NULL without resolving.
+			return nullptr;
 		}
 
-		return (T*)Result;
+		return (T*)ObjectPtr.Get();
 	}
 #endif
 
@@ -712,7 +711,10 @@ private:
 	//			to handle this case without forcing the calling code to be modified.
 	FORCEINLINE T*& GetInternalRef()
 	{
-		ObjectPtr.Get();
+		ObjectPtr_Private::Get<T>(ObjectPtr);
+#if UE_WITH_OBJECT_HANDLE_LATE_RESOLVE || UE_WITH_OBJECT_HANDLE_TRACKING
+		check(ObjectPtr.IsResolved());
+#endif
 		return (T*&)ObjectPtr.GetHandleRef();
 	}
 
@@ -923,6 +925,7 @@ ToRawPtrArrayUnsafe(TObjectPtr<T>(&ArrayOfPtr)[Size])
 	{
 		// NOTE: Relying on the fact that the TObjectPtr will cache the resolved pointer in place after calling Get.
 		Item.Get();
+		check(Item.IsResolved());
 	}
 #endif
 
@@ -969,6 +972,11 @@ ToRawPtrTArrayUnsafe(ArrayType&& Array)
 	using NewArrayType             = TArray<QualifiedRawPointerType, ArrayAllocatorType>;
 	using RefQualifiedNewArrayType = typename TCopyQualifiersAndRefsFromTo<ArrayType, NewArrayType>::Type;
 
+#if UE_WITH_OBJECT_HANDLE_LATE_RESOLVE || UE_WITH_OBJECT_HANDLE_TRACKING
+	using TypeCompat = TContainerElementTypeCompatibility<ArrayElementType>;
+	TypeCompat::ReinterpretRange(Array.begin(), Array.end());
+#endif
+
 	return (RefQualifiedNewArrayType&)Array;
 }
 PRAGMA_ENABLE_DEPRECATION_WARNINGS
@@ -987,6 +995,7 @@ struct TContainerElementTypeCompatibility<TObjectPtr<T>>
 		while (Iter != IterEnd)
 		{
 			Operator(Iter).Get();
+			check(Operator(Iter).IsResolved());
 			++Iter;
 		}
 #endif
@@ -1003,6 +1012,7 @@ struct TContainerElementTypeCompatibility<TObjectPtr<T>>
 			auto& Ptr = Operator(Iter);
 			const FObjectPtr& ObjPtr = reinterpret_cast<const FObjectPtr&>(Ptr);
 			UE::CoreUObject::Private::ResolveObjectHandleNoRead(ObjPtr.GetHandleRef());
+			check(ObjPtr.IsResolved());
 			++Iter;
 		}
 		const UObject* const* ObjPtr = reinterpret_cast<const UObject* const*>(Begin);
@@ -1028,6 +1038,7 @@ struct TContainerElementTypeCompatibility<const TObjectPtr<T>>
 		while (Iter != IterEnd)
 		{
 			Operator(Iter).Get();
+			check(Operator(Iter).IsResolved());
 			++Iter;
 		}
 #endif
@@ -1044,6 +1055,7 @@ struct TContainerElementTypeCompatibility<const TObjectPtr<T>>
 			auto& Ptr = Operator(Iter);
 			const FObjectPtr& ObjPtr = reinterpret_cast<const FObjectPtr&>(Ptr);
 			UE::CoreUObject::Private::ResolveObjectHandleNoRead(ObjPtr.GetHandleRef());
+			check(ObjPtr.IsResolved());
 			++Iter;
 		}
 		const UObject* const* ObjPtr = reinterpret_cast<const UObject* const*>(Begin);
diff --git a/Engine/Source/Runtime/CoreUObject/Tests/ObjectHandleTest.cpp b/Engine/Source/Runtime/CoreUObject/Tests/ObjectHandleTest.cpp
--- a/Engine/Source/Runtime/CoreUObject/Tests/ObjectHandleTest.cpp
+++ b/Engine/Source/Runtime/CoreUObject/Tests/ObjectHandleTest.cpp
@@ -390,98 +390,108 @@ TEST_CASE_METHOD(FObjectHandleTestBase, "CoreUObject::FObjectHandle::Type Safety
 	UObjectPtrTestClass* TestSafeObject = NewObject<UObjectPtrTestClass>(TestPackage, TEXT("TestSafeObject"), RF_Transient);
 	UObject* TestUnsafeObject = NewObject<UObject>(TestPackage, TestUnsafeClass, TEXT("TestUnsafeObject"), RF_Transient);
 
+	// invalid address value for testing
+	UObject* TestInvalidAddress = (UObject*)0xFFFF'FFFF'FFFF'FFFCull;
+
 	// construct object handles for testing
 	FObjectHandle NullObjectHandle = UE::CoreUObject::Private::MakeObjectHandle(nullptr);
 	FObjectHandle TestSafeObjectHandle = UE::CoreUObject::Private::MakeObjectHandle(TestSafeObject);
-	FObjectHandle TestUnsafeObjectHandle = UE::CoreUObject::Private::MakeObjectHandle(TestUnsafeObject);
+	FObjectHandle TestSafeInvalidAddressHandle = UE::CoreUObject::Private::MakeObjectHandle(TestInvalidAddress);
 #if UE_WITH_OBJECT_HANDLE_LATE_RESOLVE
 	FObjectHandle TestLateResolveSafeObjectHandle = { UE::CoreUObject::Private::MakePackedObjectRef(TestSafeObject).EncodedRef };
 	FObjectHandle TestLateResolveUnsafeObjectHandle = { UE::CoreUObject::Private::MakePackedObjectRef(TestUnsafeObject).EncodedRef };
+	FObjectHandle TestLateResolveSafeInvalidAddressHandle = { (UE::CoreUObject::Private::FPackedObjectRef { reinterpret_cast<UPTRINT>(TestInvalidAddress) | 1 }).EncodedRef };
+	FObjectHandle TestLateResolveUnsafeInvalidAddressHandle = { (UE::CoreUObject::Private::FPackedObjectRef { reinterpret_cast<UPTRINT>(TestInvalidAddress) | (1 << UE::CoreUObject::Private::TypeIdShift) | 1 }).EncodedRef };
 #endif
 
 	// NULL/type-safe object handles should report as being safe
 	CHECK(IsObjectHandleTypeSafe(NullObjectHandle));
 	CHECK(IsObjectHandleTypeSafe(TestSafeObjectHandle));
+	CHECK(IsObjectHandleTypeSafe(TestSafeInvalidAddressHandle));
 #if UE_WITH_OBJECT_HANDLE_LATE_RESOLVE
 	CHECK(IsObjectHandleTypeSafe(TestLateResolveSafeObjectHandle));
+	CHECK(IsObjectHandleTypeSafe(TestLateResolveSafeInvalidAddressHandle));
 #endif
 
-	// unsafe type object handles should report as being unsafe, but only when late resolve is enabled (required for support)
+	// unsafe type object handles should report as being unsafe
 #if UE_WITH_OBJECT_HANDLE_LATE_RESOLVE
-	CHECK_FALSE(IsObjectHandleTypeSafe(TestUnsafeObjectHandle));
 	CHECK_FALSE(IsObjectHandleTypeSafe(TestLateResolveUnsafeObjectHandle));
-#else
-	CHECK(IsObjectHandleTypeSafe(TestUnsafeObjectHandle));
+	CHECK_FALSE(IsObjectHandleTypeSafe(TestLateResolveUnsafeInvalidAddressHandle));
 #endif
 
 	// unsafe type object handles should resolve the class to the unsafe type
-	CHECK(UE::CoreUObject::Private::ResolveObjectHandleClass(TestUnsafeObjectHandle) == TestUnsafeClass);
 #if UE_WITH_OBJECT_HANDLE_LATE_RESOLVE
 	CHECK(UE::CoreUObject::Private::ResolveObjectHandleClass(TestLateResolveUnsafeObjectHandle) == TestUnsafeClass);
 #endif
 
 	// an unsafe type object handle should not equate to other unsafe type object handles (including NULL), except for itself
-	CHECK(NullObjectHandle != TestUnsafeObjectHandle);			// note: this behavior intentionally differs from object *pointers* (see below)
-	CHECK(TestUnsafeObjectHandle != NullObjectHandle);			// see note directly above
-	CHECK(TestSafeObjectHandle != TestUnsafeObjectHandle);
-	CHECK(TestUnsafeObjectHandle != TestSafeObjectHandle);
-	CHECK(TestUnsafeObjectHandle == TestUnsafeObjectHandle);
 #if UE_WITH_OBJECT_HANDLE_LATE_RESOLVE
 	CHECK(NullObjectHandle != TestLateResolveUnsafeObjectHandle);
 	CHECK(TestLateResolveUnsafeObjectHandle != NullObjectHandle);
 	CHECK(TestSafeObjectHandle != TestLateResolveUnsafeObjectHandle);
 	CHECK(TestLateResolveUnsafeObjectHandle != TestSafeObjectHandle);
-	CHECK(TestUnsafeObjectHandle == TestLateResolveUnsafeObjectHandle);
-	CHECK(TestLateResolveUnsafeObjectHandle == TestUnsafeObjectHandle);
 	CHECK(TestLateResolveSafeObjectHandle != TestLateResolveUnsafeObjectHandle);
 	CHECK(TestLateResolveUnsafeObjectHandle != TestLateResolveSafeObjectHandle);
 	CHECK(TestLateResolveUnsafeObjectHandle == TestLateResolveUnsafeObjectHandle);
+//	CHECK(TestSafeInvalidAddressHandle != TestLateResolveUnsafeInvalidAddressHandle);			// note: commented out for now; FObjectHandle::operator==() will call FindExistingPackedObjectRef() when comparing
+//	CHECK(TestLateResolveUnsafeInvalidAddressHandle != TestSafeInvalidAddressHandle);			// resolved to unresolved values, which will then attempt to dereference the resolved address and crash (known issue)
+	CHECK(TestLateResolveSafeInvalidAddressHandle != TestLateResolveUnsafeInvalidAddressHandle);
+	CHECK(TestLateResolveUnsafeInvalidAddressHandle != TestLateResolveSafeInvalidAddressHandle);
+	CHECK(TestLateResolveUnsafeInvalidAddressHandle == TestLateResolveUnsafeInvalidAddressHandle);
 #endif
 
 	// the type safety and class queries above should not have resolved an object handle that's using late resolve
 #if UE_WITH_OBJECT_HANDLE_LATE_RESOLVE
 	CHECK_FALSE(IsObjectHandleResolved(TestLateResolveSafeObjectHandle));
 	CHECK_FALSE(IsObjectHandleResolved(TestLateResolveUnsafeObjectHandle));
+	CHECK_FALSE(IsObjectHandleResolved(TestLateResolveSafeInvalidAddressHandle));
+	CHECK_FALSE(IsObjectHandleResolved(TestLateResolveUnsafeInvalidAddressHandle));
 #endif
 
-	// unsafe type object handles should resolve/evaluate to the original type object
-	CHECK(UE::CoreUObject::Private::ResolveObjectHandle(TestSafeObjectHandle) == TestSafeObject);
-	CHECK(UE::CoreUObject::Private::ResolveObjectHandle(TestUnsafeObjectHandle) == TestUnsafeObject);
+	// unsafe type object handles should resolve/evaluate to the original object/address, or NULL for an invalid object w/ late resolve
+	CHECK(UE::CoreUObject::Private::ResolveObjectHandle(TestSafeInvalidAddressHandle) == TestInvalidAddress);
 #if UE_WITH_OBJECT_HANDLE_LATE_RESOLVE
 	CHECK(UE::CoreUObject::Private::ResolveObjectHandle(TestLateResolveSafeObjectHandle) == TestSafeObject);
 	CHECK(UE::CoreUObject::Private::ResolveObjectHandle(TestLateResolveUnsafeObjectHandle) == TestUnsafeObject);
+	CHECK(UE::CoreUObject::Private::ResolveObjectHandle(TestLateResolveSafeInvalidAddressHandle) == nullptr);
+	CHECK(UE::CoreUObject::Private::ResolveObjectHandle(TestLateResolveUnsafeInvalidAddressHandle) == nullptr);
 #endif
 
-	// all handles should now report as being resolved
-	CHECK(IsObjectHandleResolved(NullObjectHandle));
-	CHECK(IsObjectHandleResolved(TestSafeObjectHandle));
-	CHECK(IsObjectHandleResolved(TestUnsafeObjectHandle));
+	// unsafe type object handles should report as NOT being resolved (in order to preserve the bit flag on the underlying packed reference)
+	CHECK(IsObjectHandleResolved(TestSafeInvalidAddressHandle));
 #if UE_WITH_OBJECT_HANDLE_LATE_RESOLVE
 	CHECK(IsObjectHandleResolved(TestLateResolveSafeObjectHandle));
-	CHECK(IsObjectHandleResolved(TestLateResolveUnsafeObjectHandle));
+	CHECK_FALSE(IsObjectHandleResolved(TestLateResolveUnsafeObjectHandle));
+	CHECK(IsObjectHandleResolved(TestLateResolveSafeInvalidAddressHandle));
+	CHECK_FALSE(IsObjectHandleResolved(TestLateResolveUnsafeInvalidAddressHandle));
 #endif
 
 	// construct object pointers for testing intentionally different behaviors of UObject-type vs. non-UObject-type bindings
 	TObjectPtr<UObject> NullObjectPtr(nullptr);
 	TObjectPtr<UObject> TestSafeObjectPtr(TestUnsafeObject);								// type safe pointer to placeholder (bound to UObject type)
 	TObjectPtr<const UObject> TestSafeConstObjectPtr(TestUnsafeObject);						// type safe const pointer to placeholder (bound to UObject type)
-	FObjectPtr TestUnsafeObjectPtr_Untyped(TestUnsafeObject);
-	TObjectPtr<UObjectPtrTestClass> TestUnsafeObjectPtr(TestUnsafeObjectPtr_Untyped);		// unsafe pointer to placeholder object (bound to non-UObject type)
-	TObjectPtr<const UObjectPtrTestClass> TestUnsafeConstObjectPtr(TestUnsafeObjectPtr_Untyped);	// unsafe const pointer to placeholder object (bound to non-UObject type)
+	FObjectPtr TestSafeInvalidObjectPtr_Untyped(TestSafeInvalidAddressHandle);
+	TObjectPtr<UObject> TestSafeInvalidObjectPtr(TestSafeInvalidObjectPtr_Untyped);
+	TObjectPtr<const UObject> TestSafeInvalidConstObjectPtr(TestSafeInvalidObjectPtr_Untyped);
 #if UE_WITH_OBJECT_HANDLE_LATE_RESOLVE
-	// note: "safe" in this context means the pointer should be type safe because it's bound to the UObject base type, but both reference the same "unsafe" object
-	TObjectPtr<UObject> TestLateResolveSafeObjectPtr(FObjectPtr({ UE::CoreUObject::Private::MakePackedObjectRef(TestUnsafeObject).EncodedRef }));
-	TObjectPtr<const UObject> TestLateResolveSafeConstObjectPtr(FObjectPtr({ UE::CoreUObject::Private::MakePackedObjectRef(TestUnsafeObject).EncodedRef }));
-	TObjectPtr<UObjectPtrTestClass> TestLateResolveUnsafeObjectPtr(FObjectPtr({ UE::CoreUObject::Private::MakePackedObjectRef(TestUnsafeObject).EncodedRef }));
-	TObjectPtr<const UObjectPtrTestClass> TestLateResolveUnsafeConstObjectPtr(FObjectPtr({ UE::CoreUObject::Private::MakePackedObjectRef(TestUnsafeObject).EncodedRef }));
+	// note: "safe" in this context means the pointer should be type safe because it's bound to the UObject base type, but both reference the same "unsafe" object and as such do not update the handle on resolve
+	FObjectPtr TestLateResolveUnsafeObjectPtr_Untyped(TestLateResolveUnsafeObjectHandle);
+	TObjectPtr<UObject> TestLateResolveSafeObjectPtr(TestLateResolveUnsafeObjectPtr_Untyped);
+	TObjectPtr<const UObject> TestLateResolveSafeConstObjectPtr(TestLateResolveUnsafeObjectPtr_Untyped);
+	TObjectPtr<UObjectPtrTestClass> TestLateResolveUnsafeObjectPtr(TestLateResolveUnsafeObjectPtr_Untyped);
+	TObjectPtr<const UObjectPtrTestClass> TestLateResolveUnsafeConstObjectPtr(TestLateResolveUnsafeObjectPtr_Untyped);
+	FObjectPtr TestLateResolveUnsafeInvalidObjectPtr_Untyped(TestLateResolveUnsafeInvalidAddressHandle);
+	TObjectPtr<UObject> TestLateResolveSafeInvalidObjectPtr(TestLateResolveUnsafeInvalidObjectPtr_Untyped);
+	TObjectPtr<const UObject> TestLateResolveSafeInvalidConstObjectPtr(TestLateResolveUnsafeInvalidObjectPtr_Untyped);
+	TObjectPtr<UObjectPtrTestClass> TestLateResolveUnsafeInvalidObjectPtr(TestLateResolveUnsafeInvalidObjectPtr_Untyped);
+	TObjectPtr<const UObjectPtrTestClass> TestLateResolveUnsafeInvalidConstObjectPtr(TestLateResolveUnsafeInvalidObjectPtr_Untyped);
+	TArray<TObjectPtr<UObject>> TestLateResolveSafeObjectPtrArray = { TestLateResolveSafeObjectPtr, TestLateResolveSafeInvalidObjectPtr };
+	TArray<TObjectPtr<UObjectPtrTestClass>> TestLateResolveUnsafeObjectPtrArray = { TestLateResolveUnsafeObjectPtr, TestLateResolveUnsafeInvalidObjectPtr };
+	TArray<TObjectPtr<const UObject>> TestLateResolveSafeConstObjectPtrArray = { TestLateResolveSafeConstObjectPtr, TestLateResolveSafeInvalidConstObjectPtr };
+	TArray<TObjectPtr<const UObjectPtrTestClass>> TestLateResolveUnsafeConstObjectPtrArray = { TestLateResolveUnsafeConstObjectPtr, TestLateResolveUnsafeInvalidConstObjectPtr };
 #endif
 
-	// an unsafe packed object reference should now exist with the unsafe type bit set for an unsafe object reference
-#if UE_WITH_OBJECT_HANDLE_LATE_RESOLVE
-	CHECK((UE::CoreUObject::Private::FindExistingPackedObjectRef(TestUnsafeObject).EncodedRef & (1 << UE::CoreUObject::Private::TypeIdShift)) != 0);
-#endif
-
-	// type safe object pointers should evaluate to true/non-NULL
+	// type safe object pointers should evaluate to true/non-NULL (including invalid pointers)
 	CHECK(TestSafeObjectPtr);
 	CHECK(!!TestSafeObjectPtr);
 	CHECK(NULL != TestSafeObjectPtr);
@@ -494,6 +504,18 @@ TEST_CASE_METHOD(FObjectHandleTestBase, "CoreUObject::FObjectHandle::Type Safety
 	CHECK(TestSafeConstObjectPtr != NULL);
 	CHECK(nullptr != TestSafeConstObjectPtr);
 	CHECK(TestSafeConstObjectPtr != nullptr);
+	CHECK(TestSafeInvalidObjectPtr);
+	CHECK(!!TestSafeInvalidObjectPtr);
+	CHECK(NULL != TestSafeInvalidObjectPtr);
+	CHECK(TestSafeInvalidObjectPtr != NULL);
+	CHECK(nullptr != TestSafeInvalidObjectPtr);
+	CHECK(TestSafeInvalidObjectPtr != nullptr);
+	CHECK(TestSafeInvalidConstObjectPtr);
+	CHECK(!!TestSafeInvalidConstObjectPtr);
+	CHECK(NULL != TestSafeInvalidConstObjectPtr);
+	CHECK(TestSafeInvalidConstObjectPtr != NULL);
+	CHECK(nullptr != TestSafeInvalidConstObjectPtr);
+	CHECK(TestSafeInvalidConstObjectPtr != nullptr);
 #if UE_WITH_OBJECT_HANDLE_LATE_RESOLVE
 	CHECK(TestLateResolveSafeObjectPtr);
 	CHECK(!!TestLateResolveSafeObjectPtr);
@@ -507,22 +529,22 @@ TEST_CASE_METHOD(FObjectHandleTestBase, "CoreUObject::FObjectHandle::Type Safety
 	CHECK(TestLateResolveSafeConstObjectPtr != NULL);
 	CHECK(nullptr != TestLateResolveSafeConstObjectPtr);
 	CHECK(TestLateResolveSafeConstObjectPtr != nullptr);
+	CHECK(TestLateResolveSafeInvalidObjectPtr);
+	CHECK(!!TestLateResolveSafeInvalidObjectPtr);
+	CHECK(NULL != TestLateResolveSafeInvalidObjectPtr);
+	CHECK(TestLateResolveSafeInvalidObjectPtr != NULL);
+	CHECK(nullptr != TestLateResolveSafeInvalidObjectPtr);
+	CHECK(TestLateResolveSafeInvalidObjectPtr != nullptr);
+	CHECK(TestLateResolveSafeInvalidConstObjectPtr);
+	CHECK(!!TestLateResolveSafeInvalidConstObjectPtr);
+	CHECK(NULL != TestLateResolveSafeInvalidConstObjectPtr);
+	CHECK(TestLateResolveSafeInvalidConstObjectPtr != NULL);
+	CHECK(nullptr != TestLateResolveSafeInvalidConstObjectPtr);
+	CHECK(TestLateResolveSafeInvalidConstObjectPtr != nullptr);
 #endif
 
-	// unsafe type object pointers should evaluate to NULL/false (for type safety), but only when late resolve is enabled (required for support)
+	// unsafe type object pointers should evaluate to NULL/false (for type safety)
 #if UE_WITH_OBJECT_HANDLE_LATE_RESOLVE
-	CHECK_FALSE(TestUnsafeObjectPtr);
-	CHECK_FALSE(!!TestUnsafeObjectPtr);
-	CHECK(NULL == TestUnsafeObjectPtr);
-	CHECK(TestUnsafeObjectPtr == NULL);
-	CHECK(nullptr == TestUnsafeObjectPtr);
-	CHECK(TestUnsafeObjectPtr == nullptr);
-	CHECK_FALSE(TestUnsafeConstObjectPtr);
-	CHECK_FALSE(!!TestUnsafeConstObjectPtr);
-	CHECK(NULL == TestUnsafeConstObjectPtr);
-	CHECK(TestUnsafeConstObjectPtr == NULL);
-	CHECK(nullptr == TestUnsafeConstObjectPtr);
-	CHECK(TestUnsafeConstObjectPtr == nullptr);
 	CHECK_FALSE(TestLateResolveUnsafeObjectPtr);
 	CHECK_FALSE(!!TestLateResolveUnsafeObjectPtr);
 	CHECK(NULL == TestLateResolveUnsafeObjectPtr);
@@ -535,31 +557,24 @@ TEST_CASE_METHOD(FObjectHandleTestBase, "CoreUObject::FObjectHandle::Type Safety
 	CHECK(TestLateResolveUnsafeConstObjectPtr == NULL);
 	CHECK(nullptr == TestLateResolveUnsafeConstObjectPtr);
 	CHECK(TestLateResolveUnsafeConstObjectPtr == nullptr);
-#else
-	CHECK(TestUnsafeObjectPtr);
-	CHECK(!!TestUnsafeObjectPtr);
-	CHECK(NULL != TestUnsafeObjectPtr);
-	CHECK(TestUnsafeObjectPtr != NULL);
-	CHECK(nullptr != TestUnsafeObjectPtr);
-	CHECK(TestUnsafeObjectPtr != nullptr);
-	CHECK(TestUnsafeConstObjectPtr);
-	CHECK(!!TestUnsafeConstObjectPtr);
-	CHECK(NULL != TestUnsafeConstObjectPtr);
-	CHECK(TestUnsafeConstObjectPtr != NULL);
-	CHECK(nullptr != TestUnsafeConstObjectPtr);
-	CHECK(TestUnsafeConstObjectPtr != nullptr);
+	CHECK_FALSE(TestLateResolveUnsafeInvalidObjectPtr);
+	CHECK_FALSE(!!TestLateResolveUnsafeInvalidObjectPtr);
+	CHECK(NULL == TestLateResolveUnsafeInvalidObjectPtr);
+	CHECK(TestLateResolveUnsafeInvalidObjectPtr == NULL);
+	CHECK(nullptr == TestLateResolveUnsafeInvalidObjectPtr);
+	CHECK(TestLateResolveUnsafeInvalidObjectPtr == nullptr);
+	CHECK_FALSE(TestLateResolveUnsafeInvalidConstObjectPtr);
+	CHECK_FALSE(!!TestLateResolveUnsafeInvalidConstObjectPtr);
+	CHECK(NULL == TestLateResolveUnsafeInvalidConstObjectPtr);
+	CHECK(TestLateResolveUnsafeInvalidConstObjectPtr == NULL);
+	CHECK(nullptr == TestLateResolveUnsafeInvalidConstObjectPtr);
+	CHECK(TestLateResolveUnsafeInvalidConstObjectPtr == nullptr);
 #endif
 
 	// an unsafe type object pointer should not equate to other unsafe type object pointers, excluding NULL and itself
 #if UE_WITH_OBJECT_HANDLE_LATE_RESOLVE
-	CHECK(NullObjectPtr == TestUnsafeObjectPtr);				// note: this intentionally differs from object *handles* (see above), and depends on late resolve
-	CHECK(TestUnsafeObjectPtr == NullObjectPtr);				// see note directly above
-	CHECK(NullObjectPtr == TestUnsafeConstObjectPtr);
-	CHECK(TestUnsafeConstObjectPtr == NullObjectPtr);
-	CHECK(NullObjectPtr != TestLateResolveSafeObjectPtr);
-	CHECK(TestLateResolveSafeObjectPtr != NullObjectPtr);
-	CHECK(NullObjectPtr == TestLateResolveUnsafeObjectPtr);
-	CHECK(TestLateResolveUnsafeObjectPtr == NullObjectPtr);
+	CHECK(NullObjectPtr == TestLateResolveUnsafeObjectPtr);		// note: this intentionally differs from object *handles* (see above)
+	CHECK(TestLateResolveUnsafeObjectPtr == NullObjectPtr);		// see note directly above
 	CHECK(TestSafeObjectPtr != TestLateResolveUnsafeObjectPtr);
 	CHECK(TestLateResolveUnsafeObjectPtr != TestSafeObjectPtr);
 	CHECK(NullObjectPtr != TestLateResolveSafeConstObjectPtr);
@@ -568,31 +583,29 @@ TEST_CASE_METHOD(FObjectHandleTestBase, "CoreUObject::FObjectHandle::Type Safety
 	CHECK(TestLateResolveUnsafeConstObjectPtr == NullObjectPtr);
 	CHECK(TestSafeConstObjectPtr != TestLateResolveUnsafeConstObjectPtr);
 	CHECK(TestLateResolveUnsafeConstObjectPtr != TestSafeConstObjectPtr);
-	CHECK(TestUnsafeObjectPtr == TestLateResolveUnsafeObjectPtr);
-	CHECK(TestLateResolveUnsafeObjectPtr == TestUnsafeObjectPtr);
-	CHECK(TestUnsafeConstObjectPtr == TestLateResolveUnsafeConstObjectPtr);
-	CHECK(TestLateResolveUnsafeConstObjectPtr == TestUnsafeConstObjectPtr);
 	CHECK(TestLateResolveSafeObjectPtr != TestLateResolveUnsafeObjectPtr);
 	CHECK(TestLateResolveUnsafeObjectPtr != TestLateResolveSafeObjectPtr);
 	CHECK(TestLateResolveSafeConstObjectPtr != TestLateResolveUnsafeConstObjectPtr);
 	CHECK(TestLateResolveUnsafeConstObjectPtr != TestLateResolveSafeConstObjectPtr);
 	CHECK(TestLateResolveUnsafeObjectPtr == TestLateResolveUnsafeObjectPtr);
 	CHECK(TestLateResolveUnsafeConstObjectPtr == TestLateResolveUnsafeConstObjectPtr);
-#else
-	CHECK(NullObjectPtr != TestUnsafeObjectPtr);				// note: if late resolve is disabled, these will no longer be equivalent
-	CHECK(TestUnsafeObjectPtr != NullObjectPtr);				// see note directly above
-	CHECK(NullObjectPtr != TestUnsafeConstObjectPtr);
-	CHECK(TestUnsafeConstObjectPtr != NullObjectPtr);
+	CHECK(NullObjectPtr == TestLateResolveUnsafeInvalidObjectPtr);
+	CHECK(TestLateResolveUnsafeInvalidObjectPtr == NullObjectPtr);
+	CHECK(TestSafeInvalidObjectPtr != TestLateResolveUnsafeInvalidObjectPtr);						// note: these do not result in a handle-to-handle test because the underlying
+	CHECK(TestLateResolveUnsafeInvalidObjectPtr != TestSafeInvalidObjectPtr);						// handle will resolve to NULL for the late resolve case with an invalid address
+	CHECK(TestLateResolveSafeInvalidObjectPtr != TestLateResolveUnsafeInvalidObjectPtr);
+	CHECK(TestLateResolveUnsafeInvalidObjectPtr != TestLateResolveSafeInvalidObjectPtr);
+	CHECK(TestLateResolveUnsafeInvalidObjectPtr == TestLateResolveUnsafeInvalidObjectPtr);
+	CHECK(NullObjectPtr == TestLateResolveUnsafeInvalidConstObjectPtr);
+	CHECK(TestLateResolveUnsafeInvalidConstObjectPtr == NullObjectPtr);
+	CHECK(TestSafeInvalidConstObjectPtr != TestLateResolveUnsafeInvalidConstObjectPtr);
+	CHECK(TestLateResolveUnsafeInvalidConstObjectPtr != TestSafeInvalidConstObjectPtr);
+	CHECK(TestLateResolveSafeInvalidConstObjectPtr != TestLateResolveUnsafeInvalidConstObjectPtr);
+	CHECK(TestLateResolveUnsafeInvalidConstObjectPtr != TestLateResolveSafeInvalidConstObjectPtr);
+	CHECK(TestLateResolveUnsafeInvalidConstObjectPtr == TestLateResolveUnsafeInvalidConstObjectPtr);
 #endif
 
-	// an unsafe type object should evaluate the object's attributes correctly
-	CHECK(TestUnsafeObjectPtr.GetName() == TestUnsafeObject->GetName());
-	CHECK(TestUnsafeObjectPtr.GetFName() == TestUnsafeObject->GetFName());
-	CHECK(TestUnsafeObjectPtr.GetPathName() == TestUnsafeObject->GetPathName());
-	CHECK(TestUnsafeObjectPtr.GetFullName() == TestUnsafeObject->GetFullName());
-	CHECK(TestUnsafeObjectPtr.GetOuter() == TestUnsafeObject->GetOuter());
-	CHECK(TestUnsafeObjectPtr.GetClass() == TestUnsafeObject->GetClass());
-	CHECK(TestUnsafeObjectPtr.GetPackage() == TestUnsafeObject->GetPackage());
+	// an unsafe type object should evaluate the object's attributes correctly (applicable only to valid object pointers)
 #if UE_WITH_OBJECT_HANDLE_LATE_RESOLVE
 	CHECK(TestLateResolveUnsafeObjectPtr.GetName() == TestUnsafeObject->GetName());
 	CHECK(TestLateResolveUnsafeObjectPtr.GetFName() == TestUnsafeObject->GetFName());
@@ -609,37 +622,61 @@ TEST_CASE_METHOD(FObjectHandleTestBase, "CoreUObject::FObjectHandle::Type Safety
 	CHECK_FALSE(TestLateResolveUnsafeObjectPtr.IsResolved());
 	CHECK_FALSE(TestLateResolveSafeConstObjectPtr.IsResolved());
 	CHECK_FALSE(TestLateResolveUnsafeConstObjectPtr.IsResolved());
+	CHECK_FALSE(TestLateResolveSafeInvalidObjectPtr.IsResolved());
+	CHECK_FALSE(TestLateResolveUnsafeInvalidObjectPtr.IsResolved());
+	CHECK_FALSE(TestLateResolveSafeInvalidConstObjectPtr.IsResolved());
+	CHECK_FALSE(TestLateResolveUnsafeInvalidConstObjectPtr.IsResolved());
 #endif
 
-	// a type safe object pointer should resolve to a non-NULL value when dereferenced
+	// a type safe object pointer should resolve to a non-NULL value when dereferenced, or NULL for an invalid object w/ late resolve
 	CHECK(TestSafeObjectPtr.Get() == TestUnsafeObject);
 	CHECK(TestSafeConstObjectPtr.Get() == TestUnsafeObject);
+	CHECK(TestSafeInvalidObjectPtr.Get() == TestInvalidAddress);
+	CHECK(TestSafeInvalidConstObjectPtr.Get() == TestInvalidAddress);
 #if UE_WITH_OBJECT_HANDLE_LATE_RESOLVE
 	CHECK(TestLateResolveSafeObjectPtr.Get() == TestUnsafeObject);
 	CHECK(TestLateResolveSafeConstObjectPtr.Get() == TestUnsafeObject);
+	CHECK(TestLateResolveSafeInvalidObjectPtr.Get() == nullptr);
+	CHECK(TestLateResolveSafeInvalidConstObjectPtr.Get() == nullptr);
 #endif
 
-	// an unsafe type object pointer should resolve to NULL when dereferenced (for type safety), but only when late resolve is enabled (required for support)
+	// an unsafe type object pointer should resolve to NULL when dereferenced (for type safety)
 #if UE_WITH_OBJECT_HANDLE_LATE_RESOLVE
-	CHECK(TestUnsafeObjectPtr.Get() == nullptr);
-	CHECK(TestUnsafeConstObjectPtr.Get() == nullptr);
 	CHECK(TestLateResolveUnsafeObjectPtr.Get() == nullptr);
 	CHECK(TestLateResolveUnsafeConstObjectPtr.Get() == nullptr);
-#else
-	CHECK(TestUnsafeObjectPtr.Get() == TestUnsafeObject);
-	CHECK(TestUnsafeConstObjectPtr.Get() == TestUnsafeObject);
+	CHECK(TestLateResolveUnsafeInvalidObjectPtr.Get() == nullptr);
+	CHECK(TestLateResolveUnsafeInvalidConstObjectPtr.Get() == nullptr);
 #endif
 
-	// all pointers should now be resolved
+	// unsafe type object pointers should report as NOT being resolved; all others as resolved
 	CHECK(TestSafeObjectPtr.IsResolved());
-	CHECK(TestUnsafeObjectPtr.IsResolved());
 	CHECK(TestSafeConstObjectPtr.IsResolved());
-	CHECK(TestUnsafeConstObjectPtr.IsResolved());
+	CHECK(TestSafeInvalidObjectPtr.IsResolved());
+	CHECK(TestSafeInvalidConstObjectPtr.IsResolved());
+#if UE_WITH_OBJECT_HANDLE_LATE_RESOLVE
+	CHECK_FALSE(TestLateResolveSafeObjectPtr.IsResolved());
+	CHECK_FALSE(TestLateResolveSafeConstObjectPtr.IsResolved());
+	CHECK_FALSE(TestLateResolveSafeInvalidObjectPtr.IsResolved());
+	CHECK_FALSE(TestLateResolveUnsafeObjectPtr.IsResolved());
+	CHECK_FALSE(TestLateResolveUnsafeConstObjectPtr.IsResolved());
+	CHECK_FALSE(TestLateResolveUnsafeInvalidObjectPtr.IsResolved());
+#endif
+
+	// unsafe type object pointers should not convert to raw pointer / array
 #if UE_WITH_OBJECT_HANDLE_LATE_RESOLVE
-	CHECK(TestLateResolveSafeObjectPtr.IsResolved());
-	CHECK(TestLateResolveUnsafeObjectPtr.IsResolved());
-	CHECK(TestLateResolveSafeConstObjectPtr.IsResolved());
-	CHECK(TestLateResolveUnsafeConstObjectPtr.IsResolved());
+	CHECK(ToRawPtr(TestLateResolveSafeObjectPtr) == TestUnsafeObject);
+	CHECK(ToRawPtr(TestLateResolveSafeConstObjectPtr) == TestUnsafeObject);
+	CHECK(ToRawPtr(TestLateResolveUnsafeObjectPtr) == nullptr);
+	CHECK(ToRawPtr(TestLateResolveUnsafeConstObjectPtr) == nullptr);
+	CHECK(ToRawPtr(TestLateResolveSafeInvalidObjectPtr) == nullptr);
+	CHECK(ToRawPtr(TestLateResolveUnsafeInvalidObjectPtr) == nullptr);
+	{
+		SKIP(); // comment this out to run the tests below (they should assert as they will not be resolved)
+		ToRawPtrTArrayUnsafe(TestLateResolveSafeObjectPtrArray);
+		ToRawPtrTArrayUnsafe(TestLateResolveSafeConstObjectPtrArray);
+		ToRawPtrTArrayUnsafe(TestLateResolveUnsafeObjectPtrArray);
+		ToRawPtrTArrayUnsafe(TestLateResolveUnsafeConstObjectPtrArray);
+	}
 #endif
 }
 #endif
