A few IDO-related changes

- Renamed CreateUnknownPropertyTree to FindOrCreateUnknownPropertyTree.
- Fixed an edge case where the property can be set wrong on the property tag.
- Fixed an edge case where the wrong property can be recorded as initialized.
- Added a test for unknown property tracking.

diff --git a/Engine/Source/Runtime/CoreUObject/Internal/UObject/PropertyBagRepository.h b/Engine/Source/Runtime/CoreUObject/Internal/UObject/PropertyBagRepository.h
--- a/Engine/Source/Runtime/CoreUObject/Internal/UObject/PropertyBagRepository.h
+++ b/Engine/Source/Runtime/CoreUObject/Internal/UObject/PropertyBagRepository.h
@@ -64,7 +64,7 @@ public:
 	/**
 	 * Finds or creates a property path name tree to collect unknown property paths within the owner.
 	 */
-	FPropertyPathNameTree* CreateUnknownPropertyTree(const UObject* Owner);
+	FPropertyPathNameTree* FindOrCreateUnknownPropertyTree(const UObject* Owner);
 
 	// Future version for reworked InstanceDataObjects - track InstanceDataObject rather than bag (directly):
 	/**
diff --git a/Engine/Source/Runtime/CoreUObject/Private/UObject/Class.cpp b/Engine/Source/Runtime/CoreUObject/Private/UObject/Class.cpp
--- a/Engine/Source/Runtime/CoreUObject/Private/UObject/Class.cpp
+++ b/Engine/Source/Runtime/CoreUObject/Private/UObject/Class.cpp
@@ -1662,10 +1662,13 @@ void UStruct::SerializeVersionedTaggedProperties(FStructuredArchive::FSlot Slot,
 				}
 
 			#if WITH_EDITORONLY_DATA
-				if (SerializeContext->bTrackInitializedProperties && Property)
+				ON_SCOPE_EXIT
 				{
-					UE::SetPropertyValueInitialized(this, Data, Property, Tag.ArrayIndex);
-				}
+					if (SerializeContext->bTrackInitializedProperties && Property)
+					{
+						UE::SetPropertyValueInitialized(this, Data, Property, Tag.ArrayIndex);
+					}
+				};
 			#endif
 
 				if (Tag.SerializeType == EPropertyTagSerializeType::Skipped)
@@ -1673,8 +1676,6 @@ void UStruct::SerializeVersionedTaggedProperties(FStructuredArchive::FSlot Slot,
 					continue;
 				}
 
-				Tag.SetProperty(Property);
-
 				if (bUseRedirects)
 				{
 					if (UE::FPropertyTypeName NewTypeName = ApplyRedirectsToPropertyType(Tag.GetType(), Property); !NewTypeName.IsEmpty())
@@ -1707,6 +1708,8 @@ void UStruct::SerializeVersionedTaggedProperties(FStructuredArchive::FSlot Slot,
 				}
 			#endif
 
+				Tag.SetProperty(Property);
+
 				bool bTryStoreUnknownPropertyPath = false;
 
 				if (Property)
@@ -1814,7 +1817,7 @@ void UStruct::SerializeVersionedTaggedProperties(FStructuredArchive::FSlot Slot,
 						{
 							if (UObject* Object = SerializeContext->SerializedObject)
 							{
-								UnknownPropertyTree = FPropertyBagRepository::Get().CreateUnknownPropertyTree(Object);
+								UnknownPropertyTree = FPropertyBagRepository::Get().FindOrCreateUnknownPropertyTree(Object);
 							}
 						}
 					}
@@ -1856,9 +1859,6 @@ void UStruct::SerializeVersionedTaggedProperties(FStructuredArchive::FSlot Slot,
 						}
 					}
 				}
-
-				// Broadcast that a property was serialized if tracking the serialized property path.
-				SerializedPropertyPath.Reset();
 			#endif // WITH_EDITORONLY_DATA
 
 				int64 Loaded = UnderlyingArchive.Tell() - StartOfProperty;
diff --git a/Engine/Source/Runtime/CoreUObject/Private/UObject/PropertyBagRepository.cpp b/Engine/Source/Runtime/CoreUObject/Private/UObject/PropertyBagRepository.cpp
--- a/Engine/Source/Runtime/CoreUObject/Private/UObject/PropertyBagRepository.cpp
+++ b/Engine/Source/Runtime/CoreUObject/Private/UObject/PropertyBagRepository.cpp
@@ -506,7 +506,7 @@ void FPropertyBagRepository::PostEditChangeChainProperty(const UObject* Object,
 }
 
 // TODO: Create these by class on construction?
-FPropertyPathNameTree* FPropertyBagRepository::CreateUnknownPropertyTree(const UObject* Owner)
+FPropertyPathNameTree* FPropertyBagRepository::FindOrCreateUnknownPropertyTree(const UObject* Owner)
 {
 	FPropertyBagRepositoryLock LockRepo(this);
 	FPropertyBagAssociationData* BagData = AssociatedData.Find(Owner);
diff --git a/Engine/Source/Runtime/CoreUObject/Tests/UObject/InstanceDataObjectUtilsTest.cpp b/Engine/Source/Runtime/CoreUObject/Tests/UObject/InstanceDataObjectUtilsTest.cpp
--- a/Engine/Source/Runtime/CoreUObject/Tests/UObject/InstanceDataObjectUtilsTest.cpp
+++ b/Engine/Source/Runtime/CoreUObject/Tests/UObject/InstanceDataObjectUtilsTest.cpp
@@ -4,6 +4,7 @@
 
 #if WITH_TESTS && WITH_EDITORONLY_DATA
 
+#include "Logging/LogScopedVerbosityOverride.h"
 #include "Misc/ScopeExit.h"
 #include "Serialization/Formatters/BinaryArchiveFormatter.h"
 #include "Serialization/Formatters/JsonArchiveInputFormatter.h"
@@ -14,6 +15,9 @@
 #include "Tests/TestHarnessAdapter.h"
 #include "UObject/InstanceDataObjectUtils.h"
 #include "UObject/Package.h"
+#include "UObject/PropertyBagRepository.h"
+#include "UObject/PropertyHelper.h"
+#include "UObject/PropertyPathNameTree.h"
 #include "UObject/UObjectThreadContext.h"
 
 namespace UE
@@ -40,19 +44,19 @@ TEST_CASE_NAMED(FInstanceDataObjectUtilsTest, "CoreUObject::Serialization::Insta
 	REQUIRE(DProperty);
 
 	FName TestObjectName = MakeUniqueObjectName(nullptr, TestClass, FName(WriteToString<128>(TestClass->GetFName(), TEXT("_Instance"))));
-	UObject* TestObject = NewObject<UObject>(GetTransientPackage(), TestClass, TestObjectName);
-	void* StructData = StructProperty->ContainerPtrToValuePtr<void>(TestObject);
+	UObject* Owner = NewObject<UObject>(GetTransientPackage(), TestClass, TestObjectName);
+	void* StructData = StructProperty->ContainerPtrToValuePtr<void>(Owner);
 
-	CHECK_FALSE(WasPropertySetBySerialization(TestClass, TestObject, StructProperty));
-	MarkPropertySetBySerialization(TestClass, TestObject, StructProperty);
-	CHECK(WasPropertySetBySerialization(TestClass, TestObject, StructProperty));
-	CHECK_FALSE(WasPropertySetBySerialization(TestClass, TestObject, Int32Property));
+	CHECK_FALSE(WasPropertySetBySerialization(TestClass, Owner, StructProperty));
+	MarkPropertySetBySerialization(TestClass, Owner, StructProperty);
+	CHECK(WasPropertySetBySerialization(TestClass, Owner, StructProperty));
+	CHECK_FALSE(WasPropertySetBySerialization(TestClass, Owner, Int32Property));
 
-	CHECK_FALSE(IsPropertyValueInitialized(TestClass, TestObject, StructProperty));
-	SetPropertyValueInitialized(TestClass, TestObject, StructProperty);
-	CHECK(IsPropertyValueInitialized(TestClass, TestObject, StructProperty));
-	ClearPropertyValueInitialized(TestClass, TestObject, StructProperty);
-	CHECK_FALSE(IsPropertyValueInitialized(TestClass, TestObject, StructProperty));
+	CHECK_FALSE(IsPropertyValueInitialized(TestClass, Owner, StructProperty));
+	SetPropertyValueInitialized(TestClass, Owner, StructProperty);
+	CHECK(IsPropertyValueInitialized(TestClass, Owner, StructProperty));
+	ClearPropertyValueInitialized(TestClass, Owner, StructProperty);
+	CHECK_FALSE(IsPropertyValueInitialized(TestClass, Owner, StructProperty));
 
 	CHECK_FALSE(WasPropertySetBySerialization(StructProperty->Struct, StructData, AProperty));
 	MarkPropertySetBySerialization(StructProperty->Struct, StructData, AProperty);
@@ -85,7 +89,7 @@ TEST_CASE_NAMED(FInstanceDataObjectUtilsTest, "CoreUObject::Serialization::Insta
 	CHECK_FALSE(IsPropertyValueInitialized(StructProperty->Struct, StructData, DProperty));
 }
 
-TEST_CASE_NAMED(FITrackInitializedPropertiesTest, "CoreUObject::Serialization::TrackInitializedProperties", "[Core][UObject][EngineFilter]")
+TEST_CASE_NAMED(FTrackInitializedPropertiesTest, "CoreUObject::Serialization::TrackInitializedProperties", "[Core][UObject][EngineFilter]")
 {
 	UTestInstanceDataObjectClass* BaseObject = NewObject<UTestInstanceDataObjectClass>();
 	UClass* TestClass = CreateInstanceDataObjectClass(nullptr, BaseObject->GetClass(), BaseObject->GetOuter());
@@ -197,6 +201,98 @@ TEST_CASE_NAMED(FITrackInitializedPropertiesTest, "CoreUObject::Serialization::T
 #endif // WITH_TEXT_ARCHIVE_SUPPORT
 }
 
+TEST_CASE_NAMED(FTrackUnknownPropertiesTest, "CoreUObject::Serialization::TrackUnknownProperties", "[Core][UObject][EngineFilter]")
+{
+	const auto MakePropertyTypeName = [](FName Name)
+	{
+		FPropertyTypeNameBuilder Builder;
+		Builder.AddName(Name);
+		return Builder.Build();
+	};
+
+	UObject* Owner = NewObject<UTestInstanceDataObjectClass>();
+
+	FUObjectSerializeContext* SerializeContext = FUObjectThreadContext::Get().GetSerializeContext();
+	TGuardValue<UObject*> SerializedObjectScope(SerializeContext->SerializedObject, Owner);
+	TGuardValue<bool> TrackSerializedPropertyPathScope(SerializeContext->bTrackSerializedPropertyPath, true);
+	TGuardValue<bool> TrackUnknownPropertiesScope(SerializeContext->bTrackUnknownProperties, true);
+	TGuardValue<bool> TrackImpersonatePropertiesScope(SerializeContext->bImpersonateProperties, true);
+	FSerializedPropertyPathScope SerializedObjectPath(SerializeContext, {"Struct"});
+
+	FTestInstanceDataObjectStructAlternate AltStructData;
+	AltStructData.B = 2.5f;
+	AltStructData.C = 3;
+	AltStructData.D = 4;
+	AltStructData.E = 5;
+
+	TArray<uint8> BinaryData;
+	{
+		FMemoryWriter Ar(BinaryData, /*bIsPersistent*/ true);
+		FBinaryArchiveFormatter Formatter(Ar);
+		FStructuredArchive StructuredAr(Formatter);
+		FTestInstanceDataObjectStructAlternate::StaticStruct()->SerializeTaggedProperties(StructuredAr.Open(), (uint8*)&AltStructData, nullptr, nullptr);
+	}
+
+#if WITH_TEXT_ARCHIVE_SUPPORT
+	TArray<uint8> JsonData;
+	{
+		FMemoryWriter Ar(JsonData, /*bIsPersistent*/ true);
+		FJsonArchiveOutputFormatter Formatter(Ar);
+		FStructuredArchive StructuredAr(Formatter);
+		FTestInstanceDataObjectStructAlternate::StaticStruct()->SerializeTaggedProperties(StructuredAr.Open(), (uint8*)&AltStructData, nullptr, nullptr);
+	}
+#endif // WITH_TEXT_ARCHIVE_SUPPORT
+
+	LOG_SCOPE_VERBOSITY_OVERRIDE(LogClass, ELogVerbosity::Error);
+
+	FTestInstanceDataObjectStruct StructData;
+
+	{
+		FMemoryReader Ar(BinaryData, /*bIsPersistent*/ true);
+		FBinaryArchiveFormatter Formatter(Ar);
+		FStructuredArchive StructuredAr(Formatter);
+		FTestInstanceDataObjectStruct::StaticStruct()->SerializeTaggedProperties(StructuredAr.Open(), (uint8*)&StructData, nullptr, nullptr);
+	}
+
+	CHECK(StructData.A == -1);
+	CHECK(StructData.B == 2);
+	CHECK(StructData.C == 3);
+	CHECK(StructData.D == 4);
+
+	FPropertyPathNameTree* Tree = FPropertyBagRepository::Get().FindOrCreateUnknownPropertyTree(Owner);
+	{
+		FSerializedPropertyPathScope Path(SerializeContext, {"B", MakePropertyTypeName(NAME_FloatProperty)});
+		CHECK(Tree->Find(SerializeContext->SerializedPropertyPath));
+	}
+	{
+		FSerializedPropertyPathScope Path(SerializeContext, {"C", MakePropertyTypeName(NAME_Int64Property)});
+		CHECK(Tree->Find(SerializeContext->SerializedPropertyPath));
+	}
+	{
+		FSerializedPropertyPathScope Path(SerializeContext, {"E", MakePropertyTypeName(NAME_IntProperty)});
+		CHECK(Tree->Find(SerializeContext->SerializedPropertyPath));
+	}
+	FPropertyBagRepository::Get().DestroyOuterBag(Owner);
+
+#if WITH_TEXT_ARCHIVE_SUPPORT
+	StructData = {};
+
+	{
+		FMemoryReader Ar(JsonData, /*bIsPersistent*/ true);
+		FJsonArchiveInputFormatter Formatter(Ar);
+		FStructuredArchive StructuredAr(Formatter);
+		FTestInstanceDataObjectStruct::StaticStruct()->SerializeTaggedProperties(StructuredAr.Open(), (uint8*)&StructData, nullptr, nullptr);
+	}
+
+	CHECK(StructData.A == -1);
+	CHECK(StructData.B == 2);
+	CHECK(StructData.C == 3);
+	CHECK(StructData.D == 4);
+
+	// Testing of the unknown property tree is skipped because it is not supported by the text format.
+#endif // WITH_TEXT_ARCHIVE_SUPPORT
+}
+
 } // UE
 
 #endif // WITH_TESTS && WITH_EDITORONLY_DATA
diff --git a/Engine/Source/Runtime/CoreUObject/Tests/UObject/InstanceDataObjectUtilsTest.h b/Engine/Source/Runtime/CoreUObject/Tests/UObject/InstanceDataObjectUtilsTest.h
--- a/Engine/Source/Runtime/CoreUObject/Tests/UObject/InstanceDataObjectUtilsTest.h
+++ b/Engine/Source/Runtime/CoreUObject/Tests/UObject/InstanceDataObjectUtilsTest.h
@@ -25,6 +25,25 @@ public:
 	int32 D = -1;
 };
 
+USTRUCT(meta=(OriginalType="TestInstanceDataObjectStruct"))
+struct FTestInstanceDataObjectStructAlternate
+{
+	GENERATED_BODY()
+
+public:
+	UPROPERTY()
+	float B = -1;
+
+	UPROPERTY()
+	int64 C = -1;
+
+	UPROPERTY()
+	int32 D = -1;
+
+	UPROPERTY()
+	int32 E = -1;
+};
+
 UCLASS()
 class UTestInstanceDataObjectClass : public UObject
 {
