UObject::Rename: Fix deadlock. Move MarkPackageDirty call outside of the  FScopedUObjectHashTablesLock, because it can call arbitrary code, which can lock the AssetRegistry, which causes lock inversion with the AssetDataGatherer code which locks the AssetRegistry and then calls FindObject which locks the HashTablesLock.
#rnx

diff --git a/Engine/Source/Runtime/CoreUObject/Private/UObject/Obj.cpp b/Engine/Source/Runtime/CoreUObject/Private/UObject/Obj.cpp
--- a/Engine/Source/Runtime/CoreUObject/Private/UObject/Obj.cpp
+++ b/Engine/Source/Runtime/CoreUObject/Private/UObject/Obj.cpp
@@ -305,7 +305,33 @@ bool UObject::Rename( const TCHAR* InName, UObject* NewOuter, ERenameFlags Flags
 
 	AutoRTFM::Open([&]()
 	{
-		// Make sure that for the remainder of the duration of the rename operation nothing else is going to modify the UObject hash tables.
+		if (!(Flags & REN_NonTransactional))
+		{
+			// Mark touched packages as dirty.
+			if (Flags & REN_DoNotDirty)
+			{
+				// This will only mark dirty if in a transaction,
+				// the object is transactional, and the object is
+				// not in a PlayInEditor package.
+				Modify(false);
+			}
+			else
+			{
+				// This will maintain previous behavior...
+				// Which was to directly call MarkPackageDirty
+				Modify(true);
+			}
+		}
+		if (NewOuter)
+		{
+			if (!(Flags & REN_DoNotDirty))
+			{
+				NewOuter->MarkPackageDirty();
+			}
+		}
+
+		// Ensure that between StaticFindObjectFast and completion of LowLevelRename nothing else modifies
+		// the UObject hash tables.
 		FScopedUObjectHashTablesLock HashTablesLock;
 
 		if (InName == nullptr)
@@ -328,24 +354,6 @@ bool UObject::Rename( const TCHAR* InName, UObject* NewOuter, ERenameFlags Flags
 
 		//UE_LOG(LogObj, Log,  TEXT("Renaming %s to %s"), *OldName.ToString(), *NewName.ToString() );
 
-		if (!(Flags & REN_NonTransactional))
-		{
-			// Mark touched packages as dirty.
-			if (Flags & REN_DoNotDirty)
-			{
-				// This will only mark dirty if in a transaction,
-				// the object is transactional, and the object is
-				// not in a PlayInEditor package.
-				Modify(false);
-			}
-			else
-			{
-				// This will maintain previous behavior...
-				// Which was to directly call MarkPackageDirty
-				Modify(true);
-			}
-		}
-
 		OldOuter = GetOuter();
 
 		if (HasAnyFlags(RF_Public))
@@ -359,13 +367,6 @@ bool UObject::Rename( const TCHAR* InName, UObject* NewOuter, ERenameFlags Flags
 			bCreateRedirector = bRootPackage == false && bUniquePathChanged == true && bRedirectionAllowed == true && bIsCaseOnlyChange == false;
 		}
 
-		if (NewOuter)
-		{
-			if (!(Flags & REN_DoNotDirty))
-			{
-				NewOuter->MarkPackageDirty();
-			}
-		}
 #if UE_WITH_OBJECT_HANDLE_LATE_RESOLVE
 		UE::CoreUObject::Private::UpdateRenamedObject(this, NewName, NewOuter);
 #endif
