InvalidateExport: When an export is invalidated, also invalidate all of its subobjects. This fixes some remaining "Content is missing from cook" errors when cooking deep blueprint hierarchies with components, when the component that is found to be unnecessary has child components.
#rnx

diff --git a/Engine/Source/Runtime/CoreUObject/Private/UObject/LinkerLoad.cpp b/Engine/Source/Runtime/CoreUObject/Private/UObject/LinkerLoad.cpp
--- a/Engine/Source/Runtime/CoreUObject/Private/UObject/LinkerLoad.cpp
+++ b/Engine/Source/Runtime/CoreUObject/Private/UObject/LinkerLoad.cpp
@@ -692,6 +692,20 @@ void FLinkerLoad::InvalidateExport(UObject* OldObject)
 		FObjectExport& ObjExport = OldObjectLinker->ExportMap[CachedLinkerIndex];
 		ObjExport.bExportLoadFailed = true;
 	}
+
+	ForEachObjectWithOuter(OldObject, [OldObjectLinker](UObject* SubObject)
+		{
+			FLinkerLoad* OldSubObjectLinker = SubObject->GetLinker();
+			if (OldSubObjectLinker == OldObjectLinker)
+			{
+				const int32 SubObjectCachedLinkerIndex = SubObject->GetLinkerIndex();
+				if (OldObjectLinker->ExportMap.IsValidIndex(SubObjectCachedLinkerIndex))
+				{
+					FObjectExport& ObjExport = OldObjectLinker->ExportMap[SubObjectCachedLinkerIndex];
+					ObjExport.bExportLoadFailed = true;
+				}
+			}
+		});
 }
 
 /**
