Add FInt64Vector2 to the BaseStructure list

diff --git a/Engine/Source/Runtime/CoreUObject/Private/UObject/Class.cpp b/Engine/Source/Runtime/CoreUObject/Private/UObject/Class.cpp
--- a/Engine/Source/Runtime/CoreUObject/Private/UObject/Class.cpp
+++ b/Engine/Source/Runtime/CoreUObject/Private/UObject/Class.cpp
@@ -7413,6 +7413,12 @@ UScriptStruct* TBaseStructure<FIntVector>::Get()
 	return ScriptStruct;
 }
 
+UScriptStruct* TBaseStructure<FInt64Vector2>::Get()
+{
+	static UScriptStruct* ScriptStruct = StaticGetBaseStructureInternal(TEXT("Int64Vector2"));
+	return ScriptStruct;
+}
+
 UScriptStruct* TBaseStructure<FIntVector4>::Get()
 {
 	static UScriptStruct* ScriptStruct = StaticGetBaseStructureInternal(TEXT("IntVector4"));
diff --git a/Engine/Source/Runtime/CoreUObject/Public/UObject/Class.h b/Engine/Source/Runtime/CoreUObject/Public/UObject/Class.h
--- a/Engine/Source/Runtime/CoreUObject/Public/UObject/Class.h
+++ b/Engine/Source/Runtime/CoreUObject/Public/UObject/Class.h
@@ -4198,6 +4198,11 @@ template<> struct TBaseStructure<FIntVector>
 	static COREUOBJECT_API UScriptStruct* Get(); 
 };
 
+template<> struct TBaseStructure<FInt64Vector2>
+{
+	static COREUOBJECT_API UScriptStruct* Get();
+};
+
 template<> struct TBaseStructure<FIntVector4>
 {
 	static COREUOBJECT_API UScriptStruct* Get();
