Deprecated GetSerializeContext and SetSerializeContext on FArchive

Implementations have been effectively using the thread context and these changes make that the only available context.

diff --git a/Engine/Plugins/Runtime/GameplayAbilities/Source/GameplayAbilities/Private/AttributeSet.cpp b/Engine/Plugins/Runtime/GameplayAbilities/Source/GameplayAbilities/Private/AttributeSet.cpp
--- a/Engine/Plugins/Runtime/GameplayAbilities/Source/GameplayAbilities/Private/AttributeSet.cpp
+++ b/Engine/Plugins/Runtime/GameplayAbilities/Source/GameplayAbilities/Private/AttributeSet.cpp
@@ -261,7 +261,7 @@ void FGameplayAttribute::PostSerialize(const FArchive& Ar)
 				Attribute = FindFProperty<FProperty>(NewClass, *NewAttributeName);
 
 				// Verbose log any applied redirectors
-				FUObjectSerializeContext* LoadContext = const_cast<FArchive*>(&Ar)->GetSerializeContext();
+				FUObjectSerializeContext* LoadContext = FUObjectThreadContext::Get().GetSerializeContext();
 				const FString AssetName = (LoadContext && LoadContext->SerializedObject) ? LoadContext->SerializedObject->GetPathName() : TEXT("Unknown Object");
 				ABILITY_LOG(Verbose, TEXT("FGameplayAttribute::PostSerialize redirected an attribute '%s' -> '%s'. (Asset: %s)"), *PathName, *RedirectedPathName, *AssetName);
 			}
@@ -286,7 +286,7 @@ void FGameplayAttribute::PostSerialize(const FArchive& Ar)
 			// Log warning if attribute failed to resolve while name + owner were non-null
 			if (!Attribute.Get())
 			{
-				FUObjectSerializeContext* LoadContext = const_cast<FArchive*>(&Ar)->GetSerializeContext();
+				FUObjectSerializeContext* LoadContext = FUObjectThreadContext::Get().GetSerializeContext();
 				const FString AssetName = (LoadContext && LoadContext->SerializedObject) ? LoadContext->SerializedObject->GetPathName() : TEXT("Unknown Object");
 				const FString OwnerName = AttributeOwner ? AttributeOwner->GetName() : TEXT("NONE");
 				ABILITY_LOG(Warning, TEXT("FGameplayAttribute::PostSerialize called on an invalid attribute with owner %s and name %s. (Asset: %s)"), *OwnerName, *AttributeName, *AssetName);
diff --git a/Engine/Source/Editor/MovieSceneTools/Private/K2Node_GetSequenceBinding.cpp b/Engine/Source/Editor/MovieSceneTools/Private/K2Node_GetSequenceBinding.cpp
--- a/Engine/Source/Editor/MovieSceneTools/Private/K2Node_GetSequenceBinding.cpp
+++ b/Engine/Source/Editor/MovieSceneTools/Private/K2Node_GetSequenceBinding.cpp
@@ -97,7 +97,7 @@ public:
 UMovieSceneSequence* UK2Node_GetSequenceBinding::GetSequence() const
 {
 	FLinkerLoad* Linker = GetLinker();
-	return Cast<UMovieSceneSequence>(SourceSequence.TryLoad(Linker ? Linker->GetSerializeContext() : nullptr));
+	return Cast<UMovieSceneSequence>(SourceSequence.TryLoad(Linker ? FUObjectThreadContext::Get().GetSerializeContext() : nullptr));
 }
 
 void UK2Node_GetSequenceBinding::ValidateNodeDuringCompilation(FCompilerResultsLog& MessageLog) const
diff --git a/Engine/Source/Editor/UMGEditor/Private/WidgetBlueprintCompiler.cpp b/Engine/Source/Editor/UMGEditor/Private/WidgetBlueprintCompiler.cpp
--- a/Engine/Source/Editor/UMGEditor/Private/WidgetBlueprintCompiler.cpp
+++ b/Engine/Source/Editor/UMGEditor/Private/WidgetBlueprintCompiler.cpp
@@ -836,7 +836,7 @@ void FWidgetBlueprintCompilerContext::FinishCompilingClass(UClass* Class)
 			if (WidgetBP->bIsRegeneratingOnLoad)
 			{
 				FLinkerLoad* Linker = WidgetBP->GetLinker();
-				LinkerLoadingContext = Linker ? Linker->GetSerializeContext() : nullptr;
+				LinkerLoadingContext = Linker ? FUObjectThreadContext::Get().GetSerializeContext() : nullptr;
 				DupParams.bSkipPostLoad = true;
 				DupParams.CreatedObjects = &DupObjectsMap;
 			}
diff --git a/Engine/Source/Editor/UnrealEd/Private/Commandlets/PackageUtilities.cpp b/Engine/Source/Editor/UnrealEd/Private/Commandlets/PackageUtilities.cpp
--- a/Engine/Source/Editor/UnrealEd/Private/Commandlets/PackageUtilities.cpp
+++ b/Engine/Source/Editor/UnrealEd/Private/Commandlets/PackageUtilities.cpp
@@ -1639,7 +1639,7 @@ int32 UPkgInfoCommandlet::Main( const FString& Params )
 			TRefCountPtr<FUObjectSerializeContext> LoadContext(FUObjectThreadContext::Get().GetSerializeContext());
 			BeginLoad(LoadContext);
 			Linker = CreateLinkerForFilename(LoadContext, Filename);
-			EndLoad(Linker ? Linker->GetSerializeContext() : LoadContext.GetReference());
+			EndLoad(LoadContext);
 		}
 		else
 		{
diff --git a/Engine/Source/Editor/UnrealEd/Private/Cooker/DiffWriterArchive.cpp b/Engine/Source/Editor/UnrealEd/Private/Cooker/DiffWriterArchive.cpp
--- a/Engine/Source/Editor/UnrealEd/Private/Cooker/DiffWriterArchive.cpp
+++ b/Engine/Source/Editor/UnrealEd/Private/Cooker/DiffWriterArchive.cpp
@@ -473,7 +473,7 @@ void FCallstacks::RecordSerialize(EOffsetFrame OffsetFrame, int64 CurrentOffset,
 
 	if (Length > 0)
 	{
-		UObject* SerializedObject = Ar.GetSerializeContext() ? Ar.GetSerializeContext()->SerializedObject : nullptr;
+		UObject* SerializedObject = FUObjectThreadContext::Get().GetSerializeContext()->SerializedObject;
 		TArrayView<const FName> DebugStack = Ar.GetDebugDataStack();
 
 		const bool bCollectingCallstacks = Accumulator.bFirstSaveComplete;
@@ -1440,14 +1440,14 @@ void DumpPackageHeaderDiffs_LinkerLoad(
 		TRefCountPtr<FUObjectSerializeContext> LinkerLoadContext(FUObjectThreadContext::Get().GetSerializeContext());
 		BeginLoad(LinkerLoadContext);
 		SourceLinker = CreateLinkerForPackage(LinkerLoadContext, SourceAssetPackageName, AssetFilename, SourcePackage);
-		EndLoad(SourceLinker ? SourceLinker->GetSerializeContext() : LinkerLoadContext.GetReference());
+		EndLoad(LinkerLoadContext);
 	}
 	
 	{
 		TRefCountPtr<FUObjectSerializeContext> LinkerLoadContext(FUObjectThreadContext::Get().GetSerializeContext());
 		BeginLoad(LinkerLoadContext);
 		DestLinker = CreateLinkerForPackage(LinkerLoadContext, DestAssetPackageName, AssetFilename, DestPackage);
-		EndLoad(DestLinker ? DestLinker->GetSerializeContext() : LinkerLoadContext.GetReference());
+		EndLoad(LinkerLoadContext);
 	}
 
 	if (SourceLinker && DestLinker)
diff --git a/Engine/Source/Editor/UnrealEd/Private/Kismet2/BlueprintEditorUtils.cpp b/Engine/Source/Editor/UnrealEd/Private/Kismet2/BlueprintEditorUtils.cpp
--- a/Engine/Source/Editor/UnrealEd/Private/Kismet2/BlueprintEditorUtils.cpp
+++ b/Engine/Source/Editor/UnrealEd/Private/Kismet2/BlueprintEditorUtils.cpp
@@ -846,7 +846,7 @@ void FBlueprintEditorUtils::PatchNewCDOIntoLinker(UObject* CDO, FLinkerLoad* Lin
 			// Copy flags from the old CDO.
 			CDO->SetFlags(OldObjectFlags);
 
-			FUObjectSerializeContext* LoadContext = InLoadContext ? InLoadContext : Linker->GetSerializeContext();
+			FUObjectSerializeContext* LoadContext = InLoadContext ? InLoadContext : FUObjectThreadContext::Get().GetSerializeContext();
 
 			// Make sure the new CDO gets PostLoad called on it, so either add it to ObjLoaded list, or replace it if already present.
 			if (LoadContext && !LoadContext->PRIVATE_PatchNewObjectIntoExport(OldCDO, CDO))
diff --git a/Engine/Source/Runtime/Core/Public/Serialization/Archive.h b/Engine/Source/Runtime/Core/Public/Serialization/Archive.h
--- a/Engine/Source/Runtime/Core/Public/Serialization/Archive.h
+++ b/Engine/Source/Runtime/Core/Public/Serialization/Archive.h
@@ -709,9 +709,11 @@ public:
 #endif
 
 	/** Sets the current UObject serialization context for this archive. */
+	UE_DEPRECATED(5.5, "SetSerializeContext is not supported. Remove calls to it.")
 	virtual void SetSerializeContext(FUObjectSerializeContext* InLoadContext) {}
 
 	/** Gets the current UObject serialization context for this archive. */
+	UE_DEPRECATED(5.5, "GetSerializeContext is not supported. Use FUObjectThreadContext::Get().GetSerializeContext().")
 	virtual FUObjectSerializeContext* GetSerializeContext() { return nullptr; }
 
 #if USE_STABLE_LOCALIZATION_KEYS
diff --git a/Engine/Source/Runtime/CoreUObject/Private/Serialization/ArchiveStackTrace.cpp b/Engine/Source/Runtime/CoreUObject/Private/Serialization/ArchiveStackTrace.cpp
--- a/Engine/Source/Runtime/CoreUObject/Private/Serialization/ArchiveStackTrace.cpp
+++ b/Engine/Source/Runtime/CoreUObject/Private/Serialization/ArchiveStackTrace.cpp
@@ -567,7 +567,7 @@ void FArchiveStackTraceWriter::Serialize(void* Data, int64 Length)
 
 	if (Length > 0)
 	{
-		UObject* SerializedObject = SerializeContext ? SerializeContext->SerializedObject : nullptr;
+		UObject* SerializedObject = FUObjectThreadContext::Get().GetSerializeContext()->SerializedObject;
 		TArrayView<const FName> DebugStack;
 #if WITH_EDITOR
 		DebugStack = DebugDataStack;
@@ -592,21 +592,6 @@ void FArchiveStackTraceWriter::Serialize(void* Data, int64 Length)
 	}
 }
 
-void FArchiveStackTraceWriter::SetSerializeContext(FUObjectSerializeContext* Context)
-{
-	SerializeContext = Context;
-
-	if (bInnerArchiveDisabled == false)
-	{
-		InnerArchive.SetSerializeContext(Context);
-	}
-}
-
-FUObjectSerializeContext* FArchiveStackTraceWriter::GetSerializeContext()
-{
-	return SerializeContext;
-}
-
 FArchiveStackTraceMemoryWriter::FArchiveStackTraceMemoryWriter(
 	FArchiveCallstacks& Callstacks,
 	const FArchiveDiffMap* DiffMap,
@@ -628,16 +613,6 @@ void FArchiveStackTraceMemoryWriter::Serialize(void* Memory, int64 Length)
 	FLargeMemoryWriter::Serialize(Memory, Length);
 }
 
-void FArchiveStackTraceMemoryWriter::SetSerializeContext(FUObjectSerializeContext* Context)
-{
-	StackTraceWriter.SetSerializeContext(Context);
-}
-
-FUObjectSerializeContext* FArchiveStackTraceMemoryWriter::GetSerializeContext()
-{
-	return StackTraceWriter.GetSerializeContext();
-}
-
 FArchiveStackTrace::FArchiveStackTrace(UObject* InAsset, const TCHAR* InFilename, bool bInCollectCallstacks, const FArchiveDiffMap* InDiffMap)
 	: FLargeMemoryWriter(0, false, InFilename)
 	, Callstacks(InAsset)
@@ -658,16 +633,6 @@ void FArchiveStackTrace::Serialize(void* Memory, int64 Length)
 	FLargeMemoryWriter::Serialize(Memory, Length);
 }
 
-void FArchiveStackTrace::SetSerializeContext(FUObjectSerializeContext* Context)
-{
-	StackTraceWriter.SetSerializeContext(Context);
-}
-
-FUObjectSerializeContext* FArchiveStackTrace::GetSerializeContext()
-{
-	return StackTraceWriter.GetSerializeContext();
-}
-
 bool FArchiveStackTrace::LoadPackageIntoMemory(const TCHAR* InFilename, FPackageData& OutPackageData,
 	TUniquePtr<uint8, UE::ArchiveStackTrace::FDeleteByFree>& OutLoadedBytes)
 {
diff --git a/Engine/Source/Runtime/CoreUObject/Private/Serialization/AsyncLoading.cpp b/Engine/Source/Runtime/CoreUObject/Private/Serialization/AsyncLoading.cpp
--- a/Engine/Source/Runtime/CoreUObject/Private/Serialization/AsyncLoading.cpp
+++ b/Engine/Source/Runtime/CoreUObject/Private/Serialization/AsyncLoading.cpp
@@ -7211,7 +7211,7 @@ EAsyncPackageState::Type FAsyncPackage::FinishObjects()
 	LastTypeOfWorkPerformed			= TEXT("finishing all objects");
 
 	FUObjectSerializeContext* LoadContext = GetSerializeContext();
-	check(!Linker || LoadContext == Linker->GetSerializeContext());		
+	check(!Linker || LoadContext == FUObjectThreadContext::Get().GetSerializeContext());
 	TArray<UObject*>& ThreadObjLoaded = LoadContext->PRIVATE_GetObjectsLoadedInternalUseOnly();
 
 	EAsyncLoadingResult::Type LoadingResult;
@@ -7270,10 +7270,6 @@ EAsyncPackageState::Type FAsyncPackage::FinishObjects()
 		const bool bInternalCallbacks = true;
 		CallCompletionCallbacks(bInternalCallbacks, LoadingResult);
 	}
-	else
-	{
-		LoadContext->DetachFromLinkers();
-	}
 
 	return EAsyncPackageState::Complete;
 }
diff --git a/Engine/Source/Runtime/CoreUObject/Private/Serialization/AsyncLoading2.cpp b/Engine/Source/Runtime/CoreUObject/Private/Serialization/AsyncLoading2.cpp
--- a/Engine/Source/Runtime/CoreUObject/Private/Serialization/AsyncLoading2.cpp
+++ b/Engine/Source/Runtime/CoreUObject/Private/Serialization/AsyncLoading2.cpp
@@ -5278,7 +5278,6 @@ void FAsyncPackage2::CreateLinker(const FLinkerInstancingContext* InstancingCont
 		FUObjectSerializeContext* LoadContext = GetSerializeContext();
 #if ALT2_ENABLE_NEW_ARCHIVE_FOR_LINKERLOAD
 		Linker = new FLinkerLoad(LinkerRoot, Desc.PackagePathToLoad, LinkerFlags, InstancingContext ? *InstancingContext : FLinkerInstancingContext());
-		Linker->SetSerializeContext(LoadContext);
 		LinkerRoot->SetLinker(Linker);
 		FLinkerLoadArchive2* Loader = new FLinkerLoadArchive2(Desc.PackagePathToLoad);
 		Linker->SetLoader(Loader, Loader->NeedsEngineVersionChecks());
diff --git a/Engine/Source/Runtime/CoreUObject/Private/Serialization/DuplicateDataReader.cpp b/Engine/Source/Runtime/CoreUObject/Private/Serialization/DuplicateDataReader.cpp
--- a/Engine/Source/Runtime/CoreUObject/Private/Serialization/DuplicateDataReader.cpp
+++ b/Engine/Source/Runtime/CoreUObject/Private/Serialization/DuplicateDataReader.cpp
@@ -41,9 +41,7 @@ FDuplicateDataReader::FDuplicateDataReader( class FUObjectAnnotationSparse<FDupl
 
 void FDuplicateDataReader::SerializeFail()
 {
-	FUObjectSerializeContext* LoadContext = GetSerializeContext();
-	check(LoadContext);
-	UE_LOG(LogObj, Fatal, TEXT("FDuplicateDataReader Overread. SerializedObject = %s SerializedProperty = %s"), *GetFullNameSafe(LoadContext->SerializedObject), *GetFullNameSafe(GetSerializedProperty()));
+	UE_LOG(LogObj, Fatal, TEXT("FDuplicateDataReader Overread. SerializedObject = %s SerializedProperty = %s"), *GetFullNameSafe(FUObjectThreadContext::Get().GetSerializeContext()->SerializedObject), *GetFullNameSafe(GetSerializedProperty()));
 }
 
 FArchive& FDuplicateDataReader::operator<<(FName& N)
diff --git a/Engine/Source/Runtime/CoreUObject/Private/UObject/FieldPath.cpp b/Engine/Source/Runtime/CoreUObject/Private/UObject/FieldPath.cpp
--- a/Engine/Source/Runtime/CoreUObject/Private/UObject/FieldPath.cpp
+++ b/Engine/Source/Runtime/CoreUObject/Private/UObject/FieldPath.cpp
@@ -191,7 +191,7 @@ UStruct* FFieldPath::TryToResolveOwnerFromLinker(FLinkerLoad* InLinker) const
 {
 	UStruct* OwnerStruct = nullptr;
 	check(InLinker);
-	FUObjectSerializeContext* Context = InLinker->GetSerializeContext();
+	FUObjectSerializeContext* Context = FUObjectThreadContext::Get().GetSerializeContext();
 	if (Context && Context->SerializedObject && Context->SerializedObject->IsA<UStruct>())
 	{
 		TArray<FName> StructPath;
diff --git a/Engine/Source/Runtime/CoreUObject/Private/UObject/Linker.cpp b/Engine/Source/Runtime/CoreUObject/Private/UObject/Linker.cpp
--- a/Engine/Source/Runtime/CoreUObject/Private/UObject/Linker.cpp
+++ b/Engine/Source/Runtime/CoreUObject/Private/UObject/Linker.cpp
@@ -673,13 +673,6 @@ FLinkerLoad* GetPackageLinker
 	// See if the linker is already loaded.
 	if (FLinkerLoad* Result = FLinkerLoad::FindExistingLinkerForPackage(InOuter))
 	{
-		if (InExistingContext && Result->GetSerializeContext() && Result->GetSerializeContext() != InExistingContext)
-		{
-			if (!Result->GetSerializeContext()->HasStartedLoading())
-			{
-				Result->SetSerializeContext(InExistingContext);
-			}
-		}
 		return Result;
 	}
 
@@ -809,22 +802,17 @@ FLinkerLoad* GetPackageLinker
 		{
 			if (InExistingContext)
 			{
-				if ((Result->GetSerializeContext() && Result->GetSerializeContext()->HasStartedLoading() && InExistingContext->GetBeginLoadCount() == 1) ||
-					(IsInAsyncLoadingThread() && Result->GetSerializeContext()))
+				TRefCountPtr<FUObjectSerializeContext> CurrentContext(FUObjectThreadContext::Get().GetSerializeContext());
+				if ((CurrentContext->HasStartedLoading() && InExistingContext->GetBeginLoadCount() == 1) || IsInAsyncLoadingThread())
 				{
-					// Use the context associated with the linker because it has already started loading objects (or we're in ALT where each package needs its own context)
-					*InOutLoadContext = Result->GetSerializeContext();
+					// Use the current context because it has already started loading objects (or we're in ALT where each package needs its own context)
+					*InOutLoadContext = CurrentContext;
 				}
-				else
+				else if (CurrentContext != InExistingContext)
 				{
-					if (Result->GetSerializeContext() && Result->GetSerializeContext() != InExistingContext)
-					{
-						// Make sure the objects already loaded with the context associated with the existing linker
-						// are copied to the context provided for this function call to make sure they all get loaded ASAP
-						InExistingContext->AddUniqueLoadedObjects(Result->GetSerializeContext()->PRIVATE_GetObjectsLoadedInternalUseOnly());
-					}
-					// Replace the linker context with the one passed into this function
-					Result->SetSerializeContext(InExistingContext);
+					// Make sure the objects already loaded with the context associated with the existing linker
+					// are copied to the context provided for this function call to make sure they all get loaded ASAP
+					InExistingContext->AddUniqueLoadedObjects(CurrentContext->PRIVATE_GetObjectsLoadedInternalUseOnly());
 				}
 			}
 			return Result;
diff --git a/Engine/Source/Runtime/CoreUObject/Private/UObject/LinkerLoad.cpp b/Engine/Source/Runtime/CoreUObject/Private/UObject/LinkerLoad.cpp
--- a/Engine/Source/Runtime/CoreUObject/Private/UObject/LinkerLoad.cpp
+++ b/Engine/Source/Runtime/CoreUObject/Private/UObject/LinkerLoad.cpp
@@ -414,7 +414,7 @@ void FLinkerLoad::CreateActiveRedirectsMap(const FString& GEngineIniName)
 
 FScopedCreateImportCounter::FScopedCreateImportCounter(FLinkerLoad* Linker, int32 Index)
 {
-	LoadContext = Linker->GetSerializeContext();
+	LoadContext = FUObjectThreadContext::Get().GetSerializeContext();
 	check(LoadContext);
 
 	// Remember the old linker and index
@@ -443,7 +443,7 @@ struct FScopedCreateExportCounter
 	 */
 	FScopedCreateExportCounter(FLinkerLoad* Linker, int32 Index)
 	{
-		LoadContext = Linker->GetSerializeContext();
+		LoadContext = FUObjectThreadContext::Get().GetSerializeContext();
 		check(LoadContext);
 
 		// Remember the old linker and index
@@ -675,10 +675,7 @@ void FLinkerLoad::PRIVATE_PatchNewObjectIntoExport(UObject* OldObject, UObject*
 		ObjExport.Object = NewObject;
 
 		// If the object was in the ObjLoaded queue (exported, but not yet serialized), swap out for our new object
-		if(OldObjectLinker->GetSerializeContext())
-		{
-			OldObjectLinker->GetSerializeContext()->PRIVATE_PatchNewObjectIntoExport(OldObject, NewObject);
-		}
+		FUObjectThreadContext::Get().GetSerializeContext()->PRIVATE_PatchNewObjectIntoExport(OldObject, NewObject);
 	}
 }
 
@@ -746,7 +743,6 @@ FLinkerLoad* FLinkerLoad::CreateLinkerAsync(FUObjectSerializeContext* LoadContex
 			LoadFlags |= LOAD_Async;
 		}
 		Linker = new FLinkerLoad(Parent, PackagePath, LoadFlags, InstancingContext ? *InstancingContext : FLinkerInstancingContext());
-		Linker->SetSerializeContext(LoadContext);
 		Parent->SetLinker(Linker);
 		if (GEventDrivenLoaderEnabled && Linker)
 		{
@@ -759,10 +755,6 @@ FLinkerLoad* FLinkerLoad::CreateLinkerAsync(FUObjectSerializeContext* LoadContex
 	return Linker;
 }
 
-void FLinkerLoad::SetSerializeContext(FUObjectSerializeContext* InLoadContext)
-{
-}
-
 FUObjectSerializeContext* FLinkerLoad::GetSerializeContext()
 {
 	return FUObjectThreadContext::Get().GetSerializeContext();
@@ -1084,12 +1076,6 @@ FLinkerLoad::~FLinkerLoad()
 	// Detaches linker.
 	Detach();
 
-	// Detach the serialize context
-	if (GetSerializeContext())
-	{
-		SetSerializeContext(nullptr);
-	}
-
 	DEC_DWORD_STAT(STAT_LiveLinkerCount);
 
 #if WITH_EDITOR
@@ -3063,7 +3049,7 @@ void FLinkerLoad::GatherImportDependencies(int32 ImportIndex, TSet<FDependencyRe
 		return;
 	}
 
-	FUObjectSerializeContext* LoadContext = GetSerializeContext();
+	FUObjectSerializeContext* LoadContext = FUObjectThreadContext::Get().GetSerializeContext();
 	check(LoadContext);
 
 	BeginLoad(LoadContext, TEXT("GatherImportDependencies"));
@@ -3192,7 +3178,7 @@ FLinkerLoad::EVerifyResult FLinkerLoad::VerifyImport(int32 ImportIndex)
 	// these checks find out if the VerifyImportInner was successful or not 
 	if (Import.SourceLinker && Import.SourceIndex == INDEX_NONE && Import.XObject == NULL && !Import.OuterIndex.IsNull() && Import.ObjectName != NAME_ObjectRedirector)
 	{
-		FUObjectSerializeContext* CurrentLoadContext = GetSerializeContext();
+		FUObjectSerializeContext* CurrentLoadContext = FUObjectThreadContext::Get().GetSerializeContext();
 
 		// if we found the package, but not the object, look for a redirector
 		FObjectImport OriginalImport = Import;
@@ -3458,7 +3444,7 @@ bool FLinkerLoad::VerifyImportInner(const int32 ImportIndex, FString& WarningSuf
 
 		UPackage* Package = nullptr;
 		uint32 InternalLoadFlags = LoadFlags & (LOAD_NoVerify | LOAD_NoWarn | LOAD_Quiet | LOAD_RegenerateBulkDataGuids);
-		FUObjectSerializeContext* SerializeContext = GetSerializeContext();
+		FUObjectSerializeContext* SerializeContext = FUObjectThreadContext::Get().GetSerializeContext();
 
 		// Resolve the package name for the import, potentially remapping it, if instancing
 		FName PackageToLoad = !Import.HasPackageName() ? Import.ObjectName : Import.GetPackageName();
@@ -3738,7 +3724,7 @@ bool FLinkerLoad::VerifyImportInner(const int32 ImportIndex, FString& WarningSuf
 #if WITH_EDITOR
 				else
 				{
-					const FUObjectSerializeContext* SerializeContext = GetSerializeContext();
+					const FUObjectSerializeContext* SerializeContext = FUObjectThreadContext::Get().GetSerializeContext();
 					if (ensure(SerializeContext))
 					{
 						// If we're serializing a redirector's destination object, validate/create the outer package object if it's a missing type. If this import
@@ -3964,7 +3950,7 @@ bool FLinkerLoad::VerifyImportInner(const int32 ImportIndex, FString& WarningSuf
 			{
 				// except if we are looking for _the_ package...in which case we are looking for TmpPkg, so we are done
 				Import.XObject = TmpPkg;
-				FUObjectSerializeContext* CurrentLoadContext = GetSerializeContext();
+				FUObjectSerializeContext* CurrentLoadContext = FUObjectThreadContext::Get().GetSerializeContext();
 				check(CurrentLoadContext);
 				CurrentLoadContext->IncrementImportCount();
 				FLinkerManager::Get().AddLoaderWithNewImports(this);
@@ -4023,7 +4009,7 @@ bool FLinkerLoad::VerifyImportInner(const int32 ImportIndex, FString& WarningSuf
 				if (FindObject != nullptr && ((LoadFlags & LOAD_FindIfFail) || bIsInMemoryOnlyOrNativeTransient))
 				{
 					Import.XObject = FindObject;
-					FUObjectSerializeContext* CurrentLoadContext = GetSerializeContext();
+					FUObjectSerializeContext* CurrentLoadContext = FUObjectThreadContext::Get().GetSerializeContext();
 					check(CurrentLoadContext);
 					CurrentLoadContext->IncrementImportCount();
 					FLinkerManager::Get().AddLoaderWithNewImports(this);
@@ -4049,11 +4035,6 @@ bool FLinkerLoad::VerifyImportInner(const int32 ImportIndex, FString& WarningSuf
 		}
 	}
 
-	if (!GEventDrivenLoaderEnabled && Import.SourceLinker && !Import.SourceLinker->GetSerializeContext())
-	{
-		Import.SourceLinker->SetSerializeContext(GetSerializeContext());
-	}
-
 	return false;
 }
 
@@ -4569,7 +4550,7 @@ void FLinkerLoad::Preload( UObject* Object )
 	// Preload the object if necessary.
 	if (Object->HasAnyFlags(RF_NeedLoad))
 	{
-		FUObjectSerializeContext* CurrentLoadContext = GetSerializeContext();
+		FUObjectSerializeContext* CurrentLoadContext = FUObjectThreadContext::Get().GetSerializeContext();
 
 		if (Object->GetLinker() == this)
 		{
@@ -5084,7 +5065,7 @@ UObject* FLinkerLoad::CreateExport( int32 Index )
 	if( !Export.Object && !FilterExport(Export) ) // for some acceptable position, it was not "not for" 
 	{
 		TGuardValue<void*> GuardThreadContextAsyncPackage(FUObjectThreadContext::Get().AsyncPackage, AsyncRoot);
-		FUObjectSerializeContext* CurrentLoadContext = GetSerializeContext();
+		FUObjectSerializeContext* CurrentLoadContext = FUObjectThreadContext::Get().GetSerializeContext();
 		check(!GEventDrivenLoaderEnabled || !bLockoutLegacyOperations || !EVENT_DRIVEN_ASYNC_LOAD_ACTIVE_AT_RUNTIME);
 		check(Export.ObjectName!=NAME_None || !(Export.ObjectFlags&RF_Public));
 		check(IsLoading());
@@ -5650,7 +5631,7 @@ UObject* FLinkerLoad::CreateExport( int32 Index )
 						// fully regenerated before another (there is no changing 
 						// that); so dependencies need to be recompiled later (with
 						// all the regenerated classes in place)
-						FScopedClassDependencyGather DependencyHelper(ClassObject, GetSerializeContext());
+						FScopedClassDependencyGather DependencyHelper(ClassObject, FUObjectThreadContext::Get().GetSerializeContext());
 #endif //WITH_EDITOR
 
 						ClassObject->Bind();
@@ -5732,7 +5713,7 @@ UObject* FLinkerLoad::CreateImport( int32 Index )
 	// Imports can have no name if they were filtered out due to package redirects, skip in that case
 	if (Import.XObject == nullptr && Import.ObjectName != NAME_None)
 	{
-		FUObjectSerializeContext* CurrentLoadContext = GetSerializeContext();
+		FUObjectSerializeContext* CurrentLoadContext = FUObjectThreadContext::Get().GetSerializeContext();
 		if (!GIsEditor && !IsRunningCommandlet())
 		{
 			// Try to find existing version in memory first.
@@ -5836,10 +5817,6 @@ UObject* FLinkerLoad::CreateImport( int32 Index )
 			{
 				VerifyImportResult = VerifyImport(Index);
 			}
-			else if (!GEventDrivenLoaderEnabled && !Import.SourceLinker->GetSerializeContext())
-			{
-				Import.SourceLinker->SetSerializeContext(GetSerializeContext());
-			}
 			if(Import.SourceIndex != INDEX_NONE)
 			{
 				check(Import.SourceLinker);
@@ -6096,7 +6073,7 @@ void FLinkerLoad::Detach()
 	FLinkerManager::Get().RemoveLoaderFromObjectLoadersAndLoadersWithNewImports(this);
 	if (!FPlatformProperties::HasEditorOnlyData())
 	{
-		FUObjectSerializeContext* CurrentLoadContext = GetSerializeContext();
+		FUObjectSerializeContext* CurrentLoadContext = FUObjectThreadContext::Get().GetSerializeContext();
 		check(CurrentLoadContext);
 		CurrentLoadContext->RemoveDelayedLinkerClosePackage(this);
 	}
@@ -6437,7 +6414,7 @@ UObject* FLinkerLoad::GetArchetypeFromLoader(const UObject* Obj)
 {
 	if (GEventDrivenLoaderEnabled)
 	{
-		FUObjectSerializeContext* CurrentLoadContext = GetSerializeContext();
+		FUObjectSerializeContext* CurrentLoadContext = FUObjectThreadContext::Get().GetSerializeContext();
 		check(CurrentLoadContext);
 		check(!TemplateForGetArchetypeFromLoader || CurrentLoadContext->SerializedObject == Obj);
 		return TemplateForGetArchetypeFromLoader;
diff --git a/Engine/Source/Runtime/CoreUObject/Private/UObject/LinkerManager.cpp b/Engine/Source/Runtime/CoreUObject/Private/UObject/LinkerManager.cpp
--- a/Engine/Source/Runtime/CoreUObject/Private/UObject/LinkerManager.cpp
+++ b/Engine/Source/Runtime/CoreUObject/Private/UObject/LinkerManager.cpp
@@ -284,10 +284,7 @@ void FLinkerManager::DissociateImportsAndForcedExports()
 				// an import that points to a redirector will fail to find the redirector
 				Import.SourceIndex = INDEX_NONE;
 			}
-			if (Linker->GetSerializeContext())
-			{
-				Linker->GetSerializeContext()->ResetImportCount();
-			}
+			FUObjectThreadContext::Get().GetSerializeContext()->ResetImportCount();
 		}
 	}
 
@@ -304,10 +301,7 @@ void FLinkerManager::DissociateImportsAndForcedExports()
 					Export.ResetObject();
 				}
 			}
-			if (Linker->GetSerializeContext())
-			{
-				Linker->GetSerializeContext()->ResetForcedExports();
-			}
+			FUObjectThreadContext::Get().GetSerializeContext()->ResetForcedExports();
 		}
 	}
 }
diff --git a/Engine/Source/Runtime/CoreUObject/Private/UObject/LinkerSave.cpp b/Engine/Source/Runtime/CoreUObject/Private/UObject/LinkerSave.cpp
--- a/Engine/Source/Runtime/CoreUObject/Private/UObject/LinkerSave.cpp
+++ b/Engine/Source/Runtime/CoreUObject/Private/UObject/LinkerSave.cpp
@@ -377,8 +377,6 @@ FArchive& FLinkerSave::operator<<( FName& InName )
 {
 	int32 Save = MapName(InName.GetDisplayIndex());
 
-	check(GetSerializeContext());
-
 	bool bNameMapped = Save != INDEX_NONE;
 	if (!bNameMapped)
 	{
@@ -387,7 +385,7 @@ FArchive& FLinkerSave::operator<<( FName& InName )
 		FString ErrorMessage = FString::Printf(TEXT("Name \"%s\" is not mapped when saving %s (object: %s, property: %s). This can mean that this object serialize function is not deterministic between reference harvesting and serialization."),
 			*InName.ToString(),
 			*GetArchiveName(),
-			*GetSerializeContext()->SerializedObject->GetFullName(),
+			*FUObjectThreadContext::Get().GetSerializeContext()->SerializedObject->GetFullName(),
 			*GetFullNameSafe(GetSerializedProperty()));
 		ensureMsgf(false, TEXT("%s"), *ErrorMessage);
 		if (LogOutput)
@@ -404,7 +402,7 @@ FArchive& FLinkerSave::operator<<( FName& InName )
 			FString ErrorMessage = FString::Printf(TEXT("Name \"%s\" is referenced from an export but not mapped in the export data names region when saving %s (object: %s, property: %s)."),
 				*InName.ToString(),
 				*GetArchiveName(),
-				*GetSerializeContext()->SerializedObject->GetFullName(),
+				*FUObjectThreadContext::Get().GetSerializeContext()->SerializedObject->GetFullName(),
 				*GetFullNameSafe(GetSerializedProperty()));
 			ensureMsgf(false, TEXT("%s"), *ErrorMessage);
 			if (LogOutput)
@@ -444,7 +442,7 @@ FArchive& FLinkerSave::operator<<(FSoftObjectPath& SoftObjectPath)
 			FString ErrorMessage = FString::Printf(TEXT("SoftObjectPath \"%s\" is not mapped when saving %s (object: %s, property: %s). This can mean that this object serialize function is not deterministic between reference harvesting and serialization."),
 				*SoftObjectPath.ToString(),
 				*GetArchiveName(),
-				*GetSerializeContext()->SerializedObject->GetFullName(),
+				*FUObjectThreadContext::Get().GetSerializeContext()->SerializedObject->GetFullName(),
 				*GetFullNameSafe(GetSerializedProperty()));
 			ensureMsgf(false, TEXT("%s"), *ErrorMessage);
 			if (LogOutput)
@@ -480,18 +478,9 @@ bool FLinkerSave::ShouldSkipProperty(const FProperty* InProperty) const
 	return false;
 }
 
-void FLinkerSave::SetSerializeContext(FUObjectSerializeContext* InLoadContext)
-{
-	SaveContext = InLoadContext;
-	if (Saver)
-	{
-		Saver->SetSerializeContext(InLoadContext);
-	}
-}
-
 FUObjectSerializeContext* FLinkerSave::GetSerializeContext()
 {
-	return SaveContext;
+	return FUObjectThreadContext::Get().GetSerializeContext();
 }
 
 void FLinkerSave::UsingCustomVersion(const struct FGuid& Guid)
diff --git a/Engine/Source/Runtime/CoreUObject/Private/UObject/SavePackage/PackageHarvester.cpp b/Engine/Source/Runtime/CoreUObject/Private/UObject/SavePackage/PackageHarvester.cpp
--- a/Engine/Source/Runtime/CoreUObject/Private/UObject/SavePackage/PackageHarvester.cpp
+++ b/Engine/Source/Runtime/CoreUObject/Private/UObject/SavePackage/PackageHarvester.cpp
@@ -420,7 +420,6 @@ FPackageHarvester::FPackageHarvester(FSaveContext& InContext)
 	this->SetPortFlags(SaveContext.GetPortFlags());
 	this->SetFilterEditorOnly(SaveContext.IsFilterEditorOnly());
 	this->SetCookData(SaveContext.GetSaveArgs().ArchiveCookData);
-	this->SetSerializeContext(SaveContext.GetSerializeContext());
 	this->SetUseUnversionedPropertySerialization(SaveContext.IsSaveUnversionedProperties());
 
 	ResolveOverrides();
diff --git a/Engine/Source/Runtime/CoreUObject/Private/UObject/SavePackage/SaveContext.h b/Engine/Source/Runtime/CoreUObject/Private/UObject/SavePackage/SaveContext.h
--- a/Engine/Source/Runtime/CoreUObject/Private/UObject/SavePackage/SaveContext.h
+++ b/Engine/Source/Runtime/CoreUObject/Private/UObject/SavePackage/SaveContext.h
@@ -472,13 +472,12 @@ public:
 	};
 
 public:
-	FSaveContext(UPackage* InPackage, UObject* InAsset, const TCHAR* InFilename, const FSavePackageArgs& InSaveArgs, FUObjectSerializeContext* InSerializeContext = nullptr)
+	FSaveContext(UPackage* InPackage, UObject* InAsset, const TCHAR* InFilename, const FSavePackageArgs& InSaveArgs)
 		: Package(InPackage)
 		, Asset(InAsset)
 		, Filename(InFilename)
 		, SaveArgs(InSaveArgs)
 		, PackageWriter(InSaveArgs.SavePackageContext ? InSaveArgs.SavePackageContext->PackageWriter : nullptr)
-		, SerializeContext(InSerializeContext)
 		, GameRealmExcludedObjectMarks(GetExcludedObjectMarksForGameRealm(SaveArgs.GetTargetPlatform()))
 	{
 		// Assumptions & checks
@@ -754,16 +753,6 @@ public:
 		return (SaveArgs.SaveFlags & ESaveFlags::SAVE_RehydratePayloads) != 0;
 	}
 
-	FUObjectSerializeContext* GetSerializeContext() const
-	{
-		return SerializeContext;
-	}
-
-	void SetSerializeContext(FUObjectSerializeContext* InContext)
-	{
-		SerializeContext = InContext;
-	}
-
 	FEDLCookCheckerThreadState* GetEDLCookChecker() const
 	{
 		return EDLCookChecker;
@@ -1202,7 +1191,6 @@ private:
 	IPackageWriter* PackageWriter;
 
 	// State context
-	FUObjectSerializeContext* SerializeContext = nullptr;
 	FObjectSaveContextData ObjectSaveContext;
 	bool bCanUseUnversionedPropertySerialization = false;
 	bool bTextFormat = false;
diff --git a/Engine/Source/Runtime/CoreUObject/Private/UObject/SavePackage2.cpp b/Engine/Source/Runtime/CoreUObject/Private/UObject/SavePackage2.cpp
--- a/Engine/Source/Runtime/CoreUObject/Private/UObject/SavePackage2.cpp
+++ b/Engine/Source/Runtime/CoreUObject/Private/UObject/SavePackage2.cpp
@@ -2034,6 +2034,8 @@ ESavePackageResult WritePackageTextHeader(FStructuredArchive::FRecord& Structure
 	FLinkerSave& Linker = *SaveContext.GetLinker();
 	FScopedSlowTask SlowTask((float)Linker.ExportMap.Num(), FText(), SaveContext.IsUsingSlowTask());
 
+	FUObjectSerializeContext* SerializeContext = FUObjectThreadContext::Get().GetSerializeContext();
+
 	for (int32 ExportIndex = 0; ExportIndex < Linker.ExportMap.Num(); ExportIndex++)
 	{
 		if (GWarn->ReceivedUserCancel())
@@ -2064,7 +2066,7 @@ ESavePackageResult WritePackageTextHeader(FStructuredArchive::FRecord& Structure
 		}
 		else
 		{
-			TGuardValue<UObject*> GuardSerializedObject(SaveContext.GetSerializeContext()->SerializedObject, Export.Object);
+			TGuardValue<UObject*> GuardSerializedObject(SerializeContext->SerializedObject, Export.Object);
 			Export.Object->Serialize(Linker);
 #if WITH_EDITOR
 			Export.Object->CookAdditionalFiles(SaveContext.GetFilename(), SaveContext.GetTargetPlatform(),
@@ -2103,6 +2105,8 @@ ESavePackageResult WriteExports(FStructuredArchive::FRecord& StructuredArchiveRo
 	FLinkerSave* Linker = SaveContext.GetLinker();
 	FScopedSlowTask SlowTask((float)Linker->ExportMap.Num(), FText(), SaveContext.IsUsingSlowTask());
 
+	FUObjectSerializeContext* SerializeContext = FUObjectThreadContext::Get().GetSerializeContext();
+
 	FStructuredArchive::FRecord ExportsRecord = StructuredArchiveRoot.EnterRecord(TEXT("Exports"));
 
 	// Save exports.
@@ -2150,7 +2154,7 @@ ESavePackageResult WriteExports(FStructuredArchive::FRecord& StructuredArchiveRo
 			}
 			else
 			{
-				TGuardValue<UObject*> GuardSerializedObject(SaveContext.GetSerializeContext()->SerializedObject, Export.Object);
+				TGuardValue<UObject*> GuardSerializedObject(SerializeContext->SerializedObject, Export.Object);
 
 				if (bSupportsText)
 				{
@@ -2748,7 +2752,6 @@ ESavePackageResult SaveHarvestedRealms(FSaveContext& SaveContext, ESaveRealm Har
 	}
 
 	FStructuredArchive::FRecord StructuredArchiveRoot = SaveContext.GetStructuredArchive()->Open().EnterRecord();
-	StructuredArchiveRoot.GetUnderlyingArchive().SetSerializeContext(SaveContext.GetSerializeContext());
 
 	// Write Header
 	SlowTask.EnterProgressFrame();
@@ -2781,7 +2784,6 @@ ESavePackageResult SaveHarvestedRealms(FSaveContext& SaveContext, ESaveRealm Har
 			const bool bIsOptionalRealm = SaveContext.GetCurrentHarvestingRealm() == ESaveRealm::Optional;
 			TUniquePtr<FLargeMemoryWriter> ExportsArchive = SaveContext.GetPackageWriter()->CreateLinkerExportsArchive(
 				SaveContext.GetPackage()->GetFName(), SaveContext.GetAsset(), bIsOptionalRealm ? 1 : 0);
-			ExportsArchive->SetSerializeContext(SaveContext.GetSerializeContext());
 			SaveContext.Result = WriteCookedExports(*ExportsArchive, SaveContext);
 
 			if (SaveContext.Result == ESavePackageResult::Success)
@@ -2947,8 +2949,6 @@ ESavePackageResult SaveHarvestedRealms(FSaveContext& SaveContext, ESaveRealm Har
  */
 ESavePackageResult InnerSave(FSaveContext& SaveContext)
 {
-	TRefCountPtr<FUObjectSerializeContext> SerializeContext(FUObjectThreadContext::Get().GetSerializeContext());
-	SaveContext.SetSerializeContext(SerializeContext);
 	SaveContext.SetEDLCookChecker(&FEDLCookCheckerThreadState::Get());
 
 	UE::FScopedIDOSerializationContext IDOSaveContext(nullptr);
@@ -3185,7 +3185,7 @@ ESavePackageResult UPackage::SaveConcurrent(TArrayView<FPackageSaveInfo> InPacka
 		SCOPED_SAVETIMER(UPackage_SaveConcurrent_PreSave);
 		for (FPackageSaveInfo& PackageSaveInfo : InPackages)
 		{
-			FSaveContext& SaveContext = PackageSaveContexts.Emplace_GetRef(PackageSaveInfo.Package, PackageSaveInfo.Package->FindAssetInPackage(), *PackageSaveInfo.Filename, SaveArgs, nullptr);
+			FSaveContext& SaveContext = PackageSaveContexts.Emplace_GetRef(PackageSaveInfo.Package, PackageSaveInfo.Package->FindAssetInPackage(), *PackageSaveInfo.Filename, SaveArgs);
 
 			// Validation
 			SaveContext.Result = ValidatePackage(SaveContext);
diff --git a/Engine/Source/Runtime/CoreUObject/Private/UObject/UObjectGlobals.cpp b/Engine/Source/Runtime/CoreUObject/Private/UObject/UObjectGlobals.cpp
--- a/Engine/Source/Runtime/CoreUObject/Private/UObject/UObjectGlobals.cpp
+++ b/Engine/Source/Runtime/CoreUObject/Private/UObject/UObjectGlobals.cpp
@@ -2473,15 +2473,6 @@ void EndLoad(FUObjectSerializeContext* LoadContext, TArray<UPackage*>* OutLoaded
 	}
 #endif	// WITH_EDITOR
 
-
-	if (LoadContext->GetBeginLoadCount() == 0)
-	{
-		if (!GEventDrivenLoaderEnabled)
-		{
-			LoadContext->DetachFromLinkers();
-		}
-	}
-
 	if (OutLoadedPackages)
 	{
 		OutLoadedPackages->Reserve(LoadedPackages.Num());
@@ -3156,7 +3147,6 @@ UObject* StaticDuplicateObjectEx( FObjectDuplicationParameters& Parameters )
 
 	TRefCountPtr<FUObjectSerializeContext> LoadContext(FUObjectThreadContext::Get().GetSerializeContext());
 	FDuplicateDataReader Reader(DuplicatedObjectAnnotation, ObjectData.Get(), Parameters.PortFlags, Parameters.DestOuter);
-	Reader.SetSerializeContext(LoadContext);
 	for(int32 ObjectIndex = 0;ObjectIndex < SerializedObjects.Num();ObjectIndex++)
 	{
 		UObject* SerializedObject = SerializedObjects[ObjectIndex];
diff --git a/Engine/Source/Runtime/CoreUObject/Private/UObject/UObjectThreadContext.cpp b/Engine/Source/Runtime/CoreUObject/Private/UObject/UObjectThreadContext.cpp
--- a/Engine/Source/Runtime/CoreUObject/Private/UObject/UObjectThreadContext.cpp
+++ b/Engine/Source/Runtime/CoreUObject/Private/UObject/UObjectThreadContext.cpp
@@ -56,7 +56,6 @@ FUObjectSerializeContext::FUObjectSerializeContext()
 FUObjectSerializeContext::~FUObjectSerializeContext()
 {
 	checkf(!HasLoadedObjects(), TEXT("FUObjectSerializeContext is being destroyed but it still has pending loaded objects in its ObjectsLoaded list."));
-	check(AttachedLinkers.Num() == 0);
 }
 
 int32 FUObjectSerializeContext::IncrementBeginLoadCount()
@@ -99,23 +98,13 @@ bool FUObjectSerializeContext::PRIVATE_PatchNewObjectIntoExport(UObject* OldObje
 void FUObjectSerializeContext::AttachLinker(FLinkerLoad* InLinker)
 {
 	check(!GEventDrivenLoaderEnabled);
-	AttachedLinkers.Add(InLinker);
 }
 
 void FUObjectSerializeContext::DetachLinker(FLinkerLoad* InLinker)
 {
-	AttachedLinkers.Remove(InLinker);
 }
 
 void FUObjectSerializeContext::DetachFromLinkers()
 {
 	check(!GEventDrivenLoaderEnabled);
-	check(ObjectsLoaded.Num() == 0 || AttachedLinkers.Num() == 0);
-	TArray<FLinkerLoad*> LinkersToDetach = AttachedLinkers.Array();
-	for (FLinkerLoad* Linker : LinkersToDetach)
-	{
-		check(Linker->GetSerializeContext() == this);
-		Linker->SetSerializeContext(nullptr);
-	}
-	check(AttachedLinkers.Num() == 0);
 }
diff --git a/Engine/Source/Runtime/CoreUObject/Public/Serialization/ArchiveStackTrace.h b/Engine/Source/Runtime/CoreUObject/Public/Serialization/ArchiveStackTrace.h
--- a/Engine/Source/Runtime/CoreUObject/Public/Serialization/ArchiveStackTrace.h
+++ b/Engine/Source/Runtime/CoreUObject/Public/Serialization/ArchiveStackTrace.h
@@ -21,7 +21,6 @@ class FLinkerLoad;
 class FProperty;
 class FUObjectThreadContext;
 class UObject;
-struct FUObjectSerializeContext;
 
 /** Structure that holds stats from comparing two packages */
 struct FArchiveDiffStats
@@ -272,8 +271,6 @@ public:
 	COREUOBJECT_API virtual ~FArchiveStackTraceWriter() override;
 
 	COREUOBJECT_API FORCENOINLINE virtual void Serialize(void* Data, int64 Length) override; // FORCENOINLINE so it can be counted during StackTrace
-	COREUOBJECT_API virtual void SetSerializeContext(FUObjectSerializeContext* Context) override;
-	COREUOBJECT_API virtual FUObjectSerializeContext* GetSerializeContext() override;
 	
 #if WITH_EDITOR
 	virtual void PushDebugDataString(const FName& DebugData) override
@@ -350,7 +347,6 @@ public:
 private:
 	FArchiveCallstacks& Callstacks;
 	const FArchiveDiffMap* DiffMap;
-	TRefCountPtr<FUObjectSerializeContext> SerializeContext;
 #if WITH_EDITOR
 	TArray<FName> DebugDataStack;
 #endif
@@ -380,8 +376,6 @@ PRAGMA_DISABLE_DEPRECATION_WARNINGS;
 		const TCHAR* Filename = nullptr);
 
 	COREUOBJECT_API FORCENOINLINE virtual void Serialize(void* Memory, int64 Length) override; // FORCENOINLINE so it can be counted during StackTrace
-	COREUOBJECT_API virtual void SetSerializeContext(FUObjectSerializeContext* Context) override;
-	COREUOBJECT_API virtual FUObjectSerializeContext* GetSerializeContext() override;
 
 private:
 	FArchiveStackTraceWriter StackTraceWriter;
@@ -417,8 +411,6 @@ public:
 	}
 
 	COREUOBJECT_API FORCENOINLINE virtual void Serialize(void* Memory, int64 Length) override; // FORCENOINLINE so it can be counted during StackTrace
-	COREUOBJECT_API virtual void SetSerializeContext(FUObjectSerializeContext* Context) override;
-	COREUOBJECT_API virtual FUObjectSerializeContext* GetSerializeContext() override;
 
 	/** Compares this archive with the given bytes from disk or FPackageData. Dumps all differences to log. */
 	COREUOBJECT_API void CompareWith(const TCHAR* InFilename, const int64 TotalHeaderSize, const TCHAR* CallstackCutoffText,
diff --git a/Engine/Source/Runtime/CoreUObject/Public/Serialization/DuplicatedDataReader.h b/Engine/Source/Runtime/CoreUObject/Public/Serialization/DuplicatedDataReader.h
--- a/Engine/Source/Runtime/CoreUObject/Public/Serialization/DuplicatedDataReader.h
+++ b/Engine/Source/Runtime/CoreUObject/Public/Serialization/DuplicatedDataReader.h
@@ -25,9 +25,6 @@ private:
 	const FLargeMemoryData&					ObjectData;
 	int64									Offset;
 
-	/** Context for duplication */
-	TRefCountPtr<FUObjectSerializeContext> DuplicateContext;
-
 	//~ Begin FArchive Interface.
 
 	virtual FArchive& operator<<(FName& N) override;
@@ -73,16 +70,6 @@ public:
 		return ObjectData.GetSize();
 	}
 
-	virtual void SetSerializeContext(FUObjectSerializeContext* InLoadContext) override
-	{
-		DuplicateContext = InLoadContext;
-	}
-
-	virtual FUObjectSerializeContext* GetSerializeContext() override
-	{
-		return DuplicateContext;
-	}
-
 	/**
 	 * Constructor
 	 * 
diff --git a/Engine/Source/Runtime/CoreUObject/Public/Serialization/DuplicatedDataWriter.h b/Engine/Source/Runtime/CoreUObject/Public/Serialization/DuplicatedDataWriter.h
--- a/Engine/Source/Runtime/CoreUObject/Public/Serialization/DuplicatedDataWriter.h
+++ b/Engine/Source/Runtime/CoreUObject/Public/Serialization/DuplicatedDataWriter.h
@@ -37,9 +37,6 @@ private:
 	 */
 	struct FObjectInstancingGraph*			InstanceGraph;
 
-	/** Context for duplication */
-	TRefCountPtr<FUObjectSerializeContext> DuplicateContext;
-
 	//~ Begin FArchive Interface.
 
 	virtual FArchive& operator<<(FName& N) override;
@@ -88,16 +85,6 @@ public:
 		return ObjectData.GetSize();
 	}
 
-	virtual void SetSerializeContext(FUObjectSerializeContext* InLoadContext) override
-	{
-		DuplicateContext = InLoadContext;
-	}
-
-	virtual FUObjectSerializeContext* GetSerializeContext() override
-	{
-		return DuplicateContext;
-	}
-
 	TArray<UObject*>	UnserializedObjects;
 
 	/**
diff --git a/Engine/Source/Runtime/CoreUObject/Public/UObject/LinkerLoad.h b/Engine/Source/Runtime/CoreUObject/Public/UObject/LinkerLoad.h
--- a/Engine/Source/Runtime/CoreUObject/Public/UObject/LinkerLoad.h
+++ b/Engine/Source/Runtime/CoreUObject/Public/UObject/LinkerLoad.h
@@ -1510,7 +1510,7 @@ private:
 public:
 
 	//~ FArchive interface
-	COREUOBJECT_API virtual void SetSerializeContext(FUObjectSerializeContext* InLoadContext) override;
+	UE_DEPRECATED(5.5, "GetSerializeContext is not supported. Use FUObjectThreadContext::Get().GetSerializeContext().")
 	COREUOBJECT_API virtual FUObjectSerializeContext* GetSerializeContext() override;
 };
 
diff --git a/Engine/Source/Runtime/CoreUObject/Public/UObject/LinkerSave.h b/Engine/Source/Runtime/CoreUObject/Public/UObject/LinkerSave.h
--- a/Engine/Source/Runtime/CoreUObject/Public/UObject/LinkerSave.h
+++ b/Engine/Source/Runtime/CoreUObject/Public/UObject/LinkerSave.h
@@ -172,7 +172,7 @@ public:
 	FArchive& operator<<(FSoftObjectPath& SoftObjectPath);
 	FArchive& operator<<( FLazyObjectPtr& LazyObjectPtr );
 	virtual bool ShouldSkipProperty(const FProperty* InProperty) const override;
-	virtual void SetSerializeContext(FUObjectSerializeContext* InLoadContext) override;
+	UE_DEPRECATED(5.5, "GetSerializeContext is not supported. Use FUObjectThreadContext::Get().GetSerializeContext().")
 	FUObjectSerializeContext* GetSerializeContext() override;
 	virtual void UsingCustomVersion(const struct FGuid& Guid) override;
 	/**
diff --git a/Engine/Source/Runtime/CoreUObject/Public/UObject/UObjectThreadContext.h b/Engine/Source/Runtime/CoreUObject/Public/UObject/UObjectThreadContext.h
--- a/Engine/Source/Runtime/CoreUObject/Public/UObject/UObjectThreadContext.h
+++ b/Engine/Source/Runtime/CoreUObject/Public/UObject/UObjectThreadContext.h
@@ -151,8 +151,6 @@ private:
 	TArray<UObject*> ObjectsLoaded;
 	/** List of linkers that we want to close the loaders for (to free file handles) - needs to be delayed until EndLoad is called with GObjBeginLoadCount of 0 */
 	TArray<FLinkerLoad*> DelayedLinkerClosePackages;
-	/** List of linkers associated with this context */
-	TSet<FLinkerLoad*> AttachedLinkers;
 
 public:
 
@@ -266,12 +264,15 @@ public:
 	}
 
 	/** Attaches a linker to this context */
+	UE_DEPRECATED(5.5, "AttachLinker is not necessary. Remove calls to it.")
 	COREUOBJECT_API void AttachLinker(FLinkerLoad* InLinker);
 	
 	/** Detaches a linker from this context */
+	UE_DEPRECATED(5.5, "DetachLinker is not necessary. Remove calls to it.")
 	COREUOBJECT_API void DetachLinker(FLinkerLoad* InLinker);
 
 	/** Detaches all linkers from this context */
+	UE_DEPRECATED(5.5, "DetachFromLinkers is not necessary. Remove calls to it.")
 	COREUOBJECT_API void DetachFromLinkers();
 
 	//~ TRefCountPtr interface
