Fix fatal error message when trying to reallocate an existing object to print the class path and object flags
#rnx

diff --git a/Engine/Source/Runtime/CoreUObject/Private/UObject/UObjectGlobals.cpp b/Engine/Source/Runtime/CoreUObject/Private/UObject/UObjectGlobals.cpp
--- a/Engine/Source/Runtime/CoreUObject/Private/UObject/UObjectGlobals.cpp
+++ b/Engine/Source/Runtime/CoreUObject/Private/UObject/UObjectGlobals.cpp
@@ -3473,12 +3473,17 @@ UObject* StaticAllocateObject
 						"This has the side effect, of using the full path name for config ini sections. Use 'OverridePerObjectConfigSection' to keep the short name.\n\n");
 			}
 
+			// This generally happens when calling NewObject with a specific object name and an object already exists at the same path.
+			// If the classes look the same but have different paths, an old version may have been renamed due to plugin unloading or class recompiling.
+			// If the object has the garbage flag set, it was marked as ready to destroy but still exists so it cannot be reallocated before garbage collection clears it.
 			UE_LOG(LogUObjectGlobals, Fatal,
-				TEXT("%sObjects have the same fully qualified name but different paths.\n"
+				TEXT("%sCannot replace existing object of a different class.\n"
 				     "\tNew Object: %s %s.%s\n"
-				     "\tExisting Object: %s"),
-				ErrorPrefix, *InClass->GetName(), InOuter ? *InOuter->GetPathName() : TEXT(""), *InName.ToString(),
-				*Obj->GetFullName());
+				     "\tExisting Object: %s %s(0x%08x 0x%08x)"),
+				ErrorPrefix, *InClass->GetPathName(), InOuter ? *InOuter->GetPathName() : TEXT(""), *InName.ToString(),
+				*Obj->GetFullName(nullptr, EObjectFullNameFlags::IncludeClassPackage),
+				Obj->HasAnyInternalFlags(EInternalObjectFlags::Garbage) ? TEXT("(garbage) ") : TEXT(""),
+				(int32)Obj->GetFlags(), (int32)Obj->GetInternalFlags());
 		}
 	}
 
