Temporary fix to make deleted components work with CreateProjectFromTemplate

diff --git a/Engine/Source/Runtime/CoreUObject/Private/UObject/PropertyArray.cpp b/Engine/Source/Runtime/CoreUObject/Private/UObject/PropertyArray.cpp
--- a/Engine/Source/Runtime/CoreUObject/Private/UObject/PropertyArray.cpp
+++ b/Engine/Source/Runtime/CoreUObject/Private/UObject/PropertyArray.cpp
@@ -277,6 +277,23 @@ void FArrayProperty::SerializeItem(FStructuredArchive::FSlot Slot, void* Value,
 								return i;
 							}
 						}
+
+						////////////////////////////////////////////////////////////////////////////
+						// HACK for CreateProjectFromTemplate, 
+						// Fallback to find object by name
+						FName ObjectNameToFind = Object->GetFName();
+						for (int i = 0; i < ArrayNum; ++i)
+						{
+							if (UObject* CurrentObject = InnerObjectProperty->GetObjectPropertyValue(ArrayHelper.GetElementPtr(i)))
+							{
+								if (CurrentObject->GetFName() == ObjectNameToFind)
+								{
+									return i;
+								}
+							}
+						}
+						// HACK for CreateProjectFromTemplate, fallback to find object by name
+						////////////////////////////////////////////////////////////////////////////
 					}
 					return INDEX_NONE;
 				};
@@ -313,6 +330,10 @@ void FArrayProperty::SerializeItem(FStructuredArchive::FSlot Slot, void* Value,
 							{
 								IndicesToRemove.Add(Index);
 							}
+							else
+							{
+								UE_LOG(LogOverridableObject, VeryVerbose, TEXT("Unable to load removed item %s(0x%p)"), *GetNameSafe(RemovedSubObject), RemovedSubObject);
+							}
 
 							// Need to fetch the ArrayOverriddenPropertyNode every loop as the previous iteration might have reallocated the node.
 							if (FOverriddenPropertyNode* ArrayOverriddenPropertyNode = OverriddenProperties ? OverriddenProperties->SetOverriddenPropertyOperation(EOverriddenPropertyOperation::Modified, UnderlyingArchive.GetSerializedPropertyChain(), /*Property*/nullptr) : nullptr)
@@ -454,6 +475,27 @@ void FArrayProperty::SerializeItem(FStructuredArchive::FSlot Slot, void* Value,
 									}
 								}
 							}
+
+							////////////////////////////////////////////////////////////////////////////
+							// HACK for CreateProjectFromTemplate, 
+							// Fallback to find object by name
+							if (ObjectToFind.Object)
+							{
+								FName ObjectNameToFind = ObjectToFind.Object->GetFName();
+								for (int i = 0; i < ArrayNum; ++i)
+								{
+									if (UObject* CurrentObject = InnerObjectProperty->GetObjectPropertyValue(ArrayHelper.GetElementPtr(i)))
+									{
+										if (CurrentObject->GetFName() == ObjectNameToFind)
+										{
+											return i;
+										}
+									}
+								}
+							}
+							// HACK for CreateProjectFromTemplate, fallback to find object by name
+							////////////////////////////////////////////////////////////////////////////
+
 							return INDEX_NONE;
 						};
 
@@ -473,6 +515,10 @@ void FArrayProperty::SerializeItem(FStructuredArchive::FSlot Slot, void* Value,
 										{
 											RemovedIndices.Add(DefaultIndex);
 										}
+										else
+										{
+											UE_LOG(LogOverridableObject, VeryVerbose, TEXT("Unable to save deleted item %s(0x%p)"), *GetNameSafe(Pair.Key.Object), Pair.Key.Object.Get());
+										}
 										break;
 									}
 								case EOverriddenPropertyOperation::Add:
@@ -483,6 +529,10 @@ void FArrayProperty::SerializeItem(FStructuredArchive::FSlot Slot, void* Value,
 											AddedIndices.Add(Index);
 											ModifiedIndices.Remove(Index);
 										}
+										else
+										{
+											UE_LOG(LogOverridableObject, VeryVerbose, TEXT("Unable to save added item %s(0x%p)"), *GetNameSafe(Pair.Key.Object), Pair.Key.Object.Get());
+										}
 										break;
 									}
 								default:
