Auto-enable placeholder types by default when IDO features are enabled.

Notes:
- Placeholder types can still be independently enabled or disabled via CVar (SceneGraph.EnablePropertyBagPlaceholderObjectSupport 1/0).

diff --git a/Engine/Source/Runtime/CoreUObject/Private/UObject/PropertyBagRepository.cpp b/Engine/Source/Runtime/CoreUObject/Private/UObject/PropertyBagRepository.cpp
--- a/Engine/Source/Runtime/CoreUObject/Private/UObject/PropertyBagRepository.cpp
+++ b/Engine/Source/Runtime/CoreUObject/Private/UObject/PropertyBagRepository.cpp
@@ -873,7 +873,7 @@ namespace Private
 	static FAutoConsoleVariableRef CVarEnablePropertyBagPlaceholderObjectSupport(
 		TEXT("SceneGraph.EnablePropertyBagPlaceholderObjectSupport"),
 		bEnablePropertyBagPlaceholderObjectSupport,
-		TEXT("If true, allows placeholder types to be created in place of missing types on load in order to redirect serialization into a property bag."),
+		TEXT("If true, allows placeholder types to be created in place of missing types in order to redirect serialization into a property bag."),
 		ECVF_Default
 	);
 #endif
@@ -883,13 +883,19 @@ bool FPropertyBagRepository::IsPropertyBagPlaceholderObjectSupportEnabled()
 {
 #if WITH_EDITOR && UE_WITH_OBJECT_HANDLE_TYPE_SAFETY
 	static bool bIsInitialized = false;
+	static bool bForceDisabled = false;
 	if (!bIsInitialized)
 	{
 		Private::bEnablePropertyBagPlaceholderObjectSupport |= FParse::Param(FCommandLine::Get(), TEXT("WithPropertyBagPlaceholderObjects"));
+		Private::CVarEnablePropertyBagPlaceholderObjectSupport->OnChangedDelegate().AddLambda([](IConsoleVariable* CVar)
+		{
+			bForceDisabled = !CVar->GetBool();
+		});
+
 		bIsInitialized = true;
 	}
 	
-	return Private::bEnablePropertyBagPlaceholderObjectSupport;
+	return Private::bEnablePropertyBagPlaceholderObjectSupport || (IsInstanceDataObjectSupportEnabled() && !bForceDisabled);
 #else
 	return false;
 #endif
