Split the logic to collect save overrides out of PreSave so that it can be queried by anything that needs to understand save overrides without invoking the other side effects of PreSave

diff --git a/Engine/Source/Runtime/CoreUObject/Private/UObject/Obj.cpp b/Engine/Source/Runtime/CoreUObject/Private/UObject/Obj.cpp
--- a/Engine/Source/Runtime/CoreUObject/Private/UObject/Obj.cpp
+++ b/Engine/Source/Runtime/CoreUObject/Private/UObject/Obj.cpp
@@ -1467,6 +1467,11 @@ void UObject::PreSave(FObjectPreSaveContext SaveContext)
 	PRAGMA_ENABLE_DEPRECATION_WARNINGS
 	FCoreUObjectDelegates::OnObjectPreSave.Broadcast(this, SaveContext);
 #endif
+	CollectSaveOverrides(FObjectCollectSaveOverridesContext(SaveContext.Data));
+}
+
+void UObject::CollectSaveOverrides(FObjectCollectSaveOverridesContext SaveContext)
+{
 }
 
 #if WITH_EDITOR
diff --git a/Engine/Source/Runtime/CoreUObject/Public/UObject/Object.h b/Engine/Source/Runtime/CoreUObject/Public/UObject/Object.h
--- a/Engine/Source/Runtime/CoreUObject/Public/UObject/Object.h
+++ b/Engine/Source/Runtime/CoreUObject/Public/UObject/Object.h
@@ -21,6 +21,7 @@ namespace UE::ConfigAccessTracking { struct FConfigAccessData; }
 class FConfigCacheIni;
 class FCustomPropertyConditionState;
 class FEditPropertyChain;
+class FObjectCollectSaveOverridesContext;
 class FObjectPostSaveContext;
 class FObjectPostSaveRootContext;
 class FObjectPreSaveContext;
@@ -291,6 +292,11 @@ public:
 	 */
 	COREUOBJECT_API virtual void PreSave(FObjectPreSaveContext SaveContext);
 
+	/**
+	 * Collect any save override state.
+	 */
+	COREUOBJECT_API virtual void CollectSaveOverrides(FObjectCollectSaveOverridesContext SaveContext);
+
 #if WITH_EDITOR
 	/**
 	 * Note that the object will be modified.  If we are currently recording into the 
diff --git a/Engine/Source/Runtime/CoreUObject/Public/UObject/ObjectSaveContext.h b/Engine/Source/Runtime/CoreUObject/Public/UObject/ObjectSaveContext.h
--- a/Engine/Source/Runtime/CoreUObject/Public/UObject/ObjectSaveContext.h
+++ b/Engine/Source/Runtime/CoreUObject/Public/UObject/ObjectSaveContext.h
@@ -28,6 +28,20 @@ struct FObjectSaveContextData
 	COREUOBJECT_API void Set(UPackage* Package, const ITargetPlatform* InTargetPlatform, const TCHAR* InTargetFilename, uint32 InSaveFlags);
 	COREUOBJECT_API void Set(UPackage* Package, const ITargetPlatform* InTargetPlatform, const FPackagePath& TargetPath, uint32 InSaveFlags);
 
+	/**
+	 * Add a save override to specific object. (i.e. mark certain objects or properties transient for this save)
+	 */
+	void AddSaveOverride(UObject* Target, FObjectSaveOverride InOverride)
+	{
+		if (FObjectSaveOverride* MatchingSaveOverride = SaveOverrides.Find(Target))
+		{
+			MatchingSaveOverride->Merge(InOverride);
+		}
+		else
+		{
+			SaveOverrides.Add(Target, MoveTemp(InOverride));
+		}
+	}
 
 	// Global Parameters that are read-only by the interfaces
 
@@ -96,6 +110,57 @@ struct FObjectSaveContextData
 
 };
 
+/** Interface used by CollectSaveOverrides to access the save parameters. */
+class FObjectCollectSaveOverridesContext
+{
+public:
+	explicit FObjectCollectSaveOverridesContext(FObjectSaveContextData& InData)
+		: Data(InData)
+	{
+		// Note: Doesn't increment NumRefPasses as CollectSaveOverrides is called from PreSave
+	}
+
+	FObjectCollectSaveOverridesContext(const FObjectCollectSaveOverridesContext& Other)
+		: Data(Other.Data)
+	{
+		// Note: Doesn't increment NumRefPasses as CollectSaveOverrides is called from PreSave
+	}
+
+	/** Report whether this is a save into a target-specific cooked format. */
+	bool IsCooking() const { return Data.TargetPlatform != nullptr; }
+
+	/** Return the targetplatform of the save, if cooking. Null if not cooking. */
+	const ITargetPlatform* GetTargetPlatform() const { return Data.TargetPlatform; }
+
+	bool IsCookByTheBook() const { return GetCookType()  == UE::Cook::ECookType::ByTheBook; }
+	bool IsCookOnTheFly() const { return GetCookType() == UE::Cook::ECookType::OnTheFly; }
+	bool IsCookTypeUnknown() const { return GetCookType() == UE::Cook::ECookType::Unknown; }
+	UE::Cook::ECookType GetCookType() const { return Data.CookType; }
+	UE::Cook::ECookingDLC GetCookingDLC() const { return Data.CookingDLC; }
+
+	/**
+	 * Return whether the package is being saved due to a procedural save.
+	 * Any save without the possibility of user-generated edits to the package is a procedural save (Cooking, EditorDomain).
+	 * This allows us to execute transforms that only need to be executed in response to new user data.
+	 */
+	bool IsProceduralSave() const { return Data.bProceduralSave; }
+
+	/** Return the save flags (ESaveFlags) of the save. */
+	uint32 GetSaveFlags() const { return Data.SaveFlags; }
+
+	/**
+	 * Add a save override to specific object. (i.e. mark certain objects or properties transient for this save)
+	 */
+	void AddSaveOverride(UObject* Target, FObjectSaveOverride InOverride)
+	{
+		Data.AddSaveOverride(Target, MoveTemp(InOverride));
+	}
+
+protected:
+	FObjectSaveContextData& Data;
+	friend class UObject;
+};
+
 /** Interface used by PreSave to access the save parameters. */
 class FObjectPreSaveContext
 {
@@ -172,9 +237,10 @@ public:
 	/**
 	 * Add a save override to specific object. (i.e. mark certain objects or properties transient for this save)
 	 */
+	// TODO: UE_DEPRECATED(5.5, "Calling AddSaveOverride in UObject::PreSave is deprecated. Override UObject::CollectSaveOverrides and call AddSaveOverride on its context instead.")
 	void AddSaveOverride(UObject* Target, FObjectSaveOverride InOverride)
 	{
-		Data.SaveOverrides.Add(Target, MoveTemp(InOverride));
+		Data.AddSaveOverride(Target, MoveTemp(InOverride));
 	}
 
 protected:
diff --git a/Engine/Source/Runtime/CoreUObject/Public/UObject/ObjectSaveOverride.h b/Engine/Source/Runtime/CoreUObject/Public/UObject/ObjectSaveOverride.h
--- a/Engine/Source/Runtime/CoreUObject/Public/UObject/ObjectSaveOverride.h
+++ b/Engine/Source/Runtime/CoreUObject/Public/UObject/ObjectSaveOverride.h
@@ -13,6 +13,12 @@
  */
 struct FPropertySaveOverride
 {
+	void Merge(const FPropertySaveOverride& Other)
+	{
+		checkf(PropertyPath == Other.PropertyPath, TEXT("Merge called with an unrelated FPropertySaveOverride!"));
+		bMarkTransient |= Other.bMarkTransient;
+	}
+
 	FFieldPath PropertyPath;
 	bool bMarkTransient;
 };
@@ -20,8 +26,32 @@ struct FPropertySaveOverride
 /** Data to specify an override to apply to an object during save without mutating the object itself. */
 struct FObjectSaveOverride
 {
+	void Merge(const FObjectSaveOverride& Other)
+	{
+		bForceTransient |= Other.bForceTransient;
+
+		if (PropOverrides.IsEmpty())
+		{
+			PropOverrides = Other.PropOverrides;
+		}
+		else if (!Other.PropOverrides.IsEmpty())
+		{
+			for (const FPropertySaveOverride& OtherPropOverride : Other.PropOverrides)
+			{
+				if (FPropertySaveOverride* MatchingPropOverride = PropOverrides.FindByPredicate([&OtherPropOverride](const FPropertySaveOverride& PropOverride) { return PropOverride.PropertyPath == OtherPropOverride.PropertyPath; }))
+				{
+					MatchingPropOverride->Merge(OtherPropOverride);
+				}
+				else
+				{
+					PropOverrides.Add(OtherPropOverride);
+				}
+			}
+		}
+	}
+
 	TArray<FPropertySaveOverride> PropOverrides;
 
 	// Treats the object as RF_Transient for the duration of the save
 	bool bForceTransient = false;
-};
\ No newline at end of file
+};
