Fix UObject getting marked as garbage during a transaction hitting atomic code
Move this to after the transaction

diff --git a/Engine/Source/Runtime/CoreUObject/Public/UObject/UObjectBaseUtility.h b/Engine/Source/Runtime/CoreUObject/Public/UObject/UObjectBaseUtility.h
--- a/Engine/Source/Runtime/CoreUObject/Public/UObject/UObjectBaseUtility.h
+++ b/Engine/Source/Runtime/CoreUObject/Public/UObject/UObjectBaseUtility.h
@@ -177,8 +177,14 @@ public:
 	FORCEINLINE void MarkAsGarbage()
 	{
 		check(!IsRooted());
-		AtomicallySetFlags(RF_MirroredGarbage);
-		GUObjectArray.IndexToObject(InternalIndex)->SetGarbage();
+
+		// only set these flags after we have finished our transaction
+		// GC should not happen during a transaction so delaying this wont cause GC issues
+		UE_AUTORTFM_ONCOMMIT(
+		{
+			AtomicallySetFlags(RF_MirroredGarbage);
+			GUObjectArray.IndexToObject(InternalIndex)->SetGarbage();
+		});
 	}
 
 	/**
