Impersonate properties in FPropertyBagRepository::CopyTaggedProperties

Impersonation is required when serializing an IDO otherwise generated types (like structs) within them will not serialize.

diff --git a/Engine/Source/Runtime/CoreUObject/Private/UObject/PropertyBagRepository.cpp b/Engine/Source/Runtime/CoreUObject/Private/UObject/PropertyBagRepository.cpp
--- a/Engine/Source/Runtime/CoreUObject/Private/UObject/PropertyBagRepository.cpp
+++ b/Engine/Source/Runtime/CoreUObject/Private/UObject/PropertyBagRepository.cpp
@@ -753,11 +753,14 @@ void FPropertyBagRepository::CreateInstanceDataObjectUnsafe(UObject* Owner, FPro
 
 void FPropertyBagRepository::CopyTaggedProperties(const UObject* Source, UObject* Dest)
 {
+	FUObjectSerializeContext* SerializeContext = FUObjectThreadContext::Get().GetSerializeContext();
+	TGuardValue<bool> ImpersonatePropertiesScope(SerializeContext->bImpersonateProperties, true);
+
 	TArray<uint8> Buffer;
 	Buffer.Reserve(Source->GetClass()->GetStructureSize());
 	FObjectWriter Writer(Buffer);
 	Source->GetClass()->SerializeTaggedProperties(Writer, (uint8*)Source, Source->GetClass(), nullptr);
-		
+
 	FObjectReader Reader(Buffer);
 	Reader.ArMergeOverrides = true;
 	Dest->GetClass()->SerializeTaggedProperties(Reader, (uint8*)Dest, Dest->GetClass(), nullptr);
@@ -771,7 +774,7 @@ FScopedIDOSerializationContext::FScopedIDOSerializationContext(UObject* InObject
 	FUObjectSerializeContext* SerializeContext = FUObjectThreadContext::Get().GetSerializeContext();
 	bHasIDOSupport = FPropertyBagRepository::IsInstanceDataObjectSupportEnabled(Object);
 	bCreateIDO = bHasIDOSupport && !SerializeContext->bImpersonateProperties && Archive->IsLoading();
-	
+
 	if (bHasIDOSupport)
 	{
 		if (Archive->IsLoading())
@@ -781,7 +784,7 @@ FScopedIDOSerializationContext::FScopedIDOSerializationContext(UObject* InObject
 			ScopedTrackSerializedPropertyPath.Emplace(SerializeContext->bTrackSerializedPropertyPath, bCreateIDO);
 			ScopedSerializeUnknownProperty.Emplace(SerializeContext->bTrackUnknownProperties, bCreateIDO);
 			ScopedSerializedObject.Emplace(SerializeContext->SerializedObject, Object);
-			
+
 			// Enable tracking of initialized properties when loading an IDO, which is implied by impersonation being enabled.
 			const bool bLoadingIDO = bHasIDOSupport && SerializeContext->bImpersonateProperties;
 			TGuardValue<bool> ScopedTrackInitializedProperties(SerializeContext->bTrackInitializedProperties, bLoadingIDO);
