Fix bulkdata streaming requests with zero length from asserting if IoStore is disabled.

- Currently, if we make a request via the bulkdata streaming api for zero bytes, the IoStore (FFileIoStore) path will return a successful read but if the IoStore is disabled we will go through FPackageResourceIoBackend which will use the IAsyncReadFileHandle api which tends to assert or fail on zero byte reads.
- There is work pending for unifying how all our different file apis behave under these sorts of conditions but for now it's easier to fix FPackageResourceIoBackend to handle 0 byte reads than fix all IAsyncReadFileHandle implementations.
- If we detect a zero byte read in FPackageResourceIoBackend we can immediately pass it back to the IoDispatcher to handle the completed result.
- Note that if no user supplied memory can be found we will make a 0 byte allocation via FIoBuffer. This mirrors the behavour found in FFileIoStore.
- Even though we don't use the file handle, we still attempt to open one before counting the 0 byte read a success. This will cause attempts to make a 0 byte read of a missing file a failure etc. The file handle is then immediately closed.
- Since the zero byte reads will have no handles, there is no need for it to have an entry in FPendingRequests, so FPendingRequests::Remove needs to handle the case where the entry will not be found.

diff --git a/Engine/Source/Runtime/CoreUObject/Private/UObject/PackageResourceIoDispatcherBackend.cpp b/Engine/Source/Runtime/CoreUObject/Private/UObject/PackageResourceIoDispatcherBackend.cpp
--- a/Engine/Source/Runtime/CoreUObject/Private/UObject/PackageResourceIoDispatcherBackend.cpp
+++ b/Engine/Source/Runtime/CoreUObject/Private/UObject/PackageResourceIoDispatcherBackend.cpp
@@ -111,12 +111,13 @@ class FPackageResourceIoBackend final
 		{
 			FScopeLock _(&CriticalSection);
 			
-			FHandles& Handles = Lookup.FindChecked(Request);
-
-			Handles.RequestHandle->WaitCompletion();
-			Handles.RequestHandle.Reset();
+			if (FHandles* Handles = Lookup.Find(Request))
+			{
+				Handles->RequestHandle->WaitCompletion();
+				Handles->RequestHandle.Reset();
 
-			Lookup.Remove(Request);
+				Lookup.Remove(Request);
+			}
 		}
 
 		void Cancel(FIoRequestImpl* Request)
@@ -225,6 +226,21 @@ bool FPackageResourceIoBackend::Resolve(FIoRequestImpl* Request)
 		return false;
 	}
 	
+	if (Request->Options.GetSize() == 0)
+	{
+		void* UserSuppliedMemory = Request->Options.GetTargetVa();
+
+		FIoBuffer Buffer = UserSuppliedMemory	? FIoBuffer(FIoBuffer::Wrap, UserSuppliedMemory, 0)
+												: FIoBuffer(0);
+
+		Request->SetResult(Buffer);
+
+		CompletedRequests.Enqueue(Request);
+		BackendContext->WakeUpDispatcherThreadDelegate.Execute();
+
+		return true;
+	}
+
 	PendingRequests.Add(Request, MoveTemp(FileHandle), [this, Request](IAsyncReadFileHandle& FileHandle)
 	{
 		FAsyncFileCallBack Callback = [this, Request](bool bWasCancelled, IAsyncReadRequest* FileReadRequest)
