Pass the save overrides through to FExportsValidationContext

#rnx

diff --git a/Engine/Source/Runtime/CoreUObject/Private/UObject/SavePackage2.cpp b/Engine/Source/Runtime/CoreUObject/Private/UObject/SavePackage2.cpp
--- a/Engine/Source/Runtime/CoreUObject/Private/UObject/SavePackage2.cpp
+++ b/Engine/Source/Runtime/CoreUObject/Private/UObject/SavePackage2.cpp
@@ -629,7 +629,7 @@ ESavePackageResult ValidateExports(FSaveContext& SaveContext)
 		FOutputDevice* OutputDevice = SaveContext.IsGenerateSaveError() ? SaveContext.GetError() : nullptr;
 		for (const TFunction<FSavePackageSettings::ExternalExportValidationFunc>& ValidateExport : SaveContext.GetExternalExportValidations())
 		{
-			SaveContext.Result = ValidateExport({ SaveContext.GetPackage(), Exports, Flags, OutputDevice});
+			SaveContext.Result = ValidateExport({ SaveContext.GetPackage(), Exports, SaveContext.GetObjectSaveContext().SaveOverrides, Flags, OutputDevice});
 			if (SaveContext.Result != ESavePackageResult::Success)
 			{
 				return SaveContext.Result;
diff --git a/Engine/Source/Runtime/CoreUObject/Public/UObject/SavePackage.h b/Engine/Source/Runtime/CoreUObject/Public/UObject/SavePackage.h
--- a/Engine/Source/Runtime/CoreUObject/Public/UObject/SavePackage.h
+++ b/Engine/Source/Runtime/CoreUObject/Public/UObject/SavePackage.h
@@ -169,16 +169,18 @@ struct FExportsValidationContext
 		IsCooking	= 1 << 0,
 	};
 
-	FExportsValidationContext(const UPackage* InPackage, const TSet<UObject*>& InExports, EFlags InFlags,
-		FOutputDevice* InOutputDevice)
+	FExportsValidationContext(const UPackage* InPackage, const TSet<UObject*>& InExports, const TMap<UObject*, FObjectSaveOverride>& InSaveOverrides, 
+		EFlags InFlags, FOutputDevice* InOutputDevice)
 		: Package(InPackage)
 		, Exports(InExports)
+		, SaveOverrides(InSaveOverrides)
 		, Flags(InFlags)
 		, OutputDevice(InOutputDevice)
 	{}
 
 	const UPackage* Package;
 	const TSet<UObject*>& Exports;
+	const TMap<UObject*, FObjectSaveOverride>& SaveOverrides;
 	const EFlags Flags;
 	FOutputDevice* OutputDevice;
 };
