Making FCoreRedirects thread safe and removing some of the changes in CL 33988617 that were meant to account for FCoreRedirects not being thread safe. Performance appears to be no worse than prior to CL 33988617. Total time for the locks appears to be <100ms over the entire editor startup.

diff --git a/Engine/Source/Runtime/AssetRegistry/Private/AssetRegistry.cpp b/Engine/Source/Runtime/AssetRegistry/Private/AssetRegistry.cpp
--- a/Engine/Source/Runtime/AssetRegistry/Private/AssetRegistry.cpp
+++ b/Engine/Source/Runtime/AssetRegistry/Private/AssetRegistry.cpp
@@ -4614,7 +4614,6 @@ Impl::EGatherStatus FAssetRegistryImpl::TickGatherer(Impl::FTickContext& TickCon
 	FEventContext& EventContext = TickContext.EventContext;
 	FInterruptionContext& InOutInterruptionContext = TickContext.InterruptionContext;
 	bool bLocalIsInGameThread = IsInGameThread();
-	bool bCanAccessCoreRedirects = bLocalIsInGameThread || IsEngineStartupModuleLoadingComplete();
 
 	EGatherStatus OutStatus = EGatherStatus::Complete;
 	if (!GlobalGatherer.IsValid())
@@ -4758,7 +4757,7 @@ Impl::EGatherStatus FAssetRegistryImpl::TickGatherer(Impl::FTickContext& TickCon
 	}
 
 	auto RunAssetSearchDataGathered = [this, &EventContext, &TickContext, &LazyStartTimer,
-		&InOutInterruptionContext, bCanAccessCoreRedirects]
+		&InOutInterruptionContext]
 	(TMultiMap<FName, TUniquePtr<FAssetData>>& InAssetResults,
 		TMultiMap<FName, TUniquePtr<FAssetData>>& OutDeferredAssetResults)
 		{
@@ -4778,7 +4777,7 @@ Impl::EGatherStatus FAssetRegistryImpl::TickGatherer(Impl::FTickContext& TickCon
 				}
 
 				AssetSearchDataGathered(EventContext, InAssetResults, OutDeferredAssetResults,
-					InOutInterruptionContext, bCanAccessCoreRedirects);
+					InOutInterruptionContext);
 			}
 		};
 	auto RunDependencyDataGathered = [this, &bLocalIsInGameThread, &LazyStartTimer, &InOutInterruptionContext]
@@ -4853,7 +4852,7 @@ Impl::EGatherStatus FAssetRegistryImpl::TickGatherer(Impl::FTickContext& TickCon
 		// so we avoid spuriously reporting status == UnableToProgress
 		if (BackgroundResults.Assets.IsEmpty()
 			&& (!bLocalIsInGameThread || BackgroundResults.AssetsForGameThread.IsEmpty())
-			&& bCanAccessCoreRedirects && TickContext.bHandleDeferred)
+			&& TickContext.bHandleDeferred)
 		{
 			if (!DeferredAssets.IsEmpty() || !DeferredDependencies.IsEmpty() ||
 				(bLocalIsInGameThread && (!DeferredAssetsForGameThread.IsEmpty() || !DeferredDependenciesForGameThread.IsEmpty())))
@@ -5120,7 +5119,6 @@ void FAssetRegistryImpl::TickGatherPackage(Impl::FEventContext& EventContext, co
 		}
 	};
 
-	bool bCanAccessCoreRedirects = IsInGameThread() || IsEngineStartupModuleLoadingComplete();
 	FName PackageFName(PackageName);
 
 	// Gather results from the background search
@@ -5163,7 +5161,7 @@ void FAssetRegistryImpl::TickGatherPackage(Impl::FEventContext& EventContext, co
 		TMultiMap<FName, TUniquePtr<FAssetData>> DeferredPackageAssetsMap;
 		UE::AssetRegistry::Impl::FInterruptionContext InterruptionContext;
 		AssetSearchDataGathered(EventContext, PackageAssetsMap, DeferredPackageAssetsMap,
-			InterruptionContext, bCanAccessCoreRedirects);
+			InterruptionContext);
 		if (DeferredPackageAssetsMap.Num())
 		{
 			UE_LOG(LogAssetRegistry, Warning, TEXT("Attempted to add package '%s' to the registry before its UClass was available. \
@@ -5953,31 +5951,22 @@ void FAssetRegistryImpl::ScanPathsSynchronous(Impl::FScanPathContext& Context)
 				Iter.RemoveCurrent();
 			}
 		}
-		bool bCanAccessCoreRedirects = IsInGameThread() || IsEngineStartupModuleLoadingComplete();
 		// Force AssetSearchDataGathered to process these assets, skipping the PostLoadAssetRegistryTags if needed
-		// But don't allow this if !bCanAccessCoreRedirects. The caller should not be calling ScanPathsSynchronous
-		// in that case.
 		const bool bOldForceCompletionEvenIfPostLoadsFail = bForceCompletionEvenIfPostLoadsFail;
-		bForceCompletionEvenIfPostLoadsFail = bCanAccessCoreRedirects;
+		bForceCompletionEvenIfPostLoadsFail = true;
 
 		int32 OriginalNumDeferredAssetsForGameThread = DeferredAssetsForGameThread.Num();
+
 		// We don't call AssetsFoundCallback here because even for deferred assets it will already have been called.
 		// We pass DeferredAssetsForGameThread as the OutDeferred parameter, but we expect nothing will be deferred.
 		AssetSearchDataGathered(Context.EventContext, CollectedDeferredAssets, DeferredAssetsForGameThread,
-			TickContext.InterruptionContext, bCanAccessCoreRedirects);
+			TickContext.InterruptionContext);
 		// All of the assets we collected should have been processed or deferred.
 		ensure(CollectedDeferredAssets.Num() == 0);
-		// We should not have deferred any assets because we set bForceCompletionEvenIfPostLoadsFail. If we couldn't
-		// set it because the caller is calling to early, log an error about it.
-		if (DeferredAssetsForGameThread.Num() > OriginalNumDeferredAssetsForGameThread)
-		{
-			// With current logic the only way we could defer is if !bCanAccessCoreRedirects
-			check(!bCanAccessCoreRedirects);
-			ensureMsgf(false,
-				TEXT("ScanPathsSynchronous was called on a thread other than the game thread, before IsEngineStartupModuleLoadingComplete. ")
-				TEXT("This caused us not to be able to complete the scan because we can not access CoreRedirects. Some assets will be missing. ")
-				TEXT("Calling code should change to not allow the ScanPathsSynchronous call this early, or queue it to the gamethread."));
-		}
+
+		// We should not have deferred any new assets because we set bForceCompletionEvenIfPostLoadsFail=true
+		ensure(DeferredAssetsForGameThread.Num() <= OriginalNumDeferredAssetsForGameThread);
+
 		bForceCompletionEvenIfPostLoadsFail = bOldForceCompletionEvenIfPostLoadsFail;
 		// Tick to perform any subsequent processing required for these assets beyond AssetSearchDataGathered
 		Impl::FTickContext AssetTickContext(Context.EventContext, Context.InheritanceContext);
@@ -6059,7 +6048,7 @@ FAssetData* FAssetRegistryImpl::ResolveAssetIdCollision(FAssetData& A, FAssetDat
 	return Keep;
 }
 
-bool FAssetRegistryImpl::TryPostLoadAssetRegistryTags(FAssetData* AssetData, bool bCanAccessCoreRedirects)
+bool FAssetRegistryImpl::TryPostLoadAssetRegistryTags(FAssetData* AssetData)
 {
 	check(AssetData);
 	if (!AssetData->TagsAndValues.Num())
@@ -6070,51 +6059,44 @@ bool FAssetRegistryImpl::TryPostLoadAssetRegistryTags(FAssetData* AssetData, boo
 	bool bCouldPostLoadAssetRegistryTags = true;
 	UClass* AssetClass = nullptr;
 	FTopLevelAssetPath AssetClassPath;
-	if (!bCanAccessCoreRedirects)
-	{
-		bCouldPostLoadAssetRegistryTags = false;
-	}
-	else
-	{
-		AssetClassPath = AssetData->AssetClassPath;
-		AssetClass = FindObject<UClass>(AssetClassPath, true);
+	AssetClassPath = AssetData->AssetClassPath;
+	AssetClass = FindObject<UClass>(AssetClassPath, true);
 
-		while (!AssetClass)
+	while (!AssetClass)
+	{
+		// this is probably a blueprint that has not yet been loaded, try to find its native base class
+		const FTopLevelAssetPath* ParentClassPath = CachedBPInheritanceMap.Find(AssetClassPath);
+		if (ParentClassPath && !ParentClassPath->IsNull())
+		{
+			AssetClassPath = *ParentClassPath;
+			AssetClass = FindObject<UClass>(AssetClassPath, true);
+		}
+		else
 		{
-			// this is probably a blueprint that has not yet been loaded, try to find its native base class
-			const FTopLevelAssetPath* ParentClassPath = CachedBPInheritanceMap.Find(AssetClassPath);
-			if (ParentClassPath && !ParentClassPath->IsNull())
+			FTopLevelAssetPath LastAssetClassPath = AssetClassPath;
+			// Maybe it's a redirector
+			FSoftObjectPath RedirectedPath = GRedirectCollector.GetAssetPathRedirection(FSoftObjectPath(AssetClassPath));
+			if (RedirectedPath.IsValid())
 			{
-				AssetClassPath = *ParentClassPath;
-				AssetClass = FindObject<UClass>(AssetClassPath, true);
+				AssetClassPath = RedirectedPath.GetAssetPath();
 			}
 			else
 			{
-				FTopLevelAssetPath LastAssetClassPath = AssetClassPath;
-				// Maybe it's a redirector
-				FSoftObjectPath RedirectedPath = GRedirectCollector.GetAssetPathRedirection(FSoftObjectPath(AssetClassPath));
-				if (RedirectedPath.IsValid())
+				FCoreRedirectObjectName NewName = FCoreRedirects::GetRedirectedName(ECoreRedirectFlags::Type_Class, FCoreRedirectObjectName(AssetClassPath));
+				if (NewName.IsValid())
 				{
-					AssetClassPath = RedirectedPath.GetAssetPath();
-				}
-				else
-				{
-					FCoreRedirectObjectName NewName = FCoreRedirects::GetRedirectedName(ECoreRedirectFlags::Type_Class, FCoreRedirectObjectName(AssetClassPath));
-					if (NewName.IsValid())
-					{
-						AssetClassPath = NewName.ToString();
-					}
+					AssetClassPath = NewName.ToString();
 				}
+			}
 
-				if (AssetClassPath != LastAssetClassPath && !AssetClassPath.IsNull())
-				{
-					AssetClass = FindObject<UClass>(AssetClassPath, true);
-				}
-				else
-				{
-					bCouldPostLoadAssetRegistryTags = false;
-					break;
-				}
+			if (AssetClassPath != LastAssetClassPath && !AssetClassPath.IsNull())
+			{
+				AssetClass = FindObject<UClass>(AssetClassPath, true);
+			}
+			else
+			{
+				bCouldPostLoadAssetRegistryTags = false;
+				break;
 			}
 		}
 	}
@@ -6134,8 +6116,6 @@ bool FAssetRegistryImpl::TryPostLoadAssetRegistryTags(FAssetData* AssetData, boo
 		// Okay, we think we're done loading and now we need to make some expensive final checks to try to either
 		// track down the classes for fixup or just give up
 		bMakeFinalChecks = true;
-		// bCanAccessCoreRedirects should be true when IsEngineStartupModuleLoadingComplete
-		check(bCanAccessCoreRedirects);
 	}
 	if (!AssetClass && bForceCompletionEvenIfPostLoadsFail)
 	{
@@ -6324,7 +6304,7 @@ bool FAssetRegistryImpl::ShouldSkipGatheredAsset(FAssetData& AssetData)
 void FAssetRegistryImpl::AssetSearchDataGathered(Impl::FEventContext& EventContext,
 	TMultiMap<FName, TUniquePtr<FAssetData>>& AssetResults,
 	TMultiMap<FName, TUniquePtr<FAssetData>>& OutDeferredAssetResults,
-	Impl::FInterruptionContext& InOutInterruptionContext, bool bCanAccessCoreRedirects)
+	Impl::FInterruptionContext& InOutInterruptionContext)
 {
 	TRACE_CPUPROFILER_EVENT_SCOPE(AssetSearchDataGathered);
 
@@ -6393,7 +6373,7 @@ void FAssetRegistryImpl::AssetSearchDataGathered(Impl::FEventContext& EventConte
 
 #if WITH_EDITOR
 		// Postload assets based on their declared class. Queue them for for later retry if their class has not yet loaded.
-		bool CouldPostLoad = TryPostLoadAssetRegistryTags(BackgroundResult.Get(), bCanAccessCoreRedirects);
+		bool CouldPostLoad = TryPostLoadAssetRegistryTags(BackgroundResult.Get());
 		if (!CouldPostLoad)
 		{
 			OutDeferredAssetResults.Add(BackgroundAssetPackageName, MoveTemp(BackgroundResult));
diff --git a/Engine/Source/Runtime/AssetRegistry/Private/AssetRegistryImpl.h b/Engine/Source/Runtime/AssetRegistry/Private/AssetRegistryImpl.h
--- a/Engine/Source/Runtime/AssetRegistry/Private/AssetRegistryImpl.h
+++ b/Engine/Source/Runtime/AssetRegistry/Private/AssetRegistryImpl.h
@@ -428,7 +428,7 @@ private:
 	void AssetSearchDataGathered(Impl::FEventContext& EventContext,
 		TMultiMap<FName, TUniquePtr<FAssetData>>& AssetResults, 
 		TMultiMap<FName, TUniquePtr<FAssetData>>& OutDeferredResults,\
-		Impl::FInterruptionContext& InOutInterruptionContext, bool bCanAccessCoreRedirects);
+		Impl::FInterruptionContext& InOutInterruptionContext);
 	/** Validate assets gathered from disk before adding them to the AssetRegistry. */
 	bool ShouldSkipGatheredAsset(FAssetData& AssetData);
 
@@ -471,7 +471,7 @@ private:
 	 * @param AssetData Existing asset data
 	 * @return Returns false if the required parent UClass is not yet available
 	 */
-	bool TryPostLoadAssetRegistryTags(FAssetData* AssetData, bool bCanAccessCoreRedirects);
+	bool TryPostLoadAssetRegistryTags(FAssetData* AssetData);
 
 	/** Update Redirect collector with redirects loaded from asset registry */
 	void UpdateRedirectCollector();
diff --git a/Engine/Source/Runtime/CoreUObject/Private/UObject/CoreRedirects.cpp b/Engine/Source/Runtime/CoreUObject/Private/UObject/CoreRedirects.cpp
--- a/Engine/Source/Runtime/CoreUObject/Private/UObject/CoreRedirects.cpp
+++ b/Engine/Source/Runtime/CoreUObject/Private/UObject/CoreRedirects.cpp
@@ -1081,16 +1081,65 @@ int FCoreRedirectObjectName::Compare(const FCoreRedirectObjectName& Other) const
 	return 0;
 }
 
-bool FCoreRedirects::bInitialized = false;
-bool FCoreRedirects::bInDebugMode = false;
+std::atomic<bool> FCoreRedirects::bInitialized = false;
+std::atomic<bool> FCoreRedirects::bInDebugMode = false;
 bool FCoreRedirects::bValidatedOnce = false;
-#if WITH_COREREDIRECTS_MULTITHREAD_WARNING
-bool FCoreRedirects::bIsInMultithreadedPhase = false;
-#endif
 
 TMap<FName, ECoreRedirectFlags> FCoreRedirects::ConfigKeyMap;
 FCoreRedirects::FRedirectTypeMap FCoreRedirects::RedirectTypeMap;
-FRWLock FCoreRedirects::KnownMissingLock;
+FCoreRedirects::FRWLockWithExclusiveRecursion FCoreRedirects::RWLock;
+
+namespace UE::CoreRedirects::Private
+{
+class FRWWithExclusiveRecursionScopeLockForRead
+{
+public:
+	explicit FRWWithExclusiveRecursionScopeLockForRead(FCoreRedirects::FRWLockWithExclusiveRecursion& InLock)
+	: InternalLock(InLock)
+	, NeedsUnlock(true)
+	{
+		InternalLock.ReadLock();
+	}
+
+	~FRWWithExclusiveRecursionScopeLockForRead()
+	{
+		if (NeedsUnlock)
+		{
+			InternalLock.ReadUnlock();
+		}
+	}
+protected:
+
+	enum class EInitFlag
+	{
+		InitFlag
+	};
+
+	FRWWithExclusiveRecursionScopeLockForRead(FCoreRedirects::FRWLockWithExclusiveRecursion& InLock, EInitFlag Unused)
+	: InternalLock(InLock)
+	, NeedsUnlock(false)
+	{}
+
+	FCoreRedirects::FRWLockWithExclusiveRecursion& InternalLock;
+	bool NeedsUnlock;
+};
+
+class FRWWithExclusiveRecursionScopeLockForWrite :  public FRWWithExclusiveRecursionScopeLockForRead
+{
+public:
+	explicit FRWWithExclusiveRecursionScopeLockForWrite(FCoreRedirects::FRWLockWithExclusiveRecursion& InLock)
+	: FRWWithExclusiveRecursionScopeLockForRead(InLock, EInitFlag::InitFlag)
+	{
+		InternalLock.WriteLock();
+	}
+
+	~FRWWithExclusiveRecursionScopeLockForWrite()
+	{
+		InternalLock.WriteUnlock();
+	}
+};
+
+}
 
 FCoreRedirects::FRedirectNameMap& FCoreRedirects::FRedirectTypeMap::FindOrAdd(ECoreRedirectFlags Key)
 {
@@ -1138,69 +1187,66 @@ void FCoreRedirects::FRedirectTypeMap::Empty()
 
 void FCoreRedirects::Initialize()
 {
-	if (bInitialized)
+	ensureMsgf(IsInGameThread(), TEXT("FCoreRedirects can only be initialized on the game thread."));
+	if (bInitialized.load(std::memory_order_relaxed))
 	{
 		return;
 	}
-	bInitialized = true;
-
-#if !UE_BUILD_SHIPPING && !NO_LOGGING
-	if (FParse::Param(FCommandLine::Get(), TEXT("FullDebugCoreRedirects")))
-	{
-		// Enable debug mode and set to maximum verbosity
-		bInDebugMode = true;
-		LogCoreRedirects.SetVerbosity(ELogVerbosity::VeryVerbose);
-		FCoreDelegates::OnFEngineLoopInitComplete.AddStatic(ValidateAllRedirects);
-	}
-	else if (FParse::Param(FCommandLine::Get(), TEXT("DebugCoreRedirects")))
+	
 	{
-		// Enable debug mode and increase log levels but don't show every message
-		bInDebugMode = true;
-		LogCoreRedirects.SetVerbosity(ELogVerbosity::Verbose);
-		FCoreDelegates::OnFEngineLoopInitComplete.AddStatic(ValidateAllRedirects);
-	}
-#endif
+		FCoreRedirectorScopeLockForWrite ScopeLock(RWLock);
 
-#if WITH_COREREDIRECTS_MULTITHREAD_WARNING
-	// Setting IsInMultithreadedPhase has to occur after LoadModule("AssetRegistry") (which can write to FCoreRedirects) and
-	// before the first package is queued onto the async loading thread (which reads from FCoreRedirects multithreaded). OnPostEngineInit is in that window.
-	FCoreDelegates::OnPostEngineInit.AddStatic(EnterMultithreadedPhase);
+#if !UE_BUILD_SHIPPING
+		if (FParse::Param(FCommandLine::Get(), TEXT("FullDebugCoreRedirects")))
+		{
+			// Enable debug mode and set to maximum verbosity
+			bInDebugMode.store(true, std::memory_order_relaxed);
+			UE_SET_LOG_VERBOSITY(LogCoreRedirects, VeryVerbose);
+			FCoreDelegates::OnFEngineLoopInitComplete.AddStatic(ValidateAllRedirects);
+		}
+		else if (FParse::Param(FCommandLine::Get(), TEXT("DebugCoreRedirects")))
+		{
+			// Enable debug mode and increase log levels but don't show every message
+			bInDebugMode.store(true, std::memory_order_relaxed);
+			UE_SET_LOG_VERBOSITY(LogCoreRedirects, Verbose);
+			FCoreDelegates::OnFEngineLoopInitComplete.AddStatic(ValidateAllRedirects);
+		}
 #endif
 
-	// Setup map
-	ConfigKeyMap.Add(TEXT("ObjectRedirects"), ECoreRedirectFlags::Type_Object);
-	ConfigKeyMap.Add(TEXT("ClassRedirects"), ECoreRedirectFlags::Type_Class);
-	ConfigKeyMap.Add(TEXT("StructRedirects"), ECoreRedirectFlags::Type_Struct);
-	ConfigKeyMap.Add(TEXT("EnumRedirects"), ECoreRedirectFlags::Type_Enum);
-	ConfigKeyMap.Add(TEXT("FunctionRedirects"), ECoreRedirectFlags::Type_Function);
-	ConfigKeyMap.Add(TEXT("PropertyRedirects"), ECoreRedirectFlags::Type_Property);
-	ConfigKeyMap.Add(TEXT("PackageRedirects"), ECoreRedirectFlags::Type_Package);
-
-	RegisterNativeRedirects();
-
-	// Prepopulate RedirectTypeMap entries that some threads write to after the engine goes multi-threaded.
-	// Most RedirectTypeMap entries are written to only from InitUObject's call to ReadRedirectsFromIni, and at that point the Engine is single-threaded.
-	// Known missing packages and plugin loads can add entries to existing lists but will not add brand new types.
-	// Taking advantage of this, we treat the list of Key/Value pairs of RedirectTypeMap as immutable and read from it without synchronization.
-	// Note that the values for those written-during-multithreading entries need to be synchronized; it is only the list of Key/Value pairs that is immutable.
-#if WITH_COREREDIRECTS_MULTITHREAD_WARNING
-	check(!bIsInMultithreadedPhase);
-#endif
-	RedirectTypeMap.FindOrAdd(ECoreRedirectFlags::Type_Package | ECoreRedirectFlags::Category_Removed | ECoreRedirectFlags::Option_MissingLoad);
+		// Setup map
+		ConfigKeyMap.Add(TEXT("ObjectRedirects"), ECoreRedirectFlags::Type_Object);
+		ConfigKeyMap.Add(TEXT("ClassRedirects"), ECoreRedirectFlags::Type_Class);
+		ConfigKeyMap.Add(TEXT("StructRedirects"), ECoreRedirectFlags::Type_Struct);
+		ConfigKeyMap.Add(TEXT("EnumRedirects"), ECoreRedirectFlags::Type_Enum);
+		ConfigKeyMap.Add(TEXT("FunctionRedirects"), ECoreRedirectFlags::Type_Function);
+		ConfigKeyMap.Add(TEXT("PropertyRedirects"), ECoreRedirectFlags::Type_Property);
+		ConfigKeyMap.Add(TEXT("PackageRedirects"), ECoreRedirectFlags::Type_Package);
 
+		RegisterNativeRedirectsUnderWriteLock(ScopeLock);
+
+		// Prepopulate RedirectTypeMap entries that some threads write to after the engine goes multi-threaded.
+		// Most RedirectTypeMap entries are written to only from InitUObject's call to ReadRedirectsFromIni, and at that point the Engine is single-threaded.
+		// Known missing packages and plugin loads can add entries to existing lists but will not add brand new types.
+		// Taking advantage of this, we treat the list of Key/Value pairs of RedirectTypeMap as immutable and read from it without synchronization.
+		// Note that the values for those written-during-multithreading entries need to be synchronized; it is only the list of Key/Value pairs that is immutable.
+		RedirectTypeMap.FindOrAdd(ECoreRedirectFlags::Type_Package | ECoreRedirectFlags::Category_Removed | ECoreRedirectFlags::Option_MissingLoad);
+
+		bInitialized.store(true, std::memory_order_release);
+	}
 	// Enable to run startup tests
 	//RunTests();
 }
 
-#if WITH_COREREDIRECTS_MULTITHREAD_WARNING
-void FCoreRedirects::EnterMultithreadedPhase()
+bool FCoreRedirects::RedirectNameAndValues(ECoreRedirectFlags Type, const FCoreRedirectObjectName& OldObjectName,
+	FCoreRedirectObjectName& NewObjectName, const FCoreRedirect** FoundValueRedirect, ECoreRedirectMatchFlags MatchFlags)
 {
-	bIsInMultithreadedPhase = true;
+	FCoreRedirectorScopeLockForRead ScopeLock(RWLock);
+	return RedirectNameAndValuesUnderReadLock(Type, OldObjectName, NewObjectName, FoundValueRedirect, MatchFlags, ScopeLock);
 }
-#endif
 
-bool FCoreRedirects::RedirectNameAndValues(ECoreRedirectFlags Type, const FCoreRedirectObjectName& OldObjectName,
-	FCoreRedirectObjectName& NewObjectName, const FCoreRedirect** FoundValueRedirect, ECoreRedirectMatchFlags MatchFlags)
+bool FCoreRedirects::RedirectNameAndValuesUnderReadLock(ECoreRedirectFlags Type, const FCoreRedirectObjectName& OldObjectName,
+	FCoreRedirectObjectName& NewObjectName, const FCoreRedirect** FoundValueRedirect, ECoreRedirectMatchFlags MatchFlags,
+	const FCoreRedirectorScopeLockForRead& HeldLock)
 {
 	auto ProcessRedirect = [&FoundValueRedirect, OldObjectName](const FCoreRedirect* Redirect, const FCoreRedirectObjectName& NewObjectName)
 	{
@@ -1226,7 +1272,7 @@ bool FCoreRedirects::RedirectNameAndValues(ECoreRedirectFlags Type, const FCoreR
 
 	NewObjectName = OldObjectName;
 	TArray<const FCoreRedirect*> FoundRedirects;
-	if (GetMatchingRedirects(Type, OldObjectName, FoundRedirects, MatchFlags))
+	if (GetMatchingRedirectsUnderReadLock(Type, OldObjectName, FoundRedirects, MatchFlags, HeldLock))
 	{
 		if (FoundRedirects.Num() > 1)
 		{
@@ -1256,7 +1302,7 @@ bool FCoreRedirects::RedirectNameAndValues(ECoreRedirectFlags Type, const FCoreR
 	}
 
 	const bool bDidRedirect = NewObjectName != OldObjectName;
-	UE_CLOG(bInDebugMode && bDidRedirect, LogCoreRedirects, Verbose,
+	UE_CLOG(IsInDebugMode() && bDidRedirect, LogCoreRedirects, Verbose,
 		TEXT("RedirectNameAndValues(%s) replaced by %s"), *OldObjectName.ToString(), *NewObjectName.ToString());
 	return bDidRedirect;
 }
@@ -1264,9 +1310,10 @@ bool FCoreRedirects::RedirectNameAndValues(ECoreRedirectFlags Type, const FCoreR
 FCoreRedirectObjectName FCoreRedirects::GetRedirectedName(ECoreRedirectFlags Type,
 	const FCoreRedirectObjectName& OldObjectName, ECoreRedirectMatchFlags MatchFlags)
 {
+	FCoreRedirectorScopeLockForRead ScopeLock(RWLock);
 	FCoreRedirectObjectName NewObjectName;
 
-	RedirectNameAndValues(Type, OldObjectName, NewObjectName, nullptr, MatchFlags);
+	RedirectNameAndValuesUnderReadLock(Type, OldObjectName, NewObjectName, nullptr, MatchFlags, ScopeLock);
 
 	return NewObjectName;
 }
@@ -1274,14 +1321,15 @@ FCoreRedirectObjectName FCoreRedirects::GetRedirectedName(ECoreRedirectFlags Typ
 const TMap<FString, FString>* FCoreRedirects::GetValueRedirects(ECoreRedirectFlags Type,
 	const FCoreRedirectObjectName& OldObjectName, ECoreRedirectMatchFlags MatchFlags)
 {
+	FCoreRedirectorScopeLockForRead ScopeLock(RWLock);
 	FCoreRedirectObjectName NewObjectName;
 	const FCoreRedirect* FoundRedirect = nullptr;
 
-	RedirectNameAndValues(Type, OldObjectName, NewObjectName, &FoundRedirect, MatchFlags);
+	RedirectNameAndValuesUnderReadLock(Type, OldObjectName, NewObjectName, &FoundRedirect, MatchFlags, ScopeLock);
 
 	if (FoundRedirect && FoundRedirect->ValueChanges.Num() > 0)
 	{
-		UE_CLOG(bInDebugMode, LogCoreRedirects, VeryVerbose, TEXT("GetValueRedirects found %d matches for %s"),
+		UE_CLOG(IsInDebugMode() , LogCoreRedirects, VeryVerbose, TEXT("GetValueRedirects found %d matches for %s"),
 			FoundRedirect->ValueChanges.Num(), *OldObjectName.ToString());
 
 		return &FoundRedirect->ValueChanges;
@@ -1292,6 +1340,13 @@ const TMap<FString, FString>* FCoreRedirects::GetValueRedirects(ECoreRedirectFla
 
 bool FCoreRedirects::GetMatchingRedirects(ECoreRedirectFlags SearchFlags, const FCoreRedirectObjectName& OldObjectName,
 	TArray<const FCoreRedirect*>& FoundRedirects, ECoreRedirectMatchFlags MatchFlags)
+{
+	FCoreRedirectorScopeLockForRead ScopeLock(RWLock);
+	return GetMatchingRedirectsUnderReadLock(SearchFlags, OldObjectName, FoundRedirects, MatchFlags, ScopeLock);
+}
+
+bool FCoreRedirects::GetMatchingRedirectsUnderReadLock(ECoreRedirectFlags SearchFlags, const FCoreRedirectObjectName& OldObjectName,
+	TArray<const FCoreRedirect*>& FoundRedirects, ECoreRedirectMatchFlags MatchFlags, const FCoreRedirectorScopeLockForRead& HeldLock)
 {
 	// Look for all redirects that match the given names and flags
 	bool bFound = false;
@@ -1340,6 +1395,8 @@ bool FCoreRedirects::GetMatchingRedirects(ECoreRedirectFlags SearchFlags, const
 
 bool FCoreRedirects::FindPreviousNames(ECoreRedirectFlags SearchFlags, const FCoreRedirectObjectName& NewObjectName, TArray<FCoreRedirectObjectName>& PreviousNames)
 {
+	FCoreRedirectorScopeLockForRead ScopeLock(RWLock);
+
 	// Look for reverse direction redirects
 	bool bFound = false;
 
@@ -1419,7 +1476,7 @@ bool FCoreRedirects::FindPreviousNames(ECoreRedirectFlags SearchFlags, const FCo
 		}
 	}
 
-	UE_CLOG(bFound && bInDebugMode, LogCoreRedirects, VeryVerbose, TEXT("FindPreviousNames found %d previous names for %s"), PreviousNames.Num(), *NewObjectName.ToString());
+	UE_CLOG(bFound && IsInDebugMode(), LogCoreRedirects, VeryVerbose, TEXT("FindPreviousNames found %d previous names for %s"), PreviousNames.Num(), *NewObjectName.ToString());
 
 	return bFound;
 }
@@ -1428,8 +1485,8 @@ bool FCoreRedirects::IsKnownMissing(ECoreRedirectFlags Type, const FCoreRedirect
 {
 	TArray<const FCoreRedirect*> FoundRedirects;
 
-	FRWScopeLock ScopeLock(KnownMissingLock, FRWScopeLockType::SLT_ReadOnly);
-	return FCoreRedirects::GetMatchingRedirects(Type | ECoreRedirectFlags::Category_Removed, ObjectName, FoundRedirects);
+	FCoreRedirectorScopeLockForRead ScopeLock(RWLock);
+	return FCoreRedirects::GetMatchingRedirectsUnderReadLock(Type | ECoreRedirectFlags::Category_Removed, ObjectName, FoundRedirects, ECoreRedirectMatchFlags::None, ScopeLock);
 }
 
 bool FCoreRedirects::AddKnownMissing(ECoreRedirectFlags Type, const FCoreRedirectObjectName& ObjectName, ECoreRedirectFlags Channel)
@@ -1439,7 +1496,6 @@ bool FCoreRedirects::AddKnownMissing(ECoreRedirectFlags Type, const FCoreRedirec
 	check((Channel & ~ECoreRedirectFlags::Option_MissingLoad) == ECoreRedirectFlags::None);
 	FCoreRedirect NewRedirect(Type | ECoreRedirectFlags::Category_Removed | Channel, ObjectName, FCoreRedirectObjectName());
 
-	FRWScopeLock ScopeLock(KnownMissingLock, FRWScopeLockType::SLT_Write);
 	return AddRedirectList(TArrayView<const FCoreRedirect>(&NewRedirect, 1), TEXT("AddKnownMissing"));
 }
 
@@ -1448,7 +1504,6 @@ bool FCoreRedirects::RemoveKnownMissing(ECoreRedirectFlags Type, const FCoreRedi
 	check((Channel & ~ECoreRedirectFlags::Option_MissingLoad) == ECoreRedirectFlags::None);
 	FCoreRedirect RedirectToRemove(Type | ECoreRedirectFlags::Category_Removed | Channel, ObjectName, FCoreRedirectObjectName());
 
-	FRWScopeLock ScopeLock(KnownMissingLock, FRWScopeLockType::SLT_Write);
 	return RemoveRedirectList(TArrayView<const FCoreRedirect>(&RedirectToRemove, 1), TEXT("RemoveKnownMissing"));
 }
 
@@ -1457,7 +1512,7 @@ void FCoreRedirects::ClearKnownMissing(ECoreRedirectFlags Type, ECoreRedirectFla
 	check((Channel & ~ECoreRedirectFlags::Option_MissingLoad) == ECoreRedirectFlags::None);
 	ECoreRedirectFlags RedirectFlags = Type | ECoreRedirectFlags::Category_Removed | Channel;
 
-	FRWScopeLock ScopeLock(KnownMissingLock, FRWScopeLockType::SLT_Write);
+	FCoreRedirectorScopeLockForWrite ScopeLock(RWLock);
 	FRedirectNameMap* RedirectNameMap = RedirectTypeMap.Find(RedirectFlags);
 	if (RedirectNameMap)
 	{
@@ -1482,10 +1537,6 @@ bool FCoreRedirects::RunTests()
 	bool bSuccess = true;
 
 	FRedirectTypeMap BackupMap = MoveTemp(RedirectTypeMap);
-#if WITH_COREREDIRECTS_MULTITHREAD_WARNING
-	bool BackupIsInMultithreadedPhase = bIsInMultithreadedPhase;
-	bIsInMultithreadedPhase = false;
-#endif
 	RedirectTypeMap.Empty();
 #if WITH_EDITOR
 	FRedirectionSummary BackupSummary = MoveTemp(GRedirectionSummary);
@@ -1684,9 +1735,6 @@ bool FCoreRedirects::RunTests()
 #if WITH_EDITOR
 	GRedirectionSummary = MoveTemp(BackupSummary);
 #endif
-#if WITH_COREREDIRECTS_MULTITHREAD_WARNING
-	bIsInMultithreadedPhase = BackupIsInMultithreadedPhase;
-#endif
 
 	UE_LOG(LogCoreRedirects, Log, TEXT("FCoreRedirect Test %s!"), (bSuccess ? TEXT("Passed") : TEXT("Failed")));
 	return bSuccess;
@@ -1700,6 +1748,7 @@ bool FCoreRedirectTest::RunTest(const FString& Parameters)
 
 bool FCoreRedirects::ReadRedirectsFromIni(const FString& IniName)
 {
+	FCoreRedirectorScopeLockForWrite ScopeLock(RWLock);
 	Initialize();
 
 	if (GConfig)
@@ -1846,11 +1895,17 @@ bool FCoreRedirects::ReadRedirectsFromIni(const FString& IniName)
 
 bool FCoreRedirects::AddRedirectList(TArrayView<const FCoreRedirect> Redirects, const FString& SourceString)
 {
+	FCoreRedirectorScopeLockForWrite ScopeLock(RWLock);
 	Initialize();
+	return AddRedirectListUnderWriteLock(Redirects, SourceString, ScopeLock);
+}
 
+bool FCoreRedirects::AddRedirectListUnderWriteLock(TArrayView<const FCoreRedirect> Redirects, const FString & SourceString, 
+	const FCoreRedirectorScopeLockForWrite& HeldLock)
+{
 	UE_LOG(LogCoreRedirects, Verbose, TEXT("AddRedirect(%s) adding %d redirects"), *SourceString, Redirects.Num());
 
-	if (bInDebugMode && bValidatedOnce)
+	if (IsInDebugMode() && bValidatedOnce)
 	{
 		// Validate on apply because we finished our initial validation pass
 		ValidateRedirectList(Redirects, SourceString);
@@ -1877,7 +1932,7 @@ bool FCoreRedirects::AddRedirectList(TArrayView<const FCoreRedirect> Redirects,
 			UE_LOG(LogCoreRedirects, Log, TEXT("AddRedirect(%s) has wildcard redirect %s, these are very slow and should be resolved as soon as possible! Please refer to the documentation in Engine/Config/BaseEngine.ini."), *SourceString, *NewRedirect.OldName.ToString());
 		}
 		
-		if (AddSingleRedirect(NewRedirect, SourceString))
+		if (AddSingleRedirectUnderWriteLock(NewRedirect, SourceString, HeldLock))
 		{
 			bAddedAny = true;
 
@@ -1887,7 +1942,7 @@ bool FCoreRedirects::AddRedirectList(TArrayView<const FCoreRedirect> Redirects,
 				FCoreRedirect ValueRedirect = NewRedirect;
 				ValueRedirect.OldName = ValueRedirect.NewName;
 
-				AddSingleRedirect(ValueRedirect, SourceString);
+				AddSingleRedirectUnderWriteLock(ValueRedirect, SourceString, HeldLock);
 			}
 		}
 	}
@@ -1895,22 +1950,13 @@ bool FCoreRedirects::AddRedirectList(TArrayView<const FCoreRedirect> Redirects,
 	return bAddedAny;
 }
 
-bool FCoreRedirects::AddSingleRedirect(const FCoreRedirect& NewRedirect, const FString& SourceString)
+bool FCoreRedirects::AddSingleRedirectUnderWriteLock(const FCoreRedirect& NewRedirect, const FString& SourceString, 
+	const FCoreRedirectorScopeLockForWrite& HeldLock)
 {
 	const bool bIsWildcardMatch = NewRedirect.IsWildcardMatch();
 	FRedirectNameMap* ExistingNameMap;
 
-#if WITH_COREREDIRECTS_MULTITHREAD_WARNING
-	if (bIsInMultithreadedPhase)
-	{
-		ExistingNameMap = RedirectTypeMap.Find(NewRedirect.RedirectFlags);
-		checkf(ExistingNameMap, TEXT("Once EnterMultithreadedPhase has been called, it is no longer valid to add redirects of new types."));
-	}
-	else
-#endif
-	{
-		ExistingNameMap = &RedirectTypeMap.FindOrAdd(NewRedirect.RedirectFlags);
-	}
+	ExistingNameMap = &RedirectTypeMap.FindOrAdd(NewRedirect.RedirectFlags);
 
 	TArray<FCoreRedirect>* ExistingRedirects = nullptr;
 	if (bIsWildcardMatch)
@@ -2004,6 +2050,7 @@ bool FCoreRedirects::AddSingleRedirect(const FCoreRedirect& NewRedirect, const F
 
 bool FCoreRedirects::RemoveRedirectList(TArrayView<const FCoreRedirect> Redirects, const FString& SourceString)
 {
+	FCoreRedirectorScopeLockForWrite ScopeLock(RWLock);
 	UE_LOG(LogCoreRedirects, Verbose, TEXT("RemoveRedirect(%s) Removing %d redirects"), *SourceString, Redirects.Num());
 
 	bool bRemovedAny = false;
@@ -2038,13 +2085,26 @@ bool FCoreRedirects::RemoveRedirectList(TArrayView<const FCoreRedirect> Redirect
 			UE_LOG(LogCoreRedirects, Log, TEXT("RemoveRedirect(%s) has wildcard redirect %s, these are very slow and should be resolved as soon as possible! Please refer to the documentation in Engine/Config/BaseEngine.ini."), *SourceString, *RedirectToRemove.OldName.ToString());
 		}
 
-		bRemovedAny |= RemoveSingleRedirect(RedirectToRemove, SourceString);
+		bRemovedAny |= RemoveSingleRedirectUnderWriteLock(RedirectToRemove, SourceString, ScopeLock);
 	}
 
 	return bRemovedAny;
 }
 
-bool FCoreRedirects::RemoveSingleRedirect(const FCoreRedirect& RedirectToRemove, const FString& SourceString)
+bool FCoreRedirects::IsInitialized()
+{
+	// Use an atomic variable rather than take a read lock because we might already be under a read lock
+	return bInitialized.load(std::memory_order_acquire);
+}
+
+bool FCoreRedirects::IsInDebugMode()
+{
+	// Use an atomic variable rather than take a read lock because we might already be under a read lock
+	return bInDebugMode.load(std::memory_order_relaxed);
+}
+
+bool FCoreRedirects::RemoveSingleRedirectUnderWriteLock(const FCoreRedirect& RedirectToRemove, const FString& SourceString,
+	const FCoreRedirectorScopeLockForWrite& HeldLock)
 {
 	const bool bIsWildcardMatch = RedirectToRemove.IsWildcardMatch();
 	FRedirectNameMap* ExistingNameMap = RedirectTypeMap.Find(RedirectToRemove.RedirectFlags);
@@ -2205,6 +2265,7 @@ void FCoreRedirects::ValidateRedirectList(TArrayView<const FCoreRedirect> Redire
 
 void FCoreRedirects::ValidateAllRedirects()
 {
+	FCoreRedirectorScopeLockForWrite ReadLock(RWLock);
 	bValidatedOnce = true;
 
 	// Validate all existing redirects
@@ -2221,6 +2282,13 @@ void FCoreRedirects::ValidateAllRedirects()
 	}
 }
 
+const TMap<FName, ECoreRedirectFlags>& FCoreRedirects::GetConfigKeyMap()
+{
+	// The config key map is only written to during initialization. That allows us to provide unguarded access after initialization
+	ensureMsgf(IsInitialized(), TEXT("It is not legal to read the config key map until after FCoreRedirects has been initialized."));
+	return ConfigKeyMap;
+}
+
 ECoreRedirectFlags FCoreRedirects::GetFlagsForTypeName(FName PackageName, FName TypeName)
 {
 	if (PackageName == GLongCoreUObjectPackageName)
@@ -3018,7 +3086,7 @@ static void RegisterNativeRedirects49(TArray<FCoreRedirect>& Redirects)
 
 UE_ENABLE_OPTIMIZATION_SHIP
 
-void FCoreRedirects::RegisterNativeRedirects()
+void FCoreRedirects::RegisterNativeRedirectsUnderWriteLock(const FCoreRedirectorScopeLockForWrite& HeldLock)
 {
 	// Registering redirects here instead of in baseengine.ini is faster to parse and can clean up the ini, but is not required
 	TArray<FCoreRedirect> Redirects;
@@ -3029,10 +3097,80 @@ void FCoreRedirects::RegisterNativeRedirects()
 
 	// 4.10 and later are in baseengine.ini
 
-	AddRedirectList(Redirects, TEXT("RegisterNativeRedirects"));
+	AddRedirectListUnderWriteLock(Redirects, TEXT("RegisterNativeRedirects"), HeldLock);
 }
 #else
-void FCoreRedirects::RegisterNativeRedirects()
+void FCoreRedirects::RegisterNativeRedirectsUnderWriteLock(const FCoreRedirectorScopeLock& HeldLock)
 {
 }
-#endif // UE_WITH_CORE_REDIRECTS
\ No newline at end of file
+#endif // UE_WITH_CORE_REDIRECTS
+
+void FCoreRedirects::FRWLockWithExclusiveRecursion::ReadLock()
+{
+	int32 InitialLockOwner = WriteLockOwnerThreadId.load(std::memory_order_relaxed);
+	
+	// Avoid calling GetCurrentThreadId in the common case where we start out unlocked
+	// If we are unowned at the start, then there are three cases to consider:
+	// 1. We are unlocked. In that case we'll pass the InitialLockOwner check and call ReadLock()
+	// 2. We are already in a ReadLock() state. Taking the read lock again will cause us to deadlock. 
+	// That is part of the contract for this lock so it's okay.
+	// 3. We are write locked, in which case we'll enter the first block 
+	if (InitialLockOwner != 0)
+	{
+		// If we are initially write locked, there are two possibilities:
+		// 1. This thread owns the write lock. In that case, it's safe to just bump the recursion count
+		// 2. Another thread owns the write lock, in which case we want to block with a ReadLock()
+		if (InitialLockOwner == FPlatformTLS::GetCurrentThreadId())
+		{
+			// Allow recursive read locking by bumping the recursion count
+			RecursionCount++;
+		}
+		else
+		{
+			InternalLock.ReadLock();
+		}
+	}
+	else
+	{
+		// We may or may not still be locked at this point if another thread took a WriteLock
+		// but it doesn't matter because in either case we want to take a ReadLock and potentially wait
+		InternalLock.ReadLock();
+	}
+}
+
+void FCoreRedirects::FRWLockWithExclusiveRecursion::WriteLock()
+{
+	uint32 CurrentThreadId = FPlatformTLS::GetCurrentThreadId();
+	if (WriteLockOwnerThreadId.load(std::memory_order_relaxed) != CurrentThreadId)
+	{
+		InternalLock.WriteLock();
+		WriteLockOwnerThreadId = CurrentThreadId;
+	}
+	RecursionCount++;
+}
+
+void FCoreRedirects::FRWLockWithExclusiveRecursion::WriteUnlock()
+{
+	ensure(WriteLockOwnerThreadId.load(std::memory_order_relaxed) == FPlatformTLS::GetCurrentThreadId());
+	ensure(RecursionCount > 0);
+	RecursionCount--;
+	if (RecursionCount == 0)
+	{
+		WriteLockOwnerThreadId.store(0, std::memory_order_relaxed);
+		InternalLock.WriteUnlock();
+	}
+}
+
+void FCoreRedirects::FRWLockWithExclusiveRecursion::ReadUnlock()
+{
+	if (RecursionCount > 0)
+	{
+		ensureMsgf(WriteLockOwnerThreadId.load(std::memory_order_relaxed) == FPlatformTLS::GetCurrentThreadId(), 
+			TEXT("Called ReadUnlock() on a lock held exclusively by another thread."));
+		RecursionCount--;
+	}
+	else
+	{
+		InternalLock.ReadUnlock();
+	}
+}
diff --git a/Engine/Source/Runtime/CoreUObject/Public/UObject/CoreRedirects.h b/Engine/Source/Runtime/CoreUObject/Public/UObject/CoreRedirects.h
--- a/Engine/Source/Runtime/CoreUObject/Public/UObject/CoreRedirects.h
+++ b/Engine/Source/Runtime/CoreUObject/Public/UObject/CoreRedirects.h
@@ -24,8 +24,8 @@ class IPakFile;
 class UClass;
 struct FTopLevelAssetPath;
 namespace UE::CoreRedirects::Private { struct FCoreRedirectObjectUtf8Name; }
-
-#define WITH_COREREDIRECTS_MULTITHREAD_WARNING !UE_BUILD_SHIPPING && !IS_PROGRAM && !WITH_EDITOR
+namespace UE::CoreRedirects::Private { class FRWWithExclusiveRecursionScopeLockForRead; }
+namespace UE::CoreRedirects::Private { class FRWWithExclusiveRecursionScopeLockForWrite; }
 
 DECLARE_LOG_CATEGORY_EXTERN(LogCoreRedirects, Log, All);
 
@@ -352,10 +352,10 @@ struct FCoreRedirects
 	static COREUOBJECT_API bool RemoveRedirectList(TArrayView<const FCoreRedirect> Redirects, const FString& SourceString);
 
 	/** Returns true if this has ever been initialized */
-	static bool IsInitialized() { return bInitialized; }
+	static COREUOBJECT_API bool IsInitialized();
 
 	/** Returns true if this is in debug mode that slows loading and adds additional warnings */
-	static bool IsInDebugMode() { return bInDebugMode; }
+	static COREUOBJECT_API bool IsInDebugMode();
 
 	/** Validate a named list of redirects */
 	static COREUOBJECT_API void ValidateRedirectList(TArrayView<const FCoreRedirect> Redirects, const FString& SourceString);
@@ -363,8 +363,8 @@ struct FCoreRedirects
 	/** Validates all known redirects and warn if they seem to point to missing things */
 	static COREUOBJECT_API void ValidateAllRedirects();
 
-	/** Gets map from config key -> Flags */
-	static const TMap<FName, ECoreRedirectFlags>& GetConfigKeyMap() { return ConfigKeyMap; }
+	/** Gets map from config key -> Flags. It may only be accessed once it becomes constant data after the system is initialized */
+	static COREUOBJECT_API const TMap<FName, ECoreRedirectFlags>& GetConfigKeyMap();
 
 	/** Goes from the containing package and name of the type to the type flag */
 	static COREUOBJECT_API ECoreRedirectFlags GetFlagsForTypeName(FName PackageName, FName TypeName);
@@ -391,24 +391,36 @@ struct FCoreRedirects
 	static COREUOBJECT_API bool RunTests();
 
 private:
+	typedef UE::CoreRedirects::Private::FRWWithExclusiveRecursionScopeLockForRead FCoreRedirectorScopeLockForRead;
+	typedef UE::CoreRedirects::Private::FRWWithExclusiveRecursionScopeLockForWrite FCoreRedirectorScopeLockForWrite;
+
 	/** Static only class, never constructed */
 	COREUOBJECT_API FCoreRedirects();
 
-	/** Add a single redirect to a type map */
-	static COREUOBJECT_API bool AddSingleRedirect(const FCoreRedirect& NewRedirect, const FString& SourceString);
+	/** Internal implementation for AddRedirectList that requires a write lock to already have been acquired */
+	static COREUOBJECT_API bool AddRedirectListUnderWriteLock(TArrayView<const FCoreRedirect> Redirects, 
+		const FString& SourceString, const FCoreRedirectorScopeLockForWrite& HeldLock);
+
+	/** Internal implementation for AddSingleRedirect that requires a write lock to already have been acquired */
+	static COREUOBJECT_API bool AddSingleRedirectUnderWriteLock(const FCoreRedirect& NewRedirect, 
+		const FString& SourceString, const FCoreRedirectorScopeLockForWrite& HeldLock);
 
-	/** Remove a single redirect from a type map */
-	static COREUOBJECT_API bool RemoveSingleRedirect(const FCoreRedirect& OldRedirect, const FString& SourceString);
+	/** Internal implementation for RemoveSingleRedirect that requires a write lock to already have been acquired */
+	static COREUOBJECT_API bool RemoveSingleRedirectUnderWriteLock(const FCoreRedirect& OldRedirect, 
+		const FString& SourceString, const FCoreRedirectorScopeLockForWrite& HeldLock);
 
 	/** Add native redirects, called before ini is parsed for the first time */
-	static COREUOBJECT_API void RegisterNativeRedirects();
+	static COREUOBJECT_API void RegisterNativeRedirectsUnderWriteLock(const FCoreRedirectorScopeLockForWrite& HeldLock);
 
-#if WITH_COREREDIRECTS_MULTITHREAD_WARNING
-	/** Mark that CoreRedirects is about to start being used from multiple threads, and writes to new types of redirects are no longer allowed.
-	  * ReadRedirectsFromIni and all other AddRedirectList calls must be called before this
-	  */
-	static COREUOBJECT_API void EnterMultithreadedPhase();
-#endif
+	/** Internal implementation for GetMatchingRedirects that requires a read lock to already have been acquired */
+	static COREUOBJECT_API bool GetMatchingRedirectsUnderReadLock(ECoreRedirectFlags Type, 
+		const FCoreRedirectObjectName& OldObjectName, TArray<const FCoreRedirect*>& FoundRedirects, ECoreRedirectMatchFlags MatchFlags, 
+		const FCoreRedirectorScopeLockForRead& HeldLock);
+
+	/** Internal implementation for RedirectNameAndValues that requires a read lock to already have been acquired */
+	static COREUOBJECT_API bool RedirectNameAndValuesUnderReadLock(ECoreRedirectFlags Type, const FCoreRedirectObjectName& OldObjectName,
+		FCoreRedirectObjectName& NewObjectName, const FCoreRedirect** FoundValueRedirect,
+		ECoreRedirectMatchFlags MatchFlags, const FCoreRedirectorScopeLockForRead& HeldLock);
 
 	/** Container for managing Wildcard redirects (substrings, prefixes, suffixes) */
 	struct FWildcardData
@@ -439,19 +451,14 @@ private:
 	};
 
 	/** Whether this has been initialized at least once */
-	static COREUOBJECT_API bool bInitialized;
+	static COREUOBJECT_API std::atomic<bool> bInitialized;
 
 	/** True if we are in debug mode that does extra validation */
-	static COREUOBJECT_API bool bInDebugMode;
+	static COREUOBJECT_API std::atomic<bool> bInDebugMode;
 
 	/** True if we have done our initial validation. After initial validation, each change to redirects will validate independently */
 	static COREUOBJECT_API bool bValidatedOnce;
 
-#if WITH_COREREDIRECTS_MULTITHREAD_WARNING
-	/** Whether CoreRedirects is now being used multithreaded and therefore does not support writes to RedirectTypeMap keyvalue pairs */
-	static COREUOBJECT_API bool bIsInMultithreadedPhase;
-#endif
-
 	/** Map from config name to flag */
 	static COREUOBJECT_API TMap<FName, ECoreRedirectFlags> ConfigKeyMap;
 
@@ -471,9 +478,33 @@ private:
 	};
 	static COREUOBJECT_API FRedirectTypeMap RedirectTypeMap;
 
+	/** This lock allows exclusive locking (WriteLock) and shared locking (ReadLock)
+	 *  Additionally, it permits limited types of recursion. It is possible to ReadLock() 
+	 *  or WriteLock() while locked for write. It is NOT possible to Read or WriteLock() while
+	 *  locked for read. I.e., if the lock is held exclusively, it re-acquiring it is always permitted.
+	 *  If the lock is held shared, re-acquiring it is never permitted.
+	 */
+	struct FRWLockWithExclusiveRecursion
+	{
+		void ReadLock();
+
+		void WriteLock();
+
+		void WriteUnlock();
+
+		void ReadUnlock();
+
+	private:
+		FRWLock InternalLock;
+		std::atomic<uint32> WriteLockOwnerThreadId = 0;
+		int32 RecursionCount = 0;
+	};
+	friend class UE::CoreRedirects::Private::FRWWithExclusiveRecursionScopeLockForRead;
+	friend class UE::CoreRedirects::Private::FRWWithExclusiveRecursionScopeLockForWrite;
+
 	/**
-	 * Lock to protect multithreaded access to *KnownMissing functions, which can be called from the async loading threads. 
-	 * TODO: The KnownMissing functions use RedirectTypeMap, which is unguarded; there is race condition vulnerability if asyncloading thread is active before all categories are added to RedirectTypeMap.
+	 * Lock to protect multithreaded access to the CoreRedirect system
 	 */
-	static COREUOBJECT_API FRWLock KnownMissingLock;
+	static COREUOBJECT_API FRWLockWithExclusiveRecursion RWLock;
 };
+
