Deprecated calling AddSaveOverride on the context given to UObject::PreSave

You should override UObject::CollectSaveOverrides and call AddSaveOverride on its context instead

diff --git a/Engine/Source/Runtime/CoreUObject/Public/UObject/ObjectSaveContext.h b/Engine/Source/Runtime/CoreUObject/Public/UObject/ObjectSaveContext.h
--- a/Engine/Source/Runtime/CoreUObject/Public/UObject/ObjectSaveContext.h
+++ b/Engine/Source/Runtime/CoreUObject/Public/UObject/ObjectSaveContext.h
@@ -237,7 +237,7 @@ public:
 	/**
 	 * Add a save override to specific object. (i.e. mark certain objects or properties transient for this save)
 	 */
-	// TODO: UE_DEPRECATED(5.5, "Calling AddSaveOverride in UObject::PreSave is deprecated. Override UObject::CollectSaveOverrides and call AddSaveOverride on its context instead.")
+	UE_DEPRECATED(5.5, "Calling AddSaveOverride in UObject::PreSave is deprecated. Override UObject::CollectSaveOverrides and call AddSaveOverride on its context instead.")
 	void AddSaveOverride(UObject* Target, FObjectSaveOverride InOverride)
 	{
 		Data.AddSaveOverride(Target, MoveTemp(InOverride));
