Fix FCoreRedirects for AutoRTFM

diff --git a/Engine/Source/Runtime/CoreUObject/Private/UObject/CoreRedirects.cpp b/Engine/Source/Runtime/CoreUObject/Private/UObject/CoreRedirects.cpp
--- a/Engine/Source/Runtime/CoreUObject/Private/UObject/CoreRedirects.cpp
+++ b/Engine/Source/Runtime/CoreUObject/Private/UObject/CoreRedirects.cpp
@@ -20,6 +20,8 @@
 #include "Serialization/DeferredMessageLog.h"
 #include "Templates/Casts.h"
 
+#include "AutoRTFM/AutoRTFM.h"
+
 DEFINE_LOG_CATEGORY(LogCoreRedirects);
 
 #if !defined UE_WITH_CORE_REDIRECTS
@@ -1098,14 +1100,16 @@ public:
 	: InternalLock(InLock)
 	, NeedsUnlock(true)
 	{
-		InternalLock.ReadLock();
+		UE_AUTORTFM_OPEN({ InternalLock.ReadLock(); });
+		AutoRTFM::PushOnAbortHandler(this, [this] { this->InternalLock.ReadUnlock(); });
 	}
 
 	~FRWWithExclusiveRecursionScopeLockForRead()
 	{
 		if (NeedsUnlock)
 		{
-			InternalLock.ReadUnlock();
+			AutoRTFM::PopOnAbortHandler(this);
+			UE_AUTORTFM_OPEN({ InternalLock.ReadUnlock(); });
 		}
 	}
 protected:
