IncrementalCook: Add ObjectRedirectors recorded by the RedirectCollector to the RedirectionSummary per package that is maintained by CoreRedirects. SoftObjectPaths are swapped (without ever being resolved) during SavePackage to point to the redirected locations, so if package A has a softobjectpath to B and B is redirected to C, the softobjectpath text in A's cooked output needs to modified to C, and we therefore need to recook A.
#rnx

diff --git a/Engine/Source/Runtime/CoreUObject/Private/Misc/RedirectCollector.cpp b/Engine/Source/Runtime/CoreUObject/Private/Misc/RedirectCollector.cpp
--- a/Engine/Source/Runtime/CoreUObject/Private/Misc/RedirectCollector.cpp
+++ b/Engine/Source/Runtime/CoreUObject/Private/Misc/RedirectCollector.cpp
@@ -3,6 +3,7 @@
 #include "Misc/RedirectCollector.h"
 #include "Algo/Transform.h"
 #include "Misc/CoreDelegates.h"
+#include "UObject/CoreRedirects.h"
 #include "UObject/UObjectGlobals.h"
 #include "UObject/Object.h"
 #include "UObject/Package.h"
@@ -118,6 +119,9 @@ void FRedirectCollector::ResolveAllSoftObjectPaths(FName FilterPackage)
 				if (Dest.ToString() != ToLoad)
 				{
 					ObjectPathRedirectionMap.Add(ToLoadPath, Dest);
+#if WITH_EDITOR
+					FCoreRedirects::RecordAddedObjectRedirector(ToLoadPath, Dest);
+#endif
 				}
 			}
 			else
@@ -246,11 +250,21 @@ void FRedirectCollector::AddAssetPathRedirection(const FSoftObjectPath& Original
 	{
 		// If RedirectedPath points back to OriginalPath, remove that to avoid a circular reference
 		// This can happen when renaming assets in the editor but not actually dropping redirectors because it was new
-		ObjectPathRedirectionMap.Remove(RedirectedPath);
+		FSoftObjectPath TargetPath;
+		bool bRemoved = ObjectPathRedirectionMap.RemoveAndCopyValue(RedirectedPath, TargetPath);
+#if WITH_EDITOR
+		if (bRemoved)
+		{
+			FCoreRedirects::RecordRemovedObjectRedirector(RedirectedPath, TargetPath);
+		}
+#endif
 	}
 
 	// This replaces an existing mapping, can happen in the editor if things are renamed twice
 	ObjectPathRedirectionMap.Add(OriginalPath, RedirectedPath);
+#if WITH_EDITOR
+	FCoreRedirects::RecordAddedObjectRedirector(OriginalPath, RedirectedPath);
+#endif
 }
 
 void FRedirectCollector::AddAssetPathRedirection(FName OriginalPath, FName RedirectedPath)
@@ -268,6 +282,9 @@ void FRedirectCollector::RemoveAssetPathRedirection(const FSoftObjectPath& Origi
 
 	if (ensureMsgf(Found, TEXT("Cannot remove redirection from %s, it was not registered"), *OriginalPath.ToString()))
 	{
+#if WITH_EDITOR
+		FCoreRedirects::RecordRemovedObjectRedirector(OriginalPath, *Found);
+#endif
 		ObjectPathRedirectionMap.Remove(OriginalPath);
 	}
 }
diff --git a/Engine/Source/Runtime/CoreUObject/Private/UObject/CoreRedirects.cpp b/Engine/Source/Runtime/CoreUObject/Private/UObject/CoreRedirects.cpp
--- a/Engine/Source/Runtime/CoreUObject/Private/UObject/CoreRedirects.cpp
+++ b/Engine/Source/Runtime/CoreUObject/Private/UObject/CoreRedirects.cpp
@@ -238,6 +238,22 @@ FCoreRedirectObjectName::FCoreRedirectObjectName(const FTopLevelAssetPath& TopLe
 {
 }
 
+FCoreRedirectObjectName::FCoreRedirectObjectName(const FSoftObjectPath& SoftObjectPath)
+{
+	if (!SoftObjectPath.IsSubobject())
+	{
+		PackageName = SoftObjectPath.GetLongPackageFName();
+		ObjectName = SoftObjectPath.GetAssetFName();
+	}
+	else
+	{
+		if (!ExpandNames(SoftObjectPath.ToString(), ObjectName, OuterName, PackageName))
+		{
+			Reset();
+		}
+	}
+}
+
 FCoreRedirectObjectName::FCoreRedirectObjectName(const FString& InString)
 {
 	if (!ExpandNames(InString, ObjectName, OuterName, PackageName))
@@ -1534,6 +1550,21 @@ void FCoreRedirects::AppendHashOfGlobalRedirects(FBlake3& Hasher)
 {
 	GRedirectionSummary.AppendHashGlobal(Hasher);
 }
+
+void FCoreRedirects::RecordAddedObjectRedirector(const FSoftObjectPath& Source, const FSoftObjectPath& Dest)
+{
+	FCoreRedirect ConvertedToCoreRedirect(ECoreRedirectFlags::Type_Object,
+		FCoreRedirectObjectName(Source), FCoreRedirectObjectName(Dest));
+	GRedirectionSummary.Add(ConvertedToCoreRedirect, false /* bIsWildcardMatch */);
+}
+
+void FCoreRedirects::RecordRemovedObjectRedirector(const FSoftObjectPath& Source, const FSoftObjectPath& Dest)
+{
+	FCoreRedirect ConvertedToCoreRedirect(ECoreRedirectFlags::Type_Object,
+		FCoreRedirectObjectName(Source), FCoreRedirectObjectName(Dest));
+	GRedirectionSummary.Remove(ConvertedToCoreRedirect, false /* bIsWildcardMatch */);
+}
+
 #endif
 
 bool FCoreRedirects::RunTests()
diff --git a/Engine/Source/Runtime/CoreUObject/Public/UObject/CoreRedirects.h b/Engine/Source/Runtime/CoreUObject/Public/UObject/CoreRedirects.h
--- a/Engine/Source/Runtime/CoreUObject/Public/UObject/CoreRedirects.h
+++ b/Engine/Source/Runtime/CoreUObject/Public/UObject/CoreRedirects.h
@@ -22,6 +22,7 @@
 class FBlake3;
 class IPakFile;
 class UClass;
+struct FSoftObjectPath;
 struct FTopLevelAssetPath;
 namespace UE::CoreRedirects::Private { struct FCoreRedirectObjectUtf8Name; }
 namespace UE::CoreRedirects::Private { class FRWWithExclusiveRecursionScopeLockForRead; }
@@ -97,6 +98,8 @@ struct FCoreRedirectObjectName
 
 	COREUOBJECT_API FCoreRedirectObjectName(const FTopLevelAssetPath& TopLevelAssetPath);
 
+	COREUOBJECT_API FCoreRedirectObjectName(const FSoftObjectPath& SoftObjectPath);
+
 	COREUOBJECT_API FCoreRedirectObjectName(const FString& InString);
 
 	/** Construct from object in memory */
@@ -385,6 +388,11 @@ struct FCoreRedirects
 	 * Used in iterative cooking to invalidate the cooked version of packages when CoreRedirects change.
 	 */
 	static COREUOBJECT_API void AppendHashOfGlobalRedirects(FBlake3& Hasher);
+
+	/** Add the given Source->Path redirector to the summary used for AppendHashOfRedirectsAffectingPackages. */
+	static COREUOBJECT_API void RecordAddedObjectRedirector(const FSoftObjectPath& Source, const FSoftObjectPath& Dest);
+	/** Remove the given Source->Path redirector to the summary used for AppendHashOfRedirectsAffectingPackages. */
+	static COREUOBJECT_API void RecordRemovedObjectRedirector(const FSoftObjectPath& Source, const FSoftObjectPath& Dest);
 #endif
 
 	/** Runs set of redirector tests, returns false on failure */
