IncrementalCook: Add TObjPtr references made from Linker::PreLoad and UObject::PostLoad to the dependencies for the package of the object being pre or post loaded.
#rnx

diff --git a/Engine/Source/Runtime/CoreUObject/Private/Misc/PackageAccessTracking.cpp b/Engine/Source/Runtime/CoreUObject/Private/Misc/PackageAccessTracking.cpp
--- a/Engine/Source/Runtime/CoreUObject/Private/Misc/PackageAccessTracking.cpp
+++ b/Engine/Source/Runtime/CoreUObject/Private/Misc/PackageAccessTracking.cpp
@@ -15,6 +15,8 @@ bool IsBuildOpName(FName OpName)
 	//If we want to capture script package dependencies, we may also need to record accesses under the "PackageAccessTrackingOps::NAME_CreateDefaultObject" operation
 	//which occurs from "UObjectLoadAllCompiledInDefaultProperties" and possibly elsewhere.
 	return (OpName == PackageAccessTrackingOps::NAME_Load)
+		|| (OpName == PackageAccessTrackingOps::NAME_PreLoad)
+		|| (OpName == PackageAccessTrackingOps::NAME_PostLoad)
 		|| (OpName == PackageAccessTrackingOps::NAME_Save)
 		|| (OpName == PackageAccessTrackingOps::NAME_CookerBuildObject)
 		// || (OpName == PackageAccessTrackingOps::NAME_CreateDefaultObject)
diff --git a/Engine/Source/Runtime/CoreUObject/Private/Serialization/AsyncLoading.cpp b/Engine/Source/Runtime/CoreUObject/Private/Serialization/AsyncLoading.cpp
--- a/Engine/Source/Runtime/CoreUObject/Private/Serialization/AsyncLoading.cpp
+++ b/Engine/Source/Runtime/CoreUObject/Private/Serialization/AsyncLoading.cpp
@@ -6770,6 +6770,7 @@ EAsyncPackageState::Type FAsyncPackage::CreateExports()
 				// This will cause the object to be serialized. We do this here for all objects and
 				// not just UClass and template objects, for which this is required in order to ensure
 				// seek free loading, to be able introduce async file I/O.
+				UE_TRACK_REFERENCING_PACKAGE_SCOPED(Object, PackageAccessTrackingOps::NAME_PreLoad);
 				Linker->Preload(Object);
 				PackageObjLoaded.Add(Object);
 			}
@@ -6835,6 +6836,7 @@ EAsyncPackageState::Type FAsyncPackage::PreLoadObjects()
 		UObject* Object = PackageObjLoaded[PreLoadIndex++];
 		if (Object && Object->GetLinker())
 		{
+			UE_TRACK_REFERENCING_PACKAGE_SCOPED(Object, PackageAccessTrackingOps::NAME_PreLoad);
 			Object->GetLinker()->Preload(Object);
 				LastObjectWorkWasPerformedOn = Object;
 				LastTypeOfWorkPerformed = TEXT("preloading");
@@ -7065,6 +7067,7 @@ EAsyncPackageState::Type FAsyncPackage::PostLoadDeferredObjects(double InTickSta
 					{
 						if (PreLoadObject && PreLoadObject->GetLinker())
 						{
+							UE_TRACK_REFERENCING_PACKAGE_SCOPED(PreLoadObject, PackageAccessTrackingOps::NAME_PreLoad);
 							PreLoadObject->GetLinker()->Preload(PreLoadObject);
 							PackageObjLoaded.Add(PreLoadObject);
 						}
diff --git a/Engine/Source/Runtime/CoreUObject/Private/Serialization/AsyncLoading2.cpp b/Engine/Source/Runtime/CoreUObject/Private/Serialization/AsyncLoading2.cpp
--- a/Engine/Source/Runtime/CoreUObject/Private/Serialization/AsyncLoading2.cpp
+++ b/Engine/Source/Runtime/CoreUObject/Private/Serialization/AsyncLoading2.cpp
@@ -5651,6 +5651,7 @@ bool FAsyncPackage2::PreloadLinkerLoadExports(FAsyncLoadingThreadState2& ThreadS
 			if (Object->HasAnyFlags(RF_NeedLoad))
 			{
 				UE_ASYNC_PACKAGE_LOG(VeryVerbose, Desc, TEXT("PreloadLinkerLoadExports"), TEXT("Preloading export %d: %s"), ExportIndex, *Object->GetPathName());
+				UE_TRACK_REFERENCING_PACKAGE_SCOPED(Object, PackageAccessTrackingOps::NAME_PreLoad);
 				LinkerLoadState->Linker->Preload(Object);
 			}
 		}
