SavePackage: PKG_ContainsNoAsset was accidentally set for some editor packages; it is only used in cooked packages. Prevent the possibility of setting it for editor packages: recalculate it every time rather than preserving it from the package's previously saved flags, and always set it to false for editor packages.
#rnx

diff --git a/Engine/Source/Runtime/CoreUObject/Private/UObject/SavePackage2.cpp b/Engine/Source/Runtime/CoreUObject/Private/UObject/SavePackage2.cpp
--- a/Engine/Source/Runtime/CoreUObject/Private/UObject/SavePackage2.cpp
+++ b/Engine/Source/Runtime/CoreUObject/Private/UObject/SavePackage2.cpp
@@ -2449,10 +2449,14 @@ ESavePackageResult UpdatePackageHeader(FStructuredArchive::FRecord& StructuredAr
 
 		// Update package flags from package, in case serialization has modified package flags.
 		uint32 PackageFlags = Linker->LinkerRoot->GetPackageFlags();
-		if (!bContainsAsset)
+		if (SaveContext.IsCooking() && !bContainsAsset)
 		{
 			PackageFlags |= PKG_ContainsNoAsset;
 		}
+		else
+		{
+			PackageFlags &= ~PKG_ContainsNoAsset;
+		}
 		// Take the Linker FilterEditorOnlyData setting over the package flags to set this flag in the summary
 		if (Linker->IsFilterEditorOnly())
 		{
