IncrementalCook: CookDependencies: Fix comparison function for FCookDependency of type ECookDependency::Function. It was ignoring the arguments to the function and treating all function invocations with the same function name as identical no matter their arguments.
#rnx

diff --git a/Engine/Source/Runtime/CoreUObject/Public/Cooker/CookDependency.h b/Engine/Source/Runtime/CoreUObject/Public/Cooker/CookDependency.h
--- a/Engine/Source/Runtime/CoreUObject/Public/Cooker/CookDependency.h
+++ b/Engine/Source/Runtime/CoreUObject/Public/Cooker/CookDependency.h
@@ -416,7 +416,7 @@ inline bool FCookDependency::operator<(const FCookDependency& Other) const
 		FMemoryView ViewA;
 		FMemoryView ViewB;
 		bool bHasViewA = FunctionData.Args.TryGetRangeView(ViewA);
-		bool bHasViewB = FunctionData.Args.TryGetRangeView(ViewB);
+		bool bHasViewB = Other.FunctionData.Args.TryGetRangeView(ViewB);
 		if ((!bHasViewA) | (!bHasViewB))
 		{
 			return bHasViewB; // If both false, return false. If only one, return true only if A is the false.
@@ -482,7 +482,7 @@ inline bool FCookDependency::operator==(const FCookDependency& Other) const
 		FMemoryView ViewA;
 		FMemoryView ViewB;
 		bool bHasViewA = FunctionData.Args.TryGetRangeView(ViewA);
-		bool bHasViewB = FunctionData.Args.TryGetRangeView(ViewB);
+		bool bHasViewB = Other.FunctionData.Args.TryGetRangeView(ViewB);
 		if (bHasViewA != bHasViewB)
 		{
 			return false;
