[BugFix] CLASS_NewerVersionExists should only impact whether new IDOs get created *not* whether IDOs get saved to disk

diff --git a/Engine/Source/Runtime/CoreUObject/Private/UObject/InstanceDataObjectUtils.cpp b/Engine/Source/Runtime/CoreUObject/Private/UObject/InstanceDataObjectUtils.cpp
--- a/Engine/Source/Runtime/CoreUObject/Private/UObject/InstanceDataObjectUtils.cpp
+++ b/Engine/Source/Runtime/CoreUObject/Private/UObject/InstanceDataObjectUtils.cpp
@@ -167,12 +167,6 @@ namespace UE
 				return true;
 			}
 
-			// reinst classes shouldn't use IDOs
-			if (InObject->GetClass()->HasAnyClassFlags(CLASS_NewerVersionExists))
-			{
-				return false;
-			}
-
 			//@todo FH: change to check trait when available or use config object
 			const UClass* ObjClass = InObject->GetClass();
 			while (ObjClass && ObjClass->GetClass()->GetFName() != NAME_VerseClass)
diff --git a/Engine/Source/Runtime/CoreUObject/Private/UObject/PropertyBagRepository.cpp b/Engine/Source/Runtime/CoreUObject/Private/UObject/PropertyBagRepository.cpp
--- a/Engine/Source/Runtime/CoreUObject/Private/UObject/PropertyBagRepository.cpp
+++ b/Engine/Source/Runtime/CoreUObject/Private/UObject/PropertyBagRepository.cpp
@@ -810,7 +810,9 @@ FScopedIDOSerializationContext::FScopedIDOSerializationContext(UObject* InObject
 {
 	FUObjectSerializeContext* SerializeContext = FUObjectThreadContext::Get().GetSerializeContext();
 	bHasIDOSupport = FPropertyBagRepository::IsInstanceDataObjectSupportEnabled(Object);
-	bCreateIDO = bHasIDOSupport && !SerializeContext->bImpersonateProperties && Archive->IsLoading();
+	
+	bool bHasReinstancedClass = InObject->GetClass()->HasAnyClassFlags(CLASS_NewerVersionExists);
+	bCreateIDO = bHasIDOSupport && !SerializeContext->bImpersonateProperties && Archive->IsLoading() && !bHasReinstancedClass;
 
 	if (bHasIDOSupport)
 	{
