Added storage of unknown enum names

diff --git a/Engine/Source/Runtime/CoreUObject/Internal/UObject/PropertyBagRepository.h b/Engine/Source/Runtime/CoreUObject/Internal/UObject/PropertyBagRepository.h
--- a/Engine/Source/Runtime/CoreUObject/Internal/UObject/PropertyBagRepository.h
+++ b/Engine/Source/Runtime/CoreUObject/Internal/UObject/PropertyBagRepository.h
@@ -17,6 +17,8 @@ namespace UE
 #if WITH_EDITORONLY_DATA
 
 class FPropertyPathNameTree;
+class FPropertyTypeName;
+class FUnknownEnumNames;
 
 // Singleton class tracking property bag association with objects
 class FPropertyBagRepository : public FGCObject
@@ -26,6 +28,7 @@ class FPropertyBagRepository : public FGCObject
 		void Destroy();
 
 		FPropertyPathNameTree* Tree = nullptr;
+		FUnknownEnumNames* EnumNames = nullptr;
 		TObjectPtr<UObject> InstanceDataObject = nullptr;
 		bool bNeedsFixup = false;
 	};
@@ -67,6 +70,31 @@ public:
 	 */
 	FPropertyPathNameTree* FindOrCreateUnknownPropertyTree(const UObject* Owner);
 
+	/**
+	 * Adds an unknown enum name to the names tracked for an object.
+	 *
+	 * @param Owner			The owner to associate the unknown name with.
+	 * @param Enum			The enum associated with the property being serialized. May be null or the wrong type.
+	 * @param EnumTypeName	The type name of the enum containing the unknown name.
+	 * @param EnumValueName	The unknown name to track.
+	 */
+	void AddUnknownEnumName(const UObject* Owner, const UEnum* Enum, FPropertyTypeName EnumTypeName, FName EnumValueName);
+
+	/**
+	 * Finds tracked unknown enum names associated with the object.
+	 *
+	 * @param Owner			The owner associated with the unknown names to find.
+	 * @param EnumTypeName	The type name of the enum containing the unknown names to find.
+	 * @param OutNames		Array to assign the unknown names to. Empty on return if no names are found.
+	 * @param bOutHasFlags	Assigned to true if the enum is known to have flags, otherwise false.
+	 */
+	void FindUnknownEnumNames(const UObject* Owner, FPropertyTypeName EnumTypeName, TArray<FName>& OutNames, bool& bOutHasFlags);
+
+	/**
+	 * Resets tracked unknown enum names associated with the object.
+	 */
+	void ResetUnknownEnumNames(const UObject* Owner);
+
 	// Future version for reworked InstanceDataObjects - track InstanceDataObject rather than bag (directly):
 	/**
 	 * Instantiate an InstanceDataObject object representing all fields within the bag, tracked against the owner object.
diff --git a/Engine/Source/Runtime/CoreUObject/Private/UObject/PropertyBagRepository.cpp b/Engine/Source/Runtime/CoreUObject/Private/UObject/PropertyBagRepository.cpp
--- a/Engine/Source/Runtime/CoreUObject/Private/UObject/PropertyBagRepository.cpp
+++ b/Engine/Source/Runtime/CoreUObject/Private/UObject/PropertyBagRepository.cpp
@@ -5,19 +5,21 @@
 #if WITH_EDITORONLY_DATA
 
 #include "Containers/Queue.h"
+#include "Misc/StringBuilder.h"
 #include "Serialization/ObjectReader.h"
 #include "Serialization/ObjectWriter.h"
+#include "String/ParseTokens.h"
 #include "UObject/GarbageCollection.h"
+#include "UObject/InstanceDataObjectUtils.h"
+#include "UObject/LinkerLoad.h"
 #include "UObject/Object.h"
-#include "UObject/UObjectGlobals.h"
+#include "UObject/Package.h"
+#include "UObject/PropertyOptional.h"
 #include "UObject/PropertyPathNameTree.h"
+#include "UObject/UObjectGlobals.h"
 #include "UObject/UObjectHash.h"
 #include "UObject/UObjectThreadContext.h"
-#include "UObject/LinkerLoad.h"
-#include "UObject/InstanceDataObjectUtils.h"
-#include "UObject/Package.h"
 #include "Templates/UnrealTemplate.h"
-#include "UObject/PropertyOptional.h"
 
 #if WITH_EDITOR
 #include "HAL/IConsoleManager.h"
@@ -123,11 +125,26 @@ public:
 	}
 };
 
+class FUnknownEnumNames
+{
+public:
+	struct FInfo
+	{
+		TSet<FName> Names;
+		bool bHasFlags = false;
+	};
+
+	TMap<FPropertyTypeName, FInfo> Enums;
+};
+
 void FPropertyBagRepository::FPropertyBagAssociationData::Destroy()
 {
 	delete Tree;
 	Tree = nullptr;
 
+	delete EnumNames;
+	EnumNames = nullptr;
+
 	if (InstanceDataObject && InstanceDataObject->IsValidLowLevel())
 	{
 		InstanceDataObject = nullptr;
@@ -512,8 +529,7 @@ FPropertyPathNameTree* FPropertyBagRepository::FindOrCreateUnknownPropertyTree(c
 	FPropertyBagAssociationData* BagData = AssociatedData.Find(Owner);
 	if (!BagData)
 	{
-		FPropertyBagAssociationData NewBagData;
-		BagData = &AssociatedData.Emplace(Owner, NewBagData);
+		BagData = &AssociatedData.Emplace(Owner);
 	}
 	if (!BagData->Tree)
 	{
@@ -522,6 +538,86 @@ FPropertyPathNameTree* FPropertyBagRepository::FindOrCreateUnknownPropertyTree(c
 	return BagData->Tree;
 }
 
+void FPropertyBagRepository::AddUnknownEnumName(const UObject* Owner, const UEnum* Enum, FPropertyTypeName EnumTypeName, FName EnumValueName)
+{
+	check(Owner);
+	checkf(Enum || !EnumTypeName.IsEmpty(), TEXT("AddUnknownEnumName requires an enum or its type name. Owner: %s"), *Owner->GetPathName());
+
+	FPropertyBagRepositoryLock LockRepo(this);
+	FPropertyBagAssociationData& BagData = AssociatedData.FindOrAdd(Owner);
+	if (!BagData.EnumNames)
+	{
+		BagData.EnumNames = new FUnknownEnumNames;
+	}
+
+	if (EnumTypeName.IsEmpty())
+	{
+		FPropertyTypeNameBuilder Builder;
+		Builder.AddPath(Enum);
+		EnumTypeName = Builder.Build();
+	}
+
+	FUnknownEnumNames::FInfo& Info = BagData.EnumNames->Enums.FindOrAdd(EnumTypeName);
+
+	TStringBuilder<128> EnumValueString(InPlace, EnumValueName);
+	if (String::FindFirstChar(EnumValueString, TEXT('|')) == INDEX_NONE)
+	{
+		Info.Names.Add(EnumValueName);
+	}
+	else
+	{
+		Info.bHasFlags = true;
+		String::ParseTokens(EnumValueString, TEXT('|'), [&Info, Enum](FStringView Token)
+		{
+			FName Name(Token);
+			if (!Enum || Enum->GetIndexByName(Name) == INDEX_NONE)
+			{
+				Info.Names.Add(Name);
+			}
+		}, String::EParseTokensOptions::SkipEmpty | String::EParseTokensOptions::Trim);
+	}
+
+	if (!Info.bHasFlags && Enum && Enum->HasAnyEnumFlags(EEnumFlags::Flags))
+	{
+		Info.bHasFlags = true;
+	}
+}
+
+void FPropertyBagRepository::FindUnknownEnumNames(const UObject* Owner, FPropertyTypeName EnumTypeName, TArray<FName>& OutNames, bool& bOutHasFlags)
+{
+	check(Owner);
+	checkf(!EnumTypeName.IsEmpty(), TEXT("FindUnknownEnumNames requires an enum type name. Owner: %s"), *Owner->GetPathName());
+
+	OutNames.Empty();
+	bOutHasFlags = false;
+
+	FPropertyBagRepositoryLock LockRepo(this);
+	FPropertyBagAssociationData* BagData = AssociatedData.Find(Owner);
+	if (!BagData || !BagData->EnumNames)
+	{
+		return;
+	}
+
+	if (const FUnknownEnumNames::FInfo* Info = BagData->EnumNames->Enums.Find(EnumTypeName))
+	{
+		OutNames = Info->Names.Array();
+		bOutHasFlags = Info->bHasFlags;
+	}
+}
+
+void FPropertyBagRepository::ResetUnknownEnumNames(const UObject* Owner)
+{
+	check(Owner);
+
+	FPropertyBagRepositoryLock LockRepo(this);
+	FPropertyBagAssociationData* BagData = AssociatedData.Find(Owner);
+	if (BagData)
+	{
+		delete BagData->EnumNames;
+		BagData->EnumNames = nullptr;
+	}
+}
+
 UObject* FPropertyBagRepository::CreateInstanceDataObject(UObject* Owner, FArchive* Archive)
 {
 	FPropertyBagRepositoryLock LockRepo(this);
diff --git a/Engine/Source/Runtime/CoreUObject/Tests/UObject/InstanceDataObjectUtilsTest.cpp b/Engine/Source/Runtime/CoreUObject/Tests/UObject/InstanceDataObjectUtilsTest.cpp
--- a/Engine/Source/Runtime/CoreUObject/Tests/UObject/InstanceDataObjectUtilsTest.cpp
+++ b/Engine/Source/Runtime/CoreUObject/Tests/UObject/InstanceDataObjectUtilsTest.cpp
@@ -23,7 +23,7 @@
 namespace UE
 {
 
-TEST_CASE_NAMED(FInstanceDataObjectUtilsTest, "CoreUObject::Serialization::InstanceDataObjectUtils", "[Core][UObject][EngineFilter]")
+TEST_CASE_NAMED(FInstanceDataObjectUtilsTest, "CoreUObject::Serialization::InstanceDataObjectUtils", "[CoreUObject][EngineFilter]")
 {
 	UTestInstanceDataObjectClass* BaseObject = NewObject<UTestInstanceDataObjectClass>();
 	UClass* TestClass = CreateInstanceDataObjectClass(nullptr, BaseObject->GetClass(), BaseObject->GetOuter());
@@ -89,7 +89,7 @@ TEST_CASE_NAMED(FInstanceDataObjectUtilsTest, "CoreUObject::Serialization::Insta
 	CHECK_FALSE(IsPropertyValueInitialized(StructProperty->Struct, StructData, DProperty));
 }
 
-TEST_CASE_NAMED(FTrackInitializedPropertiesTest, "CoreUObject::Serialization::TrackInitializedProperties", "[Core][UObject][EngineFilter]")
+TEST_CASE_NAMED(FTrackInitializedPropertiesTest, "CoreUObject::Serialization::TrackInitializedProperties", "[CoreUObject][EngineFilter]")
 {
 	UTestInstanceDataObjectClass* BaseObject = NewObject<UTestInstanceDataObjectClass>();
 	UClass* TestClass = CreateInstanceDataObjectClass(nullptr, BaseObject->GetClass(), BaseObject->GetOuter());
@@ -201,7 +201,7 @@ TEST_CASE_NAMED(FTrackInitializedPropertiesTest, "CoreUObject::Serialization::Tr
 #endif // WITH_TEXT_ARCHIVE_SUPPORT
 }
 
-TEST_CASE_NAMED(FTrackUnknownPropertiesTest, "CoreUObject::Serialization::TrackUnknownProperties", "[Core][UObject][EngineFilter]")
+TEST_CASE_NAMED(FTrackUnknownPropertiesTest, "CoreUObject::Serialization::TrackUnknownProperties", "[CoreUObject][EngineFilter]")
 {
 	const auto MakePropertyTypeName = [](FName Name)
 	{
@@ -293,6 +293,99 @@ TEST_CASE_NAMED(FTrackUnknownPropertiesTest, "CoreUObject::Serialization::TrackU
 #endif // WITH_TEXT_ARCHIVE_SUPPORT
 }
 
+TEST_CASE_NAMED(FUnknownEnumNamesTest, "CoreUObject::Serialization::UnknownEnumNames", "[CoreUObject][EngineFilter]")
+{
+	UObject* Owner = NewObject<UTestInstanceDataObjectClass>();
+
+	FPropertyBagRepository& Repo = FPropertyBagRepository::Get();
+
+	TArray<FName> Names{NAME_None};
+	bool bHasFlags = true;
+
+	// Test a non-flags enum...
+
+	FPropertyTypeName FruitTypeName = []
+	{
+		FPropertyTypeNameBuilder Builder;
+		Builder.AddPath(StaticEnum<ETestInstanceDataObjectFruit>());
+		return Builder.Build();
+	}();
+
+	Repo.FindUnknownEnumNames(Owner, FruitTypeName, Names, bHasFlags);
+	CHECK(Names.IsEmpty());
+	CHECK_FALSE(bHasFlags);
+
+	const FName NAME_Cherry = "Cherry";
+	const FName NAME_Pear = "Pear";
+
+	Repo.AddUnknownEnumName(Owner, nullptr, FruitTypeName, NAME_Pear);
+	Repo.AddUnknownEnumName(Owner, StaticEnum<ETestInstanceDataObjectFruit>(), {}, NAME_Cherry);
+	Repo.AddUnknownEnumName(Owner, StaticEnum<ETestInstanceDataObjectFruit>(), {}, NAME_Pear);
+	Repo.AddUnknownEnumName(Owner, nullptr, FruitTypeName, NAME_Cherry);
+
+	Repo.FindUnknownEnumNames(Owner, FruitTypeName, Names, bHasFlags);
+	CHECKED_IF(Names.Num() == 2)
+	{
+		CHECK(Names[0] == NAME_Pear);
+		CHECK(Names[1] == NAME_Cherry);
+	}
+	CHECK_FALSE(bHasFlags);
+
+	// Test a flags enum by name only...
+
+	FPropertyTypeName FlagsTypeName = []
+	{
+		FPropertyTypeNameBuilder Builder;
+		Builder.AddPath(StaticEnum<ETestInstanceDataObjectFlags>());
+		return Builder.Build();
+	}();
+
+	const FName NAME_South = "South";
+	const FName NAME_Down = "Down";
+	const FName NAME_Up = "Up";
+
+	TStringBuilder<128> FlagsString;
+	FlagsString.Join(MakeArrayView({NAME_Up, NAME_Down, NAME_South}), TEXTVIEW(" | "));
+
+	Repo.AddUnknownEnumName(Owner, nullptr, FlagsTypeName, NAME_Down);
+	Repo.AddUnknownEnumName(Owner, nullptr, FlagsTypeName, *FlagsString);
+
+	Repo.FindUnknownEnumNames(Owner, FlagsTypeName, Names, bHasFlags);
+	CHECKED_IF(Names.Num() == 3)
+	{
+		CHECK(Names[0] == NAME_Down);
+		CHECK(Names[1] == NAME_Up);
+		CHECK(Names[2] == NAME_South);
+	}
+	CHECK(bHasFlags);
+
+	// Test resetting unknown enum names for an owner...
+
+	Repo.ResetUnknownEnumNames(Owner);
+
+	Repo.FindUnknownEnumNames(Owner, FlagsTypeName, Names, bHasFlags);
+	CHECK(Names.IsEmpty());
+	CHECK_FALSE(bHasFlags);
+
+	// Test a flags enum by enum...
+
+	Repo.AddUnknownEnumName(Owner, StaticEnum<ETestInstanceDataObjectFlags>(), {}, NAME_Up);
+
+	Repo.FindUnknownEnumNames(Owner, FlagsTypeName, Names, bHasFlags);
+	CHECK(Names.Num() == 1);
+	CHECK(bHasFlags);
+
+	Repo.AddUnknownEnumName(Owner, StaticEnum<ETestInstanceDataObjectFlags>(), FlagsTypeName, *FlagsString);
+
+	Repo.FindUnknownEnumNames(Owner, FlagsTypeName, Names, bHasFlags);
+	CHECKED_IF(Names.Num() == 2)
+	{
+		CHECK(Names[0] == NAME_Up);
+		CHECK(Names[1] == NAME_Down);
+	}
+	CHECK(bHasFlags);
+}
+
 } // UE
 
 #endif // WITH_TESTS && WITH_EDITORONLY_DATA
diff --git a/Engine/Source/Runtime/CoreUObject/Tests/UObject/InstanceDataObjectUtilsTest.h b/Engine/Source/Runtime/CoreUObject/Tests/UObject/InstanceDataObjectUtilsTest.h
--- a/Engine/Source/Runtime/CoreUObject/Tests/UObject/InstanceDataObjectUtilsTest.h
+++ b/Engine/Source/Runtime/CoreUObject/Tests/UObject/InstanceDataObjectUtilsTest.h
@@ -6,6 +6,25 @@
 #include "UObject/Object.h"
 #include "InstanceDataObjectUtilsTest.generated.h"
 
+UENUM()
+enum class ETestInstanceDataObjectFruit : uint8
+{
+	None = 0,
+	Apple,
+	Banana,
+	Orange,
+};
+
+UENUM(Flags)
+enum class ETestInstanceDataObjectFlags : uint16
+{
+	None = 0,
+	North = 1 << 0,
+	East = 1 << 1,
+	South = 1 << 2,
+	West = 1 << 3,
+};
+
 USTRUCT()
 struct FTestInstanceDataObjectStruct
 {
