Support interfaces in FObjectPropertyBase::CheckValidObject
#rnx

diff --git a/Engine/Source/Runtime/CoreUObject/Private/UObject/PropertyBaseObject.cpp b/Engine/Source/Runtime/CoreUObject/Private/UObject/PropertyBaseObject.cpp
--- a/Engine/Source/Runtime/CoreUObject/Private/UObject/PropertyBaseObject.cpp
+++ b/Engine/Source/Runtime/CoreUObject/Private/UObject/PropertyBaseObject.cpp
@@ -707,7 +707,7 @@ void FObjectPropertyBase::CheckValidObject(void* ValueAddress, TObjectPtr<UObjec
 	auto IsDeferringValueLoad = [&]() { return false; };
 #endif // USE_CIRCULAR_DEPENDENCY_LOAD_DEFERRING
 
-	if ((PropertyClass != nullptr) && !ObjectClass->IsChildOf(PropertyClass) && !ObjectClass->GetAuthoritativeClass()->IsChildOf(PropertyClass))
+	if ((PropertyClass != nullptr) && !ObjectClass->IsChildOf(PropertyClass) && !ObjectClass->GetAuthoritativeClass()->IsChildOf(PropertyClass) && !ObjectClass->ImplementsInterface(PropertyClass))
 	{
 			
 		// we could be in the middle of replacing references to the 
