Unshelved from Robert.Millar:

Do not cache weak pointer in persistent object pointer if resolving to a being-loaded object on the async loading thread

diff --git a/Engine/Source/Runtime/CoreUObject/Public/UObject/PersistentObjectPtr.h b/Engine/Source/Runtime/CoreUObject/Public/UObject/PersistentObjectPtr.h
--- a/Engine/Source/Runtime/CoreUObject/Public/UObject/PersistentObjectPtr.h
+++ b/Engine/Source/Runtime/CoreUObject/Public/UObject/PersistentObjectPtr.h
@@ -56,7 +56,14 @@ struct TPersistentObjectPtr
 		if (Object)
 		{
 			ObjectID = TObjectID::GetOrCreateIDForObject(Object);
-			WeakPtr = Object;
+			if (CanCacheObjectPointer(Object))
+			{
+				WeakPtr = Object;
+			}
+			else
+			{
+				WeakPtr.Reset();
+			}
 		}
 		else
 		{
@@ -102,7 +109,10 @@ struct TPersistentObjectPtr
 		if (!Object && ObjectID.IsValid())
 		{
 			Object = ObjectID.ResolveObject();
-			WeakPtr = Object;
+			if (CanCacheObjectPointer(Object))
+			{
+				WeakPtr = Object;
+			}
 
 			// Make sure it isn't garbage to match the default behavior of WeakPtr.Get() without looking it up again
 			return ::GetValid(Object);
@@ -125,10 +135,14 @@ struct TPersistentObjectPtr
 		if (!Object && ObjectID.IsValid())
 		{
 			Object = ObjectID.ResolveObject();
-			WeakPtr = Object;
+			FWeakObjectPtr LocalWeakPtr = Object;
 
+			if (CanCacheObjectPointer(Object))
+			{
+				WeakPtr = Object;
+			}
 			// Get the object again using the correct flag
-			Object = WeakPtr.Get(bEvenIfPendingKill);
+			Object = LocalWeakPtr.Get(bEvenIfPendingKill);
 		}
 		return Object;
 	}
@@ -224,6 +238,17 @@ struct TPersistentObjectPtr
 
 private:
 
+	// Returns whether the object pointer can be stored in WeakPtr for later retrieval
+	// For example, objects that are in the process of being async loaded may not be cached
+	inline bool CanCacheObjectPointer(const UObject* Ptr) const
+	{
+		if (IsInAsyncLoadingThread() && Ptr && Ptr->HasAnyInternalFlags(EInternalObjectFlags::AsyncLoading))
+		{
+			return false;
+		}
+		return true;
+	}
+
 	/** Once the object has been noticed to be loaded, this is set to the object weak pointer **/
 	mutable FWeakObjectPtr	WeakPtr;
 	/** Guid for the object this pointer points to or will point to. **/
