[ZenLoader]
 - Fix OnAssetLoaded so it is called in -game to match the behavior of the previous loaders

#rnx

diff --git a/Engine/Source/Runtime/CoreUObject/Private/Serialization/AsyncLoading2.cpp b/Engine/Source/Runtime/CoreUObject/Private/Serialization/AsyncLoading2.cpp
--- a/Engine/Source/Runtime/CoreUObject/Private/Serialization/AsyncLoading2.cpp
+++ b/Engine/Source/Runtime/CoreUObject/Private/Serialization/AsyncLoading2.cpp
@@ -7876,14 +7876,11 @@ EAsyncPackageState::Type FAsyncLoadingThread2::ProcessLoadedPackagesFromGameThre
 				{
 					// In the editor we need to find any assets and packages and add them to list for later callback
 					EditorCompletedUPackages.Add(Package->LinkerRoot);
-					if (GIsEditor)
+					for (UObject* Object : Package->ConstructedObjects)
 					{
-						for (UObject* Object : Package->ConstructedObjects)
+						if (Object->IsAsset())
 						{
-							if (Object->IsAsset())
-							{
-								EditorLoadedAssets.Add(Object);
-							}
+							EditorLoadedAssets.Add(Object);
 						}
 					}
 				}
