Make VerseTestScriptCmd safe when using the retry-transaction testing modes of AutoRTFM.

#okforversepublic

diff --git a/Engine/Source/Runtime/CoreUObject/Public/UObject/UObjectBase.h b/Engine/Source/Runtime/CoreUObject/Public/UObject/UObjectBase.h
--- a/Engine/Source/Runtime/CoreUObject/Public/UObject/UObjectBase.h
+++ b/Engine/Source/Runtime/CoreUObject/Public/UObject/UObjectBase.h
@@ -19,6 +19,7 @@
 #include "UObject/UObjectGlobals.h"
 #include "UObject/UnrealNames.h"
 #include "UObject/ObjectPtr.h"
+#include "AutoRTFM/AutoRTFM.h"
 
 class UClass;
 class UEnum;
@@ -258,7 +259,15 @@ private:
 	FORCEINLINE int32 GetFlagsInternal() const
 	{
 		static_assert(sizeof(int32) == sizeof(ObjectFlags), "Flags must be 32-bit for atomics.");
-		return FPlatformAtomics::AtomicRead_Relaxed((int32*)&ObjectFlags);
+
+		int32 Result = 0;
+
+		UE_AUTORTFM_OPEN2
+		{
+			Result = FPlatformAtomics::AtomicRead_Relaxed((int32*)&ObjectFlags);
+		};
+
+		return Result;
 	}
 
 	/** Flags used to track and report various object states. This needs to be 8 byte aligned on 32-bit
