[BugFix] Prefab instances aren't getting prefab values for properties that aren't overriden

diff --git a/Engine/Source/Runtime/CoreUObject/Private/UObject/PropertyBagRepository.cpp b/Engine/Source/Runtime/CoreUObject/Private/UObject/PropertyBagRepository.cpp
--- a/Engine/Source/Runtime/CoreUObject/Private/UObject/PropertyBagRepository.cpp
+++ b/Engine/Source/Runtime/CoreUObject/Private/UObject/PropertyBagRepository.cpp
@@ -892,6 +892,7 @@ void FPropertyBagRepository::CopyTaggedProperties(const UObject* Source, UObject
 	TArray<uint8> Buffer;
 	Buffer.Reserve(Source->GetClass()->GetStructureSize());
 	FObjectWriter Writer(Buffer);
+	Writer.ArNoDelta = true;
 	Source->GetClass()->SerializeTaggedProperties(Writer, (uint8*)Source, Source->GetClass(), nullptr);
 
 	FObjectReader Reader(Buffer);
