Split the StaticFindObjectFast fatal log into two so that we can tell which condition is triggering it

#rnx

diff --git a/Engine/Source/Runtime/CoreUObject/Private/UObject/UObjectGlobals.cpp b/Engine/Source/Runtime/CoreUObject/Private/UObject/UObjectGlobals.cpp
--- a/Engine/Source/Runtime/CoreUObject/Private/UObject/UObjectGlobals.cpp
+++ b/Engine/Source/Runtime/CoreUObject/Private/UObject/UObjectGlobals.cpp
@@ -412,10 +412,8 @@ int32 UpdateSuffixForNextNewObject(UObject* Parent, const UClass* Class, TFuncti
 //
 UObject* StaticFindObjectFast(UClass* ObjectClass, UObject* ObjectPackage, FName ObjectName, bool bExactClass, bool bAnyPackage, EObjectFlags ExclusiveFlags, EInternalObjectFlags ExclusiveInternalFlags)
 {
-	if (UE::IsSavingPackage(nullptr) || IsGarbageCollectingAndLockingUObjectHashTables())
-	{
-		UE_LOG(LogUObjectGlobals, Fatal,TEXT("Illegal call to StaticFindObjectFast() while serializing object data or garbage collecting!"));
-	}
+	UE_CLOG(UE::IsSavingPackage(nullptr), LogUObjectGlobals, Fatal, TEXT("Illegal call to StaticFindObjectFast() while serializing object data!"));
+	UE_CLOG(IsGarbageCollectingAndLockingUObjectHashTables(), LogUObjectGlobals, Fatal, TEXT("Illegal call to StaticFindObjectFast() while garbage collecting!"));
 
 	// We don't want to return any objects that are currently being background loaded unless we're using FindObject during async loading.
 	ExclusiveInternalFlags |= IsInAsyncLoadingThread() ? EInternalObjectFlags::None : EInternalObjectFlags::AsyncLoading;	
@@ -435,10 +433,8 @@ UObject* StaticFindObjectFast(UClass* ObjectClass, UObject* ObjectPackage, FName
 //
 UObject* StaticFindObjectFast(UClass* ObjectClass, UObject* ObjectPackage, FName ObjectName, bool bExactClass, EObjectFlags ExclusiveFlags, EInternalObjectFlags ExclusiveInternalFlags)
 {
-	if (UE::IsSavingPackage(nullptr) || IsGarbageCollectingAndLockingUObjectHashTables())
-	{
-		UE_LOG(LogUObjectGlobals, Fatal, TEXT("Illegal call to StaticFindObjectFast() while serializing object data or garbage collecting!"));
-	}
+	UE_CLOG(UE::IsSavingPackage(nullptr), LogUObjectGlobals, Fatal, TEXT("Illegal call to StaticFindObjectFast() while serializing object data!"));
+	UE_CLOG(IsGarbageCollectingAndLockingUObjectHashTables(), LogUObjectGlobals, Fatal, TEXT("Illegal call to StaticFindObjectFast() while garbage collecting!"));
 
 	// We don't want to return any objects that are currently being background loaded unless we're using FindObject during async loading.
 	ExclusiveInternalFlags |= IsInAsyncLoadingThread() ? EInternalObjectFlags::None : EInternalObjectFlags::AsyncLoading;
