Increased FProperty::IndexInOwner from uint16 to int32 to support types with many properties

diff --git a/Engine/Source/Runtime/CoreUObject/Private/UObject/Property.cpp b/Engine/Source/Runtime/CoreUObject/Private/UObject/Property.cpp
--- a/Engine/Source/Runtime/CoreUObject/Private/UObject/Property.cpp
+++ b/Engine/Source/Runtime/CoreUObject/Private/UObject/Property.cpp
@@ -1364,7 +1364,7 @@ int32 FProperty::SetupOffset()
 		UStruct* OwnerStruct = (UStruct*)OwnerUObject;
 		Offset_Internal = Align(OwnerStruct->GetPropertiesSize(), GetMinAlignment());
 	#if WITH_EDITORONLY_DATA
-		IndexInOwner = IntCastChecked<uint16>(OwnerStruct->TotalFieldCount);
+		IndexInOwner = OwnerStruct->TotalFieldCount;
 		OwnerStruct->TotalFieldCount += ArrayDim;
 	#endif
 	}
diff --git a/Engine/Source/Runtime/CoreUObject/Public/UObject/UnrealType.h b/Engine/Source/Runtime/CoreUObject/Public/UObject/UnrealType.h
--- a/Engine/Source/Runtime/CoreUObject/Public/UObject/UnrealType.h
+++ b/Engine/Source/Runtime/CoreUObject/Public/UObject/UnrealType.h
@@ -188,7 +188,7 @@ private:
 
 #if WITH_EDITORONLY_DATA
 	/** Index of the property within its owner, inclusive of base properties. Generated during Link(). */
-	uint16		IndexInOwner = 0;
+	int32		IndexInOwner = 0;
 #endif
 
 	// In memory variables (generated during Link()).
