Fall back to an FArchive and FReferenceCollector based method of reference collection when objects that are referencing only via weak references need to be included in the result. Also report EdGraphPin SubCategoryObject based references in editor. Fix various ReloadPackage bugs and remove iterative referencer searching.
[RN] Fixed updating of objects that only reference a blueprint or its instances via weak references

diff --git a/Engine/Source/Runtime/Core/Public/Containers/Set.h b/Engine/Source/Runtime/Core/Public/Containers/Set.h
--- a/Engine/Source/Runtime/Core/Public/Containers/Set.h
+++ b/Engine/Source/Runtime/Core/Public/Containers/Set.h
@@ -2181,7 +2181,7 @@ struct TSetPrivateFriend
 		// Load the set's new elements.
 		Ar << Set.Elements;
 
-		if(Ar.IsLoading() || Ar.IsModifyingWeakAndStrongReferences())
+		if(Ar.IsLoading() || (Ar.IsModifyingWeakAndStrongReferences() && !Ar.IsSaving()))
 		{
 			// Free the old hash.
 			Set.Hash.ResizeAllocation(0,0,sizeof(FSetElementId));
@@ -2200,7 +2200,7 @@ struct TSetPrivateFriend
  	{
 		Slot << Set.Elements;
 
-		if (Slot.GetUnderlyingArchive().IsLoading() || Slot.GetUnderlyingArchive().IsModifyingWeakAndStrongReferences())
+		if (Slot.GetUnderlyingArchive().IsLoading() || (Slot.GetUnderlyingArchive().IsModifyingWeakAndStrongReferences() && !Slot.GetUnderlyingArchive().IsSaving()))
 		{
 			// Free the old hash.
 			Set.Hash.ResizeAllocation(0, 0, sizeof(FSetElementId));
diff --git a/Engine/Source/Runtime/CoreUObject/Private/UObject/Class.cpp b/Engine/Source/Runtime/CoreUObject/Private/UObject/Class.cpp
--- a/Engine/Source/Runtime/CoreUObject/Private/UObject/Class.cpp
+++ b/Engine/Source/Runtime/CoreUObject/Private/UObject/Class.cpp
@@ -5534,7 +5534,7 @@ bool UClass::IsStructTrashed() const
 
 void UClass::Serialize( FArchive& Ar )
 {
-	if ( Ar.IsLoading() || Ar.IsModifyingWeakAndStrongReferences() )
+	if ( Ar.IsLoading() || (Ar.IsModifyingWeakAndStrongReferences() && !Ar.IsSaving()) )
 	{
 		// Rehash since SuperStruct will be serialized in UStruct::Serialize
 		UnhashObject(this);
@@ -5542,7 +5542,7 @@ void UClass::Serialize( FArchive& Ar )
 
 	Super::Serialize( Ar );
 
-	if ( Ar.IsLoading() || Ar.IsModifyingWeakAndStrongReferences() )
+	if ( Ar.IsLoading() || (Ar.IsModifyingWeakAndStrongReferences() && !Ar.IsSaving()) )
 	{
 		HashObject(this);
 	}
diff --git a/Engine/Source/Runtime/CoreUObject/Private/UObject/PackageReload.cpp b/Engine/Source/Runtime/CoreUObject/Private/UObject/PackageReload.cpp
--- a/Engine/Source/Runtime/CoreUObject/Private/UObject/PackageReload.cpp
+++ b/Engine/Source/Runtime/CoreUObject/Private/UObject/PackageReload.cpp
@@ -71,24 +71,11 @@ class FPackageReferencersHelper
 	 */
 	static void RecursiveRetrieveReferencers(const TArray<UObject*>& InReferencedObjects, TSet<TWeakObjectPtr<UObject>>& OutReferencingObjects)
 	{
-		TSet<UObject*> ReferencerObjects;
-
-		// Use the fast reference collector to recursively find referencers until no more new ones are found.
-		int32 LastObjectCount = 0;
-
 		TArray<UObject*> FoundReferencerObjects = FReferencerFinder::GetAllReferencers(InReferencedObjects, nullptr, EReferencerFinderFlags::None);
-		do
-		{
-			LastObjectCount = ReferencerObjects.Num();
-			ReferencerObjects.Append(FoundReferencerObjects);
-
-			FoundReferencerObjects = FReferencerFinder::GetAllReferencers(FoundReferencerObjects, nullptr, EReferencerFinderFlags::SkipInnerReferences);
-
-		} while (LastObjectCount != ReferencerObjects.Num());
 
 		// Convert them into weak pointers to deal with objects that get GC'd during the reload operation.
-		OutReferencingObjects.Reserve(ReferencerObjects.Num());
-		for (UObject* ReferencerObject : ReferencerObjects)
+		OutReferencingObjects.Reserve(OutReferencingObjects.Num() + FoundReferencerObjects.Num());
+		for (UObject* ReferencerObject : FoundReferencerObjects)
 		{
 			OutReferencingObjects.Emplace(MakeWeakObjectPtr(ReferencerObject));
 		}
@@ -111,6 +98,9 @@ public:
 			for (const FReloadPackageData& PackageToReloadData : InPackagesToReload)
 			{
 				ObjectsToReload.Add(PackageToReloadData.PackageToReload);
+				TArray<UObject*> Inners;
+				GetObjectsWithOuter(PackageToReloadData.PackageToReload, Inners);
+				ObjectsToReload.Append(Inners);
 			}
 
 			// Collect all other objects that reference any of them.
@@ -720,6 +710,7 @@ void ReloadPackages(const TArrayView<FReloadPackageData>& InPackagesToReload, TA
 			FScopedSlowTask FixingUpReferencesSlowTask((float)((NumPackagesInBatch * 4) + GUObjectArray.GetObjectArrayNum()), NSLOCTEXT("CoreUObject", "FixingUpReferences", "Fixing-Up References"));
 
 			// Pre-pass to notify things that the package old package is about to be fixed-up
+			TArray<UObject*> PotentialReferencers;
 			TMap<UObject*, PackageReloadInternal::FObjectAndPackageIndex> OldObjectToNewData;
 			for (int32 BatchPackageIndex = BatchStartIndex; BatchPackageIndex < PackageIndex; ++BatchPackageIndex)
 			{
@@ -736,6 +727,13 @@ void ReloadPackages(const TArrayView<FReloadPackageData>& InPackagesToReload, TA
 					for (const auto& ObjectMappingPair : NewPackageData.EventData->GetRepointedObjects())
 					{
 						OldObjectToNewData.Add(ObjectMappingPair.Key, PackageReloadInternal::FObjectAndPackageIndex(ObjectMappingPair.Value, BatchPackageIndex));
+						
+						// Also do reference replacement on the loaded objects, in case they have buggy
+						// Lifecycle functions that have created references to the DEADPACKAGE:
+						if (ObjectMappingPair.Value)
+						{
+							PotentialReferencers.Add(ObjectMappingPair.Value);
+						}
 					}
 				}
 			}
@@ -744,7 +742,6 @@ void ReloadPackages(const TArrayView<FReloadPackageData>& InPackagesToReload, TA
 			// todo: multi-thread this like FHotReloadModule::ReplaceReferencesToReconstructedCDOs?
 			//The FThreadSafeObjectIterator will lock the global UObject array, to avoid potential deadlock we simply build the list
 			//of the potential referencers and do the reference fix serialize outside of the FThreadSafeObjectIterator.
-			TArray<UObject*> PotentialReferencers;
 
 			PackageReferencersHelper.ForEachObject(
 				[&PotentialReferencers] (UObject* PotentialReferencer)
@@ -769,6 +766,9 @@ void ReloadPackages(const TArrayView<FReloadPackageData>& InPackagesToReload, TA
 				}
 			);
 
+			// Main reference replacement pass - this should fix as many references as possible
+			// allowing users to run code that has no visibility into the old data and
+			// therefore is not at risk of out of date reads or writes:
 			for (UObject* PotentialReferencer : PotentialReferencers)
 			{
 				FixingUpReferencesSlowTask.EnterProgressFrame(1.0f);
diff --git a/Engine/Source/Runtime/CoreUObject/Private/UObject/ReferencerFinder.cpp b/Engine/Source/Runtime/CoreUObject/Private/UObject/ReferencerFinder.cpp
--- a/Engine/Source/Runtime/CoreUObject/Private/UObject/ReferencerFinder.cpp
+++ b/Engine/Source/Runtime/CoreUObject/Private/UObject/ReferencerFinder.cpp
@@ -2,11 +2,13 @@
 
 #include "UObject/ReferencerFinder.h"
 
-#include "UObject/UObjectIterator.h"
-#include "UObject/FastReferenceCollector.h"
-#include "UObject/UObjectArray.h"
 #include "Async/ParallelFor.h"
 #include "HAL/IConsoleManager.h"
+#include "Serialization/ArchiveUObject.h"
+#include "UObject/FastReferenceCollector.h"
+#include "UObject/UObjectArray.h"
+#include "UObject/UObjectIterator.h"
+#include "UObject/UObjectGlobals.h"
 
 class FAllReferencesProcessor : public FSimpleReferenceProcessorBase
 {
@@ -42,11 +44,6 @@ public:
 			}
 		}
 	}
-
-	bool MarkWeakObjectReferenceForClearing(UObject** WeakReference, UObject* ReferenceOwner)
-	{
-		return (Flags & EReferencerFinderFlags::SkipWeakReferences) == EReferencerFinderFlags::SkipWeakReferences;
-	}
 };
 
 class FAllReferencesCollector : public UE::GC::TDefaultCollector<FAllReferencesProcessor>
@@ -55,9 +52,14 @@ class FAllReferencesCollector : public UE::GC::TDefaultCollector<FAllReferencesP
 public:
 	using Super::Super;
 
+	// MarkWeakObjectReferenceForClearing is a strange function, it's used by raw pointers that need 
+	// weak semantics. by returning true here we indicate that the reference does not need to be reported
+	// to the collector. That's the case for us because the GC information no longer contains tokens
+	// describing the location of weak references, so if we want to process weak references we need
+	// to use the reflection data:
 	virtual bool MarkWeakObjectReferenceForClearing(UObject** WeakReference, UObject* ReferenceOwner) override
 	{
-		return Processor.MarkWeakObjectReferenceForClearing(WeakReference, ReferenceOwner);
+		return true;
 	}
 };
 
@@ -78,75 +80,282 @@ void FReferencerFinder::NotifyRegistrationComplete()
 	GUObjectRegistrationComplete = true;
 }
 
-TArray<UObject*> FReferencerFinder::GetAllReferencers(const TArray<UObject*>& Referencees, const TSet<UObject*>* ObjectsToIgnore, EReferencerFinderFlags Flags)
+namespace UE::Private
 {
-	return GetAllReferencers(TSet<UObject*>(Referencees), ObjectsToIgnore, Flags);
-}
 
-TArray<UObject*> FReferencerFinder::GetAllReferencers(const TSet<UObject*>& Referencees, const TSet<UObject*>* ObjectsToIgnore, EReferencerFinderFlags Flags)
+struct FReferencerFinderArchive : public FArchiveUObject
+								, public FReferenceCollector
 {
-	TArray<UObject*> Ret;
-	if(Referencees.Num() > 0)
+private:
+	UObject* CurrentObject = nullptr;
+	const TSet<UObject*>& Targets;
+	bool bDoesReferenceAnyTargets = false;
+	const bool bSkipInnerReferences;
+
+public:
+	FReferencerFinderArchive(const TSet<UObject*>& InTargets, bool bInSkipInnerReferences)
+		: FArchiveUObject()
+		, Targets(InTargets)
+		, bDoesReferenceAnyTargets(false)
+		, bSkipInnerReferences(bInSkipInnerReferences)
+	{
+		SetIsSaving(true);
+		SetShouldSkipCompilingAssets(true);
+		SetWantBinaryPropertySerialization(true);
+		SetUseUnversionedPropertySerialization(true);
+		SetShouldSkipUpdateCustomVersion(true);
+		// we aren't modifying, but we are searching weak references,
+		// which per the comment on IsModifyingWeakAndStrongReferences() 
+		// is a valid use of this flag. This pattern has been trailblazed
+		// by FFindReferencersArchive and FFindLightmapsArchive:
+		ArIsModifyingWeakAndStrongReferences = true;
+		ArIsObjectReferenceCollector = true;
+		ArShouldSkipBulkData = true;
+	}
+
+	// Crawls the provided object, returning true if we find any references to the 
+	// objects in Targets:
+	bool SearchForReferencesToTargets(UObject* Obj)
+	{
+		checkSlow(Obj);
+		CurrentObject = Obj;
+		bDoesReferenceAnyTargets = false;
+		if(UStruct* AsStruct = Cast<UStruct>(Obj))
+		{
+			// UStruct serialize overrides (like UControlRigBlueprintGeneratedClass::Serialize)
+			// are not thread safe, but UStruct::SerializeBin will not report any
+			// references in the reflection data. This is a compromise.
+			AsStruct->UStruct::Serialize(*this);
+		}
+		else
+		{
+			// This could miss some things in a user serialize functions,
+			// but there exist many user Serialize routines that are not threadsafe
+			// and deadlocking would be disastrous.
+			Obj->GetClass()->SerializeBin(*this, Obj);
+		}
+		// Compensate for missed stream serialized reference 
+		// by calling ARO:
+		Obj->GetClass()->CallAddReferencedObjects(Obj, *this); 
+
+		return bDoesReferenceAnyTargets;
+	}
+	
+private:
+	// FReferenceCollector:
+	virtual bool IsIgnoringArchetypeRef() const override { return false; }
+	virtual bool IsIgnoringTransient() const override { return false; }
+	virtual void HandleObjectReference(UObject*& InObject, const UObject* InReferencingObject, const FProperty* InReferencingProperty)
 	{
-		FCriticalSection ResultCritical;
+		if (InObject != nullptr && 
+			Targets.Contains(InObject) )
+		{
+			ProcessTargetObject(InObject);
+		}
+	}
 
-		// Lock the global array so that nothing can add UObjects while we're iterating over it
-		GUObjectArray.LockInternalArray();
+	// FArchiveUObject
+	virtual FArchive& operator<<(UObject*& ObjRef) override
+	{
+		if (ObjRef != nullptr && 
+			Targets.Contains(ObjRef) )
+		{
+			ProcessTargetObject(ObjRef);
+		}
 
-		const int32 MaxNumberOfObjects = GUObjectArray.GetObjectArrayNum();
-		const int32 NumThreads = GetNumCollectReferenceWorkers();
-		const int32 NumberOfObjectsPerThread = (MaxNumberOfObjects / NumThreads) + 1;
+		return *this;
+	}
 
-		ParallelFor(NumThreads, [&Referencees, ObjectsToIgnore, &ResultCritical, &Ret, NumberOfObjectsPerThread, NumThreads, MaxNumberOfObjects, Flags](int32 ThreadIndex)
+	void ProcessTargetObject(UObject* Object)
+	{
+		if (bSkipInnerReferences && CurrentObject->IsIn(Object))
 		{
-			TSet<UObject*> ThreadResult;
-			FAllReferencesProcessor Processor(Referencees, Flags, ThreadResult);
-			TArray<UObject*> ObjectsToSerialize;
-			ObjectsToSerialize.Reserve(NumberOfObjectsPerThread);
+			return;
+		}
 
-			const int32 FirstObjectIndex = ThreadIndex * NumberOfObjectsPerThread;
-			const int32 NumObjects = (ThreadIndex < (NumThreads - 1)) ? NumberOfObjectsPerThread : (MaxNumberOfObjects - (NumThreads - 1)*NumberOfObjectsPerThread);
+		bDoesReferenceAnyTargets = true;
+	}
+};
 
-			// First cache all potential referencers
-			for (int32 ObjectIndex = 0; ObjectIndex < NumObjects && (FirstObjectIndex + ObjectIndex) < MaxNumberOfObjects; ++ObjectIndex)
-			{
-				FUObjectItem& ObjectItem = GUObjectArray.GetObjectItemArrayUnsafe()[FirstObjectIndex + ObjectIndex];
-				if (ObjectItem.Object && !ObjectItem.IsUnreachable())
-				{
-					UObject* PotentialReferencer = static_cast<UObject*>(ObjectItem.Object);
-					if (ObjectsToIgnore && ObjectsToIgnore->Contains(PotentialReferencer))
-					{
-						continue;
-					}
+static TArray<UObject*> GetAllReferencersIncludingWeak(
+	const TSet<UObject*>& Referencees, const TSet<UObject*>* ObjectsToIgnore, const bool bSkipInnerReferences)
+{
+	// Lock the global array so that nothing can add UObjects while we're iterating over it
+	GUObjectArray.LockInternalArray();
 
-					if (!Referencees.Contains(PotentialReferencer))
-					{
-						ObjectsToSerialize.Add(PotentialReferencer);
-					}
-				}
-			}
+	const int32 MaxNumberOfObjects = GUObjectArray.GetObjectArrayNum();
+	const int32 NumThreads = GetNumCollectReferenceWorkers();
+	const int32 NumberOfObjectsPerThread = (MaxNumberOfObjects / NumThreads) + 1;
+	
+	// Allocate per thread results, in this case each thread will produce a list of referencers
+	TUniquePtr<TArray<UObject*>[]> ThreadResultsAlloc = MakeUnique<TArray<UObject*>[]>(NumThreads);
+	TArrayView<TArray<UObject*>> ThreadResults(ThreadResultsAlloc.Get(), NumThreads);
 
-			FGCArrayStruct ArrayStruct;
-			ArrayStruct.SetInitialObjectsUnpadded(ObjectsToSerialize);
+	ParallelFor(NumThreads, 
+		[&Referencees, 
+		ObjectsToIgnore, 
+		&ThreadResults,
+		NumberOfObjectsPerThread, 
+		NumThreads, 
+		MaxNumberOfObjects, 
+		bSkipInnerReferences](int32 ThreadIndex)
+	{
+		FReferencerFinderArchive ReferenceFinderArchive(Referencees, bSkipInnerReferences);
+		TArray<UObject*>& ThreadResult = ThreadResults[ThreadIndex];
+		
+		// Process the block of objects assigned to this thread:
+		const int32 FirstObjectIndex = ThreadIndex * NumberOfObjectsPerThread;
+		const int32 EndIndex = FMath::Min(FirstObjectIndex + NumberOfObjectsPerThread, MaxNumberOfObjects);
+		for (int32 Index = FirstObjectIndex; Index < EndIndex; ++Index)
+		{
+			FUObjectItem& ObjectItem = GUObjectArray.GetObjectItemArrayUnsafe()[Index];
+			if (!ObjectItem.Object || ObjectItem.IsUnreachable())
+			{
+				continue;
+			}
+			UObject* PotentialReferencer = static_cast<UObject*>(ObjectItem.Object);
 			
+			if(!IsValid(PotentialReferencer))
 			{
-				// Since ReferenceCollector is configured to automatically assemble reference token streams
-				// for classes that require it, make sure GC is locked because UClass::AssembleReferenceTokenStream requires it
-				FGCScopeGuard GCGuard;
-				// Now check if any of the potential referencers is referencing any of the referencees
-				CollectReferences(Processor, ArrayStruct);
+				continue;
 			}
 
-			if (ThreadResult.Num())
+			const bool bIsIgnored = ObjectsToIgnore && ObjectsToIgnore->Contains(PotentialReferencer);
+			const bool bIsReferencee = Referencees.Contains(PotentialReferencer);
+			if ( bIsIgnored || bIsReferencee )
 			{
-				// We found objects referencing some of the referencess so add them to the final results array
-				FScopeLock ResultLock(&ResultCritical);
-				Ret.Append(ThreadResult.Array());
+				continue;
 			}
-		}, (GUObjectRegistrationComplete && GAllowParallelReferenceCollection) ? EParallelForFlags::None : EParallelForFlags::ForceSingleThread);
 
-		// Release the global array lock
-		GUObjectArray.UnlockInternalArray();
+			if (ReferenceFinderArchive.SearchForReferencesToTargets(PotentialReferencer))
+			{
+				ThreadResult.Add(PotentialReferencer);
+			}
+		}
+	},  (GUObjectRegistrationComplete && GAllowParallelReferenceCollection) ? 
+			EParallelForFlags::None : EParallelForFlags::ForceSingleThread);
+			
+	// Release the global array lock
+	GUObjectArray.UnlockInternalArray();
+
+	TArray<UObject*> FinalResult;
+	for(const TArray<UObject*>& ThreadResult : ThreadResults)
+	{
+		FinalResult.Append(ThreadResult);
 	}
+
+	return FinalResult;
+}
+
+static TArray<UObject*> GetAllReferencersExcludingWeak(
+	const TSet<UObject*>& Referencees, const TSet<UObject*>* ObjectsToIgnore, const EReferencerFinderFlags Flags)
+{
+	// We can use the faster GC schema to find references when we don't need to include weak 
+	// references:
+	TArray<UObject*> Ret;
+	FCriticalSection ResultCritical;
+
+	// Lock the global array so that nothing can add UObjects while we're iterating over it
+	GUObjectArray.LockInternalArray();
+
+	const int32 MaxNumberOfObjects = GUObjectArray.GetObjectArrayNum();
+	const int32 NumThreads = GetNumCollectReferenceWorkers();
+	const int32 NumberOfObjectsPerThread = (MaxNumberOfObjects / NumThreads) + 1;
+
+	ParallelFor(NumThreads, [&Referencees, ObjectsToIgnore, &ResultCritical, &Ret, NumberOfObjectsPerThread, NumThreads, MaxNumberOfObjects, Flags](int32 ThreadIndex)
+	{
+		TSet<UObject*> ThreadResult;
+		FAllReferencesProcessor Processor(Referencees, Flags, ThreadResult);
+		TArray<UObject*> ObjectsToSerialize;
+		ObjectsToSerialize.Reserve(NumberOfObjectsPerThread);
+		
+		const int32 FirstObjectIndex = ThreadIndex * NumberOfObjectsPerThread;
+		const int32 EndIndex = FMath::Min(FirstObjectIndex + NumberOfObjectsPerThread, MaxNumberOfObjects);
+		
+		for (int32 Index = FirstObjectIndex; Index < EndIndex; ++Index)
+		{
+			FUObjectItem& ObjectItem = GUObjectArray.GetObjectItemArrayUnsafe()[Index];
+			if (ObjectItem.Object && !ObjectItem.IsUnreachable())
+			{
+				UObject* PotentialReferencer = static_cast<UObject*>(ObjectItem.Object);
+				if (ObjectsToIgnore && ObjectsToIgnore->Contains(PotentialReferencer))
+				{
+					continue;
+				}
+
+				if (!Referencees.Contains(PotentialReferencer))
+				{
+					ObjectsToSerialize.Add(PotentialReferencer);
+				}
+			}
+		}
+
+		FGCArrayStruct ArrayStruct;
+		ArrayStruct.SetInitialObjectsUnpadded(ObjectsToSerialize);
+			
+		{
+			// Since ReferenceCollector is configured to automatically assemble reference token streams
+			// for classes that require it, make sure GC is locked because UClass::AssembleReferenceTokenStream requires it
+			FGCScopeGuard GCGuard;
+			// Now check if any of the potential referencers is referencing any of the referencees
+			CollectReferences(Processor, ArrayStruct);
+		}
+
+		if (ThreadResult.Num())
+		{
+			// We found objects referencing some of the referencees so add them to the final results array
+			FScopeLock ResultLock(&ResultCritical);
+			Ret.Append(ThreadResult.Array());
+		}
+	}, (GUObjectRegistrationComplete && GAllowParallelReferenceCollection) ? EParallelForFlags::None : EParallelForFlags::ForceSingleThread);
+
+	// Release the global array lock
+	GUObjectArray.UnlockInternalArray();
 	return Ret;
+}
+
+// Feature flag for disabling this reworked serialize based reference finding, as there is a risk
+// of undiscovered race conditions:
+static bool GUseSerializeToFindWeakReferencers = true;
+static FAutoConsoleVariableRef CVar_UseSerializeToFindWeakReferencers(
+	TEXT("UObject.UseSerializeToFindWeakReferencers"),
+	GUseSerializeToFindWeakReferencers,
+	TEXT("If true use Serialize routines to find objects that are referencing only via weak references - set to false to only use GC schema for reference finding")
+);
+}// namespace UE::Private
+
+TArray<UObject*> FReferencerFinder::GetAllReferencers(
+	const TArray<UObject*>& Referencees, const TSet<UObject*>* ObjectsToIgnore, EReferencerFinderFlags Flags)
+{
+	return GetAllReferencers(TSet<UObject*>(Referencees), ObjectsToIgnore, Flags);
+}
+
+TArray<UObject*> FReferencerFinder::GetAllReferencers(
+	const TSet<UObject*>& Referencees, const TSet<UObject*>* ObjectsToIgnore, EReferencerFinderFlags Flags)
+{
+	if(Referencees.Num() == 0)
+	{
+		return TArray<UObject*>();
+	}
+
+	const bool bSkipWeakReferences = (Flags & EReferencerFinderFlags::SkipWeakReferences) 
+		!= EReferencerFinderFlags::None;
+	const bool bSkipInnerReferences = (Flags & EReferencerFinderFlags::SkipInnerReferences) 
+		!= EReferencerFinderFlags::None;
+	if(!bSkipWeakReferences && UE::Private::GUseSerializeToFindWeakReferencers)
+	{
+		return UE::Private::GetAllReferencersIncludingWeak(
+			Referencees, 
+			ObjectsToIgnore, 
+			bSkipInnerReferences
+		);
+	}
+	else
+	{
+		return UE::Private::GetAllReferencersExcludingWeak(
+			Referencees,
+			ObjectsToIgnore,
+			Flags
+		);
+	}
 }
\ No newline at end of file
diff --git a/Engine/Source/Runtime/Engine/Private/EdGraph/EdGraphPin.cpp b/Engine/Source/Runtime/Engine/Private/EdGraph/EdGraphPin.cpp
--- a/Engine/Source/Runtime/Engine/Private/EdGraph/EdGraphPin.cpp
+++ b/Engine/Source/Runtime/Engine/Private/EdGraph/EdGraphPin.cpp
@@ -851,6 +851,12 @@ void UEdGraphPin::AddStructReferencedObjects(class FReferenceCollector& Collecto
 {
 	Collector.AddReferencedObject(PinType.PinSubCategoryMemberReference.MemberParent);
 	Collector.AddReferencedObject(DefaultObject);
+#if WITH_EDITOR
+	// EdGraphPins are not serialized via the reflection system, but we need visibility 
+	// into pin references for asset reload and potentially reinstancing, so
+	// I'm publishing this as an owned reference in editor:
+	Collector.AddReferencedObject(PinType.PinSubCategoryObject);
+#endif
 }
 
 void UEdGraphPin::SerializeAsOwningNode(FArchive& Ar, TArray<UEdGraphPin*>& ArrayRef)
