Allow TBaseStructure to be used with FDoubleRange and FDoubleRangeBound

#tests Standard Build + Cooked Client testing
#rnx

diff --git a/Engine/Source/Runtime/CoreUObject/Private/UObject/Class.cpp b/Engine/Source/Runtime/CoreUObject/Private/UObject/Class.cpp
--- a/Engine/Source/Runtime/CoreUObject/Private/UObject/Class.cpp
+++ b/Engine/Source/Runtime/CoreUObject/Private/UObject/Class.cpp
@@ -7521,6 +7521,18 @@ UScriptStruct* TBaseStructure<FFloatRange>::Get()
 	return ScriptStruct;
 }
 
+UScriptStruct* TBaseStructure<FDoubleRangeBound>::Get()
+{
+	static UScriptStruct* ScriptStruct = StaticGetBaseStructureInternal(TEXT("DoubleRangeBound"));
+	return ScriptStruct;
+}
+
+UScriptStruct* TBaseStructure<FDoubleRange>::Get()
+{
+	static UScriptStruct* ScriptStruct = StaticGetBaseStructureInternal(TEXT("DoubleRange"));
+	return ScriptStruct;
+}
+
 UScriptStruct* TBaseStructure<FInt32RangeBound>::Get()
 {
 	static UScriptStruct* ScriptStruct = StaticGetBaseStructureInternal(TEXT("Int32RangeBound"));
