- Changed "PropertiesSetBySerialization" to "SerializedProperties" and made the flag get set even for properties that were skipped during serialization. This way it's a better reflection of which properties were in the old class
- removed GMarkPropertiesSetBySerialization and added bTrackSerializedProperties flag to FUObjectSerializeContext to avoid race conditinons

diff --git a/Engine/Plugins/Experimental/InstanceDataObjectFixupTool/Source/Private/InstanceDataObjectFixupPanel.cpp b/Engine/Plugins/Experimental/InstanceDataObjectFixupTool/Source/Private/InstanceDataObjectFixupPanel.cpp
--- a/Engine/Plugins/Experimental/InstanceDataObjectFixupTool/Source/Private/InstanceDataObjectFixupPanel.cpp
+++ b/Engine/Plugins/Experimental/InstanceDataObjectFixupTool/Source/Private/InstanceDataObjectFixupPanel.cpp
@@ -933,7 +933,7 @@ static void InitRedirectedPropertyTreeRec(const TSharedPtr<FRedirectedPropertyNo
 	{
 		if (Property->ArrayDim == 1)
 		{
-			if (UE::FPropertyBagRepository::WasPropertySetBySerialization(Struct, StructValue, Property))
+			if (UE::FPropertyBagRepository::WasPropertyValueSerialized(Struct, StructValue, Property))
 			{
 				const TSharedPtr<FRedirectedPropertyNode>& ChildNode = Node->FindOrAdd(FPropertyInfo(Property));
 				void* Value = Property->ContainerPtrToValuePtr<void>(StructValue);
@@ -944,7 +944,7 @@ static void InitRedirectedPropertyTreeRec(const TSharedPtr<FRedirectedPropertyNo
 		{
 			for (int32 StaticArrayIndex = 0; StaticArrayIndex < Property->ArrayDim; ++StaticArrayIndex)
             {
-            	if (UE::FPropertyBagRepository::WasPropertySetBySerialization(Struct, StructValue, Property, StaticArrayIndex))
+            	if (UE::FPropertyBagRepository::WasPropertyValueSerialized(Struct, StructValue, Property, StaticArrayIndex))
             	{
             		const TSharedPtr<FRedirectedPropertyNode>& ChildNode = Node->FindOrAdd(FPropertyInfo(Property, StaticArrayIndex));
             		void* Value = Property->ContainerPtrToValuePtr<void>(StructValue, StaticArrayIndex);
diff --git a/Engine/Source/Runtime/CoreUObject/Internal/UObject/InstanceDataObjectUtils.h b/Engine/Source/Runtime/CoreUObject/Internal/UObject/InstanceDataObjectUtils.h
--- a/Engine/Source/Runtime/CoreUObject/Internal/UObject/InstanceDataObjectUtils.h
+++ b/Engine/Source/Runtime/CoreUObject/Internal/UObject/InstanceDataObjectUtils.h
@@ -30,13 +30,12 @@ bool StructContainsLooseProperties(const UStruct* Struct);
 /** Generate a UClass that contains the union of the properties of PropertyTree and OwnerClass. */
 UClass* CreateInstanceDataObjectClass(const FPropertyPathNameTree* PropertyTree, UClass* OwnerClass, UObject* Outer);
 
-extern bool GMarkPropertiesSetBySerialization;
 /** Notify that a property in an struct was set when the struct was deserialized. */
-void MarkPropertySetBySerialization(const UStruct* Struct, void* StructData, const FProperty* Property, int32 ArrayIndex = 0);
+void MarkPropertyValueSerialized(const UStruct* Struct, void* StructData, const FProperty* Property, int32 ArrayIndex = 0);
 /** Query whether a property in the struct was set when the struct was deserialized. */
-bool WasPropertySetBySerialization(const UStruct* Struct, const void* StructData, const FProperty* Property, int32 ArrayIndex = 0);
+bool WasPropertyValueSerialized(const UStruct* Struct, const void* StructData, const FProperty* Property, int32 ArrayIndex = 0);
 /** Copy whether each property was set by serialization from one IDO to another. */
-void CopyPropertySetBySerializationData(const FFieldVariant& OldField, void* OldDataPtr, const FFieldVariant& NewField, void* NewDataPtr);
+void CopyPropertyValueSerializedData(const FFieldVariant& OldField, void* OldDataPtr, const FFieldVariant& NewField, void* NewDataPtr);
 
 /** Query whether the property value is initialized for a property in the struct. */
 UE_API bool IsPropertyValueInitialized(const UStruct* Struct, void* StructData, const FProperty* Property, int32 ArrayIndex = 0);
diff --git a/Engine/Source/Runtime/CoreUObject/Internal/UObject/PropertyBagRepository.h b/Engine/Source/Runtime/CoreUObject/Internal/UObject/PropertyBagRepository.h
--- a/Engine/Source/Runtime/CoreUObject/Internal/UObject/PropertyBagRepository.h
+++ b/Engine/Source/Runtime/CoreUObject/Internal/UObject/PropertyBagRepository.h
@@ -145,7 +145,7 @@ public:
 	COREUOBJECT_API const UObject* FindInstanceForDataObject(const UObject* InstanceDataObject) const;
 
 	// query whether a property in Struct was set when the struct was deserialized
-	COREUOBJECT_API static bool WasPropertySetBySerialization(const UStruct* Struct, const void* StructData, const FProperty* Property, int32 ArrayIndex = 0);
+	COREUOBJECT_API static bool WasPropertyValueSerialized(const UStruct* Struct, const void* StructData, const FProperty* Property, int32 ArrayIndex = 0);
 
 	// FGCObject interface
 	virtual void AddReferencedObjects(FReferenceCollector& Collector) override;
@@ -215,6 +215,8 @@ struct FScopedIDOSerializationContext
 	TOptional<TGuardValue<bool>> ScopedTrackSerializedPropertyPath;
 	TOptional<TGuardValue<bool>> ScopedSerializeUnknownProperty;
 	TOptional<TGuardValue<bool>> ScopedImpersonateProperties;
+	TOptional<TGuardValue<bool>> ScopedTrackInitializedProperties;
+	TOptional<TGuardValue<bool>> ScopedTrackSerializedProperties;
 	TOptional<TGuardValue<UObject*>> ScopedSerializedObject;
 
 private:
diff --git a/Engine/Source/Runtime/CoreUObject/Private/UObject/Class.cpp b/Engine/Source/Runtime/CoreUObject/Private/UObject/Class.cpp
--- a/Engine/Source/Runtime/CoreUObject/Private/UObject/Class.cpp
+++ b/Engine/Source/Runtime/CoreUObject/Private/UObject/Class.cpp
@@ -1713,6 +1713,13 @@ void UStruct::SerializeVersionedTaggedProperties(FStructuredArchive::FSlot Slot,
 
 				Tag.SetProperty(Property);
 
+#if WITH_EDITORONLY_DATA
+				if (SerializeContext->bTrackSerializedProperties && Property)
+				{
+					UE::MarkPropertyValueSerialized(this, Data, Property, Tag.ArrayIndex);
+				}
+#endif
+
 				bool bTryStoreUnknownPropertyPath = false;
 
 				if (Property)
@@ -1795,13 +1802,6 @@ void UStruct::SerializeVersionedTaggedProperties(FStructuredArchive::FSlot Slot,
 								checkNoEntry();
 								break;
 						}
-
-					#if WITH_EDITORONLY_DATA
-						if (bAdvanceProperty && SerializeContext->bImpersonateProperties)
-						{
-							UE::MarkPropertySetBySerialization(this, Data, Property, Tag.ArrayIndex);
-						}
-					#endif
 					}
 				}
 				else
diff --git a/Engine/Source/Runtime/CoreUObject/Private/UObject/InstanceDataObjectUtils.cpp b/Engine/Source/Runtime/CoreUObject/Private/UObject/InstanceDataObjectUtils.cpp
--- a/Engine/Source/Runtime/CoreUObject/Private/UObject/InstanceDataObjectUtils.cpp
+++ b/Engine/Source/Runtime/CoreUObject/Private/UObject/InstanceDataObjectUtils.cpp
@@ -17,7 +17,7 @@
 #include "UObject/UnrealType.h"
 
 static const FName NAME_InitializedValues(ANSITEXTVIEW("_InitializedValues"));
-static const FName NAME_ValuesSetBySerialization(ANSITEXTVIEW("_ValuesSetBySerialization"));
+static const FName NAME_SerializedValues(ANSITEXTVIEW("_SerializedValues"));
 
 /** Type used for InstanceDataObject classes. */
 class UInstanceDataObjectClass final : public UClass
@@ -26,7 +26,7 @@ public:
 	DECLARE_CASTED_CLASS_INTRINSIC(UInstanceDataObjectClass, UClass, CLASS_Transient, TEXT("/Script/CoreUObject"), CASTCLASS_UClass)
 
 	FByteProperty* InitializedValuesProperty = nullptr;
-	FByteProperty* ValuesSetBySerializationProperty = nullptr;
+	FByteProperty* SerializedValuesProperty = nullptr;
 };
 
 IMPLEMENT_CORE_INTRINSIC_CLASS(UInstanceDataObjectClass, UClass,
@@ -43,7 +43,7 @@ public:
 	FGuid GetCustomGuid() const final { return Guid; }
 
 	FByteProperty* InitializedValuesProperty = nullptr;
-	FByteProperty* ValuesSetBySerializationProperty = nullptr;
+	FByteProperty* SerializedValuesProperty = nullptr;
 	FGuid Guid;
 };
 
@@ -90,7 +90,7 @@ uint32 UInstanceDataObjectStruct::GetStructTypeHash(const void* Src) const
 	uint32 ValueHash = 0;
 	for (TFieldIterator<const FProperty> It(this); It; ++It)
 	{
-		if (It->GetFName() == NAME_InitializedValues || It->GetFName() == NAME_ValuesSetBySerialization)
+		if (It->GetFName() == NAME_InitializedValues || It->GetFName() == NAME_SerializedValues)
 		{
 			continue;
 		}
@@ -132,7 +132,6 @@ namespace UE
 	static const FName NAME_VerseClass(ANSITEXTVIEW("VerseClass"));
 	static const FName NAME_IDOMapKey(ANSITEXTVIEW("Key"));
 	static const FName NAME_IDOMapValue(ANSITEXTVIEW("Value"));
-	bool GMarkPropertiesSetBySerialization = true;
 
 	bool bEnableIDOSupport = false;
 	FAutoConsoleVariableRef EnableIDOSupportCVar(
@@ -407,7 +406,7 @@ namespace UE
 			for (FPropertyPathNameTree::FConstIterator It = PropertyTree->CreateConstIterator(); It; ++It)
 			{
 				FName Name = It.GetName();
-				if (Name == NAME_InitializedValues || Name == NAME_ValuesSetBySerialization)
+				if (Name == NAME_InitializedValues || Name == NAME_SerializedValues)
 				{
 					// In rare cases, these hidden properties will get serialized even though they are transient.
 					// Ignore them here since they are generated below.
@@ -439,24 +438,24 @@ namespace UE
 
 		// Add hidden byte array properties to record whether its sibling properties were initialized or set by serialization.
 		FByteProperty* InitializedValuesProperty = CastFieldChecked<FByteProperty>(FByteProperty::Construct(Result, NAME_InitializedValues, RF_Transient | RF_MarkAsNative));
-		FByteProperty* ValuesSetBySerializationProperty = CastFieldChecked<FByteProperty>(FByteProperty::Construct(Result, NAME_ValuesSetBySerialization, RF_Transient | RF_MarkAsNative));
+		FByteProperty* SerializedValuesProperty = CastFieldChecked<FByteProperty>(FByteProperty::Construct(Result, NAME_SerializedValues, RF_Transient | RF_MarkAsNative));
 		{
 			InitializedValuesProperty->SetPropertyFlags(CPF_Transient | CPF_EditorOnly | CPF_NativeAccessSpecifierPrivate);
-			ValuesSetBySerializationProperty->SetPropertyFlags(CPF_Transient | CPF_EditorOnly | CPF_NativeAccessSpecifierPrivate);
+			SerializedValuesProperty->SetPropertyFlags(CPF_Transient | CPF_EditorOnly | CPF_NativeAccessSpecifierPrivate);
 			Result->AddCppProperty(InitializedValuesProperty);
-			Result->AddCppProperty(ValuesSetBySerializationProperty);
+			Result->AddCppProperty(SerializedValuesProperty);
 		}
 
 		// Store generated properties to avoid scanning every property to find it when it is needed.
 		if (UInstanceDataObjectClass* IdoClass = Cast<UInstanceDataObjectClass>(Result))
 		{
 			IdoClass->InitializedValuesProperty = InitializedValuesProperty;
-			IdoClass->ValuesSetBySerializationProperty = ValuesSetBySerializationProperty;
+			IdoClass->SerializedValuesProperty = SerializedValuesProperty;
 		}
 		else if (UInstanceDataObjectStruct* IdoStruct = Cast<UInstanceDataObjectStruct>(Result))
 		{
 			IdoStruct->InitializedValuesProperty = InitializedValuesProperty;
-			IdoStruct->ValuesSetBySerializationProperty = ValuesSetBySerializationProperty;
+			IdoStruct->SerializedValuesProperty = SerializedValuesProperty;
 		}
 
 		// AddCppProperty expects reverse property order for StaticLink to work correctly
@@ -473,7 +472,7 @@ namespace UE
 		}
 		const int32 PropertyCountBytes = FMath::Max(1, FMath::DivideAndRoundUp(PropertyCount, 8));
 		InitializedValuesProperty->ArrayDim = PropertyCountBytes;
-		ValuesSetBySerializationProperty->ArrayDim = PropertyCountBytes;
+		SerializedValuesProperty->ArrayDim = PropertyCountBytes;
 
 		Result->Bind();
 		Result->StaticLink(/*bRelinkExistingProperties*/true);
@@ -534,61 +533,57 @@ namespace UE
 		return Result;
 	}
 
-	static const FByteProperty* FindValuesSetBySerializationProperty(const UStruct* Struct)
+	static const FByteProperty* FindSerializedValuesProperty(const UStruct* Struct)
 	{
 		if (const UInstanceDataObjectClass* IdoClass = Cast<UInstanceDataObjectClass>(Struct))
 		{
-			return IdoClass->ValuesSetBySerializationProperty;
+			return IdoClass->SerializedValuesProperty;
 		}
 		if (const UInstanceDataObjectStruct* IdoStruct = Cast<UInstanceDataObjectStruct>(Struct))
 		{
-			return IdoStruct->ValuesSetBySerializationProperty;
+			return IdoStruct->SerializedValuesProperty;
 		}
-		return CastField<FByteProperty>(Struct->FindPropertyByName(NAME_ValuesSetBySerialization));
+		return CastField<FByteProperty>(Struct->FindPropertyByName(NAME_SerializedValues));
 	}
 
-	void MarkPropertySetBySerialization(const UStruct* Struct, void* StructData, const FProperty* Property, int32 ArrayIndex)
+	void MarkPropertyValueSerialized(const UStruct* Struct, void* StructData, const FProperty* Property, int32 ArrayIndex)
 	{
-		if (!GMarkPropertiesSetBySerialization)
-		{
-			return;
-		}
-		if (const FByteProperty* ValuesSetBySerializationProperty = FindValuesSetBySerializationProperty(Struct))
+		if (const FByteProperty* SerializedValuesProperty = FindSerializedValuesProperty(Struct))
 		{
 			const int32 PropertyIndex = Property->GetIndexInOwner() + ArrayIndex;
 			const int32 ByteIndex = PropertyIndex / 8;
 			const int32 BitOffset = PropertyIndex % 8;
-			if (ByteIndex < ValuesSetBySerializationProperty->ArrayDim)
+			if (ByteIndex < SerializedValuesProperty->ArrayDim)
 			{
-				uint8* PropertyDataPtr = ValuesSetBySerializationProperty->ContainerPtrToValuePtr<uint8>(StructData, ByteIndex);
+				uint8* PropertyDataPtr = SerializedValuesProperty->ContainerPtrToValuePtr<uint8>(StructData, ByteIndex);
 				*PropertyDataPtr |= (1 << BitOffset);
 			}
 		}
 	}
 
-	bool WasPropertySetBySerialization(const UStruct* Struct, const void* StructData, const FProperty* Property, int32 ArrayIndex)
+	bool WasPropertyValueSerialized(const UStruct* Struct, const void* StructData, const FProperty* Property, int32 ArrayIndex)
 	{
-		if (const FByteProperty* ValuesSetBySerializationProperty = FindValuesSetBySerializationProperty(Struct))
+		if (const FByteProperty* SerializedValuesProperty = FindSerializedValuesProperty(Struct))
 		{
 			const int32 PropertyIndex = Property->GetIndexInOwner() + ArrayIndex;
 			const int32 ByteIndex = PropertyIndex / 8;
 			const int32 BitOffset = PropertyIndex % 8;
-			if (ByteIndex < ValuesSetBySerializationProperty->ArrayDim)
+			if (ByteIndex < SerializedValuesProperty->ArrayDim)
 			{
-				const uint8* PropertyDataPtr = ValuesSetBySerializationProperty->ContainerPtrToValuePtr<uint8>(StructData, ByteIndex);
+				const uint8* PropertyDataPtr = SerializedValuesProperty->ContainerPtrToValuePtr<uint8>(StructData, ByteIndex);
 				return (*PropertyDataPtr & (1 << BitOffset)) != 0;
 			}
 		}
 		return false;
 	}
 
-	void CopyPropertySetBySerializationData(const FFieldVariant& OldField, void* OldDataPtr, const FFieldVariant& NewField, void* NewDataPtr)
+	void CopyPropertyValueSerializedData(const FFieldVariant& OldField, void* OldDataPtr, const FFieldVariant& NewField, void* NewDataPtr)
 	{
 		if (const FStructProperty* OldAsStructProperty = OldField.Get<FStructProperty>())
 		{
 			const FStructProperty* NewAsStructProperty = NewField.Get<FStructProperty>();
 			checkf(NewAsStructProperty, TEXT("Type mismatch between OldField and NewField. Expected FStructProperty"));
-			CopyPropertySetBySerializationData(OldAsStructProperty->Struct, OldDataPtr, NewAsStructProperty->Struct, NewDataPtr);
+			CopyPropertyValueSerializedData(OldAsStructProperty->Struct, OldDataPtr, NewAsStructProperty->Struct, NewDataPtr);
 		}
 		else if (const FArrayProperty* OldAsArrayProperty = OldField.Get<FArrayProperty>())
 		{
@@ -601,7 +596,7 @@ namespace UE
 			{
 				if (NewArrayHelper.IsValidIndex(ArrayIndex))
 				{
-					CopyPropertySetBySerializationData(
+					CopyPropertyValueSerializedData(
 						OldAsArrayProperty->Inner, OldArrayHelper.GetElementPtr(ArrayIndex),
 						NewAsArrayProperty->Inner, NewArrayHelper.GetElementPtr(ArrayIndex));
 				}
@@ -619,7 +614,7 @@ namespace UE
 			
 			for (; OldItr && NewItr; ++OldItr, ++NewItr)
 			{
-				CopyPropertySetBySerializationData(
+				CopyPropertyValueSerializedData(
 					OldAsSetProperty->ElementProp, OldSetHelper.GetElementPtr(OldItr),
 					NewAsSetProperty->ElementProp, NewSetHelper.GetElementPtr(NewItr));
 			}
@@ -636,10 +631,10 @@ namespace UE
 			
 			for (; OldItr && NewItr; ++OldItr, ++NewItr)
 			{
-				CopyPropertySetBySerializationData(
+				CopyPropertyValueSerializedData(
 					OldAsMapProperty->KeyProp, OldMapHelper.GetKeyPtr(OldItr),
 					NewAsMapProperty->KeyProp, NewMapHelper.GetKeyPtr(NewItr));
-				CopyPropertySetBySerializationData(
+				CopyPropertyValueSerializedData(
 					OldAsMapProperty->ValueProp, OldMapHelper.GetValuePtr(OldItr),
 					NewAsMapProperty->ValueProp, NewMapHelper.GetValuePtr(NewItr));
 			}
@@ -662,9 +657,9 @@ namespace UE
 			};
 
 			// clear existing set-flags first
-			if (const FByteProperty* ValuesSetBySerializationProperty = FindValuesSetBySerializationProperty(NewAsStruct))
+			if (const FByteProperty* SerializedValuesProperty = FindSerializedValuesProperty(NewAsStruct))
 			{
-				ValuesSetBySerializationProperty->InitializeValue_InContainer(NewDataPtr);
+				SerializedValuesProperty->InitializeValue_InContainer(NewDataPtr);
 			}
 			
 			for (const FProperty* OldSubProperty : TFieldRange<FProperty>(OldAsStruct))
@@ -674,13 +669,13 @@ namespace UE
 					for (int32 ArrayIndex = 0; ArrayIndex < FMath::Min(OldSubProperty->ArrayDim, NewSubProperty->ArrayDim); ++ArrayIndex)
 					{
 						// copy set flags to new struct instance
-						if (WasPropertySetBySerialization(OldAsStruct, OldDataPtr, NewSubProperty, ArrayIndex))
+						if (WasPropertyValueSerialized(OldAsStruct, OldDataPtr, OldSubProperty, ArrayIndex))
 						{
-							MarkPropertySetBySerialization(NewAsStruct, NewDataPtr, NewSubProperty, ArrayIndex);
+							MarkPropertyValueSerialized(NewAsStruct, NewDataPtr, NewSubProperty, ArrayIndex);
 						}
 					
 						// recurse
-						CopyPropertySetBySerializationData(
+						CopyPropertyValueSerializedData(
 							OldSubProperty, OldSubProperty->ContainerPtrToValuePtr<void>(OldDataPtr, ArrayIndex),
 							NewSubProperty, NewSubProperty->ContainerPtrToValuePtr<void>(NewDataPtr, ArrayIndex));
 					}
diff --git a/Engine/Source/Runtime/CoreUObject/Private/UObject/PropertyBagRepository.cpp b/Engine/Source/Runtime/CoreUObject/Private/UObject/PropertyBagRepository.cpp
--- a/Engine/Source/Runtime/CoreUObject/Private/UObject/PropertyBagRepository.cpp
+++ b/Engine/Source/Runtime/CoreUObject/Private/UObject/PropertyBagRepository.cpp
@@ -172,7 +172,7 @@ void FPropertyBagRepository::ReassociateObjects(const TMap<UObject*, UObject*>&
 				
 				InstanceDataObjectToOwner.Add(NewBagData.InstanceDataObject, Pair.Value);
 				
-				CopyPropertySetBySerializationData(
+				CopyPropertyValueSerializedData(
 					OldBagData.InstanceDataObject->GetClass(), OldBagData.InstanceDataObject,
 					NewBagData.InstanceDataObject->GetClass(), NewBagData.InstanceDataObject);
 			}
@@ -785,9 +785,9 @@ const UObject* FPropertyBagRepository::FindInstanceForDataObject(const UObject*
 	return Owner ? *Owner : nullptr;
 }
 
-bool FPropertyBagRepository::WasPropertySetBySerialization(const UStruct* Struct, const void* StructData, const FProperty* Property, int32 ArrayIndex)
+bool FPropertyBagRepository::WasPropertyValueSerialized(const UStruct* Struct, const void* StructData, const FProperty* Property, int32 ArrayIndex)
 {
-	return UE::WasPropertySetBySerialization(Struct, StructData, Property, ArrayIndex);
+	return UE::WasPropertyValueSerialized(Struct, StructData, Property, ArrayIndex);
 }
 
 void FPropertyBagRepository::AddReferencedObjects(FReferenceCollector& Collector)
@@ -850,6 +850,8 @@ void FPropertyBagRepository::CreateInstanceDataObjectUnsafe(UObject* Owner, FPro
 	
 	// setup load context to mark properties the that were set by serialization
 	FUObjectSerializeContext* LoadContext = FUObjectThreadContext::Get().GetSerializeContext();
+	TGuardValue<bool> ScopedTrackSerializedProperties(LoadContext->bTrackSerializedProperties, true);
+	// enable impersonation so that the IDO gets loaded instead of Owner
 	TGuardValue<bool> ScopedImpersonateProperties(LoadContext->bImpersonateProperties, true);
 
 	FLinkerLoad* Linker = Owner->GetLinker();
@@ -860,8 +862,6 @@ void FPropertyBagRepository::CreateInstanceDataObjectUnsafe(UObject* Owner, FPro
 	}
 	else if (Linker)
 	{
-		// TODO: @jordan.hoffmann - this is very inefficient! We should remove this call to Preload. To do so, we'd need to change MarkPropertySetBySerialization
-		// to cache the serialized property list in the property bag instead of the structs. We'd also need to copy the property bag values to the IDO
 		Owner->SetFlags(RF_NeedLoad);
 		{
 			TGuardValue<bool> ScopedSkipKnownProperties(Linker->bSkipKnownProperties, true);
@@ -887,7 +887,7 @@ void FPropertyBagRepository::CopyTaggedProperties(const UObject* Source, UObject
 	FUObjectSerializeContext* SerializeContext = FUObjectThreadContext::Get().GetSerializeContext();
 	TGuardValue<bool> ImpersonatePropertiesScope(SerializeContext->bImpersonateProperties, true);
 	// don't mark properties as set by serialization when performing copy
-	TGuardValue<bool> MarkPropertiesSetBySerializationScope(UE::GMarkPropertiesSetBySerialization, false);
+	TGuardValue<bool> ScopedTrackSerializedProperties(SerializeContext->bTrackSerializedProperties, false);
 
 	TArray<uint8> Buffer;
 	Buffer.Reserve(Source->GetClass()->GetStructureSize());
@@ -923,7 +923,8 @@ FScopedIDOSerializationContext::FScopedIDOSerializationContext(UObject* InObject
 
 			// Enable tracking of initialized properties when loading an IDO, which is implied by impersonation being enabled.
 			const bool bLoadingIDO = bHasIDOSupport && SerializeContext->bImpersonateProperties;
-			TGuardValue<bool> ScopedTrackInitializedProperties(SerializeContext->bTrackInitializedProperties, bLoadingIDO);
+			ScopedTrackInitializedProperties.Emplace(SerializeContext->bTrackInitializedProperties, bLoadingIDO);
+			ScopedTrackSerializedProperties.Emplace(SerializeContext->bTrackSerializedProperties, bLoadingIDO);
 		}
 		else
 		{
diff --git a/Engine/Source/Runtime/CoreUObject/Private/UObject/UObjectThreadContext.cpp b/Engine/Source/Runtime/CoreUObject/Private/UObject/UObjectThreadContext.cpp
--- a/Engine/Source/Runtime/CoreUObject/Private/UObject/UObjectThreadContext.cpp
+++ b/Engine/Source/Runtime/CoreUObject/Private/UObject/UObjectThreadContext.cpp
@@ -47,6 +47,7 @@ FUObjectSerializeContext::FUObjectSerializeContext()
 #if WITH_EDITORONLY_DATA
 	, bTrackSerializedPropertyPath(false)
 	, bTrackInitializedProperties(false)
+	, bTrackSerializedProperties(false)
 	, bTrackUnknownProperties(false)
 	, bImpersonateProperties(false)
 #endif
diff --git a/Engine/Source/Runtime/CoreUObject/Public/UObject/UObjectThreadContext.h b/Engine/Source/Runtime/CoreUObject/Public/UObject/UObjectThreadContext.h
--- a/Engine/Source/Runtime/CoreUObject/Public/UObject/UObjectThreadContext.h
+++ b/Engine/Source/Runtime/CoreUObject/Public/UObject/UObjectThreadContext.h
@@ -173,6 +173,8 @@ public:
 	UE_INTERNAL bool bTrackSerializedPropertyPath;
 	/** True when initialized properties will be tracked for the serialized object. */
 	UE_INTERNAL bool bTrackInitializedProperties;
+	/** True when serialized properties will be tracked for the serialized object. */
+	UE_INTERNAL bool bTrackSerializedProperties;
 	/** True when unknown properties will tracked for the serialized object. */
 	UE_INTERNAL bool bTrackUnknownProperties;
 	/** True when the SerializedObject properties are being impersonated. */
diff --git a/Engine/Source/Runtime/CoreUObject/Tests/UObject/InstanceDataObjectUtilsTest.cpp b/Engine/Source/Runtime/CoreUObject/Tests/UObject/InstanceDataObjectUtilsTest.cpp
--- a/Engine/Source/Runtime/CoreUObject/Tests/UObject/InstanceDataObjectUtilsTest.cpp
+++ b/Engine/Source/Runtime/CoreUObject/Tests/UObject/InstanceDataObjectUtilsTest.cpp
@@ -47,10 +47,10 @@ TEST_CASE_NAMED(FInstanceDataObjectUtilsTest, "CoreUObject::Serialization::Insta
 	UObject* Owner = NewObject<UObject>(GetTransientPackage(), TestClass, TestObjectName);
 	void* StructData = StructProperty->ContainerPtrToValuePtr<void>(Owner);
 
-	CHECK_FALSE(WasPropertySetBySerialization(TestClass, Owner, StructProperty));
-	MarkPropertySetBySerialization(TestClass, Owner, StructProperty);
-	CHECK(WasPropertySetBySerialization(TestClass, Owner, StructProperty));
-	CHECK_FALSE(WasPropertySetBySerialization(TestClass, Owner, Int32Property));
+	CHECK_FALSE(WasPropertyValueSerialized(TestClass, Owner, StructProperty));
+	MarkPropertyValueSerialized(TestClass, Owner, StructProperty);
+	CHECK(WasPropertyValueSerialized(TestClass, Owner, StructProperty));
+	CHECK_FALSE(WasPropertyValueSerialized(TestClass, Owner, Int32Property));
 
 	CHECK_FALSE(IsPropertyValueInitialized(TestClass, Owner, StructProperty));
 	SetPropertyValueInitialized(TestClass, Owner, StructProperty);
@@ -58,14 +58,14 @@ TEST_CASE_NAMED(FInstanceDataObjectUtilsTest, "CoreUObject::Serialization::Insta
 	ClearPropertyValueInitialized(TestClass, Owner, StructProperty);
 	CHECK_FALSE(IsPropertyValueInitialized(TestClass, Owner, StructProperty));
 
-	CHECK_FALSE(WasPropertySetBySerialization(StructProperty->Struct, StructData, AProperty));
-	MarkPropertySetBySerialization(StructProperty->Struct, StructData, AProperty);
-	CHECK(WasPropertySetBySerialization(StructProperty->Struct, StructData, AProperty));
-	CHECK_FALSE(WasPropertySetBySerialization(StructProperty->Struct, StructData, BProperty));
-	MarkPropertySetBySerialization(StructProperty->Struct, StructData, BProperty);
-	CHECK(WasPropertySetBySerialization(StructProperty->Struct, StructData, BProperty));
-	CHECK_FALSE(WasPropertySetBySerialization(StructProperty->Struct, StructData, CProperty));
-	CHECK_FALSE(WasPropertySetBySerialization(StructProperty->Struct, StructData, DProperty));
+	CHECK_FALSE(WasPropertyValueSerialized(StructProperty->Struct, StructData, AProperty));
+	MarkPropertyValueSerialized(StructProperty->Struct, StructData, AProperty);
+	CHECK(WasPropertyValueSerialized(StructProperty->Struct, StructData, AProperty));
+	CHECK_FALSE(WasPropertyValueSerialized(StructProperty->Struct, StructData, BProperty));
+	MarkPropertyValueSerialized(StructProperty->Struct, StructData, BProperty);
+	CHECK(WasPropertyValueSerialized(StructProperty->Struct, StructData, BProperty));
+	CHECK_FALSE(WasPropertyValueSerialized(StructProperty->Struct, StructData, CProperty));
+	CHECK_FALSE(WasPropertyValueSerialized(StructProperty->Struct, StructData, DProperty));
 
 	CHECK_FALSE(IsPropertyValueInitialized(StructProperty->Struct, StructData, AProperty));
 	SetPropertyValueInitialized(StructProperty->Struct, StructData, AProperty);
