[BugFix] Made sure all loose properties are marked as serialized when rebuilding verse code. This allows those properties to be visible in the left panel of the fixup tool

diff --git a/Engine/Source/Runtime/CoreUObject/Private/UObject/InstanceDataObjectUtils.cpp b/Engine/Source/Runtime/CoreUObject/Private/UObject/InstanceDataObjectUtils.cpp
--- a/Engine/Source/Runtime/CoreUObject/Private/UObject/InstanceDataObjectUtils.cpp
+++ b/Engine/Source/Runtime/CoreUObject/Private/UObject/InstanceDataObjectUtils.cpp
@@ -680,6 +680,11 @@ namespace UE
 						{
 							MarkPropertyValueSerialized(NewAsStruct, NewDataPtr, NewSubProperty, ArrayIndex);
 						}
+						else if (NewSubProperty->GetBoolMetaData(NAME_IsLooseMetadata))
+						{
+							// loose properties should be marked as serialized regardless of whether the old struct marked them as such
+							MarkPropertyValueSerialized(NewAsStruct, NewDataPtr, NewSubProperty, ArrayIndex);
+						}
 					
 						// recurse
 						CopyPropertyValueSerializedData(
