Added static methods to SoftObjectPath add information at the callsite.
These should used instead of constrcutors.

diff --git a/Engine/Source/Runtime/CoreUObject/Private/UObject/SoftObjectPath.cpp b/Engine/Source/Runtime/CoreUObject/Private/UObject/SoftObjectPath.cpp
--- a/Engine/Source/Runtime/CoreUObject/Private/UObject/SoftObjectPath.cpp
+++ b/Engine/Source/Runtime/CoreUObject/Private/UObject/SoftObjectPath.cpp
@@ -22,6 +22,52 @@ FSoftObjectPath::FSoftObjectPath(FName InAssetPathName, FString InSubPathString)
 	SubPathString = MoveTemp(InSubPathString);
 }
 
+/** Static methods for more meaningful construction sites. */
+FSoftObjectPath FSoftObjectPath::ConstructFromPackageAssetSubpath(FName InPackageName, FName InAssetName, const FString& InSubPathString)
+{
+	return FSoftObjectPath(InPackageName, InAssetName, InSubPathString);
+}
+
+FSoftObjectPath FSoftObjectPath::ConstructFromPackageAssetSubpath(FName InPackageName, FName InAssetName, FString&& InSubPathString)
+{
+	return FSoftObjectPath(InPackageName, InAssetName, MoveTemp(InSubPathString));
+}
+
+FSoftObjectPath FSoftObjectPath::ConstructFromPackageAsset(FName InPackageName, FName InAssetName)
+{
+	return FSoftObjectPath(InPackageName, InAssetName, {});
+}
+
+FSoftObjectPath FSoftObjectPath::ConstructFromAssetPath(FTopLevelAssetPath InAssetPath)
+{
+	return FSoftObjectPath(InAssetPath);
+}
+
+FSoftObjectPath FSoftObjectPath::ConstructFromStringPath(FStringView InPath)
+{
+	FSoftObjectPath Tmp;
+	Tmp.SetPath(InPath);
+	return Tmp;
+}
+
+FSoftObjectPath FSoftObjectPath::ConstructFromStringPath(FUtf8StringView InPath)
+{
+	FSoftObjectPath Tmp;
+	Tmp.SetPath(InPath);
+	return Tmp;
+}
+
+FSoftObjectPath FSoftObjectPath::ConstructFromObject(const FObjectPtr& InObject)
+{
+	return FSoftObjectPath(InObject);
+}
+
+FSoftObjectPath FSoftObjectPath::ConstructFromObject(const UObject* InObject)
+{
+	return FSoftObjectPath(InObject);
+}
+
+
 FString FSoftObjectPath::ToString() const
 {
 	// Most of the time there is no sub path so we can do a single string allocation
diff --git a/Engine/Source/Runtime/CoreUObject/Public/UObject/SoftObjectPath.h b/Engine/Source/Runtime/CoreUObject/Public/UObject/SoftObjectPath.h
--- a/Engine/Source/Runtime/CoreUObject/Public/UObject/SoftObjectPath.h
+++ b/Engine/Source/Runtime/CoreUObject/Public/UObject/SoftObjectPath.h
@@ -100,6 +100,21 @@ PRAGMA_ENABLE_DEPRECATION_WARNINGS;
 		}
 	}
 
+	/** Static methods for more meaningful construction sites. */
+	COREUOBJECT_API static FSoftObjectPath ConstructFromPackageAssetSubpath(FName InPackageName, FName InAssetName, const FString& InSubPathString);
+	COREUOBJECT_API static FSoftObjectPath ConstructFromPackageAssetSubpath(FName InPackageName, FName InAssetName, FString&& InSubPathString);
+	COREUOBJECT_API static FSoftObjectPath ConstructFromPackageAsset(FName InPackageName, FName InAssetName);
+	COREUOBJECT_API static FSoftObjectPath ConstructFromAssetPath(FTopLevelAssetPath InAssetPath);
+	COREUOBJECT_API static FSoftObjectPath ConstructFromStringPath(FStringView InPath);
+	COREUOBJECT_API static FSoftObjectPath ConstructFromStringPath(FUtf8StringView InPath);
+	COREUOBJECT_API static FSoftObjectPath ConstructFromObject(const UObject* InObject);
+	COREUOBJECT_API static FSoftObjectPath ConstructFromObject(const FObjectPtr& InObject);
+	template <typename T>
+	static FSoftObjectPath ConstructFromObject(const TObjectPtr<T>& InObject)
+	{
+		return ConstructFromObject(InObject.Get());
+	}
+
 	FSoftObjectPath& operator=(const FTopLevelAssetPath Path)			{ SetPath(Path); return *this; }
 	FSoftObjectPath& operator=(const FString& Path)						{ SetPath(FStringView(Path)); return *this; }
 	FSoftObjectPath& operator=(FWideStringView Path)					{ SetPath(Path); return *this; }
