Adding support for a new type of FCoreRedirect, Type_Asset which can be used in place of UObjectRedirector.

diff --git a/Engine/Source/Runtime/CoreUObject/Private/UObject/CoreRedirects.cpp b/Engine/Source/Runtime/CoreUObject/Private/UObject/CoreRedirects.cpp
--- a/Engine/Source/Runtime/CoreUObject/Private/UObject/CoreRedirects.cpp
+++ b/Engine/Source/Runtime/CoreUObject/Private/UObject/CoreRedirects.cpp
@@ -10,6 +10,7 @@
 
 #include "GenericPlatform/GenericPlatformFile.h"
 #include "Hash/Blake3.h"
+#include "Logging/StructuredLog.h"
 #include "Misc/AutomationTest.h"
 #include "Misc/CommandLine.h"
 #include "Misc/ConfigCacheIni.h"
@@ -227,6 +228,154 @@ namespace UE::CoreRedirects::Private
 	{
 		return MatchWildcardRedirect<SuffixMatcher>(InSuffixRedirect, InUtf8Name, bPartialLHS, bPartialRHS);
 	}
+
+
+	bool RunAssetRedirectTests()
+	{
+		bool bSuccess = true;
+
+		// Test
+		TMap<FSoftObjectPath, FSoftObjectPath> Redirects;
+
+		struct FTestDefinition
+		{
+			const TCHAR* Origin = nullptr;
+			const TCHAR* Destination = nullptr;
+			const TCHAR* TestDescription = nullptr;
+			ECoreRedirectFlags RedirectFlags = ECoreRedirectFlags::None;
+			bool ExpectTrue = true;
+		};
+		TArray<FTestDefinition> Tests;
+
+		// Simple single redirection BasicName --> BasicNewName
+		{
+			FTestDefinition& BasicTest = Tests.AddDefaulted_GetRef();
+			BasicTest.Origin = TEXT("/Game/BasicName.BasicName");
+			BasicTest.Destination = TEXT("/Game/BasicNewName.BasicNewName");
+			BasicTest.TestDescription = TEXT("basic asset redirect with Type_Object");
+			BasicTest.RedirectFlags = ECoreRedirectFlags::Type_Object;
+			BasicTest.ExpectTrue = true;
+
+			{
+				FSoftObjectPath SourcePath(BasicTest.Origin);
+				FSoftObjectPath DestPath(BasicTest.Destination);
+				Redirects.Add(SourcePath, DestPath);
+			}
+		}
+
+		// Multi-asset package redirections. bp_orig_name + _C + Default --> bp_new_name + _C + Default
+		{
+			FTestDefinition& DefaultObjectTest = Tests.AddDefaulted_GetRef();
+			DefaultObjectTest.Origin = TEXT("/Plugin/bp_orig_name.Default__bp_orig_name_C");
+			DefaultObjectTest.Destination = TEXT("/Plugin/bp_new_name.Default__bp_new_name_C");
+			DefaultObjectTest.TestDescription = TEXT("default object asset redirect with Type_Object");
+			DefaultObjectTest.RedirectFlags = ECoreRedirectFlags::Type_Object;
+			DefaultObjectTest.ExpectTrue = true;
+
+			FTestDefinition& BPGCTest = Tests.AddDefaulted_GetRef();
+			BPGCTest.Origin = TEXT("/Plugin/bp_orig_name.bp_orig_name_C");
+			BPGCTest.Destination = TEXT("/Plugin/bp_new_name.bp_new_name_C");
+			BPGCTest.TestDescription = TEXT("default BPGC asset redirect with Type_Class");
+			BPGCTest.RedirectFlags = ECoreRedirectFlags::Type_Class;
+			BPGCTest.ExpectTrue = true;
+
+			FTestDefinition& InstanceOnlyTest = Tests.AddDefaulted_GetRef();
+			InstanceOnlyTest.Origin = BPGCTest.Origin;
+			InstanceOnlyTest.Destination = BPGCTest.Destination;
+			InstanceOnlyTest.RedirectFlags = ECoreRedirectFlags::Type_Class | ECoreRedirectFlags::Category_InstanceOnly;
+			InstanceOnlyTest.TestDescription = TEXT("category instance only");
+			InstanceOnlyTest.ExpectTrue = false;
+
+			{
+				FSoftObjectPath SourcePath(TEXT("/Plugin/bp_orig_name.Default__bp_orig_name_C"));
+				FSoftObjectPath DestPath(TEXT("/Plugin/bp_new_name.Default__bp_new_name_C"));
+				Redirects.Add(SourcePath, DestPath);
+			}
+
+			{
+				FSoftObjectPath SourcePath(TEXT("/Plugin/bp_orig_name.bp_orig_name_C"));
+				FSoftObjectPath DestPath(TEXT("/Plugin/bp_new_name.bp_new_name_C"));
+				Redirects.Add(SourcePath, DestPath);
+			}
+
+			{
+				FSoftObjectPath SourcePath(TEXT("/Plugin/bp_orig_name.bp_orig_name"));
+				FSoftObjectPath DestPath(TEXT("/Plugin/bp_new_name.bp_new_name"));
+				Redirects.Add(SourcePath, DestPath);
+			}
+		}
+
+		FCoreRedirects::AddAssetRedirects(Redirects);
+
+		if (!FCoreRedirects::ValidateAssetRedirects())
+		{
+			bSuccess = false;
+			UE_LOG(LogCoreRedirects, Error, TEXT("Failed asset redirect validation."));
+		}
+
+		for (auto& Test : Tests)
+		{
+			FCoreRedirectObjectName OldName(Test.Origin);
+			FCoreRedirectObjectName NewName = FCoreRedirects::GetRedirectedName(Test.RedirectFlags, OldName);
+			if (NewName.ToString().Equals(Test.Destination) != Test.ExpectTrue)
+			{
+				bSuccess = false;
+				UE_LOG(LogCoreRedirects, Error, TEXT("Failed %s. Source = %s was unexpectedly redirected to %s"),
+					Test.TestDescription, Test.Origin, *NewName.ToString());
+			}
+		}
+
+		// Remove all redirects temporarily and verify no test finds a redirection
+		FCoreRedirects::RemoveAllAssetRedirects();
+		for (auto& Test : Tests)
+		{
+			FCoreRedirectObjectName OldName(Test.Origin);
+			FCoreRedirectObjectName NewName = FCoreRedirects::GetRedirectedName(Test.RedirectFlags, OldName);
+			if (!NewName.ToString().Equals(Test.Origin))
+			{
+				bSuccess = false;
+				UE_LOG(LogCoreRedirects, Error, TEXT("Found unexpected redirect from %s to %s"),
+					Test.Origin, *NewName.ToString());
+			}
+		}
+
+		// Ensure validation failure if chains exist
+		{
+			TMap<FSoftObjectPath, FSoftObjectPath> ChainRedirects;
+
+			{
+				FSoftObjectPath SourcePath(TEXT("/Game/Chain_FirstName.Chain_FirstName"));
+				FSoftObjectPath DestPath(TEXT("/Game/Chain_SecondName.Chain_SecondName"));
+				ChainRedirects.Add(SourcePath, DestPath);
+			}
+
+			{
+				FSoftObjectPath SourcePath(TEXT("/Game/Chain_SecondName.Chain_SecondName"));
+				FSoftObjectPath DestPath(TEXT("/Game/Chain_ThirdName.Chain_ThirdName"));
+				ChainRedirects.Add(SourcePath, DestPath);
+			}
+
+			{
+				FSoftObjectPath SourcePath(TEXT("/Game/Chain_ThirdName.Chain_ThirdName"));
+				FSoftObjectPath DestPath(TEXT("/Game/Chain_FourthName.Chain_FourthName"));
+				ChainRedirects.Add(SourcePath, DestPath);
+			}
+
+			FCoreRedirects::AddAssetRedirects(ChainRedirects);
+			if (FCoreRedirects::ValidateAssetRedirects())
+			{
+				bSuccess = false;
+				UE_LOG(LogCoreRedirects, Error, TEXT("Failed to detect erroneous chained redirect in ValidateAssetRedirects()"));
+			}
+			FCoreRedirects::RemoveAllAssetRedirects();
+		}
+
+		// Re-add the redirects so that they are in place for subsequent tests. This can help
+		// detect unexpected interactions between asset redirects and legacy types of redirects
+		FCoreRedirects::AddAssetRedirects(Redirects);
+
+		return bSuccess;
+	}
 }
 
 #if WITH_EDITOR
@@ -582,6 +731,16 @@ void FCoreRedirectObjectName::UnionFieldsInline(const FCoreRedirectObjectName& O
 	}
 }
 
+FName FCoreRedirectObjectName::GetSearchKey(ECoreRedirectFlags Type) const
+{
+	if (EnumHasAnyFlags(Type, ECoreRedirectFlags::Type_Package | ECoreRedirectFlags::Type_Asset))
+	{
+		return PackageName;
+	}
+
+	return ObjectName;
+}
+
 bool FCoreRedirectObjectName::HasValidCharacters(ECoreRedirectFlags Type) const
 {
 	static FString InvalidRedirectCharacters = TEXT("\"',|&!~\n\r\t@#(){}[]=;^%$`");
@@ -807,6 +966,12 @@ bool FCoreRedirect::Matches(ECoreRedirectFlags InFlags, const FCoreRedirectObjec
 		return false;
 	}
 
+	if (EnumHasAnyFlags(InFlags, ECoreRedirectFlags::Type_Asset))
+	{
+		// Type_Asset matches should always be exact. They will either exactly match a package entry or an object entry
+		MatchFlags |= ECoreRedirectMatchFlags::DisallowPartialLHSMatch;
+	}
+
 	return Matches(InName, MatchFlags);
 }
 
@@ -819,6 +984,10 @@ bool FCoreRedirect::Matches(const FCoreRedirectObjectName& InName, ECoreRedirect
 		NameMatchFlags |= FCoreRedirectObjectName::EMatchFlags::AllowPartialRHSMatch;
 	}
 
+	if (EnumHasAllFlags(MatchFlags, ECoreRedirectMatchFlags::DisallowPartialLHSMatch))
+	{
+		NameMatchFlags |= FCoreRedirectObjectName::EMatchFlags::DisallowPartialLHSMatch;
+	}
 	if (IsSubstringMatch())
 	{
 		NameMatchFlags |= FCoreRedirectObjectName::EMatchFlags::CheckSubString;
@@ -1241,6 +1410,7 @@ void FCoreRedirects::Initialize()
 		ConfigKeyMap.Add(TEXT("FunctionRedirects"), ECoreRedirectFlags::Type_Function);
 		ConfigKeyMap.Add(TEXT("PropertyRedirects"), ECoreRedirectFlags::Type_Property);
 		ConfigKeyMap.Add(TEXT("PackageRedirects"), ECoreRedirectFlags::Type_Package);
+		ConfigKeyMap.Add(TEXT("AssetRedirects"), ECoreRedirectFlags::Type_Asset);
 
 		RegisterNativeRedirectsUnderWriteLock(ScopeLock);
 
@@ -1254,7 +1424,7 @@ void FCoreRedirects::Initialize()
 		bInitialized.store(true, std::memory_order_release);
 	}
 	// Enable to run startup tests
-	//RunTests();
+	// RunTests();
 }
 
 bool FCoreRedirects::RedirectNameAndValues(ECoreRedirectFlags Type, const FCoreRedirectObjectName& OldObjectName,
@@ -1327,6 +1497,34 @@ bool FCoreRedirects::RedirectNameAndValuesUnderReadLock(ECoreRedirectFlags Type,
 	return bDidRedirect;
 }
 
+bool FCoreRedirects::ValidateAssetRedirectsUnderReadLock(const FCoreRedirectorScopeLockForRead& HeldLock)
+{
+	bool bValidationSucceeded = true;
+	FRedirectNameMap* RedirectNameMap = RedirectTypeMap.Find(ECoreRedirectFlags::Type_Asset);
+	if (RedirectNameMap != nullptr)
+	{
+		for (const TPair<FName, TArray<FCoreRedirect>>& Pair : RedirectNameMap->RedirectMap)
+		{
+			// Pairs are package --> redirects because the search key for Type_Asset is the package name
+			for (const FCoreRedirect& Redirect : Pair.Value)
+			{
+				const FCoreRedirectObjectName& SearchName = Redirect.NewName;
+
+				TArray<const FCoreRedirect*> MatchingRedirects;
+				GetMatchingRedirects(ECoreRedirectFlags::Type_Asset, SearchName, MatchingRedirects);
+				for (const FCoreRedirect* MatchingRedirect : MatchingRedirects)
+				{
+					UE_LOG(LogCoreRedirects, Warning, TEXT("Found redirect from existing redirect. Chained redirects will not be followed. %s --> %s --> %s"),
+						*Redirect.OldName.ToString(), *SearchName.ToString(), *MatchingRedirect->NewName.ToString());
+					bValidationSucceeded = false;
+				}
+			}
+		}
+	}
+
+	return bValidationSucceeded;
+}
+
 FCoreRedirectObjectName FCoreRedirects::GetRedirectedName(ECoreRedirectFlags Type,
 	const FCoreRedirectObjectName& OldObjectName, ECoreRedirectMatchFlags MatchFlags)
 {
@@ -1371,6 +1569,10 @@ bool FCoreRedirects::GetMatchingRedirectsUnderReadLock(ECoreRedirectFlags Search
 	// Look for all redirects that match the given names and flags
 	bool bFound = false;
 
+	// We always search Type_Asset as well as whatever is requested
+	// That is because asset redirectors can redirect packages (implicitly) and any UObject type (explicitly)
+	SearchFlags |= ECoreRedirectFlags::Type_Asset;
+
 	// If we're not explicitly searching for packages, and not looking for removed things, and not searching for partial matches
 	// based on ObjectName only, add the implicit (Type=Package,Category=None) redirects
 	const bool bSearchPackageRedirects = !(SearchFlags & ECoreRedirectFlags::Type_Package) &&
@@ -1578,6 +1780,8 @@ bool FCoreRedirects::RunTests()
 	GRedirectionSummary = FRedirectionSummary();
 #endif
 
+	UE_LOG(LogCoreRedirects, Log, TEXT("Running FCoreRedirect Tests"));
+
 	TArray<FCoreRedirect> NewRedirects;
 
 	NewRedirects.Emplace(ECoreRedirectFlags::Type_Property, TEXT("Property"), TEXT("Property2"));
@@ -1602,6 +1806,10 @@ bool FCoreRedirects::RunTests()
 
 	AddRedirectList(NewRedirects, TEXT("RunTests"));
 
+	// Run the asset tests first so that their entries are present when we run the other tests
+	// That way we can confirm that having asset type redirects doesn't break existing tests
+	bSuccess = bSuccess && UE::CoreRedirects::Private::RunAssetRedirectTests();
+
 	struct FRedirectTest
 	{
 		FString OldName;
@@ -1615,8 +1823,6 @@ bool FCoreRedirects::RunTests()
 
 	TArray<FRedirectTest> Tests;
 
-	UE_LOG(LogCoreRedirects, Log, TEXT("Running FCoreRedirect Tests"));
-
 	// Package-specific property rename and package rename apply
 	Tests.Emplace(TEXT("/Game/PackageSpecific.Class:Property"), TEXT("/Game/PackageSpecific.Class:Property4"), ECoreRedirectFlags::Type_Property);
 	// Verify . works as well
@@ -1775,6 +1981,75 @@ bool FCoreRedirects::RunTests()
 	return bSuccess;
 }
 
+void FCoreRedirects::AddAssetRedirects(const TMap<FSoftObjectPath, FSoftObjectPath>& InRedirects)
+{
+	if (InRedirects.Num() > 0)
+	{
+		FCoreRedirectorScopeLockForWrite ScopeLock(RWLock);
+		FRedirectNameMap* ExistingMap = &RedirectTypeMap.FindOrAdd(ECoreRedirectFlags::Type_Asset);
+		for (const TPair<FSoftObjectPath, FSoftObjectPath>& Pair : InRedirects)
+		{
+			// Asset redirects are, by definition, not package redirects
+			if (Pair.Key.GetLongPackageFName().IsNone() || Pair.Value.GetAssetFName().IsNone())
+			{
+				UE_LOG(LogCoreRedirects, Warning, TEXT("Attempted to register asset redirector that was missing a package or object name. Redirector was from %s to %s"),
+					*Pair.Key.ToString(), *Pair.Value.ToString());
+				continue;
+			}
+
+			// Asset redirects use the package as the lookup key but contain multiple redirects underneath
+			// Conceptually, for each object redirector we add two core redirects: one for the package and one for the object itself
+			// In practice, we can just add the package redirect the first time
+
+			FCoreRedirect ObjectRedirector(ECoreRedirectFlags::Type_Asset, Pair.Key.ToString(), Pair.Value.ToString());
+			TArray<FCoreRedirect>& ExistingRedirects = ExistingMap->RedirectMap.FindOrAdd(ObjectRedirector.GetSearchKey());
+
+			if (ExistingRedirects.Num() == 0)
+			{
+				// This is a new redirector. Add two entries, one for the package first
+				FCoreRedirect PackageRedirect(ECoreRedirectFlags::Type_Asset, 
+					Pair.Key.GetLongPackageName(), Pair.Value.GetLongPackageName());
+
+				ExistingRedirects.Add(PackageRedirect);
+			}
+
+			// Check that, among the existing redirects, this one won't be a duplicate. Don't bother checking the package entry
+			bool bShouldAddRedirect = true;
+			for (int32 RedirectIndex = 1; RedirectIndex < ExistingRedirects.Num(); RedirectIndex++)
+			{
+				if (Pair.Key.GetLongPackageFName() == ExistingRedirects[RedirectIndex].OldName.PackageName
+					&& Pair.Key.GetAssetName() == ExistingRedirects[RedirectIndex].OldName.ObjectName)
+				{
+					const FCoreRedirectObjectName& ExistingTargetName = ExistingRedirects[0].NewName;
+
+					UE_LOGFMT(LogCoreRedirects, Error, "Skipping new redirect target '{target}' due to existing map from '{source}' to '{dest}'",
+						Pair.Value.ToString(), Pair.Key.ToString(), ExistingTargetName.ToString());
+
+					bShouldAddRedirect = false;
+					break;
+				}
+			}
+
+			if (bShouldAddRedirect)
+			{
+				ExistingRedirects.Add(ObjectRedirector);
+			}
+		}
+
+		if (IsInDebugMode())
+		{
+			ValidateAssetRedirects();
+		}
+	}
+}
+
+void FCoreRedirects::RemoveAllAssetRedirects()
+{
+	FCoreRedirectorScopeLockForWrite ScopeLock(RWLock);
+	FRedirectNameMap* ExistingMap = &RedirectTypeMap.FindOrAdd(ECoreRedirectFlags::Type_Asset);
+	ExistingMap->RedirectMap.Reset();
+}
+
 IMPLEMENT_SIMPLE_AUTOMATION_TEST(FCoreRedirectTest, "System.Core.Misc.CoreRedirects", EAutomationTestFlags_ApplicationContextMask | EAutomationTestFlags::EngineFilter)
 bool FCoreRedirectTest::RunTest(const FString& Parameters)
 {
@@ -2300,7 +2575,7 @@ void FCoreRedirects::ValidateRedirectList(TArrayView<const FCoreRedirect> Redire
 
 void FCoreRedirects::ValidateAllRedirects()
 {
-	FCoreRedirectorScopeLockForWrite ReadLock(RWLock);
+	FCoreRedirectorScopeLockForRead ReadLock(RWLock);
 	bValidatedOnce = true;
 
 	// Validate all existing redirects
@@ -2315,6 +2590,14 @@ void FCoreRedirects::ValidateAllRedirects()
 			ValidateRedirectList(ArrayPair.Value, ListName);
 		}
 	}
+
+	ValidateAssetRedirectsUnderReadLock(ReadLock);
+}
+
+bool FCoreRedirects::ValidateAssetRedirects()
+{
+	FCoreRedirectorScopeLockForRead ReadLock(RWLock);
+	return ValidateAssetRedirectsUnderReadLock(ReadLock);
 }
 
 const TMap<FName, ECoreRedirectFlags>& FCoreRedirects::GetConfigKeyMap()
diff --git a/Engine/Source/Runtime/CoreUObject/Public/UObject/CoreRedirects.h b/Engine/Source/Runtime/CoreUObject/Public/UObject/CoreRedirects.h
--- a/Engine/Source/Runtime/CoreUObject/Public/UObject/CoreRedirects.h
+++ b/Engine/Source/Runtime/CoreUObject/Public/UObject/CoreRedirects.h
@@ -45,6 +45,7 @@ enum class ECoreRedirectFlags : uint32
 	Type_Function =				0x00000010, // UFunction
 	Type_Property =				0x00000020, // FProperty
 	Type_Package =				0x00000040, // UPackage
+	Type_Asset =				0x00000080, // Redirects derived from UObjectRedirectors. Implicitly included with other search types
 	Type_AllMask =				0x0000FFFF, // Bit mask of all possible Types
 
 	// Category flags.  A Query will only match Redirects that have the same value for every category bit.
@@ -69,6 +70,10 @@ enum class ECoreRedirectMatchFlags
 	/** The passed-in CoreRedirectObjectName has null fields in Package, Outer, or Name, and should still be allowed to match
 	 against redirectors that were created with a full Package.[Outer:]Name. */
 	AllowPartialMatch = (1 << 0),
+	/** Used for Type_Asset redirects to ensure package redirects only match package queries and 
+	 *  full path redirects only match full path queries
+	 */
+	DisallowPartialLHSMatch = (1<<1)
 };
 ENUM_CLASS_FLAGS(ECoreRedirectMatchFlags);
 
@@ -165,15 +170,7 @@ struct FCoreRedirectObjectName
 	COREUOBJECT_API void UnionFieldsInline(const FCoreRedirectObjectName& Other);
 
 	/** Returns the name used as the key into the acceleration map */
-	FName GetSearchKey(ECoreRedirectFlags Type) const
-	{
-		if ((Type & ECoreRedirectFlags::Type_Package) == ECoreRedirectFlags::Type_Package)
-		{
-			return PackageName;
-		}
-
-		return ObjectName;
-	}
+	FName GetSearchKey(ECoreRedirectFlags Type) const;
 
 	/** Returns true if this refers to an actual object */
 	bool IsValid() const
@@ -363,9 +360,12 @@ struct FCoreRedirects
 	/** Validate a named list of redirects */
 	static COREUOBJECT_API void ValidateRedirectList(TArrayView<const FCoreRedirect> Redirects, const FString& SourceString);
 
-	/** Validates all known redirects and warn if they seem to point to missing things */
+	/** Validates all known redirects and warn if they seem to point to missing things or violate other constraints */
 	static COREUOBJECT_API void ValidateAllRedirects();
 
+	/** Validates asset redirects and warns if chains are detected. Chains should be resolved before adding asset redirects. */
+	static COREUOBJECT_API bool ValidateAssetRedirects();
+
 	/** Gets map from config key -> Flags. It may only be accessed once it becomes constant data after the system is initialized */
 	static COREUOBJECT_API const TMap<FName, ECoreRedirectFlags>& GetConfigKeyMap();
 
@@ -398,6 +398,15 @@ struct FCoreRedirects
 	/** Runs set of redirector tests, returns false on failure */
 	static COREUOBJECT_API bool RunTests();
 
+	/** Adds a collection of redirects as Type_Asset. These allow FCoreRedirects to support the functions
+	 *  of UObjectRedirector. Any duplicate sources are logged and discarded (only the first redirect from a path is used)
+	 *  Package redirects corresponding to the soft object paths are implicitly created.
+	 */
+	static COREUOBJECT_API void AddAssetRedirects(const TMap<FSoftObjectPath, FSoftObjectPath>& InRedirects);
+
+	/** Clears all redirects added via AddAssetRedirects */
+	static COREUOBJECT_API void RemoveAllAssetRedirects();
+	 
 private:
 	typedef UE::CoreRedirects::Private::FRWWithExclusiveRecursionScopeLockForRead FCoreRedirectorScopeLockForRead;
 	typedef UE::CoreRedirects::Private::FRWWithExclusiveRecursionScopeLockForWrite FCoreRedirectorScopeLockForWrite;
@@ -430,6 +439,9 @@ private:
 		FCoreRedirectObjectName& NewObjectName, const FCoreRedirect** FoundValueRedirect,
 		ECoreRedirectMatchFlags MatchFlags, const FCoreRedirectorScopeLockForRead& HeldLock);
 
+	/** Internal implementation for ValidateAssetRedirects that requires a read lock to already have been acquired */
+	static COREUOBJECT_API bool ValidateAssetRedirectsUnderReadLock(const FCoreRedirectorScopeLockForRead& HeldLock);
+
 	/** Container for managing Wildcard redirects (substrings, prefixes, suffixes) */
 	struct FWildcardData
 	{
