Moved UserDefinedStruct to CoreUObject

diff --git a/Engine/Plugins/Animation/ControlRig/Source/ControlRig/Private/Units/Animation/RigUnit_AnimAttribute.cpp b/Engine/Plugins/Animation/ControlRig/Source/ControlRig/Private/Units/Animation/RigUnit_AnimAttribute.cpp
--- a/Engine/Plugins/Animation/ControlRig/Source/ControlRig/Private/Units/Animation/RigUnit_AnimAttribute.cpp
+++ b/Engine/Plugins/Animation/ControlRig/Source/ControlRig/Private/Units/Animation/RigUnit_AnimAttribute.cpp
@@ -6,7 +6,7 @@
 #include "Engine/SkeletalMesh.h"
 #include "Units/RigUnitContext.h"
 #include "Components/SkeletalMeshComponent.h"
-#include "Engine/UserDefinedStruct.h"
+#include "StructUtils/UserDefinedStruct.h"
 
 #include UE_INLINE_GENERATED_CPP_BY_NAME(RigUnit_AnimAttribute)
 
diff --git a/Engine/Plugins/Animation/DeformerGraph/Source/OptimusCore/Private/DataInterfaces/OptimusDataInterfaceAnimAttribute.cpp b/Engine/Plugins/Animation/DeformerGraph/Source/OptimusCore/Private/DataInterfaces/OptimusDataInterfaceAnimAttribute.cpp
--- a/Engine/Plugins/Animation/DeformerGraph/Source/OptimusCore/Private/DataInterfaces/OptimusDataInterfaceAnimAttribute.cpp
+++ b/Engine/Plugins/Animation/DeformerGraph/Source/OptimusCore/Private/DataInterfaces/OptimusDataInterfaceAnimAttribute.cpp
@@ -16,7 +16,7 @@
 #include "Components/SkeletalMeshComponent.h"
 #include "ShaderParameterMetadataBuilder.h"
 #include "ComputeFramework/ComputeMetadataBuilder.h"
-#include "Engine/UserDefinedStruct.h"
+#include "StructUtils/UserDefinedStruct.h"
 #include "Nodes/OptimusNode_DataInterface.h"
 
 #include "Serialization/MemoryReader.h"
diff --git a/Engine/Plugins/Animation/DeformerGraph/Source/OptimusCore/Private/Nodes/OptimusNode_CustomComputeKernel.cpp b/Engine/Plugins/Animation/DeformerGraph/Source/OptimusCore/Private/Nodes/OptimusNode_CustomComputeKernel.cpp
--- a/Engine/Plugins/Animation/DeformerGraph/Source/OptimusCore/Private/Nodes/OptimusNode_CustomComputeKernel.cpp
+++ b/Engine/Plugins/Animation/DeformerGraph/Source/OptimusCore/Private/Nodes/OptimusNode_CustomComputeKernel.cpp
@@ -18,7 +18,7 @@
 #include "IOptimusDeprecatedExecutionDataInterface.h"
 #include "IOptimusUnnamedNodePinProvider.h"
 #include "OptimusNode_ResourceAccessorBase.h"
-#include "Engine/UserDefinedStruct.h"
+#include "StructUtils/UserDefinedStruct.h"
 
 #include UE_INLINE_GENERATED_CPP_BY_NAME(OptimusNode_CustomComputeKernel)
 
diff --git a/Engine/Plugins/Animation/DeformerGraph/Source/OptimusCore/Private/OptimusDataTypeRegistry.cpp b/Engine/Plugins/Animation/DeformerGraph/Source/OptimusCore/Private/OptimusDataTypeRegistry.cpp
--- a/Engine/Plugins/Animation/DeformerGraph/Source/OptimusCore/Private/OptimusDataTypeRegistry.cpp
+++ b/Engine/Plugins/Animation/DeformerGraph/Source/OptimusCore/Private/OptimusDataTypeRegistry.cpp
@@ -14,7 +14,7 @@
 #include "AssetRegistry/AssetRegistryModule.h"
 #include "ComputeFramework/ComputeMetadataBuilder.h"
 
-#include "Engine/UserDefinedStruct.h"
+#include "StructUtils/UserDefinedStruct.h"
 
 #include "UObject/UnrealType.h"
 
diff --git a/Engine/Plugins/Animation/DeformerGraph/Source/OptimusCore/Private/OptimusHelpers.cpp b/Engine/Plugins/Animation/DeformerGraph/Source/OptimusCore/Private/OptimusHelpers.cpp
--- a/Engine/Plugins/Animation/DeformerGraph/Source/OptimusCore/Private/OptimusHelpers.cpp
+++ b/Engine/Plugins/Animation/DeformerGraph/Source/OptimusCore/Private/OptimusHelpers.cpp
@@ -6,7 +6,7 @@
 #include "ComputeFramework/ShaderParamTypeDefinition.h"
 #include "ShaderParameterMetadataBuilder.h"
 #include "ShaderParameterMetadata.h"
-#include "Engine/UserDefinedStruct.h"
+#include "StructUtils/UserDefinedStruct.h"
 
 FName Optimus::GetUniqueNameForScope(UObject* InScopeObj, FName InName)
 {
diff --git a/Engine/Plugins/Animation/DeformerGraph/Source/OptimusEditor/Private/OptimusDetailsCustomization.cpp b/Engine/Plugins/Animation/DeformerGraph/Source/OptimusEditor/Private/OptimusDetailsCustomization.cpp
--- a/Engine/Plugins/Animation/DeformerGraph/Source/OptimusEditor/Private/OptimusDetailsCustomization.cpp
+++ b/Engine/Plugins/Animation/DeformerGraph/Source/OptimusEditor/Private/OptimusDetailsCustomization.cpp
@@ -6,7 +6,7 @@
 #include "DetailLayoutBuilder.h"
 #include "DetailWidgetRow.h"
 #include "Engine/SCS_Node.h"
-#include "Engine/UserDefinedStruct.h"
+#include "StructUtils/UserDefinedStruct.h"
 #include "IDetailChildrenBuilder.h"
 #include "IOptimusExecutionDomainProvider.h"
 #include "IOptimusParameterBindingProvider.h"
diff --git a/Engine/Plugins/Animation/DeformerGraph/Source/OptimusEditor/Private/Widgets/SOptimusDataTypeSelector.cpp b/Engine/Plugins/Animation/DeformerGraph/Source/OptimusEditor/Private/Widgets/SOptimusDataTypeSelector.cpp
--- a/Engine/Plugins/Animation/DeformerGraph/Source/OptimusEditor/Private/Widgets/SOptimusDataTypeSelector.cpp
+++ b/Engine/Plugins/Animation/DeformerGraph/Source/OptimusEditor/Private/Widgets/SOptimusDataTypeSelector.cpp
@@ -11,7 +11,7 @@
 #include "SListViewSelectorDropdownMenu.h"
 #include "Animation/AttributeTypes.h"
 #include "AssetRegistry/AssetRegistryModule.h"
-#include "Engine/UserDefinedStruct.h"
+#include "StructUtils/UserDefinedStruct.h"
 #include "Widgets/Images/SImage.h"
 #include "Widgets/Input/SComboBox.h"
 #include "Widgets/Input/SSearchBox.h"
diff --git a/Engine/Plugins/Developer/Concert/ConcertApp/MultiUserClient/Source/MultiUserClientLibrary/Public/MultiUserSubsystem.h b/Engine/Plugins/Developer/Concert/ConcertApp/MultiUserClient/Source/MultiUserClientLibrary/Public/MultiUserSubsystem.h
--- a/Engine/Plugins/Developer/Concert/ConcertApp/MultiUserClient/Source/MultiUserClientLibrary/Public/MultiUserSubsystem.h
+++ b/Engine/Plugins/Developer/Concert/ConcertApp/MultiUserClient/Source/MultiUserClientLibrary/Public/MultiUserSubsystem.h
@@ -3,7 +3,7 @@
 #pragma once
 
 #include "Delegates/DelegateCombinations.h"
-#include "Engine/UserDefinedStruct.h"
+#include "StructUtils/UserDefinedStruct.h"
 #include "Subsystems/EngineSubsystem.h"
 #include "Templates/SharedPointer.h"
 
diff --git a/Engine/Plugins/Editor/AssetManagerEditor/Source/AssetManagerEditor/Private/ReferenceViewer/EdGraph_ReferenceViewer.cpp b/Engine/Plugins/Editor/AssetManagerEditor/Source/AssetManagerEditor/Private/ReferenceViewer/EdGraph_ReferenceViewer.cpp
--- a/Engine/Plugins/Editor/AssetManagerEditor/Source/AssetManagerEditor/Private/ReferenceViewer/EdGraph_ReferenceViewer.cpp
+++ b/Engine/Plugins/Editor/AssetManagerEditor/Source/AssetManagerEditor/Private/ReferenceViewer/EdGraph_ReferenceViewer.cpp
@@ -17,7 +17,7 @@
 #include "Engine/AssetManager.h"
 #include "Engine/SCS_Node.h"
 #include "Engine/SimpleConstructionScript.h"
-#include "Engine/UserDefinedStruct.h"
+#include "StructUtils/UserDefinedStruct.h"
 #include "Interfaces/IPluginManager.h"
 #include "ReferenceViewer/EdGraphNode_ReferencedProperties.h"
 #include "UObject/ReferenceChainSearch.h"
diff --git a/Engine/Plugins/Editor/BlueprintHeaderView/Source/BlueprintHeaderView/Private/BlueprintHeaderView.cpp b/Engine/Plugins/Editor/BlueprintHeaderView/Source/BlueprintHeaderView/Private/BlueprintHeaderView.cpp
--- a/Engine/Plugins/Editor/BlueprintHeaderView/Source/BlueprintHeaderView/Private/BlueprintHeaderView.cpp
+++ b/Engine/Plugins/Editor/BlueprintHeaderView/Source/BlueprintHeaderView/Private/BlueprintHeaderView.cpp
@@ -10,7 +10,7 @@
 #include "SBlueprintHeaderView.h"
 #include "ContentBrowserModule.h"
 #include "Engine/Blueprint.h"
-#include "Engine/UserDefinedStruct.h"
+#include "StructUtils/UserDefinedStruct.h"
 #include "ToolMenuSection.h"
 #include "ToolMenus.h"
 #include "Toolkits/AssetEditorToolkit.h"
diff --git a/Engine/Plugins/Editor/BlueprintHeaderView/Source/BlueprintHeaderView/Private/HeaderViewStructListItem.cpp b/Engine/Plugins/Editor/BlueprintHeaderView/Source/BlueprintHeaderView/Private/HeaderViewStructListItem.cpp
--- a/Engine/Plugins/Editor/BlueprintHeaderView/Source/BlueprintHeaderView/Private/HeaderViewStructListItem.cpp
+++ b/Engine/Plugins/Editor/BlueprintHeaderView/Source/BlueprintHeaderView/Private/HeaderViewStructListItem.cpp
@@ -1,7 +1,7 @@
 // Copyright Epic Games, Inc. All Rights Reserved.
 
 #include "HeaderViewStructListItem.h"
-#include "Engine/UserDefinedStruct.h"
+#include "StructUtils/UserDefinedStruct.h"
 #include "IAssetTools.h"
 #include "String/LineEndings.h"
 #include "Kismet2/StructureEditorUtils.h"
diff --git a/Engine/Plugins/Editor/BlueprintHeaderView/Source/BlueprintHeaderView/Private/HeaderViewVariableListItem.cpp b/Engine/Plugins/Editor/BlueprintHeaderView/Source/BlueprintHeaderView/Private/HeaderViewVariableListItem.cpp
--- a/Engine/Plugins/Editor/BlueprintHeaderView/Source/BlueprintHeaderView/Private/HeaderViewVariableListItem.cpp
+++ b/Engine/Plugins/Editor/BlueprintHeaderView/Source/BlueprintHeaderView/Private/HeaderViewVariableListItem.cpp
@@ -2,7 +2,7 @@
 
 #include "HeaderViewVariableListItem.h"
 #include "EdMode.h"
-#include "Engine/UserDefinedStruct.h"
+#include "StructUtils/UserDefinedStruct.h"
 #include "String/LineEndings.h"
 #include "Kismet2/BlueprintEditorUtils.h"
 #include "Kismet2/StructureEditorUtils.h"
diff --git a/Engine/Plugins/Editor/BlueprintHeaderView/Source/BlueprintHeaderView/Private/SBlueprintHeaderView.cpp b/Engine/Plugins/Editor/BlueprintHeaderView/Source/BlueprintHeaderView/Private/SBlueprintHeaderView.cpp
--- a/Engine/Plugins/Editor/BlueprintHeaderView/Source/BlueprintHeaderView/Private/SBlueprintHeaderView.cpp
+++ b/Engine/Plugins/Editor/BlueprintHeaderView/Source/BlueprintHeaderView/Private/SBlueprintHeaderView.cpp
@@ -15,7 +15,7 @@
 #include "PropertyEditorModule.h"
 #include "ContentBrowserModule.h"
 #include "IContentBrowserSingleton.h"
-#include "Engine/UserDefinedStruct.h"
+#include "StructUtils/UserDefinedStruct.h"
 #include "Framework/Text/SlateTextRun.h"
 #include "Framework/Commands/GenericCommands.h"
 #include "Framework/MultiBox/MultiBoxBuilder.h"
diff --git a/Engine/Plugins/Editor/ContentBrowser/ContentBrowserAssetDataSource/Source/ContentBrowserAssetDataSource/Private/AssetFileContextMenu.cpp b/Engine/Plugins/Editor/ContentBrowser/ContentBrowserAssetDataSource/Source/ContentBrowserAssetDataSource/Private/AssetFileContextMenu.cpp
--- a/Engine/Plugins/Editor/ContentBrowser/ContentBrowserAssetDataSource/Source/ContentBrowserAssetDataSource/Private/AssetFileContextMenu.cpp
+++ b/Engine/Plugins/Editor/ContentBrowser/ContentBrowserAssetDataSource/Source/ContentBrowserAssetDataSource/Private/AssetFileContextMenu.cpp
@@ -5,7 +5,7 @@
 #include "AssetDefinition.h"
 #include "AssetDefinitionRegistry.h"
 #include "Engine/Blueprint.h"
-#include "Engine/UserDefinedStruct.h"
+#include "StructUtils/UserDefinedStruct.h"
 #include "Engine/UserDefinedEnum.h"
 #include "IAssetTypeActions.h"
 #include "Misc/MessageDialog.h"
diff --git a/Engine/Plugins/Editor/EngineAssetDefinitions/Source/Private/AssetDefinition_UserDefinedStruct.h b/Engine/Plugins/Editor/EngineAssetDefinitions/Source/Private/AssetDefinition_UserDefinedStruct.h
--- a/Engine/Plugins/Editor/EngineAssetDefinitions/Source/Private/AssetDefinition_UserDefinedStruct.h
+++ b/Engine/Plugins/Editor/EngineAssetDefinitions/Source/Private/AssetDefinition_UserDefinedStruct.h
@@ -2,7 +2,7 @@
 
 #pragma once
 
-#include "Engine/UserDefinedStruct.h"
+#include "StructUtils/UserDefinedStruct.h"
 #include "AssetDefinitionDefault.h"
 
 #include "AssetDefinition_UserDefinedStruct.generated.h"
diff --git a/Engine/Plugins/Experimental/AnimNext/Source/AnimNextEditor/Private/Param/ObjectPropertyLocatorEditor.cpp b/Engine/Plugins/Experimental/AnimNext/Source/AnimNextEditor/Private/Param/ObjectPropertyLocatorEditor.cpp
--- a/Engine/Plugins/Experimental/AnimNext/Source/AnimNextEditor/Private/Param/ObjectPropertyLocatorEditor.cpp
+++ b/Engine/Plugins/Experimental/AnimNext/Source/AnimNextEditor/Private/Param/ObjectPropertyLocatorEditor.cpp
@@ -8,7 +8,7 @@
 #include "SParameterPicker.h"
 #include "UniversalObjectLocator.h"
 #include "UniversalObjectLocatorFragmentTypeHandle.h"
-#include "Engine/UserDefinedStruct.h"
+#include "StructUtils/UserDefinedStruct.h"
 #include "Framework/PropertyViewer/IFieldIterator.h"
 #include "Modules/ModuleManager.h"
 #include "Param/AnimNextObjectPropertyLocatorFragment.h"
diff --git a/Engine/Plugins/Experimental/Chooser/Source/Chooser/Private/Chooser.cpp b/Engine/Plugins/Experimental/Chooser/Source/Chooser/Private/Chooser.cpp
--- a/Engine/Plugins/Experimental/Chooser/Source/Chooser/Private/Chooser.cpp
+++ b/Engine/Plugins/Experimental/Chooser/Source/Chooser/Private/Chooser.cpp
@@ -4,7 +4,7 @@
 #include "ChooserPropertyAccess.h"
 #include "ChooserTrace.h"
 #include "ObjectTrace.h"
-#include "Engine/UserDefinedStruct.h"
+#include "StructUtils/UserDefinedStruct.h"
 #include "Engine/Blueprint.h"
 #include "ChooserIndexArray.h"
 #include "IChooserParameterGameplayTag.h"
diff --git a/Engine/Plugins/Experimental/Chooser/Source/Chooser/Private/ChooserPropertyAccess.cpp b/Engine/Plugins/Experimental/Chooser/Source/Chooser/Private/ChooserPropertyAccess.cpp
--- a/Engine/Plugins/Experimental/Chooser/Source/Chooser/Private/ChooserPropertyAccess.cpp
+++ b/Engine/Plugins/Experimental/Chooser/Source/Chooser/Private/ChooserPropertyAccess.cpp
@@ -3,7 +3,7 @@
 #include "ChooserPropertyAccess.h"
 
 #include "IObjectChooser.h"
-#include "Engine/UserDefinedStruct.h"
+#include "StructUtils/UserDefinedStruct.h"
 #include "HAL/IConsoleManager.h"
 #include "UObject/UnrealType.h"
 #include "Logging/LogMacros.h"
diff --git a/Engine/Plugins/Experimental/Mutable/Source/CustomizableObjectEditor/Private/MuCOE/Nodes/CustomizableObjectNodeTable.cpp b/Engine/Plugins/Experimental/Mutable/Source/CustomizableObjectEditor/Private/MuCOE/Nodes/CustomizableObjectNodeTable.cpp
--- a/Engine/Plugins/Experimental/Mutable/Source/CustomizableObjectEditor/Private/MuCOE/Nodes/CustomizableObjectNodeTable.cpp
+++ b/Engine/Plugins/Experimental/Mutable/Source/CustomizableObjectEditor/Private/MuCOE/Nodes/CustomizableObjectNodeTable.cpp
@@ -6,7 +6,7 @@
 #include "Engine/StaticMesh.h"
 #include "Engine/SkeletalMesh.h"
 #include "Engine/Texture2DArray.h"
-#include "Engine/UserDefinedStruct.h"
+#include "StructUtils/UserDefinedStruct.h"
 #include "Kismet2/StructureEditorUtils.h"
 #include "Materials/Material.h"
 #include "Materials/MaterialInstance.h"
diff --git a/Engine/Plugins/Experimental/PythonScriptPlugin/Source/PythonScriptPlugin/Private/PyGenUtil.cpp b/Engine/Plugins/Experimental/PythonScriptPlugin/Source/PythonScriptPlugin/Private/PyGenUtil.cpp
--- a/Engine/Plugins/Experimental/PythonScriptPlugin/Source/PythonScriptPlugin/Private/PyGenUtil.cpp
+++ b/Engine/Plugins/Experimental/PythonScriptPlugin/Source/PythonScriptPlugin/Private/PyGenUtil.cpp
@@ -25,7 +25,7 @@
 #include "UObject/UnrealType.h"
 #include "Engine/BlueprintCore.h"
 #include "Engine/BlueprintGeneratedClass.h"
-#include "Engine/UserDefinedStruct.h"
+#include "StructUtils/UserDefinedStruct.h"
 #include "Engine/UserDefinedEnum.h"
 #include "GameFramework/Actor.h"
 #include "Interfaces/IPluginManager.h"
diff --git a/Engine/Plugins/Experimental/PythonScriptPlugin/Source/PythonScriptPlugin/Private/PyWrapperStruct.cpp b/Engine/Plugins/Experimental/PythonScriptPlugin/Source/PythonScriptPlugin/Private/PyWrapperStruct.cpp
--- a/Engine/Plugins/Experimental/PythonScriptPlugin/Source/PythonScriptPlugin/Private/PyWrapperStruct.cpp
+++ b/Engine/Plugins/Experimental/PythonScriptPlugin/Source/PythonScriptPlugin/Private/PyWrapperStruct.cpp
@@ -11,7 +11,7 @@
 #include "UObject/PropertyPortFlags.h"
 #include "Misc/ScopeExit.h"
 #include "Templates/Casts.h"
-#include "Engine/UserDefinedStruct.h"
+#include "StructUtils/UserDefinedStruct.h"
 
 #include UE_INLINE_GENERATED_CPP_BY_NAME(PyWrapperStruct)
 
diff --git a/Engine/Plugins/Experimental/StructUtils/Source/StructUtils/Private/InstancedStruct.cpp b/Engine/Plugins/Experimental/StructUtils/Source/StructUtils/Private/InstancedStruct.cpp
--- a/Engine/Plugins/Experimental/StructUtils/Source/StructUtils/Private/InstancedStruct.cpp
+++ b/Engine/Plugins/Experimental/StructUtils/Source/StructUtils/Private/InstancedStruct.cpp
@@ -9,7 +9,7 @@
 #if WITH_ENGINE
 #include "Engine/PackageMapClient.h"
 #include "Engine/NetConnection.h"
-#include "Engine/UserDefinedStruct.h"
+#include "StructUtils/UserDefinedStruct.h"
 #include "Net/RepLayout.h"
 #include "Serialization/MemoryWriter.h"
 #include "Serialization/MemoryReader.h"
diff --git a/Engine/Plugins/Experimental/StructUtils/Source/StructUtils/Private/InstancedStructContainer.cpp b/Engine/Plugins/Experimental/StructUtils/Source/StructUtils/Private/InstancedStructContainer.cpp
--- a/Engine/Plugins/Experimental/StructUtils/Source/StructUtils/Private/InstancedStructContainer.cpp
+++ b/Engine/Plugins/Experimental/StructUtils/Source/StructUtils/Private/InstancedStructContainer.cpp
@@ -2,7 +2,7 @@
 #include "InstancedStructContainer.h"
 
 #if WITH_ENGINE && WITH_EDITOR
-#include "Engine/UserDefinedStruct.h"
+#include "StructUtils/UserDefinedStruct.h"
 #include "Serialization/MemoryWriter.h"
 #include "Serialization/MemoryReader.h"
 #include "Serialization/ArchiveUObject.h"
diff --git a/Engine/Plugins/Experimental/StructUtils/Source/StructUtils/Private/PropertyBag.cpp b/Engine/Plugins/Experimental/StructUtils/Source/StructUtils/Private/PropertyBag.cpp
--- a/Engine/Plugins/Experimental/StructUtils/Source/StructUtils/Private/PropertyBag.cpp
+++ b/Engine/Plugins/Experimental/StructUtils/Source/StructUtils/Private/PropertyBag.cpp
@@ -6,7 +6,7 @@
 #include "UObject/TextProperty.h"
 
 #if WITH_ENGINE && WITH_EDITOR
-#include "Engine/UserDefinedStruct.h"
+#include "StructUtils/UserDefinedStruct.h"
 #include "Serialization/MemoryWriter.h"
 #include "Serialization/MemoryReader.h"
 #include "Serialization/ArchiveUObject.h"
diff --git a/Engine/Plugins/Experimental/StructUtils/Source/StructUtilsEditor/Private/Customizations/InstancedStructDetails.cpp b/Engine/Plugins/Experimental/StructUtils/Source/StructUtilsEditor/Private/Customizations/InstancedStructDetails.cpp
--- a/Engine/Plugins/Experimental/StructUtils/Source/StructUtilsEditor/Private/Customizations/InstancedStructDetails.cpp
+++ b/Engine/Plugins/Experimental/StructUtils/Source/StructUtilsEditor/Private/Customizations/InstancedStructDetails.cpp
@@ -11,8 +11,7 @@
 #include "UObject/Package.h"
 #include "Widgets/Images/SImage.h"
 #include "Modules/ModuleManager.h"
-#include "StructViewerModule.h"
-#include "Engine/UserDefinedStruct.h"
+#include "StructUtils/UserDefinedStruct.h"
 #include "InstancedStruct.h"
 #include "IStructureDataProvider.h"
 #include "Misc/ConfigCacheIni.h"
diff --git a/Engine/Plugins/Experimental/StructUtils/Source/StructUtilsEditor/Private/Customizations/PropertyBagDetails.cpp b/Engine/Plugins/Experimental/StructUtils/Source/StructUtilsEditor/Private/Customizations/PropertyBagDetails.cpp
--- a/Engine/Plugins/Experimental/StructUtils/Source/StructUtilsEditor/Private/Customizations/PropertyBagDetails.cpp
+++ b/Engine/Plugins/Experimental/StructUtils/Source/StructUtilsEditor/Private/Customizations/PropertyBagDetails.cpp
@@ -12,7 +12,7 @@
 #include "Widgets/Input/SComboButton.h"
 #include "Framework/MultiBox/MultiBoxBuilder.h"
 #include "ScopedTransaction.h"
-#include "Engine/UserDefinedStruct.h"
+#include "StructUtils/UserDefinedStruct.h"
 #include "SPinTypeSelector.h"
 #include "StructUtilsMetadata.h"
 #include "Templates/ValueOrError.h"
diff --git a/Engine/Plugins/Experimental/StructUtils/Source/StructUtilsEditor/Private/SInstancedStructPicker.cpp b/Engine/Plugins/Experimental/StructUtils/Source/StructUtilsEditor/Private/SInstancedStructPicker.cpp
--- a/Engine/Plugins/Experimental/StructUtils/Source/StructUtilsEditor/Private/SInstancedStructPicker.cpp
+++ b/Engine/Plugins/Experimental/StructUtils/Source/StructUtilsEditor/Private/SInstancedStructPicker.cpp
@@ -3,7 +3,7 @@
 #include "SInstancedStructPicker.h"
 #include "DetailLayoutBuilder.h"
 #include "Editor.h"
-#include "Engine/UserDefinedStruct.h"
+#include "StructUtils/UserDefinedStruct.h"
 #include "InstancedStruct.h"
 #include "IPropertyUtilities.h"
 #include "Modules/ModuleManager.h"
diff --git a/Engine/Plugins/Experimental/StructUtils/Source/StructUtilsEditor/Private/StructUtilsEditorModule.cpp b/Engine/Plugins/Experimental/StructUtils/Source/StructUtilsEditor/Private/StructUtilsEditorModule.cpp
--- a/Engine/Plugins/Experimental/StructUtils/Source/StructUtilsEditor/Private/StructUtilsEditorModule.cpp
+++ b/Engine/Plugins/Experimental/StructUtils/Source/StructUtilsEditor/Private/StructUtilsEditorModule.cpp
@@ -5,7 +5,7 @@
 #include "PropertyBagDetails.h"
 #include "PropertyEditorModule.h"
 #include "StructUtilsTypes.h"
-#include "Engine/UserDefinedStruct.h"
+#include "StructUtils/UserDefinedStruct.h"
 #include "UserDefinedStructure/UserDefinedStructEditorData.h"
 #include "UObject/Package.h"
 #include "UObject/UObjectGlobals.h"
diff --git a/Engine/Plugins/FX/Niagara/Source/NiagaraEditor/Private/NiagaraClipboard.cpp b/Engine/Plugins/FX/Niagara/Source/NiagaraEditor/Private/NiagaraClipboard.cpp
--- a/Engine/Plugins/FX/Niagara/Source/NiagaraEditor/Private/NiagaraClipboard.cpp
+++ b/Engine/Plugins/FX/Niagara/Source/NiagaraEditor/Private/NiagaraClipboard.cpp
@@ -12,7 +12,7 @@
 #include "HAL/PlatformApplicationMisc.h"
 #include "NiagaraNodeFunctionCall.h"
 #include "Engine/UserDefinedEnum.h"
-#include "Engine/UserDefinedStruct.h"
+#include "StructUtils/UserDefinedStruct.h"
 #include "NiagaraEditorModule.h"
 #include "INiagaraEditorTypeUtilities.h"
 #include "Logging/LogScopedVerbosityOverride.h"
diff --git a/Engine/Plugins/PCG/Source/PCG/Private/Helpers/PCGPropertyHelpers.cpp b/Engine/Plugins/PCG/Source/PCG/Private/Helpers/PCGPropertyHelpers.cpp
--- a/Engine/Plugins/PCG/Source/PCG/Private/Helpers/PCGPropertyHelpers.cpp
+++ b/Engine/Plugins/PCG/Source/PCG/Private/Helpers/PCGPropertyHelpers.cpp
@@ -10,7 +10,7 @@
 #include "Helpers/PCGSettingsHelpers.h"
 #include "Metadata/PCGMetadata.h"
 
-#include "Engine/UserDefinedStruct.h"
+#include "StructUtils/UserDefinedStruct.h"
 #include "UObject/Field.h"
 #include "UObject/UnrealType.h"
 
diff --git a/Engine/Plugins/PCG/Source/PCG/Private/Metadata/Accessors/PCGAttributeAccessorHelpers.cpp b/Engine/Plugins/PCG/Source/PCG/Private/Metadata/Accessors/PCGAttributeAccessorHelpers.cpp
--- a/Engine/Plugins/PCG/Source/PCG/Private/Metadata/Accessors/PCGAttributeAccessorHelpers.cpp
+++ b/Engine/Plugins/PCG/Source/PCG/Private/Metadata/Accessors/PCGAttributeAccessorHelpers.cpp
@@ -16,7 +16,7 @@
 #include "Metadata/Accessors/PCGPropertyAccessor.h"
 #include "Metadata/Accessors/PCGAttributeExtractor.h"
 
-#include "Engine/UserDefinedStruct.h"
+#include "StructUtils/UserDefinedStruct.h"
 #include "UObject/EnumProperty.h"
 
 namespace PCGAttributeAccessorHelpers
diff --git a/Engine/Plugins/Runtime/DataRegistry/Source/DataRegistryEditor/Private/DataRegistryEditorToolkit.cpp b/Engine/Plugins/Runtime/DataRegistry/Source/DataRegistryEditor/Private/DataRegistryEditorToolkit.cpp
--- a/Engine/Plugins/Runtime/DataRegistry/Source/DataRegistryEditor/Private/DataRegistryEditorToolkit.cpp
+++ b/Engine/Plugins/Runtime/DataRegistry/Source/DataRegistryEditor/Private/DataRegistryEditorToolkit.cpp
@@ -23,7 +23,7 @@
 #include "Widgets/Input/SHyperlink.h"
 #include "SourceCodeNavigation.h"
 #include "PropertyEditorModule.h"
-#include "Engine/UserDefinedStruct.h"
+#include "StructUtils/UserDefinedStruct.h"
  
 #define LOCTEXT_NAMESPACE "DataRegistryEditor"
 
diff --git a/Engine/Plugins/Runtime/ModelViewViewModel/Source/ModelViewViewModelBlueprint/Private/MVVMPropertyPath.cpp b/Engine/Plugins/Runtime/ModelViewViewModel/Source/ModelViewViewModelBlueprint/Private/MVVMPropertyPath.cpp
--- a/Engine/Plugins/Runtime/ModelViewViewModel/Source/ModelViewViewModelBlueprint/Private/MVVMPropertyPath.cpp
+++ b/Engine/Plugins/Runtime/ModelViewViewModel/Source/ModelViewViewModelBlueprint/Private/MVVMPropertyPath.cpp
@@ -4,7 +4,7 @@
 #include "BlueprintCompilationManager.h"
 #include "Blueprint/WidgetTree.h"
 #include "Engine/Blueprint.h"
-#include "Engine/UserDefinedStruct.h"
+#include "StructUtils/UserDefinedStruct.h"
 #include "Kismet2/BlueprintEditorUtils.h"
 #include "Kismet2/StructureEditorUtils.h"
 #include "MVVMBlueprintView.h"
diff --git a/Engine/Plugins/Runtime/PropertyBindingUtils/Source/PropertyBindingUtils/Private/PropertyBindingPath.cpp b/Engine/Plugins/Runtime/PropertyBindingUtils/Source/PropertyBindingUtils/Private/PropertyBindingPath.cpp
--- a/Engine/Plugins/Runtime/PropertyBindingUtils/Source/PropertyBindingUtils/Private/PropertyBindingPath.cpp
+++ b/Engine/Plugins/Runtime/PropertyBindingUtils/Source/PropertyBindingUtils/Private/PropertyBindingPath.cpp
@@ -8,7 +8,7 @@
 #include "UObject/CoreRedirects.h"
 #include "UObject/Package.h"
 #include "Engine/BlueprintGeneratedClass.h"
-#include "Engine/UserDefinedStruct.h"
+#include "StructUtils/UserDefinedStruct.h"
 #include "Kismet2/StructureEditorUtils.h"
 #include "UObject/Field.h"
 #endif
diff --git a/Engine/Plugins/Runtime/RigVM/Source/RigVM/Private/RigVMCore/RigVMAssetUserData.cpp b/Engine/Plugins/Runtime/RigVM/Source/RigVM/Private/RigVMCore/RigVMAssetUserData.cpp
--- a/Engine/Plugins/Runtime/RigVM/Source/RigVM/Private/RigVMCore/RigVMAssetUserData.cpp
+++ b/Engine/Plugins/Runtime/RigVM/Source/RigVM/Private/RigVMCore/RigVMAssetUserData.cpp
@@ -1,7 +1,7 @@
 // Copyright Epic Games, Inc. All Rights Reserved.
 
 #include "RigVMCore/RigVMAssetUserData.h"
-#include "Engine/UserDefinedStruct.h"
+#include "StructUtils/UserDefinedStruct.h"
 #include "Misc/PackageName.h"
 
 #include UE_INLINE_GENERATED_CPP_BY_NAME(RigVMAssetUserData)
diff --git a/Engine/Plugins/Runtime/RigVM/Source/RigVM/Private/RigVMCore/RigVMRegistry.cpp b/Engine/Plugins/Runtime/RigVM/Source/RigVM/Private/RigVMCore/RigVMRegistry.cpp
--- a/Engine/Plugins/Runtime/RigVM/Source/RigVM/Private/RigVMCore/RigVMRegistry.cpp
+++ b/Engine/Plugins/Runtime/RigVM/Source/RigVM/Private/RigVMCore/RigVMRegistry.cpp
@@ -9,7 +9,7 @@
 #include "Animation/AttributeTypes.h"
 #include "AssetRegistry/AssetRegistryModule.h"
 #include "Engine/UserDefinedEnum.h"
-#include "Engine/UserDefinedStruct.h"
+#include "StructUtils/UserDefinedStruct.h"
 #include "UObject/UObjectIterator.h"
 #include "Misc/CoreDelegates.h"
 #include "Misc/DelayedAutoRegister.h"
diff --git a/Engine/Plugins/Runtime/RigVM/Source/RigVM/Private/RigVMPropertyUtils.cpp b/Engine/Plugins/Runtime/RigVM/Source/RigVM/Private/RigVMPropertyUtils.cpp
--- a/Engine/Plugins/Runtime/RigVM/Source/RigVM/Private/RigVMPropertyUtils.cpp
+++ b/Engine/Plugins/Runtime/RigVM/Source/RigVM/Private/RigVMPropertyUtils.cpp
@@ -4,7 +4,7 @@
 
 #include "RigVMModule.h"
 #include "RigVMTypeUtils.h"
-#include "Engine/UserDefinedStruct.h"
+#include "StructUtils/UserDefinedStruct.h"
 #include "Misc/AssertionMacros.h"
 #include "UObject/TextProperty.h"
 
diff --git a/Engine/Plugins/Runtime/RigVM/Source/RigVM/Private/RigVMTypeUtils.cpp b/Engine/Plugins/Runtime/RigVM/Source/RigVM/Private/RigVMTypeUtils.cpp
--- a/Engine/Plugins/Runtime/RigVM/Source/RigVM/Private/RigVMTypeUtils.cpp
+++ b/Engine/Plugins/Runtime/RigVM/Source/RigVM/Private/RigVMTypeUtils.cpp
@@ -8,7 +8,7 @@
 #include "AssetRegistry/AssetData.h"
 #include "AssetRegistry/AssetRegistryModule.h"
 #include "Engine/UserDefinedEnum.h"
-#include "Engine/UserDefinedStruct.h"
+#include "StructUtils/UserDefinedStruct.h"
 
 namespace RigVMTypeUtils
 {
diff --git a/Engine/Plugins/Runtime/RigVM/Source/RigVM/Public/RigVMTypeUtils.h b/Engine/Plugins/Runtime/RigVM/Source/RigVM/Public/RigVMTypeUtils.h
--- a/Engine/Plugins/Runtime/RigVM/Source/RigVM/Public/RigVMTypeUtils.h
+++ b/Engine/Plugins/Runtime/RigVM/Source/RigVM/Public/RigVMTypeUtils.h
@@ -7,7 +7,7 @@
 #include "RigVMCore/RigVMUnknownType.h"
 #include "UObject/Interface.h"
 #include "Engine/UserDefinedEnum.h"
-#include "Engine/UserDefinedStruct.h"
+#include "StructUtils/UserDefinedStruct.h"
 #include "UObject/CoreRedirects.h"
 #include "UObject/Package.h"
 #include "UObject/SoftObjectPath.h"
diff --git a/Engine/Plugins/Runtime/RigVM/Source/RigVMDeveloper/Private/RigVMModel/RigVMController.cpp b/Engine/Plugins/Runtime/RigVM/Source/RigVMDeveloper/Private/RigVMModel/RigVMController.cpp
--- a/Engine/Plugins/Runtime/RigVM/Source/RigVMDeveloper/Private/RigVMModel/RigVMController.cpp
+++ b/Engine/Plugins/Runtime/RigVM/Source/RigVMDeveloper/Private/RigVMModel/RigVMController.cpp
@@ -26,7 +26,7 @@
 #include "Algo/Transform.h"
 #include "RigVMPythonUtils.h"
 #include "RigVMTypeUtils.h"
-#include "Engine/UserDefinedStruct.h"
+#include "StructUtils/UserDefinedStruct.h"
 #include "RigVMFunctions/RigVMDispatch_If.h"
 #include "RigVMFunctions/RigVMDispatch_Select.h"
 #include "RigVMModel/RigVMClient.h"
diff --git a/Engine/Plugins/Runtime/StateTree/Source/StateTreeEditorModule/Private/StateTreeEditorData.cpp b/Engine/Plugins/Runtime/StateTree/Source/StateTreeEditorModule/Private/StateTreeEditorData.cpp
--- a/Engine/Plugins/Runtime/StateTree/Source/StateTreeEditorModule/Private/StateTreeEditorData.cpp
+++ b/Engine/Plugins/Runtime/StateTree/Source/StateTreeEditorModule/Private/StateTreeEditorData.cpp
@@ -17,7 +17,7 @@
 
 #if WITH_EDITOR
 #include "StructUtilsDelegates.h"
-#include "Engine/UserDefinedStruct.h"
+#include "StructUtils/UserDefinedStruct.h"
 #endif
 
 #include UE_INLINE_GENERATED_CPP_BY_NAME(StateTreeEditorData)
diff --git a/Engine/Plugins/Runtime/StateTree/Source/StateTreeModule/Private/StateTree.cpp b/Engine/Plugins/Runtime/StateTree/Source/StateTreeModule/Private/StateTree.cpp
--- a/Engine/Plugins/Runtime/StateTree/Source/StateTreeModule/Private/StateTree.cpp
+++ b/Engine/Plugins/Runtime/StateTree/Source/StateTreeModule/Private/StateTree.cpp
@@ -19,7 +19,7 @@
 
 #if WITH_EDITOR
 #include "Editor.h"
-#include "Engine/UserDefinedStruct.h"
+#include "StructUtils/UserDefinedStruct.h"
 #endif
 #include UE_INLINE_GENERATED_CPP_BY_NAME(StateTree)
 
diff --git a/Engine/Plugins/Runtime/StateTree/Source/StateTreeModule/Private/StateTreePropertyBindings.cpp b/Engine/Plugins/Runtime/StateTree/Source/StateTreeModule/Private/StateTreePropertyBindings.cpp
--- a/Engine/Plugins/Runtime/StateTree/Source/StateTreeModule/Private/StateTreePropertyBindings.cpp
+++ b/Engine/Plugins/Runtime/StateTree/Source/StateTreeModule/Private/StateTreePropertyBindings.cpp
@@ -10,7 +10,7 @@
 #include "UObject/CoreRedirects.h"
 #include "UObject/Package.h"
 #include "Engine/BlueprintGeneratedClass.h"
-#include "Engine/UserDefinedStruct.h"
+#include "StructUtils/UserDefinedStruct.h"
 #include "Kismet2/StructureEditorUtils.h"
 #include "UObject/Field.h"
 #endif
diff --git a/Engine/Plugins/Tests/EditorTests/Source/EditorTests/Private/UnrealEd/EditorAssetAutomationTests.cpp b/Engine/Plugins/Tests/EditorTests/Source/EditorTests/Private/UnrealEd/EditorAssetAutomationTests.cpp
--- a/Engine/Plugins/Tests/EditorTests/Source/EditorTests/Private/UnrealEd/EditorAssetAutomationTests.cpp
+++ b/Engine/Plugins/Tests/EditorTests/Source/EditorTests/Private/UnrealEd/EditorAssetAutomationTests.cpp
@@ -47,7 +47,7 @@
 #include "Engine/TextureRenderTargetCube.h"
 #include "GameFramework/TouchInterface.h"
 #include "Engine/UserDefinedEnum.h"
-#include "Engine/UserDefinedStruct.h"
+#include "StructUtils/UserDefinedStruct.h"
 #include "Animation/AnimInstance.h"
 #include "Editor/UnrealEdEngine.h"
 #include "Factories/AnimBlueprintFactory.h"
diff --git a/Engine/Source/Developer/AssetTools/Private/AssetTools.cpp b/Engine/Source/Developer/AssetTools/Private/AssetTools.cpp
--- a/Engine/Source/Developer/AssetTools/Private/AssetTools.cpp
+++ b/Engine/Source/Developer/AssetTools/Private/AssetTools.cpp
@@ -111,7 +111,7 @@
 #include "DiffUtils.h"
 #include "VirtualTexturingEditorModule.h"
 #include "Algo/AnyOf.h"
-#include "Engine/UserDefinedStruct.h"
+#include "StructUtils/UserDefinedStruct.h"
 #include "Factories/SceneImportFactory.h"
 #include "Misc/AssetFilterData.h"
 
diff --git a/Engine/Source/Editor/AnimationBlueprintLibrary/Private/K2_AnimationAttributeNodes.cpp b/Engine/Source/Editor/AnimationBlueprintLibrary/Private/K2_AnimationAttributeNodes.cpp
--- a/Engine/Source/Editor/AnimationBlueprintLibrary/Private/K2_AnimationAttributeNodes.cpp
+++ b/Engine/Source/Editor/AnimationBlueprintLibrary/Private/K2_AnimationAttributeNodes.cpp
@@ -18,7 +18,7 @@
 #include "UObject/Object.h"
 #include "UObject/WeakObjectPtr.h"
 #include "UObject/WeakObjectPtrTemplates.h"
-#include "Engine/UserDefinedStruct.h"
+#include "StructUtils/UserDefinedStruct.h"
 
 #define LOCTEXT_NAMESPACE "K2Node_AnimationAttributeNodes"
 
diff --git a/Engine/Source/Editor/AnimationSettings/Private/AnimationSettingsModule.cpp b/Engine/Source/Editor/AnimationSettings/Private/AnimationSettingsModule.cpp
--- a/Engine/Source/Editor/AnimationSettings/Private/AnimationSettingsModule.cpp
+++ b/Engine/Source/Editor/AnimationSettings/Private/AnimationSettingsModule.cpp
@@ -5,7 +5,10 @@
 #include "SSettingsEditorCheckoutNotice.h"
 #include "Animation/AnimationSettings.h"
 #include "AssetRegistry/AssetRegistryModule.h"
-#include "Engine/UserDefinedStruct.h"
+#include "GenericPlatform/GenericPlatformFile.h"
+#include "HAL/PlatformFileManager.h"
+#include "Misc/Paths.h"
+#include "StructUtils/UserDefinedStruct.h"
 
 
 IMPLEMENT_MODULE(FAnimationSettingsModule, AnimationSettings);
diff --git a/Engine/Source/Editor/BlueprintGraph/Private/EdGraphSchema_K2.cpp b/Engine/Source/Editor/BlueprintGraph/Private/EdGraphSchema_K2.cpp
--- a/Engine/Source/Editor/BlueprintGraph/Private/EdGraphSchema_K2.cpp
+++ b/Engine/Source/Editor/BlueprintGraph/Private/EdGraphSchema_K2.cpp
@@ -26,7 +26,7 @@
 #include "Components/ChildActorComponent.h"
 #include "Engine/Selection.h"
 #include "Engine/UserDefinedEnum.h"
-#include "Engine/UserDefinedStruct.h"
+#include "StructUtils/UserDefinedStruct.h"
 #include "Textures/SlateIcon.h"
 #include "Framework/Commands/UIAction.h"
 #include "Framework/Commands/UICommandList.h"
diff --git a/Engine/Source/Editor/BlueprintGraph/Private/K2Node_BreakStruct.cpp b/Engine/Source/Editor/BlueprintGraph/Private/K2Node_BreakStruct.cpp
--- a/Engine/Source/Editor/BlueprintGraph/Private/K2Node_BreakStruct.cpp
+++ b/Engine/Source/Editor/BlueprintGraph/Private/K2Node_BreakStruct.cpp
@@ -14,7 +14,7 @@
 #include "EdGraphUtilities.h"
 #include "EditorCategoryUtils.h"
 #include "Engine/Blueprint.h"
-#include "Engine/UserDefinedStruct.h"
+#include "StructUtils/UserDefinedStruct.h"
 #include "EngineLogs.h"
 #include "Internationalization/Internationalization.h"
 #include "K2Node_StructOperation.h"
diff --git a/Engine/Source/Editor/BlueprintGraph/Private/K2Node_StructOperation.cpp b/Engine/Source/Editor/BlueprintGraph/Private/K2Node_StructOperation.cpp
--- a/Engine/Source/Editor/BlueprintGraph/Private/K2Node_StructOperation.cpp
+++ b/Engine/Source/Editor/BlueprintGraph/Private/K2Node_StructOperation.cpp
@@ -10,7 +10,7 @@
 #include "Containers/EnumAsByte.h"
 #include "EdGraph/EdGraphPin.h"
 #include "EdGraphSchema_K2.h"
-#include "Engine/UserDefinedStruct.h"
+#include "StructUtils/UserDefinedStruct.h"
 #include "HAL/PlatformCrt.h"
 #include "Internationalization/Internationalization.h"
 #include "Internationalization/Text.h"
diff --git a/Engine/Source/Editor/DataTableEditor/Private/DataTableEditor.cpp b/Engine/Source/Editor/DataTableEditor/Private/DataTableEditor.cpp
--- a/Engine/Source/Editor/DataTableEditor/Private/DataTableEditor.cpp
+++ b/Engine/Source/Editor/DataTableEditor/Private/DataTableEditor.cpp
@@ -12,7 +12,7 @@
 #include "Editor.h"
 #include "Editor/EditorEngine.h"
 #include "Engine/DataTable.h"
-#include "Engine/UserDefinedStruct.h"
+#include "StructUtils/UserDefinedStruct.h"
 #include "Fonts/FontMeasure.h"
 #include "Framework/Application/SlateApplication.h"
 #include "Framework/Commands/GenericCommands.h"
diff --git a/Engine/Source/Editor/DataTableEditor/Private/SRowEditor.cpp b/Engine/Source/Editor/DataTableEditor/Private/SRowEditor.cpp
--- a/Engine/Source/Editor/DataTableEditor/Private/SRowEditor.cpp
+++ b/Engine/Source/Editor/DataTableEditor/Private/SRowEditor.cpp
@@ -6,7 +6,7 @@
 #include "DataTableUtils.h"
 #include "DetailsViewArgs.h"
 #include "Engine/DataTable.h"
-#include "Engine/UserDefinedStruct.h"
+#include "StructUtils/UserDefinedStruct.h"
 #include "HAL/PlatformMisc.h"
 #include "IStructureDetailsView.h"
 #include "Internationalization/Internationalization.h"
diff --git a/Engine/Source/Editor/GraphEditor/Private/KismetPins/SGraphPinStruct.cpp b/Engine/Source/Editor/GraphEditor/Private/KismetPins/SGraphPinStruct.cpp
--- a/Engine/Source/Editor/GraphEditor/Private/KismetPins/SGraphPinStruct.cpp
+++ b/Engine/Source/Editor/GraphEditor/Private/KismetPins/SGraphPinStruct.cpp
@@ -9,7 +9,7 @@
 #include "EdGraph/EdGraphSchema.h"
 #include "Editor.h"
 #include "Editor/EditorEngine.h"
-#include "Engine/UserDefinedStruct.h"
+#include "StructUtils/UserDefinedStruct.h"
 #include "Internationalization/Internationalization.h"
 #include "Layout/Margin.h"
 #include "Misc/Attribute.h"
diff --git a/Engine/Source/Editor/Kismet/Private/BlueprintDetailsCustomization.cpp b/Engine/Source/Editor/Kismet/Private/BlueprintDetailsCustomization.cpp
--- a/Engine/Source/Editor/Kismet/Private/BlueprintDetailsCustomization.cpp
+++ b/Engine/Source/Editor/Kismet/Private/BlueprintDetailsCustomization.cpp
@@ -35,7 +35,7 @@
 #include "Engine/MemberReference.h"
 #include "Engine/SCS_Node.h"
 #include "Engine/SimpleConstructionScript.h"
-#include "Engine/UserDefinedStruct.h"
+#include "StructUtils/UserDefinedStruct.h"
 #include "EngineLogs.h"
 #include "Fonts/SlateFontInfo.h"
 #include "Framework/Application/MenuStack.h"
diff --git a/Engine/Source/Editor/Kismet/Private/BlueprintEditor.cpp b/Engine/Source/Editor/Kismet/Private/BlueprintEditor.cpp
--- a/Engine/Source/Editor/Kismet/Private/BlueprintEditor.cpp
+++ b/Engine/Source/Editor/Kismet/Private/BlueprintEditor.cpp
@@ -7,7 +7,7 @@
 #include "Engine/BlueprintGeneratedClass.h"
 #include "Engine/SimpleConstructionScript.h"
 #include "Engine/UserDefinedEnum.h"
-#include "Engine/UserDefinedStruct.h"
+#include "StructUtils/UserDefinedStruct.h"
 #include "Logging/TokenizedMessage.h"
 #include "Misc/PackageName.h"
 #include "AssetRegistry/AssetData.h"
diff --git a/Engine/Source/Editor/Kismet/Private/BlueprintNamespaceRegistry.cpp b/Engine/Source/Editor/Kismet/Private/BlueprintNamespaceRegistry.cpp
--- a/Engine/Source/Editor/Kismet/Private/BlueprintNamespaceRegistry.cpp
+++ b/Engine/Source/Editor/Kismet/Private/BlueprintNamespaceRegistry.cpp
@@ -17,7 +17,7 @@
 #include "Engine/Blueprint.h"
 #include "Engine/BlueprintGeneratedClass.h"
 #include "Engine/UserDefinedEnum.h"
-#include "Engine/UserDefinedStruct.h"
+#include "StructUtils/UserDefinedStruct.h"
 #include "HAL/IConsoleManager.h"
 #include "HAL/Platform.h"
 #include "HAL/PlatformCrt.h"
diff --git a/Engine/Source/Editor/Kismet/Private/UserDefinedStructureEditor.cpp b/Engine/Source/Editor/Kismet/Private/UserDefinedStructureEditor.cpp
--- a/Engine/Source/Editor/Kismet/Private/UserDefinedStructureEditor.cpp
+++ b/Engine/Source/Editor/Kismet/Private/UserDefinedStructureEditor.cpp
@@ -12,7 +12,7 @@
 #include "DragAndDrop/DecoratedDragDropOp.h"
 #include "EdGraph/EdGraphNode.h"
 #include "EdGraphSchema_K2.h"
-#include "Engine/UserDefinedStruct.h"
+#include "StructUtils/UserDefinedStruct.h"
 #include "Fonts/SlateFontInfo.h"
 #include "Framework/Docking/TabManager.h"
 #include "Framework/MultiBox/MultiBoxBuilder.h"
diff --git a/Engine/Source/Editor/KismetCompiler/Private/KismetCompiler.cpp b/Engine/Source/Editor/KismetCompiler/Private/KismetCompiler.cpp
--- a/Engine/Source/Editor/KismetCompiler/Private/KismetCompiler.cpp
+++ b/Engine/Source/Editor/KismetCompiler/Private/KismetCompiler.cpp
@@ -21,7 +21,7 @@
 #include "Editor/EditorEngine.h"
 #include "Components/TimelineComponent.h"
 #include "Engine/TimelineTemplate.h"
-#include "Engine/UserDefinedStruct.h"
+#include "StructUtils/UserDefinedStruct.h"
 #include "Blueprint/BlueprintExtension.h"
 #include "EdGraphUtilities.h"
 #include "INotifyFieldValueChanged.h"
diff --git a/Engine/Source/Editor/KismetCompiler/Private/KismetCompilerMisc.cpp b/Engine/Source/Editor/KismetCompiler/Private/KismetCompilerMisc.cpp
--- a/Engine/Source/Editor/KismetCompiler/Private/KismetCompilerMisc.cpp
+++ b/Engine/Source/Editor/KismetCompiler/Private/KismetCompilerMisc.cpp
@@ -17,7 +17,7 @@
 #include "UObject/UObjectHash.h"
 #include "Engine/MemberReference.h"
 #include "Engine/BlueprintGeneratedClass.h"
-#include "Engine/UserDefinedStruct.h"
+#include "StructUtils/UserDefinedStruct.h"
 #include "FieldNotification/FieldNotificationLibrary.h"
 #include "INotifyFieldValueChanged.h"
 #include "Kismet2/CompilerResultsLog.h"
diff --git a/Engine/Source/Editor/KismetCompiler/Private/KismetCompilerModule.cpp b/Engine/Source/Editor/KismetCompiler/Private/KismetCompilerModule.cpp
--- a/Engine/Source/Editor/KismetCompiler/Private/KismetCompilerModule.cpp
+++ b/Engine/Source/Editor/KismetCompiler/Private/KismetCompilerModule.cpp
@@ -24,7 +24,7 @@
 #include "Kismet2/BlueprintEditorUtils.h"
 
 #include "UserDefinedStructureCompilerUtils.h"
-#include "Engine/UserDefinedStruct.h"
+#include "StructUtils/UserDefinedStruct.h"
 #include "IMessageLogListing.h"
 #include "Engine/Engine.h"
 
diff --git a/Engine/Source/Editor/KismetCompiler/Private/KismetCompilerVMBackend.cpp b/Engine/Source/Editor/KismetCompiler/Private/KismetCompilerVMBackend.cpp
--- a/Engine/Source/Editor/KismetCompiler/Private/KismetCompilerVMBackend.cpp
+++ b/Engine/Source/Editor/KismetCompiler/Private/KismetCompilerVMBackend.cpp
@@ -19,7 +19,7 @@
 #include "EdGraph/EdGraphNode.h"
 #include "EdGraph/EdGraphPin.h"
 #include "Engine/LatentActionManager.h"
-#include "Engine/UserDefinedStruct.h"
+#include "StructUtils/UserDefinedStruct.h"
 #include "BPTerminal.h"
 #include "EdGraphSchema_K2.h"
 #include "K2Node_MacroInstance.h"
diff --git a/Engine/Source/Editor/KismetCompiler/Private/UserDefinedStructureCompilerUtils.cpp b/Engine/Source/Editor/KismetCompiler/Private/UserDefinedStructureCompilerUtils.cpp
--- a/Engine/Source/Editor/KismetCompiler/Private/UserDefinedStructureCompilerUtils.cpp
+++ b/Engine/Source/Editor/KismetCompiler/Private/UserDefinedStructureCompilerUtils.cpp
@@ -15,7 +15,7 @@
 #include "EdMode.h"
 #include "Engine/Blueprint.h"
 #include "Engine/BlueprintGeneratedClass.h"
-#include "Engine/UserDefinedStruct.h"
+#include "StructUtils/UserDefinedStruct.h"
 #include "GameFramework/Actor.h"
 #include "HAL/PlatformCrt.h"
 #include "HAL/PlatformMath.h"
diff --git a/Engine/Source/Editor/PropertyEditor/Private/PropertyEditorModule.cpp b/Engine/Source/Editor/PropertyEditor/Private/PropertyEditorModule.cpp
--- a/Engine/Source/Editor/PropertyEditor/Private/PropertyEditorModule.cpp
+++ b/Engine/Source/Editor/PropertyEditor/Private/PropertyEditorModule.cpp
@@ -18,7 +18,7 @@
 #include "SStructureDetailsView.h"
 
 #include "Engine/UserDefinedEnum.h"
-#include "Engine/UserDefinedStruct.h"
+#include "StructUtils/UserDefinedStruct.h"
 #include "Framework/Application/SlateApplication.h"
 #include "Interfaces/IMainFrameModule.h"
 #include "Modules/ModuleManager.h"
diff --git a/Engine/Source/Editor/PropertyEditor/Private/PropertyNode.cpp b/Engine/Source/Editor/PropertyEditor/Private/PropertyNode.cpp
--- a/Engine/Source/Editor/PropertyEditor/Private/PropertyNode.cpp
+++ b/Engine/Source/Editor/PropertyEditor/Private/PropertyNode.cpp
@@ -7,7 +7,7 @@
 #include "Editor.h"
 #include "Editor/UnrealEdEngine.h"
 #include "EditorSupportDelegates.h"
-#include "Engine/UserDefinedStruct.h"
+#include "StructUtils/UserDefinedStruct.h"
 #include "InstancedReferenceSubobjectHelper.h"
 #include "Kismet2/BlueprintEditorUtils.h"
 #include "Misc/ConfigCacheIni.h"
diff --git a/Engine/Source/Editor/PropertyEditor/Private/UserInterface/PropertyEditor/SPropertyEditorStruct.cpp b/Engine/Source/Editor/PropertyEditor/Private/UserInterface/PropertyEditor/SPropertyEditorStruct.cpp
--- a/Engine/Source/Editor/PropertyEditor/Private/UserInterface/PropertyEditor/SPropertyEditorStruct.cpp
+++ b/Engine/Source/Editor/PropertyEditor/Private/UserInterface/PropertyEditor/SPropertyEditorStruct.cpp
@@ -8,7 +8,7 @@
 #include "DragAndDrop/AssetDragDropOp.h"
 #include "StructViewerModule.h"
 #include "StructViewerFilter.h"
-#include "Engine/UserDefinedStruct.h"
+#include "StructUtils/UserDefinedStruct.h"
 
 #define LOCTEXT_NAMESPACE "PropertyEditor"
 
diff --git a/Engine/Source/Editor/StructViewer/Private/SStructViewer.cpp b/Engine/Source/Editor/StructViewer/Private/SStructViewer.cpp
--- a/Engine/Source/Editor/StructViewer/Private/SStructViewer.cpp
+++ b/Engine/Source/Editor/StructViewer/Private/SStructViewer.cpp
@@ -37,7 +37,7 @@
 #include "ContentBrowserDataDragDropOp.h"
 
 #include "Editor/UnrealEdEngine.h"
-#include "Engine/UserDefinedStruct.h"
+#include "StructUtils/UserDefinedStruct.h"
 #include "EditorDirectories.h"
 #include "Dialogs/Dialogs.h"
 
diff --git a/Engine/Source/Editor/StructViewer/Private/StructViewerNode.cpp b/Engine/Source/Editor/StructViewer/Private/StructViewerNode.cpp
--- a/Engine/Source/Editor/StructViewer/Private/StructViewerNode.cpp
+++ b/Engine/Source/Editor/StructViewer/Private/StructViewerNode.cpp
@@ -3,7 +3,7 @@
 #include "StructViewerNode.h"
 
 #include "AssetRegistry/AssetData.h"
-#include "Engine/UserDefinedStruct.h"
+#include "StructUtils/UserDefinedStruct.h"
 #include "HAL/Platform.h"
 #include "HAL/PlatformCrt.h"
 #include "Internationalization/Internationalization.h"
diff --git a/Engine/Source/Editor/UMGEditor/Private/WidgetBlueprint.cpp b/Engine/Source/Editor/UMGEditor/Private/WidgetBlueprint.cpp
--- a/Engine/Source/Editor/UMGEditor/Private/WidgetBlueprint.cpp
+++ b/Engine/Source/Editor/UMGEditor/Private/WidgetBlueprint.cpp
@@ -6,7 +6,7 @@
 #include "Blueprint/UserWidget.h"
 #include "MovieScene.h"
 
-#include "Engine/UserDefinedStruct.h"
+#include "StructUtils/UserDefinedStruct.h"
 #include "EdGraph/EdGraph.h"
 #include "Blueprint/WidgetTree.h"
 #include "Animation/WidgetAnimation.h"
diff --git a/Engine/Source/Editor/UnrealEd/Classes/UserDefinedStructure/UserDefinedStructEditorData.h b/Engine/Source/Editor/UnrealEd/Classes/UserDefinedStructure/UserDefinedStructEditorData.h
--- a/Engine/Source/Editor/UnrealEd/Classes/UserDefinedStructure/UserDefinedStructEditorData.h
+++ b/Engine/Source/Editor/UnrealEd/Classes/UserDefinedStructure/UserDefinedStructEditorData.h
@@ -11,6 +11,7 @@
 #include "UObject/StructOnScope.h"
 #include "EditorUndoClient.h"
 #include "Kismet2/StructureEditorUtils.h"
+#include "StructUtils/UserDefinedStructEditorUtils.h"
 #include "UserDefinedStructEditorData.generated.h"
 
 class ITransactionObjectAnnotation;
@@ -127,12 +128,12 @@ struct TStructOpsTypeTraits< FStructVariableDescription > : public TStructOpsTyp
 };
 
 UCLASS(MinimalAPI)
-class UUserDefinedStructEditorData : public UObject, public FEditorUndoClient
+class UUserDefinedStructEditorData : public UUserDefinedStructEditorDataBase, public FEditorUndoClient
 {
 	GENERATED_UCLASS_BODY()
 
 private:
-	// the property is used to generate an uniqe name id for member variable
+	// the property is used to generate a unique name id for member variable
 	UPROPERTY(NonTransactional) 
 	uint32 UniqueNameId;
 
@@ -149,7 +150,7 @@ public:
 	UNREALED_API virtual void Serialize(FArchive& Ar) override;
 	UNREALED_API virtual void PostEditUndo() override;
 	UNREALED_API virtual void PostEditUndo(TSharedPtr<ITransactionObjectAnnotation> TransactionAnnotation) override;
-	UNREALED_API virtual void PostLoadSubobjects(struct FObjectInstancingGraph* OuterInstanceGraph) override;
+	UNREALED_API virtual void PostLoadSubobjects(FObjectInstancingGraph* OuterInstanceGraph) override;
 	UNREALED_API virtual void PreSave(FObjectPreSaveContext ObjectSaveContext) override;
 	// End of UObject interface.
 
@@ -160,13 +161,17 @@ public:
 
 
 	UNREALED_API uint32 GenerateUniqueNameIdForMemberVariable();
-	UNREALED_API class UUserDefinedStruct* GetOwnerStruct() const;
+	UNREALED_API UUserDefinedStruct* GetOwnerStruct() const;
 
 	UNREALED_API const uint8* GetDefaultInstance() const;
-	UNREALED_API void RecreateDefaultInstance(FString* OutLog = nullptr);
-	UNREALED_API void ReinitializeDefaultInstance(FString* OutLog = nullptr);
+	UNREALED_API virtual void RecreateDefaultInstance(FString* OutLog = nullptr) override;
+	UNREALED_API virtual void ReinitializeDefaultInstance(FString* OutLog = nullptr) override;
 	UNREALED_API void CleanDefaultInstance();
 	UNREALED_API void RefreshValuesFromDefaultInstance();
+	
+	UNREALED_API virtual FString GetTooltip() const override;
+	UNREALED_API virtual FProperty* FindProperty(const UUserDefinedStruct* Struct, FName Name) const override;
+	UNREALED_API virtual FString GetFriendlyNameForProperty(const UUserDefinedStruct* Struct, const FProperty* Property) const override;
 
 private:
 
diff --git a/Engine/Source/Editor/UnrealEd/Private/Commandlets/StabilizeLocalizationKeys.cpp b/Engine/Source/Editor/UnrealEd/Private/Commandlets/StabilizeLocalizationKeys.cpp
--- a/Engine/Source/Editor/UnrealEd/Private/Commandlets/StabilizeLocalizationKeys.cpp
+++ b/Engine/Source/Editor/UnrealEd/Private/Commandlets/StabilizeLocalizationKeys.cpp
@@ -14,7 +14,7 @@
 #include "Internationalization/InternationalizationManifest.h"
 #include "Internationalization/InternationalizationArchive.h"
 #include "UserDefinedStructure/UserDefinedStructEditorData.h"
-#include "Engine/UserDefinedStruct.h"
+#include "StructUtils/UserDefinedStruct.h"
 #include "Serialization/JsonInternationalizationManifestSerializer.h"
 #include "Serialization/JsonInternationalizationArchiveSerializer.h"
 #include "Internationalization/TextPackageNamespaceUtil.h"
diff --git a/Engine/Source/Editor/UnrealEd/Private/DataTableEditorUtils.cpp b/Engine/Source/Editor/UnrealEd/Private/DataTableEditorUtils.cpp
--- a/Engine/Source/Editor/UnrealEd/Private/DataTableEditorUtils.cpp
+++ b/Engine/Source/Editor/UnrealEd/Private/DataTableEditorUtils.cpp
@@ -8,7 +8,7 @@
 #include "Framework/Application/SlateApplication.h"
 #include "Framework/Application/SlateUser.h"
 #include "Styling/AppStyle.h"
-#include "Engine/UserDefinedStruct.h"
+#include "StructUtils/UserDefinedStruct.h"
 #include "Misc/StringUtility.h"
 #include "ScopedTransaction.h"
 #include "K2Node_GetDataTableRow.h"
diff --git a/Engine/Source/Editor/UnrealEd/Private/Factories/EditorFactories.cpp b/Engine/Source/Editor/UnrealEd/Private/Factories/EditorFactories.cpp
--- a/Engine/Source/Editor/UnrealEd/Private/Factories/EditorFactories.cpp
+++ b/Engine/Source/Editor/UnrealEd/Private/Factories/EditorFactories.cpp
@@ -181,7 +181,7 @@
 #include "Engine/TextureRenderTargetVolume.h"
 #include "GameFramework/TouchInterface.h"
 #include "Engine/UserDefinedEnum.h"
-#include "Engine/UserDefinedStruct.h"
+#include "StructUtils/UserDefinedStruct.h"
 #include "Internationalization/StringTable.h"
 #include "Editor.h"
 #include "Materials/MaterialExpressionTextureSample.h"
diff --git a/Engine/Source/Editor/UnrealEd/Private/Kismet2/BlueprintEditorUtils.cpp b/Engine/Source/Editor/UnrealEd/Private/Kismet2/BlueprintEditorUtils.cpp
--- a/Engine/Source/Editor/UnrealEd/Private/Kismet2/BlueprintEditorUtils.cpp
+++ b/Engine/Source/Editor/UnrealEd/Private/Kismet2/BlueprintEditorUtils.cpp
@@ -39,7 +39,7 @@
 #include "Engine/LevelScriptActor.h"
 #include "Components/TimelineComponent.h"
 #include "Engine/TimelineTemplate.h"
-#include "Engine/UserDefinedStruct.h"
+#include "StructUtils/UserDefinedStruct.h"
 #include "UObject/PropertyPortFlags.h"
 #include "Serialization/ArchiveReplaceObjectRef.h"
 #include "EngineUtils.h"
diff --git a/Engine/Source/Editor/UnrealEd/Private/Kismet2/EnumEditorUtils.cpp b/Engine/Source/Editor/UnrealEd/Private/Kismet2/EnumEditorUtils.cpp
--- a/Engine/Source/Editor/UnrealEd/Private/Kismet2/EnumEditorUtils.cpp
+++ b/Engine/Source/Editor/UnrealEd/Private/Kismet2/EnumEditorUtils.cpp
@@ -17,7 +17,7 @@
 #include "Internationalization/TextPackageNamespaceUtil.h"
 #include "ScopedTransaction.h"
 #include "UObject/PropertyIterator.h"
-#include "Engine/UserDefinedStruct.h"
+#include "StructUtils/UserDefinedStruct.h"
 #include "Kismet2/StructureEditorUtils.h"
 
 #define LOCTEXT_NAMESPACE "Enum"
diff --git a/Engine/Source/Editor/UnrealEd/Private/Kismet2/StructureEditorUtils.cpp b/Engine/Source/Editor/UnrealEd/Private/Kismet2/StructureEditorUtils.cpp
--- a/Engine/Source/Editor/UnrealEd/Private/Kismet2/StructureEditorUtils.cpp
+++ b/Engine/Source/Editor/UnrealEd/Private/Kismet2/StructureEditorUtils.cpp
@@ -21,6 +21,14 @@
 
 //////////////////////////////////////////////////////////////////////////
 // FStructEditorManager
+FStructureEditorUtils::FStructEditorManager::FStructEditorManager()
+{
+	FUserDefinedStructEditorUtils::OnUserDefinedStructChanged.BindLambda([](UUserDefinedStruct* Struct)
+		{
+			OnStructureChanged(Struct);
+		});
+}
+
 FStructureEditorUtils::FStructEditorManager& FStructureEditorUtils::FStructEditorManager::Get()
 {
 	static TSharedRef< FStructEditorManager > EditorManager( new FStructEditorManager() );
@@ -567,7 +575,7 @@ void FStructureEditorUtils::OnStructureChanged(UUserDefinedStruct* Struct, EStru
 
 		Struct->Status = EUserDefinedStructureStatus::UDSS_Dirty;
 		CompileStructure(Struct);
-		Struct->MarkPackageDirty();
+		(void)Struct->MarkPackageDirty();
 		Struct->OnChanged();
 	}
 }
diff --git a/Engine/Source/Editor/UnrealEd/Private/ObjectTools.cpp b/Engine/Source/Editor/UnrealEd/Private/ObjectTools.cpp
--- a/Engine/Source/Editor/UnrealEd/Private/ObjectTools.cpp
+++ b/Engine/Source/Editor/UnrealEd/Private/ObjectTools.cpp
@@ -56,7 +56,7 @@
 #include "PhysicalMaterials/PhysicalMaterial.h"
 #include "Engine/Selection.h"
 #include "Engine/TextureRenderTarget2D.h"
-#include "Engine/UserDefinedStruct.h"
+#include "StructUtils/UserDefinedStruct.h"
 #include "Animation/MorphTarget.h"
 #include "Editor.h"
 #include "Editor/Transactor.h"
diff --git a/Engine/Source/Editor/UnrealEd/Private/UserDefinedStructEditorData.cpp b/Engine/Source/Editor/UnrealEd/Private/UserDefinedStructEditorData.cpp
--- a/Engine/Source/Editor/UnrealEd/Private/UserDefinedStructEditorData.cpp
+++ b/Engine/Source/Editor/UnrealEd/Private/UserDefinedStructEditorData.cpp
@@ -5,7 +5,7 @@
 #include "UObject/UE5ReleaseStreamObjectVersion.h"
 #include "UObject/UnrealType.h"
 #include "UObject/ObjectSaveContext.h"
-#include "Engine/UserDefinedStruct.h"
+#include "StructUtils/UserDefinedStruct.h"
 #include "Kismet2/StructureEditorUtils.h"
 #include "Kismet2/BlueprintEditorUtils.h"
 #include "Blueprint/BlueprintSupport.h"
@@ -77,7 +77,7 @@ FEdGraphPinType FStructVariableDescription::ToPinType() const
 }
 
 UUserDefinedStructEditorData::UUserDefinedStructEditorData(const FObjectInitializer& ObjectInitializer)
-	: Super(ObjectInitializer)
+	: Super(ObjectInitializer), UniqueNameId(0), CachedStructureChange()
 {
 }
 
@@ -308,4 +308,25 @@ void UUserDefinedStructEditorData::RefreshValuesFromDefaultInstance()
 	}
 }
 
+FProperty* UUserDefinedStructEditorData::FindProperty(const UUserDefinedStruct* Struct, const FName Name) const
+{
+	const FGuid PropertyGuid = FStructureEditorUtils::GetGuidFromPropertyName(Name);
+	FProperty* EditorProperty = PropertyGuid.IsValid()
+		? FStructureEditorUtils::GetPropertyByGuid(Struct, PropertyGuid)
+		: FStructureEditorUtils::GetPropertyByFriendlyName(Struct, Name.ToString());
+
+	ensure(!EditorProperty || !PropertyGuid.IsValid() || PropertyGuid == FStructureEditorUtils::GetGuidForProperty(EditorProperty));
+	return EditorProperty;
+}
+
+FString UUserDefinedStructEditorData::GetFriendlyNameForProperty(const UUserDefinedStruct* Struct, const FProperty* Property) const
+{
+	return Super::GetFriendlyNameForProperty(Struct, Property);
+}
+
+FString UUserDefinedStructEditorData::GetTooltip() const
+{
+	return ToolTip;
+}
+
 #undef LOCTEXT_NAMESPACE
diff --git a/Engine/Source/Editor/UnrealEd/Public/Kismet2/StructureEditorUtils.h b/Engine/Source/Editor/UnrealEd/Public/Kismet2/StructureEditorUtils.h
--- a/Engine/Source/Editor/UnrealEd/Public/Kismet2/StructureEditorUtils.h
+++ b/Engine/Source/Editor/UnrealEd/Public/Kismet2/StructureEditorUtils.h
@@ -3,7 +3,7 @@
 
 #include "CoreMinimal.h"
 #include "Misc/Guid.h"
-#include "Engine/UserDefinedStruct.h"
+#include "StructUtils/UserDefinedStruct.h"
 #include "Kismet2/ListenerManager.h"
 
 struct FEdGraphPinType;
@@ -27,7 +27,7 @@ public:
 
 	class FStructEditorManager : public FListenerManager<UUserDefinedStruct, EStructureEditorChangeInfo>
 	{
-		FStructEditorManager() {}
+		FStructEditorManager();
 	public:
 		UNREALED_API static FStructEditorManager& Get();
 
diff --git a/Engine/Source/Runtime/AIModule/Classes/Perception/AISenseBlueprintListener.h b/Engine/Source/Runtime/AIModule/Classes/Perception/AISenseBlueprintListener.h
--- a/Engine/Source/Runtime/AIModule/Classes/Perception/AISenseBlueprintListener.h
+++ b/Engine/Source/Runtime/AIModule/Classes/Perception/AISenseBlueprintListener.h
@@ -2,7 +2,7 @@
 
 #pragma once
 
-#include "Engine/UserDefinedStruct.h"
+#include "StructUtils/UserDefinedStruct.h"
 #include "AISenseBlueprintListener.generated.h"
 
 UCLASS(ClassGroup=AI, MinimalAPI)
diff --git a/Engine/Source/Runtime/CoreUObject/Private/StructUtils/UserDefinedStruct.cpp b/Engine/Source/Runtime/CoreUObject/Private/StructUtils/UserDefinedStruct.cpp
new file mode 100644
--- /dev/null
+++ b/Engine/Source/Runtime/CoreUObject/Private/StructUtils/UserDefinedStruct.cpp
@@ -0,0 +1,606 @@
+// Copyright Epic Games, Inc. All Rights Reserved.
+
+#include "StructUtils/UserDefinedStruct.h"
+#include "Templates/SubclassOf.h"
+#include "UObject/AssetRegistryTagsContext.h"
+#include "UObject/Package.h"
+#include "UObject/Package.h"
+#include "UObject/UnrealType.h"
+#include "UObject/LinkerLoad.h"
+#include "UObject/ObjectSaveContext.h"
+#include "UObject/FrameworkObjectVersion.h"
+#include "Misc/PackageName.h"
+#include "Blueprint/BlueprintSupport.h"
+
+#if WITH_EDITORONLY_DATA
+#include "UObject/CookedMetaData.h"
+#include "StructUtils/UserDefinedStructEditorUtils.h"
+#endif //WITH_EDITORONLY_DATA
+
+FUserStructOnScopeIgnoreDefaults::FUserStructOnScopeIgnoreDefaults(const UUserDefinedStruct* InUserStruct)
+{
+	// Can't call super constructor because we need to call our overridden initialize
+	ScriptStruct = InUserStruct;
+	SampleStructMemory = nullptr;
+	OwnsMemory = false;
+	Initialize();
+}
+
+FUserStructOnScopeIgnoreDefaults::FUserStructOnScopeIgnoreDefaults(const UUserDefinedStruct* InUserStruct, uint8* InData) : FStructOnScope(InUserStruct, InData)
+{
+}
+
+void FUserStructOnScopeIgnoreDefaults::Recreate(const UUserDefinedStruct* InUserStruct)
+{
+	Destroy();
+	ScriptStruct = InUserStruct;
+	Initialize();
+}
+
+void FUserStructOnScopeIgnoreDefaults::Initialize()
+{
+	if (const UStruct* ScriptStructPtr = ScriptStruct.Get())
+	{
+		SampleStructMemory = (uint8*)FMemory::Malloc(ScriptStructPtr->GetStructureSize() ? ScriptStructPtr->GetStructureSize() : 1);
+		((UUserDefinedStruct*)ScriptStruct.Get())->InitializeStructIgnoreDefaults(SampleStructMemory);
+		OwnsMemory = true;
+	}
+}
+
+UUserDefinedStruct::UUserDefinedStruct(const FObjectInitializer& ObjectInitializer)
+	: Super(ObjectInitializer)
+{
+	DefaultStructInstance.SetPackage(GetOutermost());
+}
+
+void UUserDefinedStruct::PrepareCppStructOps()
+{
+	// User structs can never have struct ops, this stops it from incorrectly assigning a native struct of the same raw name to this user struct
+	CppStructOps = nullptr;
+	bPrepareCppStructOpsCompleted = true;
+}
+
+void UUserDefinedStruct::Serialize(FStructuredArchive::FRecord Record)
+{
+	Super::Serialize(Record);
+	FArchive& UnderlyingArchive = Record.GetUnderlyingArchive();
+
+	if (HasAnyFlags(RF_ClassDefaultObject))
+	{
+		return;
+	}
+
+	UnderlyingArchive.UsingCustomVersion(FFrameworkObjectVersion::GUID);
+
+	if (UnderlyingArchive.CustomVer(FFrameworkObjectVersion::GUID) >= FFrameworkObjectVersion::UserDefinedStructsStoreDefaultInstance)
+	{
+		if (EUserDefinedStructureStatus::UDSS_UpToDate == Status && !(UnderlyingArchive.GetPortFlags() & PPF_Duplicate))
+		{
+			// If we're saving or loading new data, serialize our defaults
+			if (!DefaultStructInstance.IsValid() && UnderlyingArchive.IsLoading())
+			{
+				DefaultStructInstance.Recreate(this);
+			}
+
+			uint8* StructData = DefaultStructInstance.GetStructMemory();
+
+			FScopedPlaceholderRawContainerTracker TrackStruct(StructData);
+			SerializeItem(Record.EnterField(TEXT("Data")), StructData, nullptr);
+
+			// Now that defaults have been loaded we can inspect our properties
+			// and default values and set the StructFlags accordingly:
+			if(UnderlyingArchive.IsLoading())
+			{
+				UpdateStructFlags();
+			}
+		}
+	}
+
+#if WITH_EDITOR
+	if (UnderlyingArchive.IsLoading())
+	{
+		if (UnderlyingArchive.CustomVer(FFrameworkObjectVersion::GUID) < FFrameworkObjectVersion::UserDefinedStructsBlueprintVisible)
+		{
+			for (TFieldIterator<FProperty> PropIt(this); PropIt; ++PropIt)
+			{
+				PropIt->PropertyFlags |= CPF_BlueprintVisible;
+			}
+		}
+
+		if (EUserDefinedStructureStatus::UDSS_UpToDate == Status)
+		{
+			// We need to force the editor data to be preload in case anyone needs to extract variable
+			// information at editor time about the user structure.
+			if (EditorData != nullptr)
+			{
+				UnderlyingArchive.Preload(EditorData);
+				if (!(UnderlyingArchive.GetPortFlags() & PPF_Duplicate))
+				{
+					if (UUserDefinedStructEditorDataBase* UDSEditorData = Cast<UUserDefinedStructEditorDataBase>(EditorData))
+					{
+						if (!DefaultStructInstance.IsValid())
+						{
+							UDSEditorData->RecreateDefaultInstance();
+						}
+						else
+						{
+							UDSEditorData->ReinitializeDefaultInstance();
+						}
+					}
+				}
+			}
+
+			const FUserDefinedStructEditorUtils::EStructureError Result = FUserDefinedStructEditorUtils::IsStructureValid(this, NULL, &ErrorMessage);
+			if (FUserDefinedStructEditorUtils::EStructureError::Ok != Result)
+			{
+				Status = EUserDefinedStructureStatus::UDSS_Error;
+				UE_LOG(LogClass, Log, TEXT("UUserDefinedStruct.Serialize '%s' validation: %s"), *GetName(), *ErrorMessage);
+			}
+		}
+	}
+#endif
+}
+
+#if WITH_EDITOR
+
+void UUserDefinedStruct::PostDuplicate(bool bDuplicateForPIE)
+{
+	Super::PostDuplicate(bDuplicateForPIE);
+	if (!bDuplicateForPIE)
+	{
+		Guid = FGuid::NewGuid();
+	}
+	if (!bDuplicateForPIE && (GetOuter() != GetTransientPackage()))
+	{
+		SetMetaData(TEXT("BlueprintType"), TEXT("true"));
+		FUserDefinedStructEditorUtils::OnStructureChanged(this);
+	}
+}
+
+void UUserDefinedStruct::PostLoad()
+{
+	Super::PostLoad();
+
+	ValidateGuid();
+}
+
+void UUserDefinedStruct::PreSaveRoot(FObjectPreSaveRootContext ObjectSaveContext)
+{
+	Super::PreSaveRoot(ObjectSaveContext);
+
+	if (ObjectSaveContext.IsCooking() && (ObjectSaveContext.GetSaveFlags() & SAVE_Optional))
+	{
+		UStructCookedMetaData* CookedMetaData = NewCookedMetaData();
+		CookedMetaData->CacheMetaData(this);
+
+		if (!CookedMetaData->HasMetaData())
+		{
+			PurgeCookedMetaData();
+		}
+	}
+	else
+	{
+		PurgeCookedMetaData();
+	}
+}
+
+void UUserDefinedStruct::PostSaveRoot(FObjectPostSaveRootContext ObjectSaveContext)
+{
+	Super::PostSaveRoot(ObjectSaveContext);
+
+	PurgeCookedMetaData();
+}
+
+void UUserDefinedStruct::GetAssetRegistryTags(TArray<FAssetRegistryTag>& OutTags) const
+{
+	PRAGMA_DISABLE_DEPRECATION_WARNINGS;
+	Super::GetAssetRegistryTags(OutTags);
+	PRAGMA_ENABLE_DEPRECATION_WARNINGS;
+}
+
+void UUserDefinedStruct::GetAssetRegistryTags(FAssetRegistryTagsContext Context) const
+{
+	Super::GetAssetRegistryTags(Context);
+
+	if (const UUserDefinedStructEditorDataBase* StructEditorData = Cast<const UUserDefinedStructEditorDataBase>(EditorData))
+	{
+		Context.AddTag(FAssetRegistryTag(TEXT("Tooltip"), StructEditorData->GetTooltip(), FAssetRegistryTag::TT_Hidden));
+	}
+}
+
+void UUserDefinedStruct::ValidateGuid()
+{
+	// Backward compatibility:
+	// The guid is created in an deterministic way using existing name.
+	if (!Guid.IsValid() && (GetFName() != NAME_None))
+	{
+		const FString HashString = GetFName().ToString();
+		ensure(HashString.Len());
+
+		const uint32 BufferLength = HashString.Len() * sizeof(HashString[0]);
+		uint32 HashBuffer[5];
+		FSHA1::HashBuffer(*HashString, BufferLength, reinterpret_cast<uint8*>(HashBuffer));
+		Guid = FGuid(HashBuffer[1], HashBuffer[2], HashBuffer[3], HashBuffer[4]);
+	}
+}
+
+void UUserDefinedStruct::OnChanged()
+{
+	ChangedEvent.Broadcast(this);
+}
+
+#endif	// WITH_EDITOR
+
+FProperty* UUserDefinedStruct::CustomFindProperty(const FName Name) const
+{
+#if WITH_EDITOR
+	// If we have the editor data, check that first as it's more up to date
+	if (const UUserDefinedStructEditorDataBase* StructEditorData = Cast<const UUserDefinedStructEditorDataBase>(EditorData))
+	{
+		if (FProperty* EditorProperty = StructEditorData->FindProperty(this, Name))
+		{
+			return EditorProperty;
+		}
+	}
+#endif // WITH_EDITOR
+
+	// Check the authored names for each field
+	FString NameString = Name.ToString();
+	for (FProperty* CurrentProp : TFieldRange<FProperty>(this))
+	{
+		if (GetAuthoredNameForField(CurrentProp) == NameString)
+		{
+			return CurrentProp;
+		}
+	}
+	return nullptr;
+}
+
+FString UUserDefinedStruct::GetAuthoredNameForField(const FField* Field) const
+{
+	const FProperty* Property = CastField<FProperty>(Field);
+	if (!Property)
+	{
+		return Super::GetAuthoredNameForField(Field);
+	}
+
+#if WITH_EDITOR
+	if (const UUserDefinedStructEditorDataBase* StructEditorData = Cast<const UUserDefinedStructEditorDataBase>(EditorData))
+	{
+		const FString EditorName = StructEditorData->GetFriendlyNameForProperty(this, Property);
+		if (!EditorName.IsEmpty())
+		{
+			return EditorName;
+		}
+	}
+#endif	// WITH_EDITOR
+
+	const int32 GuidStrLen = 32;
+	const int32 MinimalPostfixlen = GuidStrLen + 3;
+	const FString OriginalName = Property->GetName();
+	if (OriginalName.Len() > MinimalPostfixlen)
+	{
+		FString DisplayName = OriginalName.LeftChop(GuidStrLen + 1);
+		int FirstCharToRemove = -1;
+		const bool bCharFound = DisplayName.FindLastChar(TCHAR('_'), FirstCharToRemove);
+		if (bCharFound && (FirstCharToRemove > 0))
+		{
+			return DisplayName.Mid(0, FirstCharToRemove);
+		}
+	}
+	return OriginalName;
+}
+
+void UUserDefinedStruct::InitializeStructIgnoreDefaults(void* Dest, int32 ArrayDim) const
+{
+	Super::InitializeStruct(Dest, ArrayDim);
+}
+
+void UUserDefinedStruct::InitializeStruct(void* Dest, int32 ArrayDim) const
+{
+	check(Dest);
+
+	const uint8* DefaultInstance = GetDefaultInstance();
+	if ((StructFlags & STRUCT_IsPlainOldData) == 0 || !DefaultInstance)
+	{
+		InitializeStructIgnoreDefaults(Dest, ArrayDim);
+	}
+
+	if (DefaultInstance)
+	{
+		int32 Stride = GetStructureSize();
+
+		for (int32 ArrayIndex = 0; ArrayIndex < ArrayDim; ArrayIndex++)
+		{
+			void* DestStruct = (uint8*)Dest + (Stride * ArrayIndex);
+			CopyScriptStruct(DestStruct, DefaultInstance);
+
+#if WITH_EDITOR
+			// When copying into another struct we need to register this raw struct pointer so any deferred dependencies will get fixed later
+			FScopedPlaceholderRawContainerTracker TrackStruct(DestStruct);
+			FBlueprintSupport::RegisterDeferredDependenciesInStruct(this, DestStruct);
+#endif
+		}	
+	}
+}
+
+void UUserDefinedStruct::SerializeTaggedProperties(FStructuredArchive::FSlot Slot, uint8* Data, UStruct* DefaultsStruct, uint8* Defaults, const UObject* BreakRecursionIfFullyLoad) const 
+{
+	bool bTemporarilyEnableDelta = false;
+	FArchive& Ar = Slot.GetUnderlyingArchive();
+
+#if WITH_EDITOR
+	// In the editor the default structure may change while the editor is running, so we need to always delta serialize
+	UUserDefinedStruct* UDDefaultsStruct = Cast<UUserDefinedStruct>(DefaultsStruct);
+
+	const bool bDuplicate = (0 != (Ar.GetPortFlags() & PPF_Duplicate));
+
+	// When saving delta, we want the difference between current data and true structure's default values.
+	// So if we don't have defaults we need to use the struct defaults
+	const bool bUseNewDefaults = !Defaults
+		&& UDDefaultsStruct
+		&& (UDDefaultsStruct->GetDefaultInstance() != Data)
+		&& !bDuplicate
+		&& (Ar.IsSaving() || Ar.IsLoading())
+		&& !Ar.IsCooking();
+
+	if (bUseNewDefaults)
+	{
+		Defaults = const_cast<uint8*>(UDDefaultsStruct->GetDefaultInstance());
+	}
+
+	// Temporarily enable delta serialization if this is a CPFUO 
+	bTemporarilyEnableDelta = bUseNewDefaults && Ar.IsIgnoringArchetypeRef() && Ar.IsIgnoringClassRef() && !Ar.DoDelta();
+	if (bTemporarilyEnableDelta)
+	{
+		Ar.ArNoDelta = false;
+	}
+#endif // WITH_EDITOR
+
+	Super::SerializeTaggedProperties(Slot, Data, DefaultsStruct, Defaults, BreakRecursionIfFullyLoad);
+
+	if (bTemporarilyEnableDelta)
+	{
+		Ar.ArNoDelta = true;
+	}
+}
+
+uint32 UUserDefinedStruct::GetStructTypeHash(const void* Src) const
+{
+	return GetUserDefinedStructTypeHash(Src, this);
+}
+
+void UUserDefinedStruct::RecursivelyPreload()
+{
+	FLinkerLoad* Linker = GetLinker();
+	if( Linker && (NULL == PropertyLink) )
+	{
+		TArray<UObject*> AllChildMembers;
+		GetObjectsWithOuter(this, AllChildMembers);
+		for (int32 Index = 0; Index < AllChildMembers.Num(); ++Index)
+		{
+			UObject* Member = AllChildMembers[Index];
+			Linker->Preload(Member);
+		}
+
+		Linker->Preload(this);
+		if (NULL == PropertyLink)
+		{
+			StaticLink(true);
+		}
+	}
+}
+
+FGuid UUserDefinedStruct::GetCustomGuid() const
+{
+	return Guid;
+}
+
+FString GetPathPostfix(const UObject* ForObject)
+{
+	FString FullAssetName = ForObject->GetOutermost()->GetPathName();
+	FString AssetName = FPackageName::GetLongPackageAssetName(FullAssetName);
+	// append a hash of the path, this uniquely identifies assets with the same name, but different folders:
+	FullAssetName.RemoveFromEnd(AssetName);
+	return FString::Printf(TEXT("%u"), GetTypeHash(FullAssetName));
+}
+
+FString UUserDefinedStruct::GetStructCPPName(uint32 CPPExportFlags) const
+{
+	if (CPPExportFlags & CPPF_BlueprintCppBackend)
+	{
+		return ::UnicodeToCPPIdentifier(*GetName(), false, GetPrefixCPP()) + GetPathPostfix(this);
+	}
+	else
+	{
+		return Super::GetStructCPPName(CPPExportFlags);
+	}
+}
+
+uint32 UUserDefinedStruct::GetUserDefinedStructTypeHash(const void* Src, const UScriptStruct* Type)
+{
+	// Ugliness to maximize entropy on first call to HashCombine... Alternatively we could
+	// do special logic on the first iteration of the below loops (unwind loops or similar):
+	const auto ConditionalCombineHash = [](uint32 AccumulatedHash, uint32 CurrentHash) -> uint32
+	{
+		if (AccumulatedHash != 0)
+		{
+			return HashCombine(AccumulatedHash, CurrentHash);
+		}
+		else
+		{
+			return CurrentHash;
+		}
+	};
+
+	uint32 ValueHash = 0;
+	// combining bool values and hashing them together, small range enums could get stuffed into here as well,
+	// but FBoolProperty does not actually provide GetValueTypeHash (and probably shouldn't). For structs
+	// with more than 64 boolean values we lose some information, but that is acceptable, just slightly 
+	// increasing risk of hash collision:
+	bool bHasBoolValues = false;
+	uint64 BoolValues = 0;
+	// for blueprint defined structs we can just loop and hash the individual properties:
+	for (TFieldIterator<const FProperty> It(Type); It; ++It)
+	{
+		uint32 CurrentHash = 0;
+		if (const FBoolProperty* BoolProperty = CastField<const FBoolProperty>(*It))
+		{
+			for (int32 I = 0; I < It->ArrayDim; ++I)
+			{
+				BoolValues = (BoolValues << 1) + ( BoolProperty->GetPropertyValue_InContainer(Src, I) ? 1 : 0 );
+			}
+		}
+		else if (ensure(It->HasAllPropertyFlags(CPF_HasGetValueTypeHash)))
+		{
+			for (int32 I = 0; I < It->ArrayDim; ++I)
+			{
+				uint32 Hash = It->GetValueTypeHash(It->ContainerPtrToValuePtr<void>(Src, I));
+				CurrentHash = ConditionalCombineHash(CurrentHash, Hash);
+			}
+		}
+
+		ValueHash = ConditionalCombineHash(ValueHash, CurrentHash);
+	}
+
+	if (bHasBoolValues)
+	{
+		ValueHash = ConditionalCombineHash(ValueHash, GetTypeHash(BoolValues));
+	}
+
+	return ValueHash;
+}
+
+const uint8* UUserDefinedStruct::GetDefaultInstance() const
+{
+	ensure(DefaultStructInstance.IsValid() && DefaultStructInstance.GetStruct() == this);
+	return DefaultStructInstance.GetStructMemory();
+}
+
+void UUserDefinedStruct::AddReferencedObjects(UObject* InThis, FReferenceCollector& Collector)
+{
+	UUserDefinedStruct* This = CastChecked<UUserDefinedStruct>(InThis);
+	ensure(!This->DefaultStructInstance.IsValid() || This->DefaultStructInstance.GetStruct() == This);
+
+	if (uint8* StructData = This->DefaultStructInstance.GetStructMemory())
+	{
+		Collector.AddPropertyReferences(This, StructData, This);
+	}
+
+	Super::AddReferencedObjects(This, Collector);
+}
+void UUserDefinedStruct::UpdateStructFlags()
+{
+	// Adapted from PrepareCppStructOps, where we 'discover' zero constructability
+	// for native types:
+	bool bIsZeroConstruct = true;
+	{
+		uint8* StructData = DefaultStructInstance.GetStructMemory();
+		if (StructData)
+		{
+			int32 Size = GetStructureSize();
+			for (int32 Index = 0; Index < Size; Index++)
+			{
+				if (StructData[Index])
+				{
+					bIsZeroConstruct = false;
+					break;
+				}
+			}
+		}
+
+		if(bIsZeroConstruct)
+		{	
+			for (TFieldIterator<FProperty> It(this); It; ++It)
+			{
+				FProperty* Property = *It;
+				if (Property && !Property->HasAnyPropertyFlags(CPF_ZeroConstructor))
+				{
+					bIsZeroConstruct = false;
+					break;
+				}
+			}
+		}
+	}
+
+	// IsPOD/NoDtor could be derived earlier than bIsZeroConstruct because they do not depend on 
+	// the structs default values, but it is convenient to calculate them all in one place:
+	bool bIsPOD = true;
+	{
+		for (TFieldIterator<FProperty> It(this); It; ++It)
+		{
+			FProperty* Property = *It;
+			if (Property && !Property->HasAnyPropertyFlags(CPF_IsPlainOldData))
+			{
+				bIsPOD = false;
+				break;
+			}
+		}
+	}
+
+	bool bHasNoDtor = bIsPOD;
+	{
+		if(!bHasNoDtor)
+		{
+			// we're not POD, but we still may have no destructor, check properties:
+			bHasNoDtor = true;
+			for (TFieldIterator<FProperty> It(this); It; ++It)
+			{
+				FProperty* Property = *It;
+				if (Property && !Property->HasAnyPropertyFlags(CPF_NoDestructor))
+				{
+					bHasNoDtor = false;
+					break;
+				}
+			}
+		}
+	}
+	
+	StructFlags = EStructFlags(StructFlags | STRUCT_ZeroConstructor | STRUCT_IsPlainOldData | STRUCT_NoDestructor);
+	if(!bIsZeroConstruct)
+	{
+		StructFlags = EStructFlags(StructFlags & ~STRUCT_ZeroConstructor);
+	}
+	if(!bIsPOD)
+	{
+		StructFlags = EStructFlags(StructFlags & ~STRUCT_IsPlainOldData);
+	}
+	if(!bHasNoDtor)
+	{
+		StructFlags = EStructFlags(StructFlags & ~STRUCT_NoDestructor);
+	}
+
+}
+
+#if WITH_EDITORONLY_DATA
+TSubclassOf<UStructCookedMetaData> UUserDefinedStruct::GetCookedMetaDataClass() const
+{
+	return UStructCookedMetaData::StaticClass();
+}
+
+UStructCookedMetaData* UUserDefinedStruct::NewCookedMetaData()
+{
+	if (!CachedCookedMetaDataPtr)
+	{
+		CachedCookedMetaDataPtr = CookedMetaDataUtil::NewCookedMetaData<UStructCookedMetaData>(this, "CookedStructMetaData", GetCookedMetaDataClass());
+	}
+	return CachedCookedMetaDataPtr;
+}
+
+const UStructCookedMetaData* UUserDefinedStruct::FindCookedMetaData()
+{
+	if (!CachedCookedMetaDataPtr)
+	{
+		CachedCookedMetaDataPtr = CookedMetaDataUtil::FindCookedMetaData<UStructCookedMetaData>(this, TEXT("CookedStructMetaData"));
+	}
+	return CachedCookedMetaDataPtr;
+}
+
+void UUserDefinedStruct::PurgeCookedMetaData()
+{
+	if (CachedCookedMetaDataPtr)
+	{
+		CookedMetaDataUtil::PurgeCookedMetaData<UStructCookedMetaData>(CachedCookedMetaDataPtr);
+	}
+}
+#endif // WITH_EDITORONLY_DATA
diff --git a/Engine/Source/Runtime/CoreUObject/Private/StructUtils/UserDefinedStructEditorUtils.cpp b/Engine/Source/Runtime/CoreUObject/Private/StructUtils/UserDefinedStructEditorUtils.cpp
new file mode 100644
--- /dev/null
+++ b/Engine/Source/Runtime/CoreUObject/Private/StructUtils/UserDefinedStructEditorUtils.cpp
@@ -0,0 +1,155 @@
+// Copyright Epic Games, Inc. All Rights Reserved.
+
+#if WITH_EDITOR
+
+#include "StructUtils/UserDefinedStructEditorUtils.h"
+
+#define LOCTEXT_NAMESPACE "Structure"
+
+//////////////////////////////////////////////////////////////////////////
+// FUserDefinedStructEditorUtils
+
+FUserDefinedStructEditorUtils::FOnUserDefinedStructChanged FUserDefinedStructEditorUtils::OnUserDefinedStructChanged;
+
+namespace
+{
+	bool IsObjPropertyValid(const FProperty* Property)
+	{
+		if (const FInterfaceProperty* InterfaceProperty = CastField<const FInterfaceProperty>(Property))
+		{
+			return InterfaceProperty->InterfaceClass != nullptr;
+		}
+		
+		if (const FArrayProperty* ArrayProperty = CastField<const FArrayProperty>(Property))
+		{
+			return ArrayProperty->Inner && IsObjPropertyValid(ArrayProperty->Inner);
+		}
+		
+		if (const FObjectProperty* ObjectProperty = CastField<const FObjectProperty>(Property))
+		{
+			return ObjectProperty->PropertyClass != nullptr;
+		}
+		return true;
+	}
+}
+
+
+void FUserDefinedStructEditorUtils::OnStructureChanged(UUserDefinedStruct* Struct)
+{
+	if (Struct)
+	{
+		OnUserDefinedStructChanged.ExecuteIfBound(Struct);
+	}
+}
+
+FUserDefinedStructEditorUtils::EStructureError FUserDefinedStructEditorUtils::IsStructureValid(const UScriptStruct* Struct, const UStruct* RecursionParent, FString* OutMsg)
+{
+	check(Struct);
+	if (Struct == RecursionParent)
+	{
+		if (OutMsg)
+		{
+			*OutMsg = FText::Format(LOCTEXT("StructureRecursionFmt", "Recursion: Recursion: Struct cannot have itself or a nested struct member referencing itself as a member variable. Struct '{0}', recursive parent '{1}'"), 
+				 FText::FromString(Struct->GetFullName()), FText::FromString(RecursionParent->GetFullName())).ToString();
+		}
+		return EStructureError::Recursion;
+	}
+
+	const UScriptStruct* FallbackStruct = GetFallbackStruct();
+	if (Struct == FallbackStruct)
+	{
+		if (OutMsg)
+		{
+			*OutMsg = LOCTEXT("StructureUnknown", "Struct unknown (deleted?)").ToString();
+		}
+		return EStructureError::FallbackStruct;
+	}
+
+	if (Struct->GetStructureSize() <= 0)
+	{
+		if (OutMsg)
+		{
+			*OutMsg = FText::Format(LOCTEXT("StructureSizeIsZeroFmt", "Struct '{0}' is empty"), FText::FromString(Struct->GetFullName())).ToString();
+		}
+		return EStructureError::EmptyStructure;
+	}
+
+	if (const UUserDefinedStruct* UDStruct = Cast<const UUserDefinedStruct>(Struct))
+	{
+		if (UDStruct->Status != EUserDefinedStructureStatus::UDSS_UpToDate)
+		{
+			if (OutMsg)
+			{
+				*OutMsg = FText::Format(LOCTEXT("StructureNotCompiledFmt", "Struct '{0}' is not compiled"), FText::FromString(Struct->GetFullName())).ToString();
+			}
+			return EStructureError::NotCompiled;
+		}
+
+		for (const FProperty* P = Struct->PropertyLink; P; P = P->PropertyLinkNext)
+		{
+			const FStructProperty* StructProp = CastField<const FStructProperty>(P);
+			if (nullptr == StructProp)
+			{
+				if (const FArrayProperty* ArrayProp = CastField<const FArrayProperty>(P))
+				{
+					StructProp = CastField<const FStructProperty>(ArrayProp->Inner);
+				}
+			}
+
+			if (StructProp)
+			{
+				if ((NULL == StructProp->Struct) || (FallbackStruct == StructProp->Struct))
+				{
+					if (OutMsg)
+					{
+						*OutMsg = FText::Format(
+							LOCTEXT("StructureUnknownPropertyFmt", "Struct unknown (deleted?). Parent '{0}' Property: '{1}'"),
+							FText::FromString(Struct->GetFullName()),
+							FText::FromString(StructProp->GetName())
+						).ToString();
+					}
+					return EStructureError::FallbackStruct;
+				}
+
+				FString OutMsgInner;
+				const EStructureError Result = IsStructureValid(
+					StructProp->Struct,
+					RecursionParent ? RecursionParent : Struct,
+					OutMsg ? &OutMsgInner : nullptr);
+				if (EStructureError::Ok != Result)
+				{
+					if (OutMsg)
+					{
+						*OutMsg = FText::Format(
+							LOCTEXT("StructurePropertyErrorTemplateFmt", "Struct '{0}' Property '{1}' Error ( {2} )"),
+							FText::FromString(Struct->GetFullName()),
+							FText::FromString(StructProp->GetName()),
+							FText::FromString(OutMsgInner)
+						).ToString();
+					}
+					return Result;
+				}
+			}
+
+			// The structure is loaded (from .uasset) without recompilation. All properties should be verified.
+			if (!IsObjPropertyValid(P))
+			{
+				if (OutMsg)
+				{
+					*OutMsg = FText::Format(
+						LOCTEXT("StructureUnknownObjectPropertyFmt", "Invalid object property. Structure '{0}' Property: '{1}'"),
+						FText::FromString(Struct->GetFullName()),
+						FText::FromString(P->GetName())
+					).ToString();
+				}
+				return EStructureError::NotCompiled;
+			}
+		}
+	}
+
+	return EStructureError::Ok;
+}
+
+#undef LOCTEXT_NAMESPACE
+
+#endif // WITH_EDITOR
\ No newline at end of file
diff --git a/Engine/Source/Runtime/CoreUObject/Public/StructUtils/UserDefinedStruct.h b/Engine/Source/Runtime/CoreUObject/Public/StructUtils/UserDefinedStruct.h
new file mode 100644
--- /dev/null
+++ b/Engine/Source/Runtime/CoreUObject/Public/StructUtils/UserDefinedStruct.h
@@ -0,0 +1,156 @@
+// Copyright Epic Games, Inc. All Rights Reserved.
+
+#pragma once
+
+#if UE_ENABLE_INCLUDE_ORDER_DEPRECATED_IN_5_5
+#include "CoreMinimal.h"
+#include "UObject/ObjectMacros.h"
+#endif // UE_ENABLE_INCLUDE_ORDER_DEPRECATED_IN_5_5
+#include "Misc/Guid.h"
+#include "UObject/Class.h"
+#include "UObject/WeakObjectPtr.h"
+#include "UObject/StructOnScope.h"
+#include "Templates/SubclassOf.h"
+#include "UserDefinedStruct.generated.h"
+
+class UUserDefinedStruct;
+class UUserDefinedStructEditorData;
+class UStructCookedMetaData;
+
+DECLARE_MULTICAST_DELEGATE_OneParam(FOnStructChanged, UUserDefinedStruct*);
+
+UENUM()
+enum EUserDefinedStructureStatus : int
+{
+	/** Struct is in an unknown state. */
+	UDSS_UpToDate,
+	/** Struct has been modified but not recompiled. */
+	UDSS_Dirty,
+	/** Struct tried but failed to be compiled. */
+	UDSS_Error,
+	/** Struct is a duplicate, the original one was changed. */
+	UDSS_Duplicate,
+
+	UDSS_MAX,
+};
+
+/** Wrapper for StructOnScope that tells it to ignore default values */
+class FUserStructOnScopeIgnoreDefaults : public FStructOnScope
+{
+public:
+	/** Constructor with no script struct, call Recreate later */
+	FUserStructOnScopeIgnoreDefaults() : FStructOnScope() {}
+
+	/** Constructor that initializes for you */
+	COREUOBJECT_API FUserStructOnScopeIgnoreDefaults(const UUserDefinedStruct* InUserStruct);
+
+	/** Initialize from existing data, will free when scope closes */
+	COREUOBJECT_API FUserStructOnScopeIgnoreDefaults(const UUserDefinedStruct* InUserStruct, uint8* InData);
+
+	/** Destroys and creates new struct */
+	COREUOBJECT_API void Recreate(const UUserDefinedStruct* InUserStruct);
+
+	COREUOBJECT_API virtual void Initialize() override;
+};
+
+UCLASS(MinimalAPI)
+class UUserDefinedStruct : public UScriptStruct
+{
+	GENERATED_UCLASS_BODY()
+
+public:
+#if WITH_EDITORONLY_DATA
+	/** The original struct, when current struct isn't a temporary duplicate, the field should be null */
+	UPROPERTY(Transient)
+	TWeakObjectPtr<UUserDefinedStruct> PrimaryStruct;
+
+	UPROPERTY()
+	FString ErrorMessage;
+
+	UPROPERTY()
+	TObjectPtr<UObject> EditorData;
+#endif // WITH_EDITORONLY_DATA
+
+	/** Status of this struct, outside of the editor it is assumed to always be UpToDate */
+	UPROPERTY()
+	TEnumAsByte<enum EUserDefinedStructureStatus> Status;
+
+	/** Uniquely identifies this specific user struct */
+	UPROPERTY(AssetRegistrySearchable)
+	FGuid Guid;
+
+protected:
+	/** Default instance of this struct with default values filled in, used to initialize structure */
+	FUserStructOnScopeIgnoreDefaults DefaultStructInstance;
+
+	/** Bool to indicate we want to initialize a version of this struct without defaults, this is set while allocating the DefaultStructInstance itself */
+	bool bIgnoreStructDefaults;
+
+public:
+#if WITH_EDITOR
+	// UObject interface.
+	COREUOBJECT_API virtual void PostDuplicate(bool bDuplicateForPIE) override;
+	COREUOBJECT_API virtual void GetAssetRegistryTags(FAssetRegistryTagsContext Context) const override;
+	UE_DEPRECATED(5.4, "Implement the version that takes FAssetRegistryTagsContext instead.")
+	COREUOBJECT_API virtual void GetAssetRegistryTags(TArray<FAssetRegistryTag>& OutTags) const override;
+	COREUOBJECT_API virtual void PostLoad() override;
+	COREUOBJECT_API virtual void PreSaveRoot(FObjectPreSaveRootContext ObjectSaveContext) override;
+	COREUOBJECT_API virtual void PostSaveRoot(FObjectPostSaveRootContext ObjectSaveContext) override;
+	// End of UObject interface.
+
+	/** Creates a new guid if needed */
+	COREUOBJECT_API void ValidateGuid();
+
+	COREUOBJECT_API virtual void OnChanged();
+
+	friend UUserDefinedStructEditorData;
+#endif	// WITH_EDITOR
+
+	// UObject interface.
+	COREUOBJECT_API virtual void Serialize(FStructuredArchive::FRecord Record) override;
+	COREUOBJECT_API virtual void SerializeTaggedProperties(FStructuredArchive::FSlot Slot, uint8* Data, UStruct* DefaultsStruct, uint8* Defaults, const UObject* BreakRecursionIfFullyLoad = nullptr) const override;
+	COREUOBJECT_API virtual FString GetAuthoredNameForField(const FField* Field) const override;
+	// End of UObject interface.
+
+	// UScriptStruct interface.
+	COREUOBJECT_API virtual void InitializeStruct(void* Dest, int32 ArrayDim = 1) const override;
+	COREUOBJECT_API virtual uint32 GetStructTypeHash(const void* Src) const override;
+	COREUOBJECT_API virtual void RecursivelyPreload() override;
+	COREUOBJECT_API virtual FGuid GetCustomGuid() const override;
+	COREUOBJECT_API virtual FString GetStructCPPName(uint32 CPPExportFlags) const override;
+	COREUOBJECT_API virtual FProperty* CustomFindProperty(const FName Name) const override;
+	COREUOBJECT_API virtual void PrepareCppStructOps() override;
+	// End of  UScriptStruct interface.
+
+	/** Returns the raw memory of the default instance */
+	COREUOBJECT_API const uint8* GetDefaultInstance() const;
+
+	/** Specifically initialize this struct without using the default instance data */
+	COREUOBJECT_API void InitializeStructIgnoreDefaults(void* Dest, int32 ArrayDim = 1) const;
+
+	/** Computes hash */
+	static COREUOBJECT_API uint32 GetUserDefinedStructTypeHash(const void* Src, const UScriptStruct* Type);
+
+	/** returns references from default instance */
+	static COREUOBJECT_API void AddReferencedObjects(UObject* InThis, FReferenceCollector& Collector);
+
+	/** Inspects properties and default values, setting appropriate StructFlags */
+	COREUOBJECT_API void UpdateStructFlags();
+
+#if WITH_EDITORONLY_DATA
+	FORCEINLINE FOnStructChanged& OnStructChanged() { return ChangedEvent; }
+public:
+	FOnStructChanged ChangedEvent;
+
+protected:
+	COREUOBJECT_API virtual TSubclassOf<UStructCookedMetaData> GetCookedMetaDataClass() const;
+
+private:
+	COREUOBJECT_API UStructCookedMetaData* NewCookedMetaData();
+	COREUOBJECT_API const UStructCookedMetaData* FindCookedMetaData();
+	COREUOBJECT_API void PurgeCookedMetaData();
+
+	UPROPERTY()
+	TObjectPtr<UStructCookedMetaData> CachedCookedMetaDataPtr;
+#endif // WITH_EDITORONLY_DATA
+};
diff --git a/Engine/Source/Runtime/CoreUObject/Public/StructUtils/UserDefinedStructEditorUtils.h b/Engine/Source/Runtime/CoreUObject/Public/StructUtils/UserDefinedStructEditorUtils.h
new file mode 100644
--- /dev/null
+++ b/Engine/Source/Runtime/CoreUObject/Public/StructUtils/UserDefinedStructEditorUtils.h
@@ -0,0 +1,66 @@
+// Copyright Epic Games, Inc. All Rights Reserved.
+#pragma once
+
+#if WITH_EDITOR
+#include "StructUtils/UserDefinedStruct.h"
+#endif // WITH_EDITOR
+#include "UserDefinedStructEditorUtils.generated.h"
+
+UCLASS(MinimalAPI, Abstract)
+class UUserDefinedStructEditorDataBase : public UObject
+{
+	GENERATED_BODY()
+
+#if WITH_EDITOR
+public:
+	COREUOBJECT_API virtual void RecreateDefaultInstance(FString* OutLog = nullptr)
+	{		
+	}
+	
+	COREUOBJECT_API virtual void ReinitializeDefaultInstance(FString* OutLog = nullptr)
+	{		
+	}
+	
+	COREUOBJECT_API virtual FProperty* FindProperty(const UUserDefinedStruct* Struct, FName Name) const
+	{
+		return nullptr;
+	}
+
+	COREUOBJECT_API virtual FString GetFriendlyNameForProperty(const UUserDefinedStruct* Struct, const FProperty* Property) const
+	{
+		return {};
+	}
+
+	COREUOBJECT_API virtual FString GetTooltip() const
+	{
+		return {};
+	}
+#endif // WITH_EDITOR
+};
+
+#if WITH_EDITOR
+struct FUserDefinedStructEditorUtils
+{
+	// NOTIFICATION
+	DECLARE_DELEGATE_OneParam(FOnUserDefinedStructChanged, UUserDefinedStruct*);
+	static COREUOBJECT_API FOnUserDefinedStructChanged OnUserDefinedStructChanged;
+
+	/** called after UDS was changed by editor*/
+	static COREUOBJECT_API void OnStructureChanged(UUserDefinedStruct* Struct);
+	
+	// VALIDATION
+	enum EStructureError
+	{
+		Ok, 
+		Recursion,
+		FallbackStruct,
+		NotCompiled,
+		NotBlueprintType,
+		NotSupportedType,
+		EmptyStructure
+	};
+
+	/** Can the structure be a member variable for a BPGClass or BPGStruct */
+	static COREUOBJECT_API EStructureError IsStructureValid(const UScriptStruct* Struct, const UStruct* RecursionParent = nullptr, FString* OutMsg = nullptr);
+};
+#endif // WITH_EDITOR
diff --git a/Engine/Source/Runtime/Engine/Classes/Animation/AnimationSettings.h b/Engine/Source/Runtime/Engine/Classes/Animation/AnimationSettings.h
--- a/Engine/Source/Runtime/Engine/Classes/Animation/AnimationSettings.h
+++ b/Engine/Source/Runtime/Engine/Classes/Animation/AnimationSettings.h
@@ -6,16 +6,24 @@
 
 #pragma once
 
+#if UE_ENABLE_INCLUDE_ORDER_DEPRECATED_IN_5_5
 #include "CoreMinimal.h"
 #include "UObject/ObjectMacros.h"
 #include "Templates/SubclassOf.h"
 #include "Animation/AnimSequence.h"
+#endif
 #include "Engine/DeveloperSettings.h"
 #include "CustomAttributes.h"
+#if UE_ENABLE_INCLUDE_ORDER_DEPRECATED_IN_5_5
 #include "MirrorDataTable.h"
-#include "Engine/UserDefinedStruct.h"
+#include "StructUtils/UserDefinedStruct.h"
+#endif
 #include "AnimationSettings.generated.h"
 
+class UUserDefinedStruct;
+struct FCustomAttributeSetting;
+struct FMirrorFindReplaceExpression;
+
 /**
  * Default animation settings.
  */
@@ -40,7 +48,7 @@ class UAnimationSettings : public UDeveloperSettings
 	UPROPERTY(config, EditAnywhere, Category = Compression)
 	bool bForceBelowThreshold;
 
-	/** If true, then the animation will be first recompressed with it's current compressor if non-NULL, or with the global default compressor (specified in the engine ini) 
+	/** If true, then the animation will be first recompressed with its current compressor if non-NULL, or with the global default compressor (specified in the engine ini) 
 	* Also known as "Run Current Default Compressor"
 	*/
 	UPROPERTY(config, EditAnywhere, Category = Compression)
diff --git a/Engine/Source/Runtime/Engine/Classes/Engine/UserDefinedStruct.h b/Engine/Source/Runtime/Engine/Classes/Engine/UserDefinedStruct.h
deleted file mode 100644
diff --git a/Engine/Source/Runtime/Engine/Private/Animation/AnimationSettings.cpp b/Engine/Source/Runtime/Engine/Private/Animation/AnimationSettings.cpp
--- a/Engine/Source/Runtime/Engine/Private/Animation/AnimationSettings.cpp
+++ b/Engine/Source/Runtime/Engine/Private/Animation/AnimationSettings.cpp
@@ -4,7 +4,7 @@
 #include "Animation/AttributeTypes.h"
 #include "Animation/CustomAttributes.h"
 #include "Animation/MirrorDataTable.h"
-#include "Engine/UserDefinedStruct.h"
+#include "StructUtils/UserDefinedStruct.h"
 
 #include UE_INLINE_GENERATED_CPP_BY_NAME(AnimationSettings)
 
diff --git a/Engine/Source/Runtime/Engine/Private/Animation/AttributeTypes.cpp b/Engine/Source/Runtime/Engine/Private/Animation/AttributeTypes.cpp
--- a/Engine/Source/Runtime/Engine/Private/Animation/AttributeTypes.cpp
+++ b/Engine/Source/Runtime/Engine/Private/Animation/AttributeTypes.cpp
@@ -6,7 +6,7 @@
 #include "AnimationRuntime.h"
 #include "Misc/DelayedAutoRegister.h"
 #include "Animation/AnimationSettings.h"
-#include "Engine/UserDefinedStruct.h"
+#include "StructUtils/UserDefinedStruct.h"
 
 namespace UE
 {
diff --git a/Engine/Source/Runtime/Engine/Private/DataTable.cpp b/Engine/Source/Runtime/Engine/Private/DataTable.cpp
--- a/Engine/Source/Runtime/Engine/Private/DataTable.cpp
+++ b/Engine/Source/Runtime/Engine/Private/DataTable.cpp
@@ -12,7 +12,7 @@
 #include "DataTableCSV.h"
 #include "DataTableJSON.h"
 #include "EditorFramework/AssetImportData.h"
-#include "Engine/UserDefinedStruct.h"
+#include "StructUtils/UserDefinedStruct.h"
 
 #include UE_INLINE_GENERATED_CPP_BY_NAME(DataTable)
 
diff --git a/Engine/Source/Runtime/Engine/Private/Tests/Loading/AsyncLoadingTests.cpp b/Engine/Source/Runtime/Engine/Private/Tests/Loading/AsyncLoadingTests.cpp
--- a/Engine/Source/Runtime/Engine/Private/Tests/Loading/AsyncLoadingTests.cpp
+++ b/Engine/Source/Runtime/Engine/Private/Tests/Loading/AsyncLoadingTests.cpp
@@ -7,7 +7,7 @@
 #include "AssetRegistry/IAssetRegistry.h"
 #include "Async/ParallelFor.h"
 #include "AsyncLoadingTests_Shared.h"
-#include "Engine/UserDefinedStruct.h"
+#include "StructUtils/UserDefinedStruct.h"
 #include "UObject/UObjectHash.h"
 
 #if WITH_DEV_AUTOMATION_TESTS
diff --git a/Engine/Source/Runtime/Engine/Private/UserDefinedStruct.cpp b/Engine/Source/Runtime/Engine/Private/UserDefinedStruct.cpp
deleted file mode 100644
diff --git a/Engine/Source/Runtime/Engine/Public/Engine.h b/Engine/Source/Runtime/Engine/Public/Engine.h
--- a/Engine/Source/Runtime/Engine/Public/Engine.h
+++ b/Engine/Source/Runtime/Engine/Public/Engine.h
@@ -317,7 +317,7 @@ MONOLITHIC_HEADER_BOILERPLATE()
 #include "Engine/TimelineTemplate.h"
 #include "GameFramework/TouchInterface.h"
 #include "Engine/UserDefinedEnum.h"
-#include "Engine/UserDefinedStruct.h"
+#include "StructUtils/UserDefinedStruct.h"
 #include "Animation/MorphTarget.h"
 #include "Animation/AnimInstance.h"
 
diff --git a/Engine/Source/Runtime/MovieSceneTracks/Private/Sections/MovieSceneEventSection.cpp b/Engine/Source/Runtime/MovieSceneTracks/Private/Sections/MovieSceneEventSection.cpp
--- a/Engine/Source/Runtime/MovieSceneTracks/Private/Sections/MovieSceneEventSection.cpp
+++ b/Engine/Source/Runtime/MovieSceneTracks/Private/Sections/MovieSceneEventSection.cpp
@@ -8,7 +8,7 @@
 #include "MovieSceneFwd.h"
 #include "Channels/MovieSceneChannelProxy.h"
 #include "Serialization/MemoryArchive.h"
-#include "Engine/UserDefinedStruct.h"
+#include "StructUtils/UserDefinedStruct.h"
 #include "MovieSceneFrameMigration.h"
 
 #include UE_INLINE_GENERATED_CPP_BY_NAME(MovieSceneEventSection)
