Make SerializeForLog overloads consistent in their declaration as friends

#rnx

diff --git a/Engine/Source/Runtime/CoreUObject/Public/UObject/UObjectGlobals.h b/Engine/Source/Runtime/CoreUObject/Public/UObject/UObjectGlobals.h
--- a/Engine/Source/Runtime/CoreUObject/Public/UObject/UObjectGlobals.h
+++ b/Engine/Source/Runtime/CoreUObject/Public/UObject/UObjectGlobals.h
@@ -3301,8 +3301,6 @@ COREUOBJECT_API bool GetAllowNativeComponentClassOverrides();
 
 namespace UE
 {
-class FAssetLog;
-COREUOBJECT_API void SerializeForLog(FCbWriter& Writer, const FAssetLog& AssetLog);
 
 class FAssetLog
 {
