Adding support for Type_Asset redirections to FCoreRedirects::FindPreviousNames.

diff --git a/Engine/Source/Runtime/CoreUObject/Private/UObject/CoreRedirects.cpp b/Engine/Source/Runtime/CoreUObject/Private/UObject/CoreRedirects.cpp
--- a/Engine/Source/Runtime/CoreUObject/Private/UObject/CoreRedirects.cpp
+++ b/Engine/Source/Runtime/CoreUObject/Private/UObject/CoreRedirects.cpp
@@ -323,6 +323,32 @@ namespace UE::CoreRedirects::Private
 				UE_LOG(LogCoreRedirects, Error, TEXT("Failed %s. Source = %s was unexpectedly redirected to %s"),
 					Test.TestDescription, Test.Origin, *NewName.ToString());
 			}
+			if (Test.ExpectTrue)
+			{
+				TArray<FCoreRedirectObjectName> OldNames;
+				if (!FCoreRedirects::FindPreviousNames(Test.RedirectFlags, NewName, OldNames))
+				{
+					bSuccess = false;
+					UE_LOG(LogCoreRedirects, Error, TEXT("Failed to FindPreviousNames for %s"), Test.Destination);
+				}
+				else
+				{
+					bool bContainsName = false;
+					for (const FCoreRedirectObjectName& ReverseOldName : OldNames)
+					{
+						if (ReverseOldName.ToString().Compare(Test.Origin))
+						{
+							bContainsName = true;
+							break;
+						}
+					}
+					if (!bContainsName)
+					{
+						bSuccess = false;
+						UE_LOG(LogCoreRedirects, Error, TEXT("Failed to find expected previous name for %s"), Test.Destination);
+					}
+				}
+			}
 		}
 
 		// Remove all redirects temporarily and verify no test finds a redirection
@@ -1625,6 +1651,10 @@ bool FCoreRedirects::FindPreviousNames(ECoreRedirectFlags SearchFlags, const FCo
 	// If we're not explicitly searching for packages or looking for removed things, add the implicit (Type=Package,Category=None) redirects
 	const bool bSearchPackageRedirects = !(SearchFlags & ECoreRedirectFlags::Type_Package) && !(SearchFlags & ECoreRedirectFlags::Category_Removed);
 
+	// We always search Type_Asset as well as whatever is requested
+	// That is because asset redirectors can redirect packages (implicitly) and any UObject type (explicitly)
+	SearchFlags |= ECoreRedirectFlags::Type_Asset;
+
 	auto TryReverseRedirect = [](const FCoreRedirect& Redirect, const FCoreRedirectObjectName& NewObjectName, TArray<FCoreRedirectObjectName>& PreviousNames)
 		{
 			FCoreRedirect ReverseRedirect = FCoreRedirect(Redirect);
@@ -1684,11 +1714,14 @@ bool FCoreRedirects::FindPreviousNames(ECoreRedirectFlags SearchFlags, const FCo
 			}
 			else
 			{
+				FCoreRedirectObjectName::EMatchFlags MatchFlags = 
+					EnumHasAnyFlags(PairFlags, ECoreRedirectFlags::Type_Asset) 
+						? FCoreRedirectObjectName::EMatchFlags::AllowPartialRHSMatch : FCoreRedirectObjectName::EMatchFlags::None;
 				for (const TPair<FName, TArray<FCoreRedirect>>& RedirectPair : Pair.Value.RedirectMap)
 				{
 					for (const FCoreRedirect& Redirect : RedirectPair.Value)
 					{
-						if (Redirect.NewName.Matches(NewObjectName, FCoreRedirectObjectName::EMatchFlags::None))
+						if (Redirect.NewName.Matches(NewObjectName, MatchFlags))
 						{
 							bFound |= TryReverseRedirect(Redirect, NewObjectName, PreviousNames);
 						}
