Added TBaseStructure specialization for FFrameRate

diff --git a/Engine/Source/Runtime/CoreUObject/Private/UObject/Class.cpp b/Engine/Source/Runtime/CoreUObject/Private/UObject/Class.cpp
--- a/Engine/Source/Runtime/CoreUObject/Private/UObject/Class.cpp
+++ b/Engine/Source/Runtime/CoreUObject/Private/UObject/Class.cpp
@@ -7605,6 +7605,12 @@ UScriptStruct* TBaseStructure<FFrameTime>::Get()
 	return ScriptStruct;
 }
 
+UScriptStruct* TBaseStructure<FFrameRate>::Get()
+{
+	static UScriptStruct* ScriptStruct = StaticGetBaseStructureInternal(TEXT("FrameRate"));
+	return ScriptStruct;
+}
+
 UScriptStruct* TBaseStructure<FAssetBundleData>::Get()
 {
 	static UScriptStruct* ScriptStruct = StaticGetBaseStructureInternal(TEXT("AssetBundleData"));
diff --git a/Engine/Source/Runtime/CoreUObject/Public/UObject/Class.h b/Engine/Source/Runtime/CoreUObject/Public/UObject/Class.h
--- a/Engine/Source/Runtime/CoreUObject/Public/UObject/Class.h
+++ b/Engine/Source/Runtime/CoreUObject/Public/UObject/Class.h
@@ -4338,6 +4338,13 @@ template<> struct TBaseStructure<FFrameTime>
 	static COREUOBJECT_API UScriptStruct* Get();
 };
 
+struct FFrameRate;
+
+template<> struct TBaseStructure<FFrameRate>
+{
+	static COREUOBJECT_API UScriptStruct* Get();
+};
+
 struct FSoftObjectPath;
 
 template<> struct TBaseStructure<FSoftObjectPath>
