Fix OverridenPropertyNodeID Path collision when re-instantiate is triggered

diff --git a/Engine/Source/Runtime/CoreUObject/Public/UObject/OverriddenPropertySet.h b/Engine/Source/Runtime/CoreUObject/Public/UObject/OverriddenPropertySet.h
--- a/Engine/Source/Runtime/CoreUObject/Public/UObject/OverriddenPropertySet.h
+++ b/Engine/Source/Runtime/CoreUObject/Public/UObject/OverriddenPropertySet.h
@@ -115,9 +115,12 @@ struct FOverriddenPropertyNodeID
 	{}
 
 	FOverriddenPropertyNodeID(const UObject& InObject)
-		: Path(FString::Printf(TEXT("%d"), GUObjectArray.ObjectToIndex(&InObject)))
-		, Object(&InObject)
+		: Object(&InObject)
 	{
+		// Note: Using ObjectIndex by itself is not sufficient for an enduring unique identifier
+		// as re-instantiation can cause a reuse of the index for another object. Appending the serial solves this issue
+		const int32 ObjectIndex = GUObjectArray.ObjectToIndex(&InObject);
+		Path = *(FString::Printf(TEXT("%d%d"), ObjectIndex, GUObjectArray.AllocateSerialNumber(ObjectIndex)));
 	}
 
 	bool operator==(const FOverriddenPropertyNodeID& Other) const
