Exposed the deterministic package namespace generation for use outside of USE_STABLE_LOCALIZATION_KEYS

#rnx

diff --git a/Engine/Source/Runtime/CoreUObject/Private/Internationalization/TextPackageNamespaceUtil.cpp b/Engine/Source/Runtime/CoreUObject/Private/Internationalization/TextPackageNamespaceUtil.cpp
--- a/Engine/Source/Runtime/CoreUObject/Private/Internationalization/TextPackageNamespaceUtil.cpp
+++ b/Engine/Source/Runtime/CoreUObject/Private/Internationalization/TextPackageNamespaceUtil.cpp
@@ -45,20 +45,7 @@ FString FindOrAddPackageNamespace(UPackage* InPackage, const bool bCanAdd)
 				FString& PackageLocalizationNamespaceValue = PackageMetaData->RootMetaDataMap.FindOrAdd(PackageLocalizationNamespaceKey);
 				if (PackageLocalizationNamespaceValue.IsEmpty())
 				{
-					// Make a determinstic new guid that will vary based on the package
-					FBlake3 Builder;
-					FString PackagePath = InPackage->GetName();
-					FGuid NonUniqueGuid = InPackage->GetPersistentGuid(); // Can be the same for duplicated packages
-					Builder.Update(&NonUniqueGuid, sizeof(FGuid));
-					Builder.Update(*PackagePath, PackagePath.Len() * sizeof(**PackagePath));
-					FBlake3Hash Hash = Builder.Finalize();
-					// We use the first 16 bytes of the FIoHash to create the guid, there is
-					// no specific reason why these were chosen, we could take any pattern or combination
-					// of bytes.
-					uint32* HashBytes = (uint32*)Hash.GetBytes();
-					FGuid NewGuid(HashBytes[0], HashBytes[1], HashBytes[2], HashBytes[3]);
-
-					PackageLocalizationNamespaceValue = NewGuid.ToString();
+					PackageLocalizationNamespaceValue = TextNamespaceUtil::GenerateDeterministicPackageNamespace(InPackage);
 				}
 				return PackageLocalizationNamespaceValue;
 			}
@@ -172,6 +159,26 @@ FText TextNamespaceUtil::CopyTextToPackage(const FText& InText, UObject* InObjec
 #endif	// USE_STABLE_LOCALIZATION_KEYS
 }
 
+#if WITH_EDITORONLY_DATA
+FString TextNamespaceUtil::GenerateDeterministicPackageNamespace(const UPackage* InPackage)
+{
+	// Make a deterministic new guid that will vary based on the package
+	FBlake3 Builder;
+	FString PackagePath = InPackage->GetName();
+	FGuid NonUniqueGuid = InPackage->GetPersistentGuid(); // Can be the same for duplicated packages
+	Builder.Update(&NonUniqueGuid, sizeof(FGuid));
+	Builder.Update(*PackagePath, PackagePath.Len() * sizeof(**PackagePath));
+	FBlake3Hash Hash = Builder.Finalize();
+	// We use the first 16 bytes of the FIoHash to create the guid, there is
+	// no specific reason why these were chosen, we could take any pattern or combination
+	// of bytes.
+	uint32* HashBytes = (uint32*)Hash.GetBytes();
+	FGuid NewGuid(HashBytes[0], HashBytes[1], HashBytes[2], HashBytes[3]);
+
+	return NewGuid.ToString();
+}
+#endif // WITH_EDITORONLY_DATA
+
 FString TextNamespaceUtil::GenerateRandomTextKey()
 {
 	return FGuid::NewGuid().ToString();
diff --git a/Engine/Source/Runtime/CoreUObject/Public/Internationalization/TextPackageNamespaceUtil.h b/Engine/Source/Runtime/CoreUObject/Public/Internationalization/TextPackageNamespaceUtil.h
--- a/Engine/Source/Runtime/CoreUObject/Public/Internationalization/TextPackageNamespaceUtil.h
+++ b/Engine/Source/Runtime/CoreUObject/Public/Internationalization/TextPackageNamespaceUtil.h
@@ -100,6 +100,16 @@ COREUOBJECT_API void ForcePackageNamespace(UObject* InObject, const FString& InN
 COREUOBJECT_API FText CopyTextToPackage(const FText& InText, UPackage* InPackage, const ETextCopyMethod InCopyMethod = ETextCopyMethod::NewKey, const bool bAlwaysApplyPackageNamespace = false);
 COREUOBJECT_API FText CopyTextToPackage(const FText& InText, UObject* InObject, const ETextCopyMethod InCopyMethod = ETextCopyMethod::NewKey, const bool bAlwaysApplyPackageNamespace = false);
 
+#if WITH_EDITORONLY_DATA
+/**
+ * Generate a deterministic package namespace based on the given package info.
+ * @note This key will be formatted like a GUID, but the value will actually be based on deterministic hashes.
+ *
+ * @param InPackage						The package to generate the namespace for.
+ */
+COREUOBJECT_API FString GenerateDeterministicPackageNamespace(const UPackage* InPackage);
+#endif // WITH_EDITORONLY_DATA
+
 /**
  * Generate a random text key.
  * @note This key will be a GUID.
