Adding an API to FRedirectCollector to provide read access to the global list of discovered asset redirects.

diff --git a/Engine/Source/Runtime/CoreUObject/Private/Misc/RedirectCollector.cpp b/Engine/Source/Runtime/CoreUObject/Private/Misc/RedirectCollector.cpp
--- a/Engine/Source/Runtime/CoreUObject/Private/Misc/RedirectCollector.cpp
+++ b/Engine/Source/Runtime/CoreUObject/Private/Misc/RedirectCollector.cpp
@@ -11,10 +11,12 @@
 #include "UObject/UnrealType.h"
 #include "UObject/ObjectRedirector.h"
 #include "Misc/PackageName.h"
+#include "Misc/Paths.h"
 #include "UObject/LinkerLoad.h"
 #include "UObject/UObjectThreadContext.h"
 #include "UObject/SoftObjectPath.h"
 #include "HAL/IConsoleManager.h"
+#include "HAL/FileManager.h"
 
 #if WITH_EDITOR
 
@@ -29,6 +31,35 @@ FAutoConsoleCommand CVarResolveAllSoftObjects(
 	})
 );
 
+void RedirectCollectorDumpAllAssetRedirects();
+
+FAutoConsoleCommand CVarRedirectCollectorDumpAllAssetRedirects(
+	TEXT("redirectcollector.DumpAllAssetRedirects"),
+	TEXT("Prints all tracked redirectors to the log."),
+	FConsoleCommandDelegate::CreateStatic(&RedirectCollectorDumpAllAssetRedirects)
+);
+
+void RedirectCollectorDumpAllAssetRedirects()
+{
+	UE::TDynamicUniqueLock<FCriticalSection> ScopeLock(GRedirectCollector.AcquireLock());
+	const TMap<FSoftObjectPath, FSoftObjectPath>& Redirects = GRedirectCollector.GetObjectPathRedirectionMapUnderLock(ScopeLock);
+	FString FullyQualifiedFileName = FPaths::ProfilingDir() + FString::Printf(TEXT("AllRedirects (%s).csv"), *FDateTime::Now().ToString(TEXT("%Y%m%d_%H%M%S")));
+	TUniquePtr<FArchive> OutputFile(IFileManager::Get().CreateFileWriter(*FullyQualifiedFileName));
+	if (OutputFile.IsValid())
+	{
+		TAnsiStringBuilder<4096> StringBuilder;
+
+		for (const TPair<FSoftObjectPath, FSoftObjectPath>& Pair : Redirects)
+		{
+			StringBuilder	<< WriteToAnsiString<256>(Pair.Key.ToString()).ToString() 
+							<< TEXT(",") 
+							<< WriteToAnsiString<256>(Pair.Value.ToString()).ToString()
+							<< LINE_TERMINATOR;
+		}
+		OutputFile->Serialize(StringBuilder.GetData(), StringBuilder.Len());
+	}
+}
+
 void FRedirectCollector::OnSoftObjectPathLoaded(const FSoftObjectPath& ObjectPath, FArchive* InArchive)
 {
 	if (ObjectPath.IsNull() || !GIsEditor)
@@ -224,6 +255,20 @@ void FRedirectCollector::OnStartupPackageLoadComplete()
 	TrackingReferenceTypesState = ETrackingReferenceTypesState::Disabled;
 }
 
+void FRedirectCollector::GetAllSourcePathsForTargetPath(const FSoftObjectPath& TargetPath, TArray<FSoftObjectPath>& OutSourcePaths) const
+{
+	FScopeLock ScopeLock(&CriticalSection);
+
+	OutSourcePaths.Reset();
+	for (const TPair<FSoftObjectPath, FSoftObjectPath>& Pair : ObjectPathRedirectionMap)
+	{
+		if (Pair.Value == TargetPath)
+		{
+			OutSourcePaths.Add(Pair.Key);
+		}
+	}
+}
+
 bool FRedirectCollector::ShouldTrackPackageReferenceTypes()
 {
 	// Called from within CriticalSection
diff --git a/Engine/Source/Runtime/CoreUObject/Public/Misc/RedirectCollector.h b/Engine/Source/Runtime/CoreUObject/Public/Misc/RedirectCollector.h
--- a/Engine/Source/Runtime/CoreUObject/Public/Misc/RedirectCollector.h
+++ b/Engine/Source/Runtime/CoreUObject/Public/Misc/RedirectCollector.h
@@ -6,6 +6,7 @@
 
 #pragma once
 
+#include "Async/UniqueLock.h"
 #include "Containers/Map.h"
 #include "Containers/Set.h"
 #include "Containers/SparseArray.h"
@@ -143,6 +144,25 @@ public:
 	/** Called from the cooker to stop the tracking of exclusions. */
 	COREUOBJECT_API void OnStartupPackageLoadComplete();
 
+	/** Access to the collected list of redirects when already holding the lock. */
+	COREUOBJECT_API const TMap<FSoftObjectPath, FSoftObjectPath>& GetObjectPathRedirectionMapUnderLock(const UE::TDynamicUniqueLock<FCriticalSection>& Lock) const
+	{
+		ensure(Lock.OwnsLock());
+		return ObjectPathRedirectionMap;
+	}
+	
+	/** Returns the set of paths, if any, that are redirected TO the provided path. This performs a relatively slow linear search. */
+	COREUOBJECT_API void GetAllSourcePathsForTargetPath(const FSoftObjectPath& TargetPath, TArray<FSoftObjectPath>& OutSourcePaths) const;
+
+	/** Used with GetObjectPathRedirectionMapUnderLock with code like:
+	 *  UE::TDynamicUniqueLock<FCriticalSection> ScopeLock(GRedirectCollector.AcquireLock());
+	 *  GRedirectCollector.GetObjectPathRedirectionMapUnderLock(ScopeLock);
+	 */
+	COREUOBJECT_API UE::TDynamicUniqueLock<FCriticalSection> AcquireLock() const
+	{
+		return UE::TDynamicUniqueLock<FCriticalSection>(CriticalSection);
+	}
+
 private:
 
 	/** A map of assets referenced by soft object paths, with the key being the package with the reference */
@@ -160,8 +180,8 @@ private:
 	/** When saving, apply this remapping to all soft object paths */
 	TMap<FSoftObjectPath, FSoftObjectPath> ObjectPathRedirectionMap;
 
-	/** For SoftObjectPackageMap map */
-	FCriticalSection CriticalSection;
+	/** For ObjectPathRedirectionMap map */
+	mutable FCriticalSection CriticalSection;
 
 	enum class ETrackingReferenceTypesState : uint8
 	{
