Disabled text character validation for Type_Asset FCoreRedirects because they are derived from UObjectRedirectors (meaning the values should be valid since they exist on disk) and because the existing validation is overly strict. Validation is not as useful in the Type_Asset case because redirects are not manually authored. Added a test to ensure that any redirector on disk (as captured by GRedirectCollector) can be added as a Type_Asset redirect.

#rnx

diff --git a/Engine/Source/Runtime/CoreUObject/Private/UObject/CoreRedirects.cpp b/Engine/Source/Runtime/CoreUObject/Private/UObject/CoreRedirects.cpp
--- a/Engine/Source/Runtime/CoreUObject/Private/UObject/CoreRedirects.cpp
+++ b/Engine/Source/Runtime/CoreUObject/Private/UObject/CoreRedirects.cpp
@@ -23,6 +23,10 @@
 
 #include "AutoRTFM/AutoRTFM.h"
 
+#if WITH_EDITOR
+#include "Misc/RedirectCollector.h"
+#endif
+
 DEFINE_LOG_CATEGORY(LogCoreRedirects);
 
 #if !defined UE_WITH_CORE_REDIRECTS
@@ -365,6 +369,27 @@ namespace UE::CoreRedirects::Private
 			}
 		}
 
+#if WITH_EDITOR
+		// ensure that it's safe to add any redirector in GRedirectCollector
+		{
+			TDynamicUniqueLock<FCriticalSection> ScopeLock(GRedirectCollector.AcquireLock());
+			const TMap<FSoftObjectPath, FSoftObjectPath>& RedirectMap = GRedirectCollector.GetObjectPathRedirectionMapUnderLock(ScopeLock);
+			TArray<const FCoreRedirect> RedirectList;
+			for (const TPair<FSoftObjectPath, FSoftObjectPath>& Redirect : RedirectMap)
+			{
+				RedirectList.Emplace(ECoreRedirectFlags::Type_Asset, Redirect.Key.ToString(), Redirect.Value.ToString());
+			}
+
+			FCoreRedirects::AddRedirectList(RedirectList, TEXT("GRedirectCollector"));
+			if (!FCoreRedirects::ValidateAssetRedirects())
+			{
+				bSuccess = false;
+				UE_LOG(LogCoreRedirects, Error, TEXT("Failed asset redirect validation."));
+			}
+		}
+		FCoreRedirects::RemoveAllAssetRedirects();
+#endif
+
 		// Ensure validation failure if chains exist
 		{
 			TMap<FSoftObjectPath, FSoftObjectPath> ChainRedirects;
@@ -2263,11 +2288,21 @@ bool FCoreRedirects::AddRedirectListUnderWriteLock(TArrayView<const FCoreRedirec
 			continue;
 		}
 
-		if ((!NewRedirect.OldName.HasValidCharacters(NewRedirect.RedirectFlags) && !FPackageName::IsVersePackage(NewRedirect.OldName.PackageName.ToString()))
-			|| (!NewRedirect.NewName.HasValidCharacters(NewRedirect.RedirectFlags) && !FPackageName::IsVersePackage(NewRedirect.NewName.PackageName.ToString())))
+		// Type_Asset redirects derive from UObjectRedirector instances on disk.
+		// Therefore we can assume that the names in it are valid (because they exist)
+		// This is as opposed to redirects that are manually created either in code or in 
+		// ini files and which should be subject to additional validation.
+		// Because the validation in HasValidCharacters does not perfectly align with other systems,
+		// there can exist assets which are valid, are referenced by UObjectRedirector assets, but which
+		// fail the check.
+		if (!EnumHasAnyFlags(NewRedirect.RedirectFlags, ECoreRedirectFlags::Type_Asset))
 		{
-			UE_LOG(LogCoreRedirects, Error, TEXT("AddRedirect(%s) failed to add redirect from %s to %s with invalid characters!"), *SourceString, *NewRedirect.OldName.ToString(), *NewRedirect.NewName.ToString());
-			continue;
+			if ((!NewRedirect.OldName.HasValidCharacters(NewRedirect.RedirectFlags) && !FPackageName::IsVersePackage(NewRedirect.OldName.PackageName.ToString()))
+				|| (!NewRedirect.NewName.HasValidCharacters(NewRedirect.RedirectFlags) && !FPackageName::IsVersePackage(NewRedirect.NewName.PackageName.ToString())))
+			{
+				UE_LOG(LogCoreRedirects, Error, TEXT("AddRedirect(%s) failed to add redirect from %s to %s with invalid characters!"), *SourceString, *NewRedirect.OldName.ToString(), *NewRedirect.NewName.ToString());
+				continue;
+			}
 		}
 
 		if (NewRedirect.IsWildcardMatch())
