Replication
* IsFullNameStableForNetworking now only returns true if his Outer is also a FullNameStable object.
* Added engine unit test: Net.SubObjects.TestStableNameSubObjects to test stable object replication behavior

#rnx

diff --git a/Engine/Source/Runtime/CoreUObject/Private/UObject/Obj.cpp b/Engine/Source/Runtime/CoreUObject/Private/UObject/Obj.cpp
--- a/Engine/Source/Runtime/CoreUObject/Private/UObject/Obj.cpp
+++ b/Engine/Source/Runtime/CoreUObject/Private/UObject/Obj.cpp
@@ -5612,7 +5612,7 @@ bool UObject::IsNameStableForNetworking() const
 /** IsFullNameStableForNetworking means an object can be referred to its full path name over the network */
 bool UObject::IsFullNameStableForNetworking() const
 {
-	if ( GetOuter() != NULL && !GetOuter()->IsNameStableForNetworking() )
+	if ( GetOuter() != NULL && !GetOuter()->IsFullNameStableForNetworking() )
 	{
 		return false;	// If any outer isn't stable, we can't consider the full name stable
 	}
