Don't collect UGCObjectReferencer::FImpl::InitialReferencedObjects twice when running Garbage Collector in single-threaded mode

diff --git a/Engine/Source/Runtime/CoreUObject/Private/UObject/GarbageCollection.cpp b/Engine/Source/Runtime/CoreUObject/Private/UObject/GarbageCollection.cpp
--- a/Engine/Source/Runtime/CoreUObject/Private/UObject/GarbageCollection.cpp
+++ b/Engine/Source/Runtime/CoreUObject/Private/UObject/GarbageCollection.cpp
@@ -3242,7 +3242,12 @@ protected:
 public:
 	virtual bool IsIgnoringArchetypeRef() const override final { return false;}
 	virtual bool IsIgnoringTransient() const override final { return false; }
-	virtual bool NeedsInitialReferences() const override final { return !IsParallel(Options); }
+	virtual bool NeedsInitialReferences() const override final 
+	{
+		// The logic here needs to match a hypothetical 'else' block in FRealtimeGC::GetInitialReferences(EGCOptions Options)
+		// so that UGCObjectReferencer::AddReferencedObjects(...) gathers InitialReferencedObjects when GC runs in single-threaded mode
+		return !IsParallel(Options); 
+	}
 	virtual void AllowEliminatingReferences(bool bAllow) override final { bAllowEliminatingReferences = bAllow; }
 
 	virtual void SetIsProcessingNativeReferences(bool bIsNative) override final
@@ -3881,14 +3886,13 @@ class FRealtimeGC : public FGarbageCollectionTracer
 	TConstArrayView<UObject**> GetInitialReferences(EGCOptions Options)
 	{
 		const double StartTime = FPlatformTime::Seconds();
-		if (!!(Options & EGCOptions::Parallel))
+		if (IsParallel(Options))
 		{
 			InitialCollection.Wait();
 		}
-		else
-		{
-			FGCObject::GGCObjectReferencer->AddInitialReferences(InitialReferences);
-		}
+		// Note that we don't gather initial references here in single-threaded mode because we'd need block and wait for FGCObject::GGCObjectReferencer->AddInitialReferences(InitialReferences) anyway
+		// so we just collect them in UGCObjectReferencer::AddReferencedObjects(...). The logic here needs to match TReachabilityCollectorBase::NeedsInitialReferences()
+
 		GGCStats.InitialReferenceCollectionTime += FPlatformTime::Seconds() - StartTime;
 
 		return InitialReferences;
