IncrementalCook: Add class metadata to the class schema, since class metadata (such as displayname) is often used in serialization functions that impact cooked packages.
This fixes the iterative false skip of UStateTree when a StateTreeNode's class's display name changes in c++, since the display name is incorporated into node names and saved into the cooked package.
#rnx

diff --git a/Engine/Source/Runtime/CoreUObject/Private/Serialization/UnversionedPropertySerialization.cpp b/Engine/Source/Runtime/CoreUObject/Private/Serialization/UnversionedPropertySerialization.cpp
--- a/Engine/Source/Runtime/CoreUObject/Private/Serialization/UnversionedPropertySerialization.cpp
+++ b/Engine/Source/Runtime/CoreUObject/Private/Serialization/UnversionedPropertySerialization.cpp
@@ -10,7 +10,10 @@
 #include "UObject/UnrealType.h"
 
 #if WITH_EDITORONLY_DATA
+#include "Algo/Sort.h"
 #include "Misc/FileHelper.h"
+#include "UObject/MetaData.h"
+#include "UObject/Package.h"
 #include "UObject/UObjectIterator.h"
 #endif
 
@@ -327,12 +330,7 @@ struct FUnversionedStructSchema
 		CA_ASSUME(Struct != nullptr);
 #if WITH_EDITORONLY_DATA
 		FBlake3 HashBuilder;
-		// Append the full name of the struct. To improve performance, append the FNames of its outers individually
-		// rather than calculating the full name as a string for performance
-		for (const UObject* Outer = Struct; Outer; Outer = Outer->GetOuter())
-		{
-			AppendHash(HashBuilder, Outer->GetFName());
-		}
+		AppendClassMetaData(Struct, HashBuilder);
 #endif
 		TArray<FUnversionedPropertySerializer, TInlineAllocator<256>> Serializers;
 		for (FProperty* Property = Struct->PropertyLink; Property; Property = Property->PropertyLinkNext)
@@ -379,6 +377,48 @@ struct FUnversionedStructSchema
 	}
 
 #if WITH_EDITORONLY_DATA
+	static void AppendClassMetaData(const UStruct* Struct, FBlake3& HashBuilder)
+	{
+		// Append the full name of the struct. To improve performance, append the FNames of its outers individually
+		// rather than calculating the full name as a string.
+		const UObject* Outermost = nullptr;
+		for (const UObject* Outer = Struct; Outer; Outer = Outer->GetOuter())
+		{
+			AppendHash(HashBuilder, Outer->GetFName());
+			Outermost = Outer;
+		}
+		// Append metadata properties of the struct since they influence builds.
+		// Skip this for the structs UPackage and UMetaData to avoid infinite recursion.
+		const TMap<FName, FString>* MetaData = nullptr;
+		if (Struct != UPackage::StaticClass() && Struct != UMetaData::StaticClass())
+		{
+			if (const UPackage* Package = Cast<UPackage>(Outermost))
+			{
+				// Package is const, but there is no const accessor for GetMetaData, so we const-cast it.
+				// Avoid calling GetMetaData if the metadata does not already exist, so that we do not
+				// create the metadata on the const object. 
+				if (Package->HasMetaData())
+				{
+					const UMetaData* MetaDataObject = const_cast<UPackage*>(Package)->GetMetaData();
+					check(MetaDataObject);
+					MetaData = MetaDataObject->GetMapForObject(Struct);
+				}
+			}
+		}
+		if (MetaData)
+		{
+			TArray<FName, TInlineAllocator<16>> MetaDataNames;
+			MetaData->GenerateKeyArray(MetaDataNames);
+			Algo::Sort(MetaDataNames, FNameLexicalLess());
+			for (FName MetaDataName : MetaDataNames)
+			{
+				AppendHash(HashBuilder, MetaDataName);
+				const FString& Value = MetaData->FindChecked(MetaDataName);
+				HashBuilder.Update(*Value, Value.Len() * sizeof(**Value));
+			}
+		}
+	}
+
 	static FBlake3Hash CalculateSchemaHash(UStruct* Struct, bool bSkipEditorOnly)
 	{
 		FBlake3 HashBuilder;
