[BugFix] Fixed bug that caused CDOs to not get copied properly when generating IDOs

diff --git a/Engine/Source/Runtime/CoreUObject/Internal/UObject/PropertyBagRepository.h b/Engine/Source/Runtime/CoreUObject/Internal/UObject/PropertyBagRepository.h
--- a/Engine/Source/Runtime/CoreUObject/Internal/UObject/PropertyBagRepository.h
+++ b/Engine/Source/Runtime/CoreUObject/Internal/UObject/PropertyBagRepository.h
@@ -192,8 +192,6 @@ private:
 
 	// Instantiate InstanceDataObject within BagData. Returns InstanceDataObject object. 
 	void CreateInstanceDataObjectUnsafe(UObject* Owner, FPropertyBagAssociationData& BagData, FArchive* Archive = nullptr);
-
-	static void CopyTaggedProperties (const UObject* Source, UObject* Dest);
 };
 
 #endif // WITH_EDITORONLY_DATA
diff --git a/Engine/Source/Runtime/CoreUObject/Private/UObject/InstanceDataObjectUtils.cpp b/Engine/Source/Runtime/CoreUObject/Private/UObject/InstanceDataObjectUtils.cpp
--- a/Engine/Source/Runtime/CoreUObject/Private/UObject/InstanceDataObjectUtils.cpp
+++ b/Engine/Source/Runtime/CoreUObject/Private/UObject/InstanceDataObjectUtils.cpp
@@ -2,6 +2,9 @@
 
 #include "UObject/InstanceDataObjectUtils.h"
 
+#include "Serialization/ObjectReader.h"
+#include "Serialization/ObjectWriter.h"
+
 #if WITH_EDITORONLY_DATA
 
 #include "HAL/IConsoleManager.h"
@@ -485,29 +488,23 @@ namespace UE
 		Result->StaticLink(/*bRelinkExistingProperties*/true);
 		return Result;
 	}
-
-	void CopyCDO(const UObject* Source, UObject* Destination)
+	
+	void CopyTaggedProperties(const UObject* Source, UObject* Dest)
 	{
-		for (const FProperty* SourceProperty : TFieldRange<FProperty>(Source->GetClass()))
-		{
-			if (const FProperty* DestinationProperty = Destination->GetClass()->FindPropertyByName(SourceProperty->GetFName()))
-			{
-				if (SourceProperty->SameType(DestinationProperty))
-				{
-					const void* SourceValue = SourceProperty->ContainerPtrToValuePtr<void>(Source);
-					void* DestinationValue = DestinationProperty->ContainerPtrToValuePtr<void>(Destination);
-					DestinationProperty->CopyCompleteValue(DestinationValue, SourceValue);
-				}
-				else
-				{
-					FString ValueText;
-					const void* SourceValue = SourceProperty->ContainerPtrToValuePtr<void>(Source);
-					void* DestinationValue = DestinationProperty->ContainerPtrToValuePtr<void>(Destination);
-					SourceProperty->ExportText_Direct(ValueText, SourceValue, SourceValue, const_cast<UObject*>(Source), PPF_None);
-					DestinationProperty->ImportText_Direct(*ValueText, DestinationValue, Destination, PPF_None);
-				}
-			}
-		}
+		FUObjectSerializeContext* SerializeContext = FUObjectThreadContext::Get().GetSerializeContext();
+		TGuardValue<bool> ImpersonatePropertiesScope(SerializeContext->bImpersonateProperties, true);
+		// don't mark properties as set by serialization when performing copy
+		TGuardValue<bool> ScopedTrackSerializedProperties(SerializeContext->bTrackSerializedProperties, false);
+
+		TArray<uint8> Buffer;
+		Buffer.Reserve(Source->GetClass()->GetStructureSize());
+		FObjectWriter Writer(Buffer);
+		Writer.ArNoDelta = true;
+		Source->GetClass()->SerializeTaggedProperties(Writer, (uint8*)Source, Source->GetClass(), nullptr);
+
+		FObjectReader Reader(Buffer);
+		Reader.ArMergeOverrides = true;
+		Dest->GetClass()->SerializeTaggedProperties(Reader, (uint8*)Dest, Dest->GetClass(), nullptr);
 	}
 
 	static void SetClassFlags(UClass* IDOClass, const UClass* OwnerClass)
@@ -535,7 +532,7 @@ namespace UE
 		UObject* ResultCDO = Result->GetDefaultObject();
 		if (ensure(OwnerCDO && ResultCDO))
 		{
-			CopyCDO(OwnerCDO, ResultCDO);
+			CopyTaggedProperties(OwnerCDO, ResultCDO);
 		}
 		return Result;
 	}
diff --git a/Engine/Source/Runtime/CoreUObject/Private/UObject/PropertyBagRepository.cpp b/Engine/Source/Runtime/CoreUObject/Private/UObject/PropertyBagRepository.cpp
--- a/Engine/Source/Runtime/CoreUObject/Private/UObject/PropertyBagRepository.cpp
+++ b/Engine/Source/Runtime/CoreUObject/Private/UObject/PropertyBagRepository.cpp
@@ -31,6 +31,9 @@ DEFINE_LOG_CATEGORY_STATIC(LogPropertyBagRepository, Log, All);
 namespace UE
 {
 
+/** Defined in InstanceDataObjectUtils.cpp */
+void CopyTaggedProperties(const UObject* Source, UObject* Dest);
+
 /** Internal registry that tracks the current set of types for property bag container objects instanced as placeholders for package exports that have invalid or missing class imports on load. */
 class FPropertyBagPlaceholderTypeRegistry
 {
@@ -882,24 +885,6 @@ void FPropertyBagRepository::CreateInstanceDataObjectUnsafe(UObject* Owner, FPro
 	}
 }
 
-void FPropertyBagRepository::CopyTaggedProperties(const UObject* Source, UObject* Dest)
-{
-	FUObjectSerializeContext* SerializeContext = FUObjectThreadContext::Get().GetSerializeContext();
-	TGuardValue<bool> ImpersonatePropertiesScope(SerializeContext->bImpersonateProperties, true);
-	// don't mark properties as set by serialization when performing copy
-	TGuardValue<bool> ScopedTrackSerializedProperties(SerializeContext->bTrackSerializedProperties, false);
-
-	TArray<uint8> Buffer;
-	Buffer.Reserve(Source->GetClass()->GetStructureSize());
-	FObjectWriter Writer(Buffer);
-	Writer.ArNoDelta = true;
-	Source->GetClass()->SerializeTaggedProperties(Writer, (uint8*)Source, Source->GetClass(), nullptr);
-
-	FObjectReader Reader(Buffer);
-	Reader.ArMergeOverrides = true;
-	Dest->GetClass()->SerializeTaggedProperties(Reader, (uint8*)Dest, Dest->GetClass(), nullptr);
-}
-
 FScopedIDOSerializationContext::FScopedIDOSerializationContext(UObject* InObject, FArchive& InArchive)
 	: Archive(&InArchive)
 	, Object(InObject)
