[ZenLoader]
 - Make UObject still being serialized invisible from postload functions
 - Fixes tons of issues where PostLoad could see uninitialized and RF_NeedLoad objects requiring additional code to filter these cases everywhere
 - Make the visibility logic reusable since it was duplicated at lots of places
 - Make sure behavior is the same whether async loading thread is active or not
 - Can be toggled with s.UseObjectVisibilityFilterForAsyncLoading (current default is false)
 - Will be made active by default soon

diff --git a/Engine/Source/Runtime/CoreUObject/Private/Serialization/AsyncLoading2.cpp b/Engine/Source/Runtime/CoreUObject/Private/Serialization/AsyncLoading2.cpp
--- a/Engine/Source/Runtime/CoreUObject/Private/Serialization/AsyncLoading2.cpp
+++ b/Engine/Source/Runtime/CoreUObject/Private/Serialization/AsyncLoading2.cpp
@@ -2733,12 +2733,23 @@ struct FAsyncPackage2
 			// Mark objects created during async loading process (e.g. from within PostLoad or CreateExport) as async loaded so they 
 			// cannot be found. This requires also keeping track of them so we can remove the async loading flag later one when we 
 			// finished routing PostLoad to all objects.
-			Object->SetInternalFlags(EInternalObjectFlags::AsyncLoading);
+			// 
+			// Once we pass the publish gate, we mark new objects in phase 2, they are generally preloaded as soon
+			// as they are created. However, it leaves a small window where GT could find a RF_NeedLoad object from phase 2...
+			if (ObjectsNowInPhase2)
+			{
+				Object->SetInternalFlags(EInternalObjectFlags::AsyncLoadingPhase2 | EInternalObjectFlags::Async);
+			}
+			else
+			{
+				Object->SetInternalFlags(EInternalObjectFlags::AsyncLoadingPhase1 | EInternalObjectFlags::Async);
+			}
 
 			ConstructedObjects.Add(Object);
 		}
 	}
 
+	void MoveConstructedObjectsToPhase2();
 	void ClearConstructedObjects();
 
 	/** Class specific callback for initializing non-native objects */
@@ -2815,6 +2826,8 @@ private:
 	int32						PostLoadInstanceIndex = 0;
 	/** Current loading state of a package. */
 	std::atomic<EAsyncPackageLoadingState2> AsyncPackageLoadingState { EAsyncPackageLoadingState2::NewPackage };
+	/** Whether the constructed objects have been moved to the second loading phase. */
+	std::atomic<bool> ObjectsNowInPhase2 = false;
 
 	struct FAllDependenciesState
 	{
@@ -4091,6 +4104,14 @@ void FAsyncLoadingThread2::ConditionalBeginDeferredPostLoad(FAsyncLoadingThreadS
 	check(ThreadState.bCanAccessAsyncLoadingThreadData);
 	if (DeferredPostLoadGroup->PackagesWithExportsToPostLoadCount == 0)
 	{
+		// Move everything to the second phase so that StaticFind from inside postload can
+		// see objects ready for postload.
+		for (int32 Index = DeferredPostLoadGroup->Packages.Num() - 1; Index >= 0; --Index)
+		{
+			FAsyncPackage2* Package = DeferredPostLoadGroup->Packages[Index];
+			Package->MoveConstructedObjectsToPhase2();
+		}
+
 		// Release the post load node of packages in the post load group in reverse order that they were added to the group
 		// This usually means that dependencies will be post load first, similarly to how they are also serialized first
 		for (int32 Index = DeferredPostLoadGroup->Packages.Num() - 1; Index >= 0; --Index)
@@ -4295,6 +4316,9 @@ bool FAsyncLoadingThread2::CreateAsyncPackagesFromQueue(FAsyncLoadingThreadState
 {
 	TRACE_CPUPROFILER_EVENT_SCOPE(CreateAsyncPackagesFromQueue);
 	
+	// Package creation needs access to all objects
+	TGuardValue GuardVisibilityFilter(FUObjectThreadContext::Get().AsyncVisibilityFilter, EInternalObjectFlags::None);
+
 	bool bPackagesCreated = false;
 	const int32 TimeSliceGranularity = ThreadState.UseTimeLimit() ? 4 : MAX_int32;
 
@@ -4711,6 +4735,9 @@ void FAsyncLoadEventQueue2::PushExternal(FEventLoadNode2* Node)
 
 bool FAsyncLoadEventQueue2::PopAndExecute(FAsyncLoadingThreadState2& ThreadState)
 {
+	// By default, nodes are all executed without visibility filter
+	TGuardValue GuardVisibilityFilter(FUObjectThreadContext::Get().AsyncVisibilityFilter, EInternalObjectFlags::None);
+
 	if (TimedOutEventNode)
 	{
 		// Backup and reset the node before executing it in case we end up with a recursive flush call, we don't want the same node run multiple time.
@@ -4762,6 +4789,9 @@ bool FAsyncLoadEventQueue2::PopAndExecute(FAsyncLoadingThreadState2& ThreadState
 
 bool FAsyncLoadEventQueue2::ExecuteSyncLoadEvents(FAsyncLoadingThreadState2& ThreadState)
 {
+	// By default, nodes are all executed without visibility filter
+	TGuardValue GuardVisibilityFilter(FUObjectThreadContext::Get().AsyncVisibilityFilter, EInternalObjectFlags::None);
+
 	check(!ThreadState.SyncLoadContextStack.IsEmpty());
 	FAsyncLoadingSyncLoadContext& SyncLoadContext = *ThreadState.SyncLoadContextStack.Top();
 
@@ -5946,6 +5976,8 @@ EEventLoadNodeExecutionResult FAsyncPackage2::Event_CreateLinkerLoadExports(FAsy
 {
 	TRACE_CPUPROFILER_EVENT_SCOPE(Event_CreateLinkerLoadExports);
 	UE_ASYNC_PACKAGE_DEBUG(Package->Desc);
+
+	check(FUObjectThreadContext::Get().AsyncVisibilityFilter == EInternalObjectFlags::None);
 	check(Package->AsyncPackageLoadingState == EAsyncPackageLoadingState2::CreateLinkerLoadExports);
 
 	FAsyncPackageScope2 Scope(Package);
@@ -5990,6 +6022,8 @@ EEventLoadNodeExecutionResult FAsyncPackage2::Event_ResolveLinkerLoadImports(FAs
 {
 	TRACE_CPUPROFILER_EVENT_SCOPE(Event_ResolveLinkerLoadImports);
 	UE_ASYNC_PACKAGE_DEBUG(Package->Desc);
+
+	check(FUObjectThreadContext::Get().AsyncVisibilityFilter == EInternalObjectFlags::None);
 	check(Package->AsyncPackageLoadingState == EAsyncPackageLoadingState2::ResolveLinkerLoadImports);
 
 	FAsyncPackageScope2 Scope(Package);
@@ -6014,6 +6048,8 @@ EEventLoadNodeExecutionResult FAsyncPackage2::Event_PreloadLinkerLoadExports(FAs
 {
 	TRACE_CPUPROFILER_EVENT_SCOPE(Event_PreloadLinkerLoadExports);
 	UE_ASYNC_PACKAGE_DEBUG(Package->Desc);
+
+	check(FUObjectThreadContext::Get().AsyncVisibilityFilter == EInternalObjectFlags::None);
 	check(Package->AsyncPackageLoadingState == EAsyncPackageLoadingState2::PreloadLinkerLoadExports);
 
 	FAsyncPackageScope2 Scope(Package);
@@ -6047,6 +6083,8 @@ EEventLoadNodeExecutionResult FAsyncPackage2::Event_ProcessPackageSummary(FAsync
 {
 	TRACE_CPUPROFILER_EVENT_SCOPE(Event_ProcessPackageSummary);
 	UE_ASYNC_PACKAGE_DEBUG(Package->Desc);
+	
+	check(FUObjectThreadContext::Get().AsyncVisibilityFilter == EInternalObjectFlags::None);
 	check(Package->AsyncPackageLoadingState == EAsyncPackageLoadingState2::WaitingForIo);
 	Package->AsyncPackageLoadingState = EAsyncPackageLoadingState2::ProcessPackageSummary;
 
@@ -6167,6 +6205,8 @@ EEventLoadNodeExecutionResult FAsyncPackage2::Event_DependenciesReady(FAsyncLoad
 {
 	TRACE_CPUPROFILER_EVENT_SCOPE(Event_DependenciesReady);
 	UE_ASYNC_PACKAGE_DEBUG(Package->Desc);
+	
+	check(FUObjectThreadContext::Get().AsyncVisibilityFilter == EInternalObjectFlags::None);
 	check(Package->AsyncPackageLoadingState == EAsyncPackageLoadingState2::WaitingForDependencies);
 
 	Package->AsyncPackageLoadingState = EAsyncPackageLoadingState2::DependenciesReady;
@@ -6234,6 +6274,8 @@ EEventLoadNodeExecutionResult FAsyncPackage2::Event_ProcessExportBundle(FAsyncLo
 {
 	TRACE_CPUPROFILER_EVENT_SCOPE(Event_ProcessExportBundle);
 	UE_ASYNC_PACKAGE_DEBUG(Package->Desc);
+	
+	check(FUObjectThreadContext::Get().AsyncVisibilityFilter == EInternalObjectFlags::None);
 	check(Package->AsyncPackageLoadingState >= EAsyncPackageLoadingState2::DependenciesReady);
 	Package->AsyncPackageLoadingState = EAsyncPackageLoadingState2::ProcessExportBundles;
 
@@ -6669,7 +6711,7 @@ void FAsyncPackage2::EventDrivenCreateExport(const FAsyncPackageHeaderData& Head
 	if (Object)
 	{
 		// If it has the AsyncLoading flag set it was created during the current load of this package (likely as a subobject)
-		if (!Object->HasAnyInternalFlags(EInternalObjectFlags::AsyncLoading))
+		if (!Object->HasAnyInternalFlags(EInternalObjectFlags_AsyncLoading))
 		{
 			ExportObject.bWasFoundInMemory = true;
 		}
@@ -6983,6 +7025,8 @@ EEventLoadNodeExecutionResult FAsyncPackage2::Event_ExportsDone(FAsyncLoadingThr
 {
 	TRACE_CPUPROFILER_EVENT_SCOPE(Event_ExportsDone);
 	UE_ASYNC_PACKAGE_DEBUG(Package->Desc);
+
+	check(FUObjectThreadContext::Get().AsyncVisibilityFilter == EInternalObjectFlags::None);
 #if ALT2_ENABLE_LINKERLOAD_SUPPORT
 	check(Package->AsyncPackageLoadingState == EAsyncPackageLoadingState2::ProcessExportBundles || 
 		  Package->AsyncPackageLoadingState == EAsyncPackageLoadingState2::WaitingForExternalReads ||
@@ -7276,6 +7320,8 @@ EEventLoadNodeExecutionResult FAsyncPackage2::Event_PostLoadExportBundle(FAsyncL
 {
 	TRACE_CPUPROFILER_EVENT_SCOPE(Event_PostLoad);
 	UE_ASYNC_PACKAGE_DEBUG(Package->Desc);
+
+	check(FUObjectThreadContext::Get().AsyncVisibilityFilter == EInternalObjectFlags::None);
 	check(Package->AsyncPackageLoadingState == EAsyncPackageLoadingState2::PostLoad);
 	check(Package->ExternalReadDependencies.Num() == 0);
 
@@ -7422,6 +7468,12 @@ EEventLoadNodeExecutionResult FAsyncPackage2::Event_DeferredPostLoadExportBundle
 	UE_ASYNC_PACKAGE_DEBUG(Package->Desc);
 	check(Package->AsyncPackageLoadingState == EAsyncPackageLoadingState2::DeferredPostLoad);
 
+	check(FUObjectThreadContext::Get().AsyncVisibilityFilter == EInternalObjectFlags::None);
+	// Prevent objects still being loaded from being visible from the game thread during postload.
+	// This avoids StaticFind, ObjectIterators or SoftObjectPtr from being able to resolve RF_NeedLoad objects along with other potential
+	// race condition during creation and serialization.
+	TGuardValue GuardVisibilityFilter(FUObjectThreadContext::Get().AsyncVisibilityFilter, EInternalObjectFlags::AsyncLoadingPhase1);
+
 	FAsyncPackageScope2 PackageScope(Package);
 	TGuardValue<bool> GuardIsRoutingPostLoad(PackageScope.ThreadContext.IsRoutingPostLoad, true);
 
@@ -7646,6 +7698,9 @@ EAsyncPackageState::Type FAsyncLoadingThread2::ProcessAsyncLoadingFromGameThread
 
 	check(IsInGameThread());
 
+	// By default, game thread can only access objects in Phase 2. Further scoping is needed to access everything.
+	TGuardValue GuardVisibilityFilter(FUObjectThreadContext::Get().AsyncVisibilityFilter, EInternalObjectFlags::AsyncLoadingPhase1);
+
 	FAsyncLoadingTickScope2 InAsyncLoadingTick(*this);
 	uint32 LoopIterations = 0;
 
@@ -7786,7 +7841,7 @@ EAsyncPackageState::Type FAsyncLoadingThread2::ProcessLoadedPackagesFromGameThre
 			if (Package->AsyncPackageLoadingState == EAsyncPackageLoadingState2::Finalize)
 			{
 				TArray<UObject*> CDODefaultSubobjects;
-				// Clear async loading flags (we still want RF_Async, but EInternalObjectFlags::AsyncLoading can be cleared)
+				// Clear async loading flags (we still want EInternalObjectFlags::Async, but EInternalObjectFlags::AsyncLoading can be cleared)
 				for (const FExportObject& Export : Package->Data.Exports)
 				{
 					if (Export.bFiltered || Export.bExportLoadFailed)
@@ -7805,9 +7860,9 @@ EAsyncPackageState::Type FAsyncLoadingThread2::ProcessLoadedPackagesFromGameThre
 						CDOToHandle->GetDefaultSubobjects(CDODefaultSubobjects);
 						for (UObject* SubObject : CDODefaultSubobjects)
 						{
-							if (SubObject && SubObject->HasAnyInternalFlags(EInternalObjectFlags::AsyncLoading))
+							if (SubObject && SubObject->HasAnyInternalFlags(EInternalObjectFlags_AsyncLoading))
 							{
-								SubObject->AtomicallyClearInternalFlags(EInternalObjectFlags::AsyncLoading);
+								SubObject->AtomicallyClearInternalFlags(EInternalObjectFlags_AsyncLoading);
 							}
 						}
 						CDODefaultSubobjects.Reset();
@@ -7997,10 +8052,12 @@ EAsyncPackageState::Type FAsyncLoadingThread2::TickAsyncLoadingFromGameThread(FA
 	//TRACE_INT_VALUE(GraphArcCount, GraphAllocator.TotalArcCount);
 	//TRACE_MEMORY_VALUE(GraphMemory, GraphAllocator.TotalAllocated);
 
-
 	check(IsInGameThread());
 	check(!IsGarbageCollecting());
 
+	// By default, game thread can only access objects in Phase 2. Further scoping is needed to access everything.
+	TGuardValue GuardVisibilityFilter(FUObjectThreadContext::Get().AsyncVisibilityFilter, EInternalObjectFlags::AsyncLoadingPhase1);
+
 #if WITH_EDITOR
 	// In the editor loading cannot be part of a transaction as it cannot be undone, and may result in recording half-loaded objects. So we suppress any active transaction while in this stack, and set the editor loading flag
 	TGuardValue<ITransaction*> SuppressTransaction(GUndo, nullptr);
@@ -8441,7 +8498,7 @@ static void VerifyObjectLoadFlagsWhenFinishedLoading()
 	TRACE_CPUPROFILER_EVENT_SCOPE(VerifyObjectLoadFlagsWhenFinishedLoading);
 
 	const EInternalObjectFlags AsyncFlags =
-		EInternalObjectFlags::Async | EInternalObjectFlags::AsyncLoading;
+		EInternalObjectFlags::Async | EInternalObjectFlags_AsyncLoading;
 
 	const EObjectFlags LoadIntermediateFlags = 
 		EObjectFlags::RF_NeedLoad | EObjectFlags::RF_WillBeLoaded |
@@ -8818,13 +8875,68 @@ void FAsyncPackage2::ClearImportedPackages()
 	}
 }
 
+void FAsyncPackage2::MoveConstructedObjectsToPhase2()
+{
+	TRACE_CPUPROFILER_EVENT_SCOPE(MoveConstructedObjectsToPhase2);
+
+	for (UObject* Object : ConstructedObjects)
+	{
+		check(!Object->HasAnyFlags(RF_NeedLoad | RF_NeedInitialization));
+		Object->SetInternalFlags(EInternalObjectFlags::AsyncLoadingPhase2);
+		Object->AtomicallyClearInternalFlags(EInternalObjectFlags::AsyncLoadingPhase1);
+	}
+
+	TArray<UObject*> CDODefaultSubobjects;
+	for (FExportObject& Export : Data.Exports)
+	{
+		UObject* Object = Export.Object;
+
+		// CDO need special handling, no matter if it's listed in DeferredFinalizeObjects
+		UObject* CDOToHandle = ((Object != nullptr) && Object->HasAnyFlags(RF_ClassDefaultObject)) ? Object : nullptr;
+
+		if (CDOToHandle != nullptr)
+		{
+			CDOToHandle->GetDefaultSubobjects(CDODefaultSubobjects);
+			for (UObject* SubObject : CDODefaultSubobjects)
+			{
+				if (SubObject && SubObject->HasAnyInternalFlags(EInternalObjectFlags::AsyncLoadingPhase1))
+				{
+					check(!SubObject->HasAnyFlags(RF_NeedLoad | RF_NeedInitialization));
+					SubObject->SetInternalFlags(EInternalObjectFlags::AsyncLoadingPhase2);
+					SubObject->AtomicallyClearInternalFlags(EInternalObjectFlags::AsyncLoadingPhase1);
+				}
+			}
+			CDODefaultSubobjects.Reset();
+		}
+
+		if (Object)
+		{
+			check(!Object->HasAnyFlags(RF_NeedLoad | RF_NeedInitialization));
+			if (Object->HasAnyInternalFlags(EInternalObjectFlags::AsyncLoadingPhase1))
+			{
+				Object->SetInternalFlags(EInternalObjectFlags::AsyncLoadingPhase2);
+				Object->AtomicallyClearInternalFlags(EInternalObjectFlags::AsyncLoadingPhase1);
+			}
+		}
+	}
+
+	if (LinkerRoot && LinkerRoot->HasAnyInternalFlags(EInternalObjectFlags::AsyncLoadingPhase1))
+	{
+		check(!LinkerRoot->HasAnyFlags(RF_NeedLoad | RF_NeedInitialization));
+		LinkerRoot->SetInternalFlags(EInternalObjectFlags::AsyncLoadingPhase2);
+		LinkerRoot->AtomicallyClearInternalFlags(EInternalObjectFlags::AsyncLoadingPhase1);
+	}
+
+	ObjectsNowInPhase2 = true;
+}
+
 void FAsyncPackage2::ClearConstructedObjects()
 {
 	TRACE_CPUPROFILER_EVENT_SCOPE(ClearConstructedObjects);
 
 	for (UObject* Object : ConstructedObjects)
 	{
-		Object->AtomicallyClearInternalFlags(EInternalObjectFlags::AsyncLoading | EInternalObjectFlags::Async);
+		Object->AtomicallyClearInternalFlags(EInternalObjectFlags_AsyncLoading | EInternalObjectFlags::Async);
 	}
 	ConstructedObjects.Empty();
 
@@ -8833,11 +8945,11 @@ void FAsyncPackage2::ClearConstructedObjects()
 		if (Export.bWasFoundInMemory)
 		{
 			check(Export.Object);
-			Export.Object->AtomicallyClearInternalFlags(EInternalObjectFlags::AsyncLoading | EInternalObjectFlags::Async);
+			Export.Object->AtomicallyClearInternalFlags(EInternalObjectFlags_AsyncLoading | EInternalObjectFlags::Async);
 		}
 		else
 		{
-			checkf(!Export.Object || !Export.Object->HasAnyInternalFlags(EInternalObjectFlags::AsyncLoading | EInternalObjectFlags::Async),
+			checkf(!Export.Object || !Export.Object->HasAnyInternalFlags(EInternalObjectFlags_AsyncLoading | EInternalObjectFlags::Async),
 				TEXT("Export object: %s (ObjectFlags=%x, InternalObjectFlags=%x)"),
 					*Export.Object->GetFullName(),
 					Export.Object->GetFlags(),
@@ -8847,7 +8959,7 @@ void FAsyncPackage2::ClearConstructedObjects()
 
 	if (LinkerRoot)
 	{
-		LinkerRoot->AtomicallyClearInternalFlags(EInternalObjectFlags::AsyncLoading | EInternalObjectFlags::Async);
+		LinkerRoot->AtomicallyClearInternalFlags(EInternalObjectFlags_AsyncLoading | EInternalObjectFlags::Async);
 	}
 }
 
@@ -9073,6 +9185,10 @@ void FAsyncPackage2::FinishUPackage()
 void FAsyncLoadingThread2::ConditionalProcessEditorCallbacks()
 {
 	check(IsInGameThread());
+	
+	// Prevent objects still being loaded from being accessible from the game thread.
+	TGuardValue GuardVisibilityFilter(FUObjectThreadContext::Get().AsyncVisibilityFilter, EInternalObjectFlags::AsyncLoadingPhase1);
+
 	if (!GameThreadState->SyncLoadContextStack.IsEmpty())
 	{
 		return;
@@ -9326,6 +9442,9 @@ void FAsyncLoadingThread2::FlushLoading(TConstArrayView<int32> RequestIDs)
 {
 	if (IsAsyncLoadingPackages())
 	{
+		// Prevent objects still being loaded from being accessible from the game thread.
+		TGuardValue GuardVisibilityFilter(FUObjectThreadContext::Get().AsyncVisibilityFilter, EInternalObjectFlags::AsyncLoadingPhase1);
+
 		// We can't possibly support flushing from async loading thread unless we have the partial request support active.
 #if WITH_PARTIAL_REQUEST_DURING_RECURSION
 		const bool bIsFlushSupportedOnCurrentThread = IsInGameThread() || IsInAsyncLoadingThread();
diff --git a/Engine/Source/Runtime/CoreUObject/Private/UObject/ObjectVisibility.cpp b/Engine/Source/Runtime/CoreUObject/Private/UObject/ObjectVisibility.cpp
new file mode 100644
--- /dev/null
+++ b/Engine/Source/Runtime/CoreUObject/Private/UObject/ObjectVisibility.cpp
@@ -0,0 +1,28 @@
+// Copyright Epic Games, Inc. All Rights Reserved.
+
+#include "UObject/ObjectVisibility.h"
+#include "UObject/UObjectThreadContext.h"
+#include "HAL/IConsoleManager.h"
+
+static bool GUseObjectVisibilityFilterForAsyncLoading = false;
+static FAutoConsoleVariableRef CVar(
+	TEXT("s.UseObjectVisibilityFilterForAsyncLoading"),
+	GUseObjectVisibilityFilterForAsyncLoading,
+	TEXT("When active and supported by the current loader, will prevent objects still in the early loading phase from being discovered during postload\n")
+	TEXT("to avoid race conditions and manipulation of objects that still haven't been deserialized.")
+);
+
+namespace UE
+{
+
+	EInternalObjectFlags GetAsyncLoadingInternalFlagsExclusion()
+	{
+		if (IsInAsyncLoadingThread())
+		{
+			return GUseObjectVisibilityFilterForAsyncLoading ? FUObjectThreadContext::Get().AsyncVisibilityFilter : EInternalObjectFlags::None;
+		}
+
+		return EInternalObjectFlags_AsyncLoading;
+	}
+
+} // namespace UE
\ No newline at end of file
diff --git a/Engine/Source/Runtime/CoreUObject/Private/UObject/UObjectGlobals.cpp b/Engine/Source/Runtime/CoreUObject/Private/UObject/UObjectGlobals.cpp
--- a/Engine/Source/Runtime/CoreUObject/Private/UObject/UObjectGlobals.cpp
+++ b/Engine/Source/Runtime/CoreUObject/Private/UObject/UObjectGlobals.cpp
@@ -29,6 +29,7 @@
 #include "UObject/UObjectHash.h"
 #include "UObject/UObjectHashPrivate.h"
 #include "UObject/Object.h"
+#include "UObject/ObjectVisibility.h"
 #include "UObject/GarbageCollection.h"
 #include "UObject/Class.h"
 #include "UObject/CoreRedirects.h"
@@ -271,10 +272,11 @@ namespace
 	 * @param ObjectPackage Package of the object to find.
 	 * @param ObjectName Name of the object to find.
 	 * @param ExactClass If the class match has to be exact. I.e. ObjectClass == FoundObjects.GetClass()
+	 * @param ExclusiveInternalFlags Do not return objects that have any of these internal flags
 	 *
 	 * @returns Found object.
 	 */
-	UObject* StaticFindObjectWithChangedLegacyPath(UClass* ObjectClass, UObject* ObjectPackage, FName ObjectName, bool ExactClass)
+	UObject* StaticFindObjectWithChangedLegacyPath(UClass* ObjectClass, UObject* ObjectPackage, FName ObjectName, bool ExactClass, EInternalObjectFlags ExclusiveInternalFlags)
 	{
 		UObject* MatchingObject = nullptr;
 
@@ -304,17 +306,25 @@ namespace
 
 			MatchingObject = StaticFindObject(ObjectClass, ObjectPackage->GetOutermost(), *ObjectName.ToString(), ExactClass);
 
-			if (MatchingObject && bSubclassOfPathChangedClass)
+			if (MatchingObject)
 			{
-				// If the class wasn't given exactly, check if found object is of class that outers were changed.
-				UClass* MatchingObjectClass = MatchingObject->GetClass();
-				if (!(MatchingObjectClass == UEnum::StaticClass()	// Enums
-					|| MatchingObjectClass == UScriptStruct::StaticClass() || MatchingObjectClass == UStruct::StaticClass() // Structs
-					|| (MatchingObjectClass == UFunction::StaticClass() && bHasDelegateSignaturePostfix)) // Delegates
-					)
+				if (MatchingObject->HasAnyInternalFlags(ExclusiveInternalFlags))
 				{
 					return nullptr;
 				}
+
+				if (bSubclassOfPathChangedClass)
+				{
+					// If the class wasn't given exactly, check if found object is of class that outers were changed.
+					UClass* MatchingObjectClass = MatchingObject->GetClass();
+					if (!(MatchingObjectClass == UEnum::StaticClass()	// Enums
+						|| MatchingObjectClass == UScriptStruct::StaticClass() || MatchingObjectClass == UStruct::StaticClass() // Structs
+						|| (MatchingObjectClass == UFunction::StaticClass() && bHasDelegateSignaturePostfix)) // Delegates
+						)
+					{
+						return nullptr;
+					}
+				}
 			}
 		}
 
@@ -416,13 +426,13 @@ UObject* StaticFindObjectFast(UClass* ObjectClass, UObject* ObjectPackage, FName
 	UE_CLOG(IsGarbageCollectingAndLockingUObjectHashTables(), LogUObjectGlobals, Fatal, TEXT("Illegal call to StaticFindObjectFast() while garbage collecting!"));
 
 	// We don't want to return any objects that are currently being background loaded unless we're using FindObject during async loading.
-	ExclusiveInternalFlags |= IsInAsyncLoadingThread() ? EInternalObjectFlags::None : EInternalObjectFlags::AsyncLoading;	
+	ExclusiveInternalFlags |= UE::GetAsyncLoadingInternalFlagsExclusion();
 	PRAGMA_DISABLE_DEPRECATION_WARNINGS
 	UObject* FoundObject = StaticFindObjectFastInternal(ObjectClass, ObjectPackage, ObjectName, bExactClass, bAnyPackage, ExclusiveFlags, ExclusiveInternalFlags);
 	PRAGMA_ENABLE_DEPRECATION_WARNINGS
 	if (!FoundObject)
 	{
-		FoundObject = StaticFindObjectWithChangedLegacyPath(ObjectClass, ObjectPackage, ObjectName, bExactClass);
+		FoundObject = StaticFindObjectWithChangedLegacyPath(ObjectClass, ObjectPackage, ObjectName, bExactClass, ExclusiveInternalFlags);
 	}
 
 	return FoundObject;
@@ -437,12 +447,12 @@ UObject* StaticFindObjectFast(UClass* ObjectClass, UObject* ObjectPackage, FName
 	UE_CLOG(IsGarbageCollectingAndLockingUObjectHashTables(), LogUObjectGlobals, Fatal, TEXT("Illegal call to StaticFindObjectFast() while garbage collecting!"));
 
 	// We don't want to return any objects that are currently being background loaded unless we're using FindObject during async loading.
-	ExclusiveInternalFlags |= IsInAsyncLoadingThread() ? EInternalObjectFlags::None : EInternalObjectFlags::AsyncLoading;
+	ExclusiveInternalFlags |= UE::GetAsyncLoadingInternalFlagsExclusion();
 	UObject* FoundObject = StaticFindObjectFastInternal(ObjectClass, ObjectPackage, ObjectName, bExactClass, ExclusiveFlags, ExclusiveInternalFlags);
 
 	if (!FoundObject)
 	{
-		FoundObject = StaticFindObjectWithChangedLegacyPath(ObjectClass, ObjectPackage, ObjectName, bExactClass);
+		FoundObject = StaticFindObjectWithChangedLegacyPath(ObjectClass, ObjectPackage, ObjectName, bExactClass, ExclusiveInternalFlags);
 	}
 
 	return FoundObject;
@@ -455,13 +465,13 @@ UObject* StaticFindObjectFastSafe(UClass* ObjectClass, UObject* ObjectPackage, F
 	if (!UE::IsSavingPackage(nullptr) && !IsGarbageCollectingAndLockingUObjectHashTables())
 	{
 		// We don't want to return any objects that are currently being background loaded unless we're using FindObject during async loading.
-		ExclusiveInternalFlags |= IsInAsyncLoadingThread() ? EInternalObjectFlags::None : EInternalObjectFlags::AsyncLoading;
+		ExclusiveInternalFlags |= UE::GetAsyncLoadingInternalFlagsExclusion();
 		PRAGMA_DISABLE_DEPRECATION_WARNINGS
 		FoundObject = StaticFindObjectFastInternal(ObjectClass, ObjectPackage, ObjectName, bExactClass, bAnyPackage, ExclusiveFlags, ExclusiveInternalFlags);
 		PRAGMA_ENABLE_DEPRECATION_WARNINGS
 		if (!FoundObject)
 		{
-			FoundObject = StaticFindObjectWithChangedLegacyPath(ObjectClass, ObjectPackage, ObjectName, bExactClass);
+			FoundObject = StaticFindObjectWithChangedLegacyPath(ObjectClass, ObjectPackage, ObjectName, bExactClass, ExclusiveInternalFlags);
 		}
 	}
 
@@ -475,11 +485,11 @@ UObject* StaticFindObjectFastSafe(UClass* ObjectClass, UObject* ObjectPackage, F
 	if (!UE::IsSavingPackage(nullptr) && !IsGarbageCollectingAndLockingUObjectHashTables())
 	{
 		// We don't want to return any objects that are currently being background loaded unless we're using FindObject during async loading.
-		ExclusiveInternalFlags |= IsInAsyncLoadingThread() ? EInternalObjectFlags::None : EInternalObjectFlags::AsyncLoading;
+		ExclusiveInternalFlags |= UE::GetAsyncLoadingInternalFlagsExclusion();;
 		FoundObject = StaticFindObjectFastInternal(ObjectClass, ObjectPackage, ObjectName, bExactClass, ExclusiveFlags, ExclusiveInternalFlags);
 		if (!FoundObject)
 		{
-			FoundObject = StaticFindObjectWithChangedLegacyPath(ObjectClass, ObjectPackage, ObjectName, bExactClass);
+			FoundObject = StaticFindObjectWithChangedLegacyPath(ObjectClass, ObjectPackage, ObjectName, bExactClass, ExclusiveInternalFlags);
 		}
 	}
 
@@ -651,7 +661,7 @@ bool StaticFindAllObjectsFast(TArray<UObject*>& OutFoundObjects, UClass* ObjectC
 	UE_CLOG(UE::IsSavingPackage(nullptr) || IsGarbageCollectingAndLockingUObjectHashTables(), LogUObjectGlobals, Fatal, TEXT("Illegal call to StaticFindAllObjectsFast() while serializing object data or garbage collecting!"));
 
 	// We don't want to return any objects that are currently being background loaded unless we're using FindObject during async loading.
-	ExclusiveInternalFlags |= IsInAsyncLoadingThread() ? EInternalObjectFlags::None : EInternalObjectFlags::AsyncLoading;
+	ExclusiveInternalFlags |= UE::GetAsyncLoadingInternalFlagsExclusion();
 	return StaticFindAllObjectsFastInternal(OutFoundObjects, ObjectClass, ObjectName, ExactClass, ExclusiveFlags, ExclusiveInternalFlags);
 }
 
@@ -661,7 +671,7 @@ bool StaticFindAllObjectsFastSafe(TArray<UObject*>& OutFoundObjects, UClass* Obj
 	if (!UE::IsSavingPackage(nullptr) && !IsGarbageCollectingAndLockingUObjectHashTables())
 	{
 		// We don't want to return any objects that are currently being background loaded unless we're using FindObject during async loading.
-		ExclusiveInternalFlags |= IsInAsyncLoadingThread() ? EInternalObjectFlags::None : EInternalObjectFlags::AsyncLoading;
+		ExclusiveInternalFlags |= UE::GetAsyncLoadingInternalFlagsExclusion();
 		bFoundObjects = StaticFindAllObjectsFastInternal(OutFoundObjects, ObjectClass, ObjectName, ExactClass, ExclusiveFlags, ExclusiveInternalFlags);
 	}
 	return bFoundObjects;
@@ -3038,9 +3048,9 @@ UObject* StaticDuplicateObjectEx( FObjectDuplicationParameters& Parameters )
 		*Parameters.SourceObject->GetClass()->GetName(), Parameters.SourceObject->GetClass()->GetPropertiesSize(),
 		*Parameters.DestClass->GetName(), Parameters.DestClass->GetPropertiesSize());
 	
-	UE_CLOG(FPlatformProperties::RequiresCookedData() && Parameters.SourceObject->HasAnyInternalFlags(EInternalObjectFlags::AsyncLoading), LogUObjectGlobals, Warning, TEXT("Duplicating object '%s' that's still being async loaded"), *Parameters.SourceObject->GetFullName());
+	UE_CLOG(FPlatformProperties::RequiresCookedData() && Parameters.SourceObject->HasAnyInternalFlags(EInternalObjectFlags_AsyncLoading), LogUObjectGlobals, Warning, TEXT("Duplicating object '%s' that's still being async loaded"), *Parameters.SourceObject->GetFullName());
 	// Make sure we're not duplicating the AsyncLoading, Async or LoaderImport internal flags, they will prevent the object from being gcd.
-	Parameters.InternalFlagMask &= ~(EInternalObjectFlags::Async | EInternalObjectFlags::LoaderImport | EInternalObjectFlags::AsyncLoading);
+	Parameters.InternalFlagMask &= ~(EInternalObjectFlags::Async | EInternalObjectFlags::LoaderImport | EInternalObjectFlags_AsyncLoading);
 
 	if (!IsAsyncLoading() && Parameters.SourceObject->HasAnyFlags(RF_ClassDefaultObject))
 	{
@@ -3255,7 +3265,7 @@ UObject* StaticDuplicateObjectEx( FObjectDuplicationParameters& Parameters )
 
 bool SaveToTransactionBuffer(UObject* Object, bool bMarkDirty)
 {
-	check(!Object->HasAnyInternalFlags(EInternalObjectFlags::Async | EInternalObjectFlags::AsyncLoading));
+	check(!Object->HasAnyInternalFlags(EInternalObjectFlags::Async | EInternalObjectFlags_AsyncLoading));
 	bool bSavedToTransactionBuffer = false;
 
 	// Script packages should not end up in the transaction buffer.
@@ -3660,7 +3670,7 @@ UObject* StaticAllocateObject
 		// Sanity checks for async flags.
 		// It's possible to duplicate an object on the game thread that is still being referenced 
 		// by async loading code or has been created on a different thread than the main thread.
-		Obj->ClearInternalFlags(EInternalObjectFlags::AsyncLoading);
+		Obj->ClearInternalFlags(EInternalObjectFlags_AsyncLoading);
 		if (Obj->HasAnyInternalFlags(EInternalObjectFlags::Async) && IsInGameThread())
 		{
 			Obj->ClearInternalFlags(EInternalObjectFlags::Async);
@@ -4543,7 +4553,7 @@ UObject* StaticConstructObject_Internal(const FStaticConstructObjectParameters&
 	
 	if (GIsEditor && 
 		// Do not consider object creation in transaction if the object is marked as async or in being async loaded 
-		!Result->HasAnyInternalFlags(EInternalObjectFlags::Async | EInternalObjectFlags::AsyncLoading) &&
+		!Result->HasAnyInternalFlags(EInternalObjectFlags::Async | EInternalObjectFlags_AsyncLoading) &&
 		// Read GUndo only if not having Async flags set to avoid making TSAN unhappy that we're trying to read an unsynchronized global
 		GUndo &&
 		(InFlags & RF_Transactional) && !(InFlags & RF_NeedLoad) && 
diff --git a/Engine/Source/Runtime/CoreUObject/Private/UObject/UObjectHash.cpp b/Engine/Source/Runtime/CoreUObject/Private/UObject/UObjectHash.cpp
--- a/Engine/Source/Runtime/CoreUObject/Private/UObject/UObjectHash.cpp
+++ b/Engine/Source/Runtime/CoreUObject/Private/UObject/UObjectHash.cpp
@@ -9,6 +9,7 @@
 #include "UObject/GarbageCollectionGlobals.h"
 #include "UObject/Package.h"
 #include "UObject/UObjectIterator.h"
+#include "UObject/ObjectVisibility.h"
 #include "Templates/Casts.h"
 #include "Misc/AsciiSet.h"
 #include "Misc/PackageName.h"
@@ -1188,11 +1189,7 @@ void GetObjectsWithOuter(const class UObjectBase* Outer, TArray<UObject *>& Resu
 #endif
 
 	// We don't want to return any objects that are currently being background loaded unless we're using the object iterator during async loading.
-	ExclusionInternalFlags |= EInternalObjectFlags::Unreachable;
-	if (!IsInAsyncLoadingThread())
-	{
-		ExclusionInternalFlags |= EInternalObjectFlags::AsyncLoading;
-	}
+	ExclusionInternalFlags |= EInternalObjectFlags::Unreachable | UE::GetAsyncLoadingInternalFlagsExclusion();
 
 	int32 StartNum = Results.Num();
 	FUObjectHashTables& ThreadHash = FUObjectHashTables::Get();
@@ -1256,11 +1253,7 @@ void ForEachObjectWithOuterBreakable(const class UObjectBase* Outer, TFunctionRe
 #endif
 
 	// We don't want to return any objects that are currently being background loaded unless we're using the object iterator during async loading.
-	ExclusionInternalFlags |= EInternalObjectFlags::Unreachable;
-	if (!IsInAsyncLoadingThread())
-	{
-		ExclusionInternalFlags |= EInternalObjectFlags::AsyncLoading;
-	}
+	ExclusionInternalFlags |= EInternalObjectFlags::Unreachable | UE::GetAsyncLoadingInternalFlagsExclusion();
 
 	FUObjectHashTables& ThreadHash = FUObjectHashTables::Get();
 	FHashTableLock HashLock(ThreadHash);
@@ -1305,11 +1298,7 @@ UObjectBase* FindObjectWithOuter(const class UObjectBase* Outer, const class UCl
 	UObject* Result = nullptr;
 	check( Outer );
 	// We don't want to return any objects that are currently being background loaded unless we're using the object iterator during async loading.
-	EInternalObjectFlags ExclusionInternalFlags = EInternalObjectFlags::Unreachable;
-	if (!IsInAsyncLoadingThread())
-	{
-		ExclusionInternalFlags = EInternalObjectFlags::AsyncLoading;
-	}
+	EInternalObjectFlags ExclusionInternalFlags = EInternalObjectFlags::Unreachable | UE::GetAsyncLoadingInternalFlagsExclusion();
 
 	if( NameToLookFor != NAME_None )
 	{
@@ -1359,11 +1348,7 @@ void ForEachObjectWithPackage(const class UPackage* Package, TFunctionRef<bool(U
 	check(Package != nullptr);
 
 	// We don't want to return any objects that are currently being background loaded unless we're using the object iterator during async loading.
-	ExclusionInternalFlags |= EInternalObjectFlags::Unreachable;
-	if (!IsInAsyncLoadingThread())
-	{
-		ExclusionInternalFlags |= EInternalObjectFlags::AsyncLoading;
-	}
+	ExclusionInternalFlags |= EInternalObjectFlags::Unreachable | UE::GetAsyncLoadingInternalFlagsExclusion();
 
 	FUObjectHashTables& ThreadHash = FUObjectHashTables::Get();
 	FHashTableLock HashLock(ThreadHash);
@@ -1470,11 +1455,7 @@ FORCEINLINE void ForEachObjectOfClasses_Implementation(FUObjectHashTables& Threa
 	TRACE_CPUPROFILER_EVENT_SCOPE(ForEachObjectOfClasses_Implementation);
 
 	// We don't want to return any objects that are currently being background loaded unless we're using the object iterator during async loading.
-	ExclusionInternalFlags |= EInternalObjectFlags::Unreachable;
-	if (!IsInAsyncLoadingThread())
-	{
-		ExclusionInternalFlags |= EInternalObjectFlags::AsyncLoading;
-	}
+	ExclusionInternalFlags |= EInternalObjectFlags::Unreachable | UE::GetAsyncLoadingInternalFlagsExclusion();
 
 	TBucketMapLock ClassToObjectListMapLock(ThreadHash.ClassToObjectListMap);
 
@@ -1576,7 +1557,7 @@ bool ClassHasInstancesAsyncLoading(const UClass* ClassToLookFor)
 				UObject *Object = static_cast<UObject*>(*ObjectIt);
 				// If the object is async loading we'll want to indicate as such to the caller,
 				// excepting two cases - garbage objects and the CDO:
-				if (Object->HasAnyInternalFlags(EInternalObjectFlags::AsyncLoading) 
+				if (Object->HasAnyInternalFlags(EInternalObjectFlags_AsyncLoading) 
 					// garbage objects won't require that the class be kept alive:
 					&& !Object->HasAnyInternalFlags(EInternalObjectFlags::Garbage) 
 					// CDO is required and owned by the class - also doesn't need to keep the class alive:
diff --git a/Engine/Source/Runtime/CoreUObject/Private/UObject/UObjectMarks.cpp b/Engine/Source/Runtime/CoreUObject/Private/UObject/UObjectMarks.cpp
--- a/Engine/Source/Runtime/CoreUObject/Private/UObject/UObjectMarks.cpp
+++ b/Engine/Source/Runtime/CoreUObject/Private/UObject/UObjectMarks.cpp
@@ -8,6 +8,7 @@
 #include "UObject/ObjectMacros.h"
 #include "UObject/UObjectIterator.h"
 #include "UObject/GarbageCollectionGlobals.h"
+#include "UObject/ObjectVisibility.h"
 
 struct FObjectMark
 {
@@ -218,11 +219,8 @@ EObjectMark ObjectGetAllMarks(const class UObjectBase* Object)
 void GetObjectsWithAllMarks(TArray<UObject *>& Results, EObjectMark Marks)
 {
 	// We don't want to return any objects that are currently being background loaded unless we're using the object iterator during async loading.
-	EInternalObjectFlags ExclusionFlags = EInternalObjectFlags::Unreachable;
-	if (!IsInAsyncLoadingThread())
-	{
-		ExclusionFlags |= EInternalObjectFlags::AsyncLoading;
-	}
+	EInternalObjectFlags ExclusionFlags = EInternalObjectFlags::Unreachable | UE::GetAsyncLoadingInternalFlagsExclusion();
+
 	const TMap<const UObjectBase *, FObjectMark>& Map = FThreadMarkAnnotation::Get().MarkAnnotation.GetAnnotationMap();
 	Results.Empty(Map.Num());
 	for (TMap<const UObjectBase *, FObjectMark>::TConstIterator It(Map); It; ++It)
@@ -241,11 +239,8 @@ void GetObjectsWithAllMarks(TArray<UObject *>& Results, EObjectMark Marks)
 void GetObjectsWithAnyMarks(TArray<UObject *>& Results, EObjectMark Marks)
 {
 	// We don't want to return any objects that are currently being background loaded unless we're using the object iterator during async loading.
-	EInternalObjectFlags ExclusionFlags = EInternalObjectFlags::Unreachable;
-	if (!IsInAsyncLoadingThread())
-	{
-		ExclusionFlags |= EInternalObjectFlags::AsyncLoading;
-	}
+	EInternalObjectFlags ExclusionFlags = EInternalObjectFlags::Unreachable | UE::GetAsyncLoadingInternalFlagsExclusion();
+
 	const TMap<const UObjectBase *, FObjectMark>& Map = FThreadMarkAnnotation::Get().MarkAnnotation.GetAnnotationMap();
 	Results.Empty(Map.Num());
 	for (TMap<const UObjectBase *, FObjectMark>::TConstIterator It(Map); It; ++It)
diff --git a/Engine/Source/Runtime/CoreUObject/Public/UObject/ObjectMacros.h b/Engine/Source/Runtime/CoreUObject/Public/UObject/ObjectMacros.h
--- a/Engine/Source/Runtime/CoreUObject/Public/UObject/ObjectMacros.h
+++ b/Engine/Source/Runtime/CoreUObject/Public/UObject/ObjectMacros.h
@@ -618,11 +618,12 @@ enum class EInternalObjectFlags : int32
 
 	LoaderImport = 1 << 20, ///< Object is ready to be imported by another package during loading
 	Garbage = 1 << 21, ///< Garbage from logical point of view and should not be referenced. This flag is mirrored in EObjectFlags as RF_Garbage for performance
+	AsyncLoadingPhase1 = 1 << 22, ///< Object is being asynchronously loaded.
 	ReachableInCluster = 1 << 23, ///< External reference to object in cluster exists
 	ClusterRoot = 1 << 24, ///< Root of a cluster
 	Native = 1 << 25, ///< Native (UClass only). 
 	Async = 1 << 26, ///< Object exists only on a different thread than the game thread.
-	AsyncLoading = 1 << 27, ///< Object is being asynchronously loaded.
+	AsyncLoadingPhase2 = 1 << 27, ///< Object is being asynchronously loaded.
 	Unreachable = 1 << 28, ///< Object is not reachable on the object graph.
 	RefCounted = 1 << 29, ///< Object currently has ref-counts associated with it.
 	RootSet = 1 << 30, ///< Object will not be garbage collected, even if unreferenced.
@@ -631,14 +632,18 @@ enum class EInternalObjectFlags : int32
 	// !!!!!!!!!!!!!!!!!!!!!!!!!!!!!!!!!!!!!!!!!!!!!!!!!!!!!!!!!!!!!!!!!!!!!!!!!!!!!!!!!!!!!!!!!!!!!!!!!!!!!!!!!!!!!!!!!!!!!!!!!!!!!!!!!!!
 	// DO NOT ADD composite flags to EInternalObjectFlags. Debugger visualisations have trouble displaying composite flag values as text.
 	// !!!!!!!!!!!!!!!!!!!!!!!!!!!!!!!!!!!!!!!!!!!!!!!!!!!!!!!!!!!!!!!!!!!!!!!!!!!!!!!!!!!!!!!!!!!!!!!!!!!!!!!!!!!!!!!!!!!!!!!!!!!!!!!!!!!
+
+	// Required for backward compatibility, will be deprecated soon.
+	AsyncLoading = AsyncLoadingPhase1 | AsyncLoadingPhase2
 };
 ENUM_CLASS_FLAGS(EInternalObjectFlags);
 
 //~ Make sure these macros are up to date!
-#define EInternalObjectFlags_GarbageCollectionKeepFlags (EInternalObjectFlags::Native | EInternalObjectFlags::Async | EInternalObjectFlags::AsyncLoading | EInternalObjectFlags::LoaderImport | EInternalObjectFlags::RefCounted)
-#define EInternalObjectFlags_AllFlags (EInternalObjectFlags::ReachabilityFlag0 | EInternalObjectFlags::ReachabilityFlag1 | EInternalObjectFlags::ReachabilityFlag2 | EInternalObjectFlags::LoaderImport | EInternalObjectFlags::Garbage | EInternalObjectFlags::ReachableInCluster | EInternalObjectFlags::ClusterRoot | EInternalObjectFlags::Native | EInternalObjectFlags::RefCounted | EInternalObjectFlags::Async | EInternalObjectFlags::AsyncLoading | EInternalObjectFlags::RootSet | EInternalObjectFlags::PendingConstruction | (EInternalObjectFlags::Unreachable))
+#define EInternalObjectFlags_GarbageCollectionKeepFlags (EInternalObjectFlags::Native | EInternalObjectFlags::Async | EInternalObjectFlags::AsyncLoadingPhase1 | EInternalObjectFlags::AsyncLoadingPhase2 | EInternalObjectFlags::LoaderImport | EInternalObjectFlags::RefCounted)
+#define EInternalObjectFlags_AllFlags (EInternalObjectFlags::ReachabilityFlag0 | EInternalObjectFlags::ReachabilityFlag1 | EInternalObjectFlags::ReachabilityFlag2 | EInternalObjectFlags::LoaderImport | EInternalObjectFlags::Garbage | EInternalObjectFlags::ReachableInCluster | EInternalObjectFlags::ClusterRoot | EInternalObjectFlags::Native | EInternalObjectFlags::RefCounted | EInternalObjectFlags::Async | EInternalObjectFlags::AsyncLoadingPhase1 | EInternalObjectFlags::AsyncLoadingPhase2 | EInternalObjectFlags::RootSet | EInternalObjectFlags::PendingConstruction | (EInternalObjectFlags::Unreachable))
 #define EInternalObjectFlags_RootFlags (EInternalObjectFlags::RootSet | EInternalObjectFlags_GarbageCollectionKeepFlags)
 #define EInternalObjectFlags_ReachabilityFlags (EInternalObjectFlags::ReachabilityFlag0 | EInternalObjectFlags::ReachabilityFlag1 | EInternalObjectFlags::ReachabilityFlag2 | EInternalObjectFlags::Unreachable)
+#define EInternalObjectFlags_AsyncLoading (EInternalObjectFlags::AsyncLoadingPhase1 | EInternalObjectFlags::AsyncLoadingPhase2)
 
 /**
  * Flags describing a UEnum 
diff --git a/Engine/Source/Runtime/CoreUObject/Public/UObject/ObjectVisibility.h b/Engine/Source/Runtime/CoreUObject/Public/UObject/ObjectVisibility.h
new file mode 100644
--- /dev/null
+++ b/Engine/Source/Runtime/CoreUObject/Public/UObject/ObjectVisibility.h
@@ -0,0 +1,12 @@
+// Copyright Epic Games, Inc. All Rights Reserved.
+
+#pragma once
+
+#include "UObject/ObjectMacros.h"
+
+namespace UE
+{
+
+	COREUOBJECT_API EInternalObjectFlags GetAsyncLoadingInternalFlagsExclusion();
+
+} // namespace UE
\ No newline at end of file
diff --git a/Engine/Source/Runtime/CoreUObject/Public/UObject/UObjectIterator.h b/Engine/Source/Runtime/CoreUObject/Public/UObject/UObjectIterator.h
--- a/Engine/Source/Runtime/CoreUObject/Public/UObject/UObjectIterator.h
+++ b/Engine/Source/Runtime/CoreUObject/Public/UObject/UObjectIterator.h
@@ -10,6 +10,7 @@
 #include "UObject/GarbageCollectionGlobals.h"
 #include "UObject/UObjectHash.h"
 #include "UObject/UObjectArray.h"
+#include "UObject/ObjectVisibility.h"
 #include "UObject/Object.h"
 #include "UObject/Class.h"
 
@@ -57,14 +58,7 @@ enum class EObjectIteratorThreadSafetyOptions : uint8
 
 inline EInternalObjectFlags GetObjectIteratorDefaultInternalExclusionFlags(EInternalObjectFlags InternalExclusionFlags)
 {
-	InternalExclusionFlags |= EInternalObjectFlags::Unreachable | EInternalObjectFlags::PendingConstruction;
-	if (!IsInAsyncLoadingThread())
-	{
-		// We don't want to return any objects that are currently being background loaded unless we're using the object iterator during async loading.
-		InternalExclusionFlags |= EInternalObjectFlags::AsyncLoading;
-	}
-
-	return InternalExclusionFlags;
+	return InternalExclusionFlags | EInternalObjectFlags::Unreachable | EInternalObjectFlags::PendingConstruction | UE::GetAsyncLoadingInternalFlagsExclusion();
 }
 
 /**
@@ -178,7 +172,7 @@ public:
 	{
 		//@warning: behavior is partially mirrored in UnObjGC.cpp. Make sure to adapt code there as well if you make changes below.
 		// verify that the async loading exclusion flag still matches (i.e. we didn't start/stop async loading within the scope of the iterator)
-		checkSlow(IsInAsyncLoadingThread() || int32(InternalExclusionFlags & EInternalObjectFlags::AsyncLoading));
+		checkSlow(IsInAsyncLoadingThread() || int32(InternalExclusionFlags & EInternalObjectFlags_AsyncLoading));
 
 		while (AdvanceIterator())
 		{
@@ -401,7 +395,7 @@ public:
 	void operator++()
 	{
 		// verify that the async loading exclusion flag still matches (i.e. we didn't start/stop async loading within the scope of the iterator)
-		checkSlow(IsInAsyncLoadingThread() || int32(InternalExclusionFlags & EInternalObjectFlags::AsyncLoading));
+		checkSlow(IsInAsyncLoadingThread() || int32(InternalExclusionFlags & EInternalObjectFlags_AsyncLoading));
 		while (AdvanceIterator())
 		{
 			if (!(*this)->HasAnyFlags(ExclusionFlags) && !(*this)->HasAnyInternalFlags(InternalExclusionFlags))
diff --git a/Engine/Source/Runtime/CoreUObject/Public/UObject/UObjectThreadContext.h b/Engine/Source/Runtime/CoreUObject/Public/UObject/UObjectThreadContext.h
--- a/Engine/Source/Runtime/CoreUObject/Public/UObject/UObjectThreadContext.h
+++ b/Engine/Source/Runtime/CoreUObject/Public/UObject/UObjectThreadContext.h
@@ -22,6 +22,7 @@
 #include "Templates/UnrealTemplate.h"
 #include "Trace/Detail/Channel.h"
 #include "UObject/PropertyPathName.h"
+#include "UObject/ObjectMacros.h"
 
 class FLinkerLoad;
 class FName;
@@ -99,6 +100,8 @@ public:
 	void* AsyncPackage;
 	/** Async package loader currently processing objects */
 	IAsyncPackageLoader* AsyncPackageLoader;
+	/** Async loading visibility filter */
+	EInternalObjectFlags AsyncVisibilityFilter = EInternalObjectFlags::None;
 
 #if !(UE_BUILD_SHIPPING || UE_BUILD_TEST)
 	/** Stack to ensure that PostInitProperties is routed through Super:: calls. **/
