[UObject]
 - Fix ObjectVisibility CVar name

#rnx

diff --git a/Engine/Source/Runtime/CoreUObject/Private/UObject/ObjectVisibility.cpp b/Engine/Source/Runtime/CoreUObject/Private/UObject/ObjectVisibility.cpp
--- a/Engine/Source/Runtime/CoreUObject/Private/UObject/ObjectVisibility.cpp
+++ b/Engine/Source/Runtime/CoreUObject/Private/UObject/ObjectVisibility.cpp
@@ -5,7 +5,7 @@
 #include "HAL/IConsoleManager.h"
 
 static bool GUseObjectVisibilityFilterForAsyncLoading = false;
-static FAutoConsoleVariableRef CVar(
+static FAutoConsoleVariableRef CVarUseObjectVisibilityFilterForAsyncLoading(
 	TEXT("s.UseObjectVisibilityFilterForAsyncLoading"),
 	GUseObjectVisibilityFilterForAsyncLoading,
 	TEXT("When active and supported by the current loader, will prevent objects still in the early loading phase from being discovered during postload\n")
