[ZenLoader]
 - Fix assert when nodes are fired from the game thread

#rnx

diff --git a/Engine/Source/Runtime/CoreUObject/Private/Serialization/AsyncLoading2.cpp b/Engine/Source/Runtime/CoreUObject/Private/Serialization/AsyncLoading2.cpp
--- a/Engine/Source/Runtime/CoreUObject/Private/Serialization/AsyncLoading2.cpp
+++ b/Engine/Source/Runtime/CoreUObject/Private/Serialization/AsyncLoading2.cpp
@@ -4608,6 +4608,9 @@ void FEventLoadNode2::Fire(FAsyncLoadingThreadState2* ThreadState)
 	//TRACE_CPUPROFILER_EVENT_SCOPE(Fire);
 	if (Spec->bExecuteImmediately && ThreadState)
 	{
+		// Firing a node can come from anywhere, needs to remove any current visibility filter
+		TGuardValue GuardVisibilityFilter(FUObjectThreadContext::Get().AsyncVisibilityFilter, EInternalObjectFlags::None);
+
 		EEventLoadNodeExecutionResult Result = Execute(*ThreadState);
 		check(Result == EEventLoadNodeExecutionResult::Complete);
 	}
