UObject::SetLinker: Add callstack information in error reported when detaching from existing linker while the object still needs loading.
#rnx

diff --git a/Engine/Source/Runtime/CoreUObject/Private/UObject/UObjectLinker.cpp b/Engine/Source/Runtime/CoreUObject/Private/UObject/UObjectLinker.cpp
--- a/Engine/Source/Runtime/CoreUObject/Private/UObject/UObjectLinker.cpp
+++ b/Engine/Source/Runtime/CoreUObject/Private/UObject/UObjectLinker.cpp
@@ -115,12 +115,16 @@ void UObject::SetLinker( FLinkerLoad* LinkerLoad, int32 LinkerIndex, bool bShoul
 	// Detach from existing linker.
 	if( Existing.Linker && bShouldDetachExisting )
 	{
-		UE_CLOG(HasAnyFlags(RF_NeedLoad|RF_NeedPostLoad), LogUObjectLinker, Error,
-			TEXT("Detaching from existing linker %s while object %s needs loading (%s). Setting linker to %s."),
-			*Existing.Linker->GetArchiveName(),
-			*GetFullName(),
-			*LexToString(GetFlags()),
-			LinkerLoad ? *LinkerLoad->GetDebugName() : TEXT("nullptr"));
+		if (HasAnyFlags(RF_NeedLoad | RF_NeedPostLoad))
+		{
+			UE_LOG(LogUObjectLinker, Error,
+				TEXT("Detaching from existing linker %s while object %s needs loading (%s). Setting linker to %s. See log for more information."),
+				*Existing.Linker->GetArchiveName(),
+				*GetFullName(),
+				*LexToString(GetFlags()),
+				LinkerLoad ? *LinkerLoad->GetDebugName() : TEXT("nullptr"));
+			FDebug::DumpStackTraceToLog(ELogVerbosity::Display);
+		}
 
 		FObjectExport& ExportObject = Existing.Linker->ExportMap[Existing.LinkerIndex];
 		checkf(ExportObject.Object != nullptr, TEXT("Expected ExportMap[%d].Object to not be null for this ('%s')"), Existing.LinkerIndex, *GetFName().ToString());
