Adding a reverse lookup map FRedirectCollector to accelerate queries from redirected asset to source asset. Also cleaning up some redundant #if WITH_EDITOR bits

diff --git a/Engine/Source/Runtime/CoreUObject/Private/Misc/RedirectCollector.cpp b/Engine/Source/Runtime/CoreUObject/Private/Misc/RedirectCollector.cpp
--- a/Engine/Source/Runtime/CoreUObject/Private/Misc/RedirectCollector.cpp
+++ b/Engine/Source/Runtime/CoreUObject/Private/Misc/RedirectCollector.cpp
@@ -17,6 +17,7 @@
 #include "UObject/SoftObjectPath.h"
 #include "HAL/IConsoleManager.h"
 #include "HAL/FileManager.h"
+#include "Misc/AutomationTest.h"
 
 #if WITH_EDITOR
 
@@ -149,10 +150,8 @@ void FRedirectCollector::ResolveAllSoftObjectPaths(FName FilterPackage)
 				UE_LOG(LogRedirectors, Verbose, TEXT("    Resolved to '%s'"), *Dest.ToString());
 				if (Dest.ToString() != ToLoad)
 				{
-					ObjectPathRedirectionMap.Add(ToLoadPath, Dest);
-#if WITH_EDITOR
+					AddObjectPathRedirectionInternal(ToLoadPath, Dest);
 					FCoreRedirects::RecordAddedObjectRedirector(ToLoadPath, Dest);
-#endif
 				}
 			}
 			else
@@ -260,13 +259,32 @@ void FRedirectCollector::GetAllSourcePathsForTargetPath(const FSoftObjectPath& T
 	FScopeLock ScopeLock(&CriticalSection);
 
 	OutSourcePaths.Reset();
-	for (const TPair<FSoftObjectPath, FSoftObjectPath>& Pair : ObjectPathRedirectionMap)
+	if (const ObjectPathSourcesArray* Sources = ObjectPathRedirectionReverseMap.Find(TargetPath))
 	{
-		if (Pair.Value == TargetPath)
+		OutSourcePaths = *Sources;
+	}
+}
+
+void FRedirectCollector::AddObjectPathRedirectionInternal(const FSoftObjectPath& Source, const FSoftObjectPath& Destination)
+{
+	if (FSoftObjectPath* ExistingDestination = ObjectPathRedirectionMap.Find(Source))
+	{
+		// We are replacing a redirect not adding one. 
+		// That means we need to remove an old reverse lookup and add a new one
+		if (ObjectPathSourcesArray* ReverseLookupArray = ObjectPathRedirectionReverseMap.Find(*ExistingDestination))
 		{
-			OutSourcePaths.Add(Pair.Key);
+			ReverseLookupArray->Remove(Source);
 		}
+
+		*ExistingDestination = Destination;
 	}
+	else
+	{
+		ObjectPathRedirectionMap.Add(Source, Destination);
+	}
+
+	ObjectPathSourcesArray& ReverseLookupArray = ObjectPathRedirectionReverseMap.FindOrAdd(Destination);
+	ReverseLookupArray.AddUnique(Source);
 }
 
 bool FRedirectCollector::ShouldTrackPackageReferenceTypes()
@@ -295,21 +313,31 @@ void FRedirectCollector::AddAssetPathRedirection(const FSoftObjectPath& Original
 	{
 		// If RedirectedPath points back to OriginalPath, remove that to avoid a circular reference
 		// This can happen when renaming assets in the editor but not actually dropping redirectors because it was new
-		FSoftObjectPath TargetPath;
-		bool bRemoved = ObjectPathRedirectionMap.RemoveAndCopyValue(RedirectedPath, TargetPath);
-#if WITH_EDITOR
-		if (bRemoved)
-		{
-			FCoreRedirects::RecordRemovedObjectRedirector(RedirectedPath, TargetPath);
-		}
-#endif
+		TryRemoveObjectPathRedirectionInternal(RedirectedPath);
 	}
 
 	// This replaces an existing mapping, can happen in the editor if things are renamed twice
-	ObjectPathRedirectionMap.Add(OriginalPath, RedirectedPath);
-#if WITH_EDITOR
+	AddObjectPathRedirectionInternal(OriginalPath, RedirectedPath);
 	FCoreRedirects::RecordAddedObjectRedirector(OriginalPath, RedirectedPath);
-#endif
+}
+
+bool FRedirectCollector::TryRemoveObjectPathRedirectionInternal(const FSoftObjectPath& Source)
+{
+	FSoftObjectPath Destination;
+	bool bRemoved = ObjectPathRedirectionMap.RemoveAndCopyValue(Source, Destination);
+	if (bRemoved)
+	{
+		FCoreRedirects::RecordRemovedObjectRedirector(Source, Destination);
+		if (ObjectPathSourcesArray* ReverseArray = ObjectPathRedirectionReverseMap.Find(Destination))
+		{
+			ReverseArray->RemoveSingleSwap(Source);
+			if (ReverseArray->Num() == 0)
+			{
+				ObjectPathRedirectionReverseMap.Remove(Destination);
+			}
+		}
+	}
+	return bRemoved;
 }
 
 void FRedirectCollector::AddAssetPathRedirection(FName OriginalPath, FName RedirectedPath)
@@ -323,15 +351,8 @@ void FRedirectCollector::RemoveAssetPathRedirection(const FSoftObjectPath& Origi
 {
 	FScopeLock ScopeLock(&CriticalSection);
 
-	FSoftObjectPath* Found = ObjectPathRedirectionMap.Find(OriginalPath);
-
-	if (ensureMsgf(Found, TEXT("Cannot remove redirection from %s, it was not registered"), *OriginalPath.ToString()))
-	{
-#if WITH_EDITOR
-		FCoreRedirects::RecordRemovedObjectRedirector(OriginalPath, *Found);
-#endif
-		ObjectPathRedirectionMap.Remove(OriginalPath);
-	}
+	bool bDidRemove = TryRemoveObjectPathRedirectionInternal(OriginalPath);
+	ensureMsgf(bDidRemove, TEXT("Cannot remove redirection from %s, it was not registered"), *OriginalPath.ToString());
 }
 
 void FRedirectCollector::RemoveAssetPathRedirection(FName OriginalPath)
@@ -391,5 +412,58 @@ FSoftObjectPath FRedirectCollector::GetAssetPathRedirection(const FSoftObjectPat
 }
 
 FRedirectCollector GRedirectCollector;
+#if WITH_AUTOMATION_WORKER
+
+IMPLEMENT_SIMPLE_AUTOMATION_TEST(FRedirectCollectorReverseLookupTest, "System.Core.Misc.RedirectCollector.ReverseLookup", EAutomationTestFlags_ApplicationContextMask | EAutomationTestFlags::EngineFilter)
+bool FRedirectCollectorReverseLookupTest::RunTest(const FString& Parameters)
+{
+	FScopeLock ScopeLock(&GRedirectCollector.CriticalSection);
+	bool bSuccess = true;
+
+	// Validate that every forward redirect has a corresponding reverse redirect
+	for (const TPair<FSoftObjectPath, FSoftObjectPath>& ForwardRedirect : GRedirectCollector.ObjectPathRedirectionMap)
+	{
+		bool bFoundReverseEntry = false;
+		if (const FRedirectCollector::ObjectPathSourcesArray* Sources =
+			GRedirectCollector.ObjectPathRedirectionReverseMap.Find(ForwardRedirect.Value))
+		{
+			if (Sources->Contains(ForwardRedirect.Key))
+			{
+				bFoundReverseEntry = true;
+			}
+		}
+
+		if (!bFoundReverseEntry)
+		{
+			bSuccess = false;
+			AddError(FString::Printf(TEXT("Failed to find matching reverse lookup for redirect %s --> %s"),
+				*ForwardRedirect.Key.ToString(), *ForwardRedirect.Value.ToString()));
+		}
+	}
 
+	// Validate that every reverse redirect has a corresponding forward redirect
+	for (const TPair<FSoftObjectPath, FRedirectCollector::ObjectPathSourcesArray>& ReverseRedirectList : GRedirectCollector.ObjectPathRedirectionReverseMap)
+	{
+		for (const FSoftObjectPath& Source : ReverseRedirectList.Value)
+		{
+			bool bFoundForwardEntry = false;
+			if (const FSoftObjectPath* Target = GRedirectCollector.ObjectPathRedirectionMap.Find(Source))
+			{
+				if (*Target == ReverseRedirectList.Key)
+				{
+					bFoundForwardEntry = true;
+				}
+			}
+
+			if (!bFoundForwardEntry)
+			{
+				bSuccess = false;
+				AddError(FString::Printf(TEXT("Failed to find matching reverse lookup for redirect %s --> %s"),
+					*Source.ToString(), *ReverseRedirectList.Key.ToString()));
+			}
+		}
+	}
+	return bSuccess;
+}
+#endif // WITH_AUTOMATION_WORKER
 #endif // WITH_EDITOR
diff --git a/Engine/Source/Runtime/CoreUObject/Public/Misc/RedirectCollector.h b/Engine/Source/Runtime/CoreUObject/Public/Misc/RedirectCollector.h
--- a/Engine/Source/Runtime/CoreUObject/Public/Misc/RedirectCollector.h
+++ b/Engine/Source/Runtime/CoreUObject/Public/Misc/RedirectCollector.h
@@ -151,7 +151,7 @@ public:
 		return ObjectPathRedirectionMap;
 	}
 	
-	/** Returns the set of paths, if any, that are redirected TO the provided path. This performs a relatively slow linear search. */
+	/** Returns the set of paths, if any, that are redirected TO the provided path.*/
 	COREUOBJECT_API void GetAllSourcePathsForTargetPath(const FSoftObjectPath& TargetPath, TArray<FSoftObjectPath>& OutSourcePaths) const;
 
 	/** Used with GetObjectPathRedirectionMapUnderLock with code like:
@@ -164,6 +164,12 @@ public:
 	}
 
 private:
+	
+	/** Handles adding forward and reverse map entries. Must be called while holding the critical section */
+	void AddObjectPathRedirectionInternal(const FSoftObjectPath& Source, const FSoftObjectPath& Destination);
+
+	/** Handles removing forward and reverse map entries. Must be called while holding the critical section */
+	bool TryRemoveObjectPathRedirectionInternal(const FSoftObjectPath& Source);
 
 	/** A map of assets referenced by soft object paths, with the key being the package with the reference */
 	typedef TSet<FSoftObjectPathProperty> FSoftObjectPathPropertySet;
@@ -180,6 +186,10 @@ private:
 	/** When saving, apply this remapping to all soft object paths */
 	TMap<FSoftObjectPath, FSoftObjectPath> ObjectPathRedirectionMap;
 
+	/** A reverse lookup map for use with GetAllSourcePathsForTargetPath */
+	typedef TArray<FSoftObjectPath, TInlineAllocator<1>> ObjectPathSourcesArray;
+	TMap<FSoftObjectPath, TArray<FSoftObjectPath, TInlineAllocator<1>>> ObjectPathRedirectionReverseMap;
+
 	/** For ObjectPathRedirectionMap map */
 	mutable FCriticalSection CriticalSection;
 
@@ -190,6 +200,8 @@ private:
 		Enabled,
 	};
 	ETrackingReferenceTypesState TrackingReferenceTypesState;
+
+	friend class FRedirectCollectorReverseLookupTest;
 };
 
 // global redirect collector callback structure
