[BugFix] Made OverriddenPropertySet hash it's properties by type and name rather than just name

diff --git a/Engine/Source/Runtime/CoreUObject/Private/UObject/OverriddenPropertySet.cpp b/Engine/Source/Runtime/CoreUObject/Private/UObject/OverriddenPropertySet.cpp
--- a/Engine/Source/Runtime/CoreUObject/Private/UObject/OverriddenPropertySet.cpp
+++ b/Engine/Source/Runtime/CoreUObject/Private/UObject/OverriddenPropertySet.cpp
@@ -85,6 +85,52 @@ FEnableOverridableSerializationScope::~FEnableOverridableSerializationScope()
 	}
 }
 
+FOverriddenPropertyNodeID::FOverriddenPropertyNodeID(const FProperty* Property)
+	: Object(nullptr)
+{
+	if (Property)
+	{
+		// append typename to the end of the property ID
+		UE::FPropertyTypeNameBuilder TypeNameBuilder;
+		Property->SaveTypeName(TypeNameBuilder);
+		UE::FPropertyTypeName TypeName = TypeNameBuilder.Build();
+		TStringBuilder<256> StringBuilder;
+		StringBuilder << Property->GetFName();
+		StringBuilder << " - ";
+		StringBuilder << TypeName;
+		Path = FName(StringBuilder.ToView());
+	}
+}
+
+FOverriddenPropertyNodeID FOverriddenPropertyNodeID::RootNodeId()
+{
+	FOverriddenPropertyNodeID Result;
+	Result.Path = FName(TEXT("root"));
+	return Result;
+}
+
+FOverriddenPropertyNodeID FOverriddenPropertyNodeID::ForMapKey(const FProperty* KeyProperty, const void* KeyData)
+{
+	if (const FObjectProperty* KeyObjectProperty = CastField<FObjectProperty>(KeyProperty))
+	{
+		if (const UObject* Object = KeyObjectProperty->GetObjectPropertyValue(KeyData))
+		{
+			return FOverriddenPropertyNodeID(*Object);
+		}
+	}
+	else
+	{
+		FString KeyString;
+		KeyProperty->ExportTextItem_Direct(KeyString, KeyData, /*DefaultValue*/nullptr, /*Parent*/nullptr, PPF_None);
+		FOverriddenPropertyNodeID Result;
+		Result.Path = FName(KeyString);
+		return Result;
+	}
+		
+	checkf(false, TEXT("This case is not handled"))
+	return FOverriddenPropertyNodeID();
+}
+
 //----------------------------------------------------------------------//
 // FOverriddenPropertyNodeID
 //----------------------------------------------------------------------//
@@ -150,8 +196,7 @@ EOverriddenPropertyOperation FOverriddenPropertySet::GetOverriddenPropertyOperat
 		const FOverriddenPropertyNode* CurrentOverriddenPropertyNode = nullptr;
 		if (OverriddenPropertyNode)
 		{
-			const FName CurrentPropID = CurrentProperty->GetFName();
-			if (const FOverriddenPropertyNodeID* CurrentPropKey = OverriddenPropertyNode->SubPropertyNodeKeys.Find(CurrentPropID))
+			if (const FOverriddenPropertyNodeID* CurrentPropKey = OverriddenPropertyNode->SubPropertyNodeKeys.Find(CurrentProperty))
 			{
 				CurrentOverriddenPropertyNode = OverriddenPropertyNodes.FindByHash(GetTypeHash(*CurrentPropKey), *CurrentPropKey);
 				checkf(CurrentOverriddenPropertyNode, TEXT("Expecting a node"));
@@ -236,19 +281,7 @@ EOverriddenPropertyOperation FOverriddenPropertySet::GetOverriddenPropertyOperat
 				{
 					// Caller wants to know about any override state on the reference of the map pair itself
 					checkf(MapProperty->KeyProp, TEXT("Expecting a key type for Maps"));
-					FObjectProperty* KeyObjectProperty = CastField<FObjectProperty>(MapProperty->KeyProp);
-
-					FOverriddenPropertyNodeID OverriddenKeyID;
-					if (const UObject* OverriddenKeyObject = KeyObjectProperty ? KeyObjectProperty->GetObjectPropertyValue(MapHelper.GetKeyPtr(InternalMapIndex)) : nullptr)
-					{
-						OverriddenKeyID = FOverriddenPropertyNodeID(*OverriddenKeyObject);
-					}
-					else
-					{
-						FString OverriddenKey;
-						MapProperty->KeyProp->ExportTextItem_Direct(OverriddenKey, MapHelper.GetKeyPtr(InternalMapIndex), nullptr, nullptr, PPF_None);
-						OverriddenKeyID = FName(OverriddenKey);
-					}
+					FOverriddenPropertyNodeID OverriddenKeyID = FOverriddenPropertyNodeID::ForMapKey(MapProperty->KeyProp, MapHelper.GetKeyPtr(InternalMapIndex));
 
 					if (const FOverriddenPropertyNodeID* CurrentPropKey = CurrentOverriddenPropertyNode->SubPropertyNodeKeys.Find(OverriddenKeyID))
 					{
@@ -296,8 +329,7 @@ bool FOverriddenPropertySet::ClearOverriddenProperty(FOverriddenPropertyNode& Pa
 		FOverriddenPropertyNode* CurrentOverriddenPropertyNode = nullptr;
 		if (OverriddenPropertyNode)
 		{
-			const FName CurrentPropID = CurrentProperty->GetFName();
-			if (const FOverriddenPropertyNodeID* CurrentPropKey = OverriddenPropertyNode->SubPropertyNodeKeys.Find(CurrentPropID))
+			if (const FOverriddenPropertyNodeID* CurrentPropKey = OverriddenPropertyNode->SubPropertyNodeKeys.Find(CurrentProperty))
 			{
 				CurrentOverriddenPropertyNode = OverriddenPropertyNodes.FindByHash(GetTypeHash(*CurrentPropKey), *CurrentPropKey);
 				checkf(CurrentOverriddenPropertyNode, TEXT("Expecting a node"));
@@ -413,20 +445,7 @@ bool FOverriddenPropertySet::ClearOverriddenProperty(FOverriddenPropertyNode& Pa
 			else if (MapHelper.IsValidIndex(InternalMapIndex) && CurrentOverriddenPropertyNode)
 			{
 				checkf(MapProperty->KeyProp, TEXT("Expecting a key type for Maps"));
-				FObjectProperty* KeyObjectProperty = CastField<FObjectProperty>(MapProperty->KeyProp);
-
-				// Calculate the node from the key
-				FOverriddenPropertyNodeID OverriddenKeyID;
-				if (const UObject* OverriddenKeyObject = KeyObjectProperty ? KeyObjectProperty->GetObjectPropertyValue(MapHelper.GetKeyPtr(InternalMapIndex)) : nullptr)
-				{
-					OverriddenKeyID = FOverriddenPropertyNodeID(*OverriddenKeyObject);
-				}
-				else
-				{
-					FString OverriddenKey;
-					MapProperty->KeyProp->ExportTextItem_Direct(OverriddenKey, MapHelper.GetKeyPtr(InternalMapIndex), nullptr, nullptr, PPF_None);
-					OverriddenKeyID = FName(OverriddenKey);
-				}
+				FOverriddenPropertyNodeID OverriddenKeyID = FOverriddenPropertyNodeID::ForMapKey(MapProperty->KeyProp, MapHelper.GetKeyPtr(InternalMapIndex));
 
 				FOverriddenPropertyNodeID CurrentPropKey;
 				if (CurrentOverriddenPropertyNode->SubPropertyNodeKeys.RemoveAndCopyValue(OverriddenKeyID, CurrentPropKey))
@@ -522,13 +541,12 @@ void FOverriddenPropertySet::NotifyPropertyChange(FOverriddenPropertyNode* Paren
 	const FProperty* Property = PropertyNode->GetValue();
 	checkf(Property, TEXT("Expecting a valid property"));
 
-	const FOverriddenPropertyNodeID PropID(Property->GetFName());
 	const void* SubValuePtr = Property->ContainerPtrToValuePtr<void>(Data, 0); //@todo support static arrays
 
 	FOverriddenPropertyNode* SubPropertyNode = nullptr;
 	if (ParentPropertyNode)
 	{
-		FOverriddenPropertyNode& SubPropertyNodeRef = FindOrAddNode(*ParentPropertyNode, PropID);
+		FOverriddenPropertyNode& SubPropertyNodeRef = FindOrAddNode(*ParentPropertyNode, Property);
 		SubPropertyNode = SubPropertyNodeRef.Operation != EOverriddenPropertyOperation::Replace ? &SubPropertyNodeRef : nullptr;
 	}
 
@@ -788,17 +806,7 @@ void FOverriddenPropertySet::NotifyPropertyChange(FOverriddenPropertyNode* Paren
 				{
 					if(SubPropertyNode)
 					{
-						FOverriddenPropertyNodeID OverriddenKeyID;
-						if (const UObject* OverriddenKeyObject = KeyObjectProperty ? KeyObjectProperty->GetObjectPropertyValue(MapHelper.GetKeyPtr(It.GetInternalIndex())) : nullptr)
-						{
-							OverriddenKeyID = FOverriddenPropertyNodeID(*OverriddenKeyObject);
-						}
-						else
-						{
-							FString OverriddenKey;
-							MapProperty->KeyProp->ExportTextItem_Direct(OverriddenKey, MapHelper.GetKeyPtr(It.GetInternalIndex()), nullptr, nullptr, PPF_None);
-							OverriddenKeyID = FName(OverriddenKey);
-						}
+						FOverriddenPropertyNodeID OverriddenKeyID = FOverriddenPropertyNodeID::ForMapKey(MapProperty->KeyProp, MapHelper.GetKeyPtr(It.GetInternalIndex()));
 						FOverriddenPropertyNode& OverriddenKeyNode = FindOrAddNode(*SubPropertyNode, OverriddenKeyID);
 						OverriddenKeyNode.Operation = EOverriddenPropertyOperation::Replace;
 					}
@@ -821,18 +829,7 @@ void FOverriddenPropertySet::NotifyPropertyChange(FOverriddenPropertyNode* Paren
 
 				if(SubPropertyNode)
 				{
-					FOverriddenPropertyNodeID AddedKeyID;
-					if (const UObject* AddedKeyObject = KeyObjectProperty ? KeyObjectProperty->GetObjectPropertyValue(MapHelper.GetKeyPtr(InternalMapIndex)) : nullptr)
-					{
-						AddedKeyID = FOverriddenPropertyNodeID(*AddedKeyObject);
-					}
-					else
-					{
-						FString AddedKey;
-						MapProperty->KeyProp->ExportTextItem_Direct(AddedKey, MapHelper.GetKeyPtr(InternalMapIndex), nullptr, nullptr, PPF_None);
-						AddedKeyID = FName(AddedKey);
-					}
-
+					FOverriddenPropertyNodeID AddedKeyID = FOverriddenPropertyNodeID::ForMapKey(MapProperty->KeyProp, MapHelper.GetKeyPtr(InternalMapIndex));
 					FOverriddenPropertyNode& AddedKeyNode = FindOrAddNode(*SubPropertyNode, AddedKeyID);
 					AddedKeyNode.Operation = EOverriddenPropertyOperation::Add;
 				}
@@ -844,17 +841,7 @@ void FOverriddenPropertySet::NotifyPropertyChange(FOverriddenPropertyNode* Paren
 
 				if(SubPropertyNode)
 				{
-					FOverriddenPropertyNodeID RemovedKeyID;
-					if (const UObject* RemovedKeyObject = KeyObjectProperty ? KeyObjectProperty->GetObjectPropertyValue(PreEditMapHelper.GetKeyPtr(InternalPreEditMapIndex)) : nullptr)
-					{
-						RemovedKeyID = FOverriddenPropertyNodeID(*RemovedKeyObject);
-					}
-					else
-					{
-						FString RemovedKey;
-						MapProperty->KeyProp->ExportTextItem_Direct(RemovedKey, PreEditMapHelper.GetKeyPtr(InternalPreEditMapIndex), nullptr, nullptr, PPF_None);
-						RemovedKeyID = FName(RemovedKey);
-					}
+					FOverriddenPropertyNodeID RemovedKeyID = FOverriddenPropertyNodeID::ForMapKey(MapProperty->KeyProp, PreEditMapHelper.GetKeyPtr(InternalPreEditMapIndex));
 					FOverriddenPropertyNode& RemovedKeyNode = FindOrAddNode(*SubPropertyNode, RemovedKeyID);
 					if (RemovedKeyNode.Operation == EOverriddenPropertyOperation::Add)
 					{
@@ -1109,8 +1096,7 @@ EOverriddenPropertyOperation FOverriddenPropertySet::GetOverriddenPropertyOperat
 	while (PropertyIterator && OverriddenPropertyNode && OverriddenPropertyNode->Operation != EOverriddenPropertyOperation::Replace)
 	{
 		const FProperty* CurrentProperty = (*PropertyIterator);
-		const FName CurrentPropID = CurrentProperty->GetFName();
-		if (const FOverriddenPropertyNodeID* CurrentPropKey = OverriddenPropertyNode->SubPropertyNodeKeys.Find(CurrentPropID))
+		if (const FOverriddenPropertyNodeID* CurrentPropKey = OverriddenPropertyNode->SubPropertyNodeKeys.Find(CurrentProperty))
 		{
 			OverriddenPropertyNode = OverriddenPropertyNodes.FindByHash(GetTypeHash(*CurrentPropKey), *CurrentPropKey);
 			checkf(OverriddenPropertyNode, TEXT("Expecting a node"));
@@ -1148,8 +1134,7 @@ FOverriddenPropertyNode* FOverriddenPropertySet::SetOverriddenPropertyOperation(
 	while (PropertyIterator && OverriddenPropertyNode->Operation != EOverriddenPropertyOperation::Replace)
 	{
 		const FProperty* CurrentProperty = (*PropertyIterator);
-		const FName CurrentPropID = CurrentProperty->GetFName();
-		OverriddenPropertyNode = &FindOrAddNode(*OverriddenPropertyNode, CurrentPropID);
+		OverriddenPropertyNode = &FindOrAddNode(*OverriddenPropertyNode, CurrentProperty);
 		++PropertyIterator;
 	}
 
@@ -1227,8 +1212,7 @@ const FOverriddenPropertyNode* FOverriddenPropertySet::GetOverriddenPropertyNode
 	while (PropertyIterator && OverriddenPropertyNode)
 	{
 		const FProperty* CurrentProperty = (*PropertyIterator);
-		const FName CurrentPropID = CurrentProperty->GetFName();
-		if (const FOverriddenPropertyNodeID* CurrentPropKey = OverriddenPropertyNode->SubPropertyNodeKeys.Find(CurrentPropID))
+		if (const FOverriddenPropertyNodeID* CurrentPropKey = OverriddenPropertyNode->SubPropertyNodeKeys.Find(CurrentProperty))
 		{
 			OverriddenPropertyNode = OverriddenPropertyNodes.FindByHash(GetTypeHash(*CurrentPropKey), *CurrentPropKey);
 			checkf(OverriddenPropertyNode, TEXT("Expecting a node"));
diff --git a/Engine/Source/Runtime/CoreUObject/Private/UObject/PropertyMap.cpp b/Engine/Source/Runtime/CoreUObject/Private/UObject/PropertyMap.cpp
--- a/Engine/Source/Runtime/CoreUObject/Private/UObject/PropertyMap.cpp
+++ b/Engine/Source/Runtime/CoreUObject/Private/UObject/PropertyMap.cpp
@@ -310,22 +310,7 @@ void FMapProperty::SerializeItem(FStructuredArchive::FSlot Slot, void* Value, co
 
 			auto GetIDFromKey = [&](uint8* KeyData) -> FOverriddenPropertyNodeID
 			{
-				if (FObjectProperty* KeyObjectProperty = CastField<FObjectProperty>(KeyProp))
-				{
-					if (const UObject* Object = KeyObjectProperty->GetObjectPropertyValue(KeyData))
-					{
-						return FOverriddenPropertyNodeID(*Object);
-					}
-				}
-				else
-				{
-					FString KeyString;
-					KeyProp->ExportTextItem_Direct(KeyString, KeyData, /*DefaultValue*/nullptr, /*Parent*/nullptr, PPF_None);
-					return FOverriddenPropertyNodeID(FName(KeyString));
-				}
-		
-				checkf(false, TEXT("This case is not handled"))
-				return FOverriddenPropertyNodeID();
+				return FOverriddenPropertyNodeID::ForMapKey(KeyProp, KeyData);
 			};
 
 			if (UnderlyingArchive.IsLoading())
diff --git a/Engine/Source/Runtime/CoreUObject/Public/UObject/OverriddenPropertySet.h b/Engine/Source/Runtime/CoreUObject/Public/UObject/OverriddenPropertySet.h
--- a/Engine/Source/Runtime/CoreUObject/Public/UObject/OverriddenPropertySet.h
+++ b/Engine/Source/Runtime/CoreUObject/Public/UObject/OverriddenPropertySet.h
@@ -109,10 +109,7 @@ struct FOverriddenPropertyNodeID
 {
 	GENERATED_BODY()
 
-	FOverriddenPropertyNodeID(FName InPath = NAME_None)
-		: Path(InPath)
-		, Object(nullptr)
-	{}
+	FOverriddenPropertyNodeID(const FProperty* Property = nullptr);
 
 	FOverriddenPropertyNodeID(const UObject& InObject)
 		: Object(&InObject)
@@ -123,6 +120,9 @@ struct FOverriddenPropertyNodeID
 		Path = *(FString::Printf(TEXT("%d%d"), ObjectIndex, GUObjectArray.AllocateSerialNumber(ObjectIndex)));
 	}
 
+	static FOverriddenPropertyNodeID RootNodeId();
+	static FOverriddenPropertyNodeID ForMapKey(const FProperty* KeyProperty, const void* KeyData);
+
 	bool operator==(const FOverriddenPropertyNodeID& Other) const
 	{
 		return Path == Other.Path || (Object && Other.Object && Object == Other.Object);
@@ -329,7 +329,7 @@ private:
 	UPROPERTY()
 	TSet<FOverriddenPropertyNode> OverriddenPropertyNodes;
 
-	static inline FOverriddenPropertyNodeID RootNodeID = FOverriddenPropertyNodeID(FName(TEXT("root")));
+	static inline FOverriddenPropertyNodeID RootNodeID = FOverriddenPropertyNodeID::RootNodeId();
 
 	/** Archetype is cache when it is or its archetype is about to be replaced **/
 	TObjectPtr<UObject> CachedArchetype = nullptr;
