#UE Add a verbose log when extra files are written during cook, and moved a couple other logs from Verbose to VeryVerbose since they happen for every package

diff --git a/Engine/Source/Runtime/CoreUObject/Private/UObject/SavePackage2.cpp b/Engine/Source/Runtime/CoreUObject/Private/UObject/SavePackage2.cpp
--- a/Engine/Source/Runtime/CoreUObject/Private/UObject/SavePackage2.cpp
+++ b/Engine/Source/Runtime/CoreUObject/Private/UObject/SavePackage2.cpp
@@ -1737,7 +1737,7 @@ void SavePreloadDependencies(FStructuredArchive::FRecord& StructuredArchiveRoot,
 				AddArcForDepChecking(false, Index, false);
 			}
 		}
-		UE_LOG(LogSavePackage, Verbose, TEXT("Saved %d dependencies for %d exports."), Linker->Summary.PreloadDependencyCount, Linker->ExportMap.Num());
+		UE_LOG(LogSavePackage, VeryVerbose, TEXT("Saved %d dependencies for %d exports."), Linker->Summary.PreloadDependencyCount, Linker->ExportMap.Num());
 	}
 }
 
@@ -2379,6 +2379,8 @@ ESavePackageResult WriteAdditionalExportFiles(FSaveContext& SaveContext)
 			FileInfo.PackageName = SaveContext.GetPackage()->GetFName();
 			FileInfo.Filename = *Writer.GetArchiveName();
 
+			UE_LOG(LogSavePackage, Verbose, TEXT("SavePackage writing additional file '%s' due to package '%s'."), *FileInfo.Filename, *FileInfo.PackageName.ToString());
+
 			FIoBuffer FileData(FIoBuffer::AssumeOwnership, Writer.ReleaseOwnership(), Size);
 
 			// This might not actuallly write the file, but instead add it to a queue to write later.
@@ -2519,7 +2521,7 @@ ESavePackageResult FinalizeFile(FStructuredArchive::FRecord& StructuredArchiveRo
 	{
 		bool bIsOptionalRealm = SaveContext.GetCurrentHarvestingRealm() == ESaveRealm::Optional;
 		FLinkerSave* Linker = SaveContext.GetLinker();
-		UE_LOG(LogSavePackage, Verbose, TEXT("Async saving from memory to '%s'"), SaveContext.GetFilename());
+		UE_LOG(LogSavePackage, VeryVerbose, TEXT("Async saving from memory to '%s'"), SaveContext.GetFilename());
 		FLargeMemoryWriter* Writer = static_cast<FLargeMemoryWriter*>(Linker->Saver);
 
 		if (PackageWriter)
