Fix for verse non nullable proxy variables(properties with getter/setter) that are being copied inside CPUFO and being assign dummy objects, which then fails the cooking process because they are referencing private objects in other packages which are those dummy objects..

diff --git a/Engine/Source/Programs/Shared/EpicGames.Core/UnrealEngineTypes.cs b/Engine/Source/Programs/Shared/EpicGames.Core/UnrealEngineTypes.cs
--- a/Engine/Source/Programs/Shared/EpicGames.Core/UnrealEngineTypes.cs
+++ b/Engine/Source/Programs/Shared/EpicGames.Core/UnrealEngineTypes.cs
@@ -1225,6 +1225,11 @@ namespace EpicGames.Core
 		/// </summary>
 		ExperimentalAlwaysOverriden = 0x0400000000000000,
 
+		/// <summary>
+		/// ****Experimental*** Property should never be overridden when using overridable serialization
+		/// </summary>
+		ExperimentalNeverOverriden = 0x0800000000000000,
+
 		/// <summary>
 		/// All Native Access Specifier flags
 		/// </summary>
diff --git a/Engine/Source/Programs/Shared/EpicGames.UHT/Specifiers/UhtPropertyMemberSpecifiers.cs b/Engine/Source/Programs/Shared/EpicGames.UHT/Specifiers/UhtPropertyMemberSpecifiers.cs
--- a/Engine/Source/Programs/Shared/EpicGames.UHT/Specifiers/UhtPropertyMemberSpecifiers.cs
+++ b/Engine/Source/Programs/Shared/EpicGames.UHT/Specifiers/UhtPropertyMemberSpecifiers.cs
@@ -396,6 +396,13 @@ namespace EpicGames.UHT.Parsers
 			context.PropertySettings.PropertyFlags |= EPropertyFlags.ExperimentalAlwaysOverriden;
 		}
 
+		[UhtSpecifier(Extends = UhtTableNames.PropertyMember, ValueType = UhtSpecifierValueType.Legacy)]
+		private static void ExperimentalNeverOverridenSpecifier(UhtSpecifierContext specifierContext)
+		{
+			UhtPropertySpecifierContext context = (UhtPropertySpecifierContext)specifierContext;
+			context.PropertySettings.PropertyFlags |= EPropertyFlags.ExperimentalNeverOverriden;
+		}
+
 		[UhtSpecifier(Extends = UhtTableNames.PropertyMember, ValueType = UhtSpecifierValueType.OptionalString)]
 		private static void GetterSpecifier(UhtSpecifierContext specifierContext, StringView? value)
 		{
diff --git a/Engine/Source/Runtime/CoreUObject/Private/UObject/OverriddenPropertySet.cpp b/Engine/Source/Runtime/CoreUObject/Private/UObject/OverriddenPropertySet.cpp
--- a/Engine/Source/Runtime/CoreUObject/Private/UObject/OverriddenPropertySet.cpp
+++ b/Engine/Source/Runtime/CoreUObject/Private/UObject/OverriddenPropertySet.cpp
@@ -27,6 +27,12 @@ EOverriddenPropertyOperation FOverridableSerializationLogic::GetOverriddenProper
 	checkf(bUseOverridableSerialization, TEXT("Nobody should use this method if it is not setup to use overridable serialization"));
 
 	const FArchiveSerializedPropertyChain* CurrentPropertyChain = Ar.GetSerializedPropertyChain();
+	const FProperty* CurrentProperty = Property ? Property : (CurrentPropertyChain ? CurrentPropertyChain->GetPropertyFromStack(0) : nullptr);
+	if (CurrentProperty && CurrentProperty->HasAnyPropertyFlags(CPF_ExperimentalNeverOverriden))
+	{
+		return EOverriddenPropertyOperation::None;
+	}
+
 	const EOverriddenPropertyOperation OverriddenOperation = OverriddenProperties ? OverriddenProperties->GetOverriddenPropertyOperation(CurrentPropertyChain, Property) : EOverriddenPropertyOperation::None;
 	if (OverriddenOperation != EOverriddenPropertyOperation::None)
 	{
@@ -34,7 +40,7 @@ EOverriddenPropertyOperation FOverridableSerializationLogic::GetOverriddenProper
 	}
 
 	// It does not mean that if we have no record of an overriden operation that a subobject might have one, need to traverse all possible subobjects. 
-	if (const FProperty* CurrentProperty = Property ? Property : (CurrentPropertyChain ? CurrentPropertyChain->GetPropertyFromStack(0) : nullptr) )
+	if (CurrentProperty)
 	{
 		if (CurrentProperty->HasAnyPropertyFlags(CPF_ExperimentalAlwaysOverriden))
 		{
diff --git a/Engine/Source/Runtime/CoreUObject/Public/UObject/ObjectMacros.h b/Engine/Source/Runtime/CoreUObject/Public/UObject/ObjectMacros.h
--- a/Engine/Source/Runtime/CoreUObject/Public/UObject/ObjectMacros.h
+++ b/Engine/Source/Runtime/CoreUObject/Public/UObject/ObjectMacros.h
@@ -458,6 +458,7 @@ enum EPropertyFlags : uint64
 	CPF_TObjectPtr						= 0x0100000000000000,	///< Property is a TObjectPtr<T> instead of a USomething*. Need to differentiate between TObjectclassOf and TObjectPtr
 	CPF_ExperimentalOverridableLogic	= 0x0200000000000000,	///< ****Experimental*** Property will use different logic to serialize knowing what changes are done against its default use the overridable information provided by the overridable manager on the object
 	CPF_ExperimentalAlwaysOverriden		= 0x0400000000000000,	///< ****Experimental*** Property should never inherit from the parent when using overridable serialization
+	CPF_ExperimentalNeverOverriden		= 0x0800000000000000,	///< ****Experimental*** Property should never be overridden when using overridable serialization
 };
 
 /** All Native Access Specifier flags */
