Made FSoftObjectPath::LoadAsync const

diff --git a/Engine/Source/Runtime/CoreUObject/Private/UObject/SoftObjectPath.cpp b/Engine/Source/Runtime/CoreUObject/Private/UObject/SoftObjectPath.cpp
--- a/Engine/Source/Runtime/CoreUObject/Private/UObject/SoftObjectPath.cpp
+++ b/Engine/Source/Runtime/CoreUObject/Private/UObject/SoftObjectPath.cpp
@@ -609,7 +609,7 @@ UObject* FSoftObjectPath::TryLoad(FUObjectSerializeContext* InLoadContext) const
 	return LoadedObject;
 }
 
-int32 FSoftObjectPath::LoadAsync(FLoadSoftObjectPathAsyncDelegate InCompletionDelegate, FLoadAssetAsyncOptionalParams InOptionalParams)
+int32 FSoftObjectPath::LoadAsync(FLoadSoftObjectPathAsyncDelegate InCompletionDelegate, FLoadAssetAsyncOptionalParams InOptionalParams) const
 {
 	FSoftObjectPath RequestedPath = *this;
 	FSoftObjectPath PathToLoad = RequestedPath;
diff --git a/Engine/Source/Runtime/CoreUObject/Public/UObject/SoftObjectPath.h b/Engine/Source/Runtime/CoreUObject/Public/UObject/SoftObjectPath.h
--- a/Engine/Source/Runtime/CoreUObject/Public/UObject/SoftObjectPath.h
+++ b/Engine/Source/Runtime/CoreUObject/Public/UObject/SoftObjectPath.h
@@ -239,7 +239,7 @@ PRAGMA_ENABLE_DEPRECATION_WARNINGS;
 	 * @param	InOptionalParams		Optional parameters for async loading the asset
 	 * @return Unique ID associated with this load request (the same object or package can be associated with multiple IDs).
 	 */
-	COREUOBJECT_API int32 LoadAsync(FLoadSoftObjectPathAsyncDelegate InCompletionDelegate, FLoadAssetAsyncOptionalParams InOptionalParams = FLoadAssetAsyncOptionalParams());
+	COREUOBJECT_API int32 LoadAsync(FLoadSoftObjectPathAsyncDelegate InCompletionDelegate, FLoadAssetAsyncOptionalParams InOptionalParams = FLoadAssetAsyncOptionalParams()) const;
 
 	/**
 	 * Attempts to find a currently loaded object that matches this path
