Clean up FBulkData deprecations

#rnx

- Removed all code marked as deprecated for UE 5.2 and earlier.
- This removes the old async streaming api, which allowed the user to request that the bulkdata payload start streaming from disk and be stored internally by the bulkdata object to be used at some point in the future via a call to FBulkData::StartAsyncLoading which has been deprecated for a number of engine releases.
- Seems we missed deprecating ::IsAsyncLoadingComplete and it's associated flag 'EBulkDataFlags::BULKDATA_HasAsyncReadPending' when deprecating ::StartAsyncLoad so marked it for deprecation now along with code comments explaining why it's not used anymore.
- Since we are removing ::StartAsyncLoad, we can completely remove UE::BulkData::Private::StartAsyncLoad as it was that only place it was being used.
- ::FlushAsyncLoading only provided functionality if 'BULKDATA_HasAsyncReadPending' was set, which would only occur if ::StartAsyncLoad was called. Since we are not removing ::StartAsyncLoad we can remove the implementation of ::FlushAsyncLoading and deprecate it.
-- In turn this means we can remove UE::BulkData::Private::FlushAsyncLoad entirely as it was only being called by ::FlushAsyncLoading.
- With both  UE::BulkData::Private::StartAsyncLoad and UE::BulkData::Private::FlushAsyncLoad removed, we can remove FAsyncBulkDataRequests aswell as there is no longer any way for people to use it.
- Removed calls to ::IsAsyncLoadingComplete from the rest of the engine, it will always be true at this point.

diff --git a/Engine/Source/Runtime/CoreUObject/Private/Serialization/BulkData.cpp b/Engine/Source/Runtime/CoreUObject/Private/Serialization/BulkData.cpp
--- a/Engine/Source/Runtime/CoreUObject/Private/Serialization/BulkData.cpp
+++ b/Engine/Source/Runtime/CoreUObject/Private/Serialization/BulkData.cpp
@@ -67,7 +67,9 @@ FStringBuilderBase& LexToString(EBulkDataFlags Flags, FStringBuilderBase& Sb)
 	TEST_AND_ADD_FLAG(Sb, Flags, BULKDATA_LazyLoadable);
 	TEST_AND_ADD_FLAG(Sb, Flags, BULKDATA_UsesIoDispatcher);
 	TEST_AND_ADD_FLAG(Sb, Flags, BULKDATA_DataIsMemoryMapped);
+	PRAGMA_DISABLE_DEPRECATION_WARNINGS
 	TEST_AND_ADD_FLAG(Sb, Flags, BULKDATA_HasAsyncReadPending);
+	PRAGMA_ENABLE_DEPRECATION_WARNINGS
 	TEST_AND_ADD_FLAG(Sb, Flags, BULKDATA_AlwaysAllowDiscard);
 	PRAGMA_DISABLE_DEPRECATION_WARNINGS
 	TEST_AND_ADD_FLAG(Sb, Flags, BULKDATA_BadDataVersion);
@@ -167,19 +169,6 @@ bool TryMemoryMapBulkData(
 	int64 Size,
 	FIoMappedRegion& OutRegion);
 
-/** Start load the internal bulk data payload. */ 
-bool StartAsyncLoad(
-	FBulkData* Owner,
-	const FBulkMetaData& BulkMeta,
-	const FIoChunkId& BulkChunkId,
-	int64 Offset,
-	int64 Size,
-	EAsyncIOPriorityAndFlags Priority,
-	TFunction<void(TIoStatusOr<FIoBuffer>)>&& Callback);
-
-/** Flush pending async load. */
-void FlushAsyncLoad(FBulkData* BulkData);
-
 //////////////////////////////////////////////////////////////////////////////
 
 static FArchive& SerializeAsInt32(FArchive& Ar, int64& Value)
@@ -624,7 +613,6 @@ FBulkData::~FBulkData()
 	DECLARE_SCOPE_CYCLE_COUNTER(TEXT("FBulkData::~FBulkData"), STAT_UBD_Destructor, STATGROUP_Memory);
 
 	check(IsUnlocked());
-	FlushAsyncLoading();
 
 	// Free memory.
 	FreeData();
@@ -858,7 +846,8 @@ bool FBulkData::IsBulkDataLoaded() const
 
 bool FBulkData::IsAsyncLoadingComplete() const
 {
-	return (GetBulkDataFlags() & BULKDATA_HasAsyncReadPending) == 0;
+	// This call no longer returns anything useful as ::StartAsyncLoading was removed in UE 5.5
+	return true;
 }
 
 bool FBulkData::IsAvailableForUse() const
@@ -885,9 +874,6 @@ void FBulkData::GetCopy( void** Dest, bool bDiscardInternalCopy )
 	check( IsUnlocked() );
 	check( Dest );
 
-	// Make sure any async loads have completed and moved the data into BulkData
-	FlushAsyncLoading();
-
 	const int64 BulkDataSize = GetBulkDataSize();
 
 	// Passed in memory is going to be used.
@@ -1093,7 +1079,6 @@ bool FBulkData::UnloadBulkData()
 #if WITH_EDITOR
 	if (IsUnlocked())
 	{
-		FlushAsyncLoading();
 		FreeData();
 		return true;
 	}
@@ -1162,60 +1147,6 @@ void FBulkData::ForceBulkDataResident()
 #endif // WITH_EDITOR
 }
 
-bool FBulkData::StartAsyncLoading()
-{
-	DECLARE_SCOPE_CYCLE_COUNTER(TEXT("FBulkData::StartAsyncLoading"), STAT_UBD_StartSerializingBulkData, STATGROUP_Memory);
-
-	if (!IsAsyncLoadingComplete())
-	{
-		return true; // Early out if an asynchronous load is already in progress.
-	}
-
-	if (IsBulkDataLoaded())
-	{
-		return false; // Early out if we do not need to actually load any data
-	}
-
-	if (!CanLoadFromDisk())
-	{
-		return false; // Early out if we cannot load from disk
-	}
-
-	checkf(IsUnlocked(), TEXT("Attempting to modify a BulkData object that is locked"));
-	check(GetBulkDataSize() > 0);
-	check(BulkChunkId.IsValid());
-	
-	BulkMeta.SetLockStatus(LOCKSTATUS_ReadWriteLock); // Bulkdata is effectively locked while streaming!
-	SetBulkDataFlags(BULKDATA_HasAsyncReadPending);
-	FPlatformMisc::MemoryBarrier();
-
-	return UE::BulkData::Private::StartAsyncLoad(
-		this,
-		BulkMeta,
-		BulkChunkId,
-		BulkMeta.GetOffset(), 
-		BulkMeta.GetSize(),
-		AIOP_Low,
-		[this](TIoStatusOr<FIoBuffer> Result)
-	{
-		if (Result.IsOk())
-		{
-			FIoBuffer Buffer = Result.ConsumeValueOrDie();
-			Buffer.EnsureOwned();
-			DataAllocation.SetData(this, Buffer.Release().ConsumeValueOrDie());
-		}
-		else
-		{
-			UE_LOG(LogSerialization, Error, TEXT("Async load bulk data '%s' FAILED, reason '%s'"), *LexToString(BulkChunkId), GetIoErrorText(Result.Status().GetErrorCode()));
-			RemoveBulkData();
-		}
-
-		FPlatformMisc::MemoryBarrier();
-		BulkMeta.SetLockStatus(LOCKSTATUS_Unlocked);
-		ClearBulkDataFlags(BULKDATA_HasAsyncReadPending);
-	});
-}
-
 /**
  * Sets the passed in bulk data flags.
  *
@@ -1241,15 +1172,6 @@ uint32 FBulkData::GetBulkDataFlags() const
 	return BulkMeta.GetFlags();
 }
 
-/**
- * Sets the passed in bulk data alignment.
- *
- * @param BulkDataAlignmentToSet	Bulk data alignment to set
- */
-void FBulkData::SetBulkDataAlignment(uint16 BulkDataAlignmentToSet)
-{
-}
-
 /**
 * Gets the current bulk data alignment.
 *
@@ -1488,15 +1410,6 @@ PRAGMA_ENABLE_DEPRECATION_WARNINGS
 }
 #endif
 
-FCustomVersionContainer FBulkData::GetCustomVersions(FArchive& InlineArchive) const
-{
-	FPackageFileVersion OutUEVersion;
-	int32 OutLicenseeUEVersion;
-	FCustomVersionContainer OutCustomVersions;
-	GetBulkDataVersions(InlineArchive, OutUEVersion, OutLicenseeUEVersion, OutCustomVersions);
-	return OutCustomVersions;
-}
-
 void FBulkData::GetBulkDataVersions(FArchive& InlineArchive, FPackageFileVersion& OutUEVersion,
 	int32& OutLicenseeUEVersion, FCustomVersionContainer& OutCustomVersions) const
 {
@@ -1680,44 +1593,11 @@ IBulkDataIORequest* FBulkData::CreateStreamingRequest(int64 OffsetInBulkData, in
 		UserSuppliedMemory).Release();
 }
 
-IBulkDataIORequest* FBulkData::CreateStreamingRequestForRange(const BulkDataRangeArray& RangeArray, EAsyncIOPriorityAndFlags Priority, FBulkDataIORequestCallBack* CompleteCallback)
-{	
-	check(RangeArray.Num() > 0);
-
-	const FBulkData& Start = *(RangeArray[0]);
-	const FBulkData& End = *(RangeArray[RangeArray.Num()-1]);
-	const int64 ReadOffset = Start.GetBulkDataOffsetInFile();
-	const int64 ReadSize = (End.GetBulkDataOffsetInFile() + End.GetBulkDataSize()) - ReadOffset;
-	
-	check(ReadSize > 0);
-
-	checkf(
-		Start.IsUsingIODispatcher() == false || Start.IsInSeparateFile(),
-		TEXT("Create bulkdata stream request from package '%s' FAILED, inline bulk data cannot be streamed from I/O store"),
-		*LexToString(Start.BulkChunkId));
-
-	checkf(
-		Start.IsUsingIODispatcher() == false || (End.IsInSeparateFile() && Start.BulkChunkId == End.BulkChunkId),
-		TEXT("Create bulk data stream request FAILED, range spans from package '%s' to package '%s'"),
-		*LexToString(Start.BulkChunkId), *LexToString(End.BulkChunkId));
-
-	return UE::BulkData::Private::CreateStreamingRequest(
-		Start.BulkMeta,
-		Start.BulkChunkId,
-		ReadOffset,
-		ReadSize,
-		Priority,
-		CompleteCallback,
-		nullptr).Release();
-}
-
 /**
  * Loads the bulk data if it is not already loaded.
  */
 void FBulkData::MakeSureBulkDataIsLoaded()
 {
-	FlushAsyncLoading();
-
 	if (IsBulkDataLoaded())
 	{
 		return;
@@ -1738,15 +1618,6 @@ void FBulkData::MakeSureBulkDataIsLoaded()
 	}
 }
 
-void FBulkData::FlushAsyncLoading()
-{
-	if (IsAsyncLoadingComplete() == false)
-	{
-		DECLARE_SCOPE_CYCLE_COUNTER(TEXT("FBulkData::FlushAsyncLoading"), STAT_UBD_WaitForAsyncLoading, STATGROUP_Memory);
-		UE::BulkData::Private::FlushAsyncLoad(this);
-	}
-}
-
 /**
  * Loads the data from disk into the specified memory block. This requires us still being attached to an
  * archive we can use for serialization.
@@ -1814,107 +1685,6 @@ bool FBulkData::CanDiscardInternalData() const
 	return BulkMeta.HasAnyFlags(static_cast<EBulkDataFlags>(BULKDATA_AlwaysAllowDiscard | BULKDATA_SingleUse)) || CanLoadFromDisk();
 }
 
-void FUntypedBulkData::SerializeElements(FArchive& Ar, void* Data)
-{
-	for (int64 ElementIndex = 0, Count = GetElementCount(); ElementIndex < Count; ++ElementIndex)
-	{
-		SerializeElement(Ar, Data, ElementIndex);
-	}
-}
-
-void FUntypedBulkData::SerializeBulkData(FArchive& Ar, void* Data, int64 DataSize, EBulkDataFlags InBulkDataFlags)
-{
-	SCOPED_LOADTIMER(BulkData_SerializeBulkData);
-
-	// skip serializing of unused data
-PRAGMA_DISABLE_DEPRECATION_WARNINGS
-	if (InBulkDataFlags & BULKDATA_Unused)
-PRAGMA_ENABLE_DEPRECATION_WARNINGS
-	{
-		return;
-	}
-
-	// Skip serialization for bulk data of zero length
-	if (DataSize == 0)
-	{
-		return;
-	}
-
-	// Allow backward compatible serialization by forcing bulk serialization off if required. Saving also always uses single
-	// element serialization so errors or oversight when changing serialization code is recoverable.
-	bool bSerializeInBulk = true;
-	if (RequiresSingleElementSerialization(Ar) 
-	// Set when serialized like a lazy array.
-	|| (InBulkDataFlags & BULKDATA_ForceSingleElementSerialization)
-	// We use bulk serialization even when saving 1 byte types (texture & sound bulk data) as an optimization for those.
-	|| (Ar.IsSaving() && (GetElementSize() > 1)))
-	{
-		bSerializeInBulk = false;
-	}
-
-	// Raw serialize the bulk data without any possibility for potential endian conversion.
-	if (bSerializeInBulk)
-	{
-		// Serialize data compressed.
-		if (InBulkDataFlags & BULKDATA_SerializeCompressed)
-		{
-			Ar.SerializeCompressed(Data, DataSize,
-				GetDecompressionFormat(InBulkDataFlags), COMPRESS_NoFlags, false);
-		}
-		// Uncompressed/ regular serialization.
-		else
-		{
-			Ar.Serialize(Data, DataSize);
-		}
-	}
-	// Serialize an element at a time via the virtual SerializeElement function potentially allowing and dealing with 
-	// endian conversion. Dealing with compression makes this a bit more complex as SerializeCompressed expects the 
-	// full data to be compressed en block and not piecewise.
-	else
-	{
-		// Serialize data compressed.
-		if (InBulkDataFlags & BULKDATA_SerializeCompressed)
-		{
-			// Loading, data is compressed in archive and needs to be decompressed.
-			if (Ar.IsLoading())
-			{
-				TUniquePtr<uint8[]> SerializedData = MakeUnique<uint8[]>(DataSize);
-
-				// Serialize data with passed in archive and compress.
-				Ar.SerializeCompressed(SerializedData.Get(), DataSize,
-					GetDecompressionFormat(InBulkDataFlags), COMPRESS_NoFlags, false);
-
-				// Initialize memory reader with uncompressed data array and propagate forced byte swapping
-				FLargeMemoryReader MemoryReader(SerializedData.Get(), DataSize, ELargeMemoryReaderFlags::Persistent);
-				MemoryReader.SetByteSwapping(Ar.ForceByteSwapping());
-
-				// Serialize each element individually via memory reader.
-				SerializeElements(MemoryReader, Data);
-			}
-			// Saving, data is uncompressed in memory and needs to be compressed.
-			else if (Ar.IsSaving())
-			{
-				// Initialize memory writer with blank data array and propagate forced byte swapping
-				FLargeMemoryWriter MemoryWriter(GetBulkDataSize(), true);
-				MemoryWriter.SetByteSwapping(Ar.ForceByteSwapping());
-
-				// Serialize each element individually via memory writer.
-				SerializeElements(MemoryWriter, Data);
-
-				// Serialize data with passed in archive and compress.
-				Ar.SerializeCompressed(MemoryWriter.GetData(), DataSize,
-					GetDecompressionFormat(InBulkDataFlags), COMPRESS_NoFlags, false);
-			}
-		}
-		// Uncompressed/ regular serialization.
-		else
-		{
-			// We can use the passed in archive if we're not compressing the data.
-			SerializeElements(Ar, Data);
-		}
-	}
-}
-
 void FFormatContainer::Serialize(FArchive& Ar, UObject* Owner, const TArray<FName>* FormatsToSave, bool bSingleUse, uint16 InAlignment, bool bInline, bool bMapped)
 {
 	if (Ar.IsLoading())
diff --git a/Engine/Source/Runtime/CoreUObject/Private/Serialization/BulkDataStreaming.cpp b/Engine/Source/Runtime/CoreUObject/Private/Serialization/BulkDataStreaming.cpp
--- a/Engine/Source/Runtime/CoreUObject/Private/Serialization/BulkDataStreaming.cpp
+++ b/Engine/Source/Runtime/CoreUObject/Private/Serialization/BulkDataStreaming.cpp
@@ -27,80 +27,6 @@ TRACE_DECLARE_ATOMIC_INT_COUNTER(BulkDataBatchRequest_PendingCount, TEXT("BulkDa
  */
 #define UE_ENABLE_BULKDATA_RANGE_TEST 0
 
-FBulkDataIORequest::FBulkDataIORequest(IAsyncReadFileHandle* InFileHandle)
-	: FileHandle(InFileHandle)
-	, ReadRequest(nullptr)
-	, Size(INDEX_NONE)
-{
-}
-
-FBulkDataIORequest::FBulkDataIORequest(IAsyncReadFileHandle* InFileHandle, IAsyncReadRequest* InReadRequest, int64 BytesToRead)
-	: FileHandle(InFileHandle)
-	, ReadRequest(InReadRequest)
-	, Size(BytesToRead)
-{
-
-}
-
-FBulkDataIORequest::~FBulkDataIORequest()
-{
-	delete ReadRequest;
-	delete FileHandle;
-}
-
-bool FBulkDataIORequest::MakeReadRequest(int64 Offset, int64 BytesToRead, EAsyncIOPriorityAndFlags PriorityAndFlags, FBulkDataIORequestCallBack* CompleteCallback, uint8* UserSuppliedMemory)
-{
-	check(ReadRequest == nullptr);
-
-	FBulkDataIORequestCallBack LocalCallback = *CompleteCallback;
-	FAsyncFileCallBack AsyncFileCallBack = [LocalCallback, BytesToRead, this](bool bWasCancelled, IAsyncReadRequest* InRequest)
-	{
-		// In some cases the call to ReadRequest can invoke the callback immediately (if the requested data is cached 
-		// in the pak file system for example) which means that FBulkDataIORequest::ReadRequest might not actually be
-		// set correctly, so we need to make sure it is assigned before we invoke LocalCallback!
-		ReadRequest = InRequest;
-
-		Size = BytesToRead;
-		LocalCallback(bWasCancelled, this);
-	};
-
-	ReadRequest = FileHandle->ReadRequest(Offset, BytesToRead, PriorityAndFlags, &AsyncFileCallBack, UserSuppliedMemory);
-	
-	if (ReadRequest != nullptr)
-	{
-		return true;
-	}
-	else
-	{
-		return false;
-	}
-}
-
-bool FBulkDataIORequest::PollCompletion() const
-{
-	return ReadRequest->PollCompletion();
-}
-
-bool FBulkDataIORequest::WaitCompletion(float TimeLimitSeconds)
-{
-	return ReadRequest->WaitCompletion(TimeLimitSeconds);
-}
-
-uint8* FBulkDataIORequest::GetReadResults()
-{
-	return ReadRequest->GetReadResults();
-}
-
-int64 FBulkDataIORequest::GetSize() const
-{
-	return Size;
-}
-
-void FBulkDataIORequest::Cancel()
-{
-	ReadRequest->Cancel();
-}
-
 //////////////////////////////////////////////////////////////////////////////
 
 namespace UE::BulkData::Private
@@ -437,24 +363,6 @@ uint8* FChunkBulkDataRequest::GetReadResults()
 
 //////////////////////////////////////////////////////////////////////////////
 
-TUniquePtr<IBulkDataIORequest> CreateBulkDataIoDispatcherRequest(
-	const FIoChunkId& ChunkId,
-	int64 Offset,
-	int64 Size,
-	FBulkDataIORequestCallBack* Callback,
-	uint8* UserSuppliedMemory,
-	int32 Priority)
-{
-	FIoBuffer Buffer = UserSuppliedMemory ? FIoBuffer(FIoBuffer::Wrap, UserSuppliedMemory, Size) : FIoBuffer(Size);
-
-	TUniquePtr<FChunkBulkDataRequest> Request = MakeUnique<FChunkBulkDataRequest>(Callback, MoveTemp(Buffer));
-	Request->Issue(ChunkId, FIoReadOptions(Offset, Size), Priority);
-
-	return Request;
-}
-
-//////////////////////////////////////////////////////////////////////////////
-
 bool OpenReadBulkData(
 	const FBulkMetaData& BulkMeta,
 	const FIoChunkId& BulkChunkId,
@@ -556,93 +464,6 @@ bool TryMemoryMapBulkData(
 	return false;
 }
 
-//////////////////////////////////////////////////////////////////////////////
-
-class FAsyncBulkDataRequests
-{
-public:
-	void AddPendingRequest(
-		FBulkData* Owner,
-		TUniquePtr<IAsyncReadFileHandle>&& FileHandle,
-		TUniquePtr<IAsyncReadRequest>&& ReadRequest)
-	{
-		FScopeLock _(&RequestsCS);
-		PendingRequests.Add(Owner, FPendingRequest { MoveTemp(FileHandle), MoveTemp(ReadRequest) });
-	}
-
-	void Flush(FBulkData* Owner)
-	{
-		FPendingRequest PendingRequest;
-
-		{
-			FScopeLock _(&RequestsCS);
-			PendingRequest = MoveTemp(PendingRequests.FindChecked(Owner));
-			PendingRequests.Remove(Owner);
-		}
-		
-		PendingRequest.ReadRequest->WaitCompletion();
-	}
-
-	static FAsyncBulkDataRequests& Get()
-	{
-		static FAsyncBulkDataRequests Instance;
-		return Instance;
-	}
-
-private:
-	struct FPendingRequest
-	{
-		TUniquePtr<IAsyncReadFileHandle> FileHandle;
-		TUniquePtr<IAsyncReadRequest> ReadRequest;
-	};
-
-	TMap<FBulkData*, FPendingRequest> PendingRequests;
-	FCriticalSection RequestsCS;
-};
-
-bool StartAsyncLoad(
-	FBulkData* Owner,
-	const FBulkMetaData& BulkMeta,
-	const FIoChunkId& BulkChunkId,
-	int64 Offset,
-	int64 Size,
-	EAsyncIOPriorityAndFlags Priority,
-	TFunction<void(TIoStatusOr<FIoBuffer>)>&& Callback)
-{
-	if (TUniquePtr<IAsyncReadFileHandle> FileHandle = OpenAsyncReadBulkData(BulkMeta, BulkChunkId))
-	{
-		FAsyncFileCallBack FileReadCallback = [Callback = MoveTemp(Callback), Size = Size]
-			(bool bCanceled, IAsyncReadRequest* Request)
-			{
-				if (bCanceled)
-				{
-					Callback(FIoStatus(EIoErrorCode::Cancelled));
-				}
-				else if (uint8* Data = Request->GetReadResults())
-				{
-					Callback(FIoBuffer(FIoBuffer::AssumeOwnership, Data, Size));
-				}
-				else
-				{
-					Callback(FIoStatus(EIoErrorCode::ReadError));
-				}
-			};
-		
-		if (IAsyncReadRequest* Request = FileHandle->ReadRequest(Offset, Size, Priority, &FileReadCallback))
-		{
-			FAsyncBulkDataRequests::Get().AddPendingRequest(Owner, MoveTemp(FileHandle), TUniquePtr<IAsyncReadRequest>(Request));
-			return true;
-		}
-	}
-
-	return false;
-}
-
-void FlushAsyncLoad(FBulkData* Owner)
-{
-	FAsyncBulkDataRequests::Get().Flush(Owner);
-}
-
 } // namespace UE::BulkData
 
 //////////////////////////////////////////////////////////////////////////////
diff --git a/Engine/Source/Runtime/CoreUObject/Public/Serialization/BulkData.h b/Engine/Source/Runtime/CoreUObject/Public/Serialization/BulkData.h
--- a/Engine/Source/Runtime/CoreUObject/Public/Serialization/BulkData.h
+++ b/Engine/Source/Runtime/CoreUObject/Public/Serialization/BulkData.h
@@ -35,10 +35,6 @@
 #include "Templates/PimplPtr.h"
 #include "UObject/NameTypes.h"
 
-#if UE_ENABLE_INCLUDE_ORDER_DEPRECATED_IN_5_2
-#include "IO/IoDispatcher.h"
-#endif
-
 class FIoBuffer;
 class FIoChunkId;
 class FLinkerLoad;
@@ -52,10 +48,8 @@ namespace UE { namespace Serialization { class FEditorBulkData; } }
 enum class EFileRegionType : uint8;
 
 #if WITH_EDITOR == 0 && WITH_EDITORONLY_DATA == 0
-	#define USE_NEW_BULKDATA UE_DEPRECATED_MACRO(5.1, "The USE_NEW_BULKDATA macro has been deprecated in favor of USE_RUNTIME_BULKDATA.") 1
 	#define USE_RUNTIME_BULKDATA 1
 #else
-	#define USE_NEW_BULKDATA UE_DEPRECATED_MACRO(5.1, "The USE_NEW_BULKDATA macro has been deprecated in favor of USE_RUNTIME_BULKDATA.") 0
 	#define USE_RUNTIME_BULKDATA 0
 #endif
 
@@ -168,8 +162,8 @@ enum EBulkDataFlags : uint32
 	BULKDATA_UsesIoDispatcher = 1u << 31u,
 	/** Assigned at runtime to indicate that the BulkData allocation is a memory mapped region of a file and not raw data. */
 	BULKDATA_DataIsMemoryMapped = 1 << 30,
-	/** Assigned at runtime to indicate that the BulkData object has an async loading request in flight and will need to wait on it. */
-	BULKDATA_HasAsyncReadPending = 1 << 29,
+	/** DEPRECATED */
+	BULKDATA_HasAsyncReadPending UE_DEPRECATED(5.5, "This flag no longer does anything")  = 1 << 29,
 	/** Assigned at runtime to indicate that the BulkData object should be considered for discard even if it cannot load from disk. */
 	BULKDATA_AlwaysAllowDiscard = 1 << 28,
 };
@@ -214,15 +208,6 @@ enum EBulkDataLockFlags
  */
 typedef TFunction<void(bool bWasCancelled, IBulkDataIORequest*)> FBulkDataIORequestCallBack;
 
-UE_DEPRECATED(5.1, "Use CreateStreamingRequest instead")
-TUniquePtr<IBulkDataIORequest> CreateBulkDataIoDispatcherRequest(
-	const FIoChunkId& InChunkID,
-	int64 InOffsetInBulkData = 0,
-	int64 InBytesToRead = INDEX_NONE,
-	FBulkDataIORequestCallBack* InCompleteCallback = nullptr,
-	uint8* InUserSuppliedMemory = nullptr,
-	int32 InPriority = IoDispatcherPriority_Low);
-
 /**
  * @documentation @todo documentation
  */
@@ -280,32 +265,6 @@ private:
 	IMappedFileRegion* MappedRegion;
 };
 
-UE_DEPRECATED(5.1, "Use CreateStreamingRequest instead");
-class FBulkDataIORequest : public IBulkDataIORequest
-{
-public:
-	FBulkDataIORequest(IAsyncReadFileHandle* InFileHandle);
-	FBulkDataIORequest(IAsyncReadFileHandle* InFileHandle, IAsyncReadRequest* InReadRequest, int64 BytesToRead);
-
-	virtual ~FBulkDataIORequest();
-
-	bool MakeReadRequest(int64 Offset, int64 BytesToRead, EAsyncIOPriorityAndFlags PriorityAndFlags, FBulkDataIORequestCallBack* CompleteCallback, uint8* UserSuppliedMemory);
-
-	virtual bool PollCompletion() const override;
-	virtual bool WaitCompletion( float TimeLimitSeconds = 0.0f ) override;
-
-	virtual uint8* GetReadResults() override;
-	virtual int64 GetSize() const override;
-
-	virtual void Cancel() override;
-
-private:
-	IAsyncReadFileHandle* FileHandle;
-	IAsyncReadRequest* ReadRequest;
-
-	int64 Size;
-};
-
 namespace UE::BulkData::Private
 {
 
@@ -646,10 +605,11 @@ public:
 	COREUOBJECT_API bool IsBulkDataLoaded() const;
 
 	/**
-	* Returns whether the bulk data asynchronous load has completed.
+	* Deprecated
 	*
-	* @return true if bulk data has been loaded or async loading was not used to load this data, false otherwise
+	* @return As of 5.5 this will always return true as ::StartAsyncLoading was removed
 	*/
+	UE_DEPRECATED(5.5, "No longer provides any functionality (assume that this will always return true")
 	COREUOBJECT_API bool IsAsyncLoadingComplete() const;
 
 	/**
@@ -752,9 +712,6 @@ public:
 	*/
 	COREUOBJECT_API uint32 GetBulkDataFlags() const;
 
-	UE_DEPRECATED(5.1, "Bulk Data will always use default alignment")
-	COREUOBJECT_API void SetBulkDataAlignment(uint16 BulkDataAlignmentToSet);
-
 	/**
 	* Gets the current bulk data alignment.
 	*
@@ -769,21 +726,6 @@ public:
 	 */
 	COREUOBJECT_API void ClearBulkDataFlags( uint32 BulkDataFlagsToClear );
 
-	/** Returns the PackagePath this bulkdata resides in */
-	UE_DEPRECATED(5.1, "Deprecated, no replacement")
-	const FPackagePath& GetPackagePath() const
-	{ 
-		static FPackagePath Empty;
-		return Empty;
-	}
-
-	/** Returns which segment of its PackagePath this bulkdata resides in */
-	UE_DEPRECATED(5.1, "Deprecated, no replacement")
-	EPackageSegment GetPackageSegment() const
-	{ 
-		return EPackageSegment::Header;
-	}
-
 	/** 
 	 * Returns the io filename hash associated with this bulk data.
 	 *
@@ -871,17 +813,6 @@ public:
 	 * Forces the bulk data to be resident in memory and detaches the archive.
 	 */
 	COREUOBJECT_API void ForceBulkDataResident();
-
-	/** 
-	* Initiates a new asynchronous operation to load the dulkdata from disk assuming that it is not already
-	* loaded.
-	* Note that a new asynchronous loading operation will not be created if one is already in progress.
-	*
-	* @return True if an asynchronous loading operation is in progress by the time that the method returns
-	* and false if the data is already loaded or cannot be loaded from disk.
-	*/
-	UE_DEPRECATED(5.1, "Use FBulkDataRequest or CreateStreamingRequest instead")
-	COREUOBJECT_API bool StartAsyncLoading();
 	
 	/**
 	 * Sets whether we should store the data compressed on disk.
@@ -937,9 +868,6 @@ public:
 		return DataAllocation.StealFileMapping(this); 
 	}
 
-	UE_DEPRECATED(5.1, "Call GetBulkDataVersions instead.")
-	COREUOBJECT_API FCustomVersionContainer GetCustomVersions(FArchive& InlineArchive) const;
-
 	/**
 	 * Get the CustomVersions used in the file containing the BulkData payload. If !IsInSeparateFile, this will be
 	 * the custom versions from the archive used to serialize the FBulkData, which the caller must provide.
@@ -996,21 +924,6 @@ public:
 	 */
 	COREUOBJECT_API IBulkDataIORequest* CreateStreamingRequest(int64 OffsetInBulkData, int64 BytesToRead, EAsyncIOPriorityAndFlags Priority, FBulkDataIORequestCallBack* CompleteCallback, uint8* UserSuppliedMemory) const;
 
-	/**
-	 * Create an async read request for a range of bulk data streaming tokens
-	 * The request will read all data between the two given streaming tokens objects. They must both represent areas of data in the file!
-	 * There is no way to validate this and it is up to the caller to make sure that it is correct.
-	 * The memory to be read into will be automatically allocated the size of which can be retrieved by calling IBulkDataIORequest::GetSize()
-	 *
-	 * @param Start				The bulk data to start reading from.
-	 * @param End				The bulk data to finish reading from.
-	 * @param Priority			Priority and flags of the request. If this includes AIOP_FLAG_PRECACHE, then memory will never be returned. The request should always be canceled and waited for, even for a precache request.
-	 * @param CompleteCallback	Called from an arbitrary thread when the request is complete. Can be nullptr, if non-null, must remain valid until it is called. It will always be called.
-	 * @return					A request for the read. This is owned by the caller and must be deleted by the caller.
-	**/
-	UE_DEPRECATED(5.1, "Use FBulkDataRequest instead")
-	static COREUOBJECT_API IBulkDataIORequest* CreateStreamingRequestForRange(const BulkDataRangeArray& RangeArray, EAsyncIOPriorityAndFlags Priority, FBulkDataIORequestCallBack* CompleteCallback);
-
 	/** Enable the given flags in the given accumulator variable. */
 	static COREUOBJECT_API void SetBulkDataFlagsOn(EBulkDataFlags& InOutAccumulator, EBulkDataFlags FlagsToSet);
 	/** Disable the given flags in the given accumulator variable. */
@@ -1061,8 +974,8 @@ private:
 	 */
 	COREUOBJECT_API bool TryLoadDataIntoMemory(FIoBuffer Dest);
 
-	/** Flushes any pending async load of bulk data  and copies the data to Dest buffer*/
-	COREUOBJECT_API void FlushAsyncLoading();
+	UE_DEPRECATED(5.5, "No longer provides any functionality")
+	COREUOBJECT_API void FlushAsyncLoading() {}
 
 	/** Returns if the offset needs fixing when serialized */
 	COREUOBJECT_API bool NeedsOffsetFixup() const;
@@ -1207,75 +1120,6 @@ public:
 	}
 };
 
-UE_DEPRECATED(5.1, "Use FBulkData/TBulkData");
-struct FUntypedBulkData : public FBulkData
-{
-	virtual ~FUntypedBulkData() = default;
-
-	/**
-	 * Returns the number of elements allocated.
-	 *
-	 * @return Number of elements
-	 */
-	int64 GetElementCount() const
-	{
-		return GetBulkDataSize() / GetElementSize();
-	}
-
-	/**
-	 * Returns size in bytes of single element.
-	 *
-	 * Pure virtual that needs to be overloaded in derived classes.
-	 *
-	 * @return Size in bytes of single element
-	 */
-	virtual int32 GetElementSize() const = 0;
-	
-	/**
-	 * Change size of locked bulk data. Only valid if locked via read-write lock.
-	 *
-	 * @param ElementCount	Number of elements to allocate.
-	 */
-	void* Realloc(int64 ElementCount)
-	{
-		return FBulkData::Realloc(ElementCount, GetElementSize());
-	}
-
-	void Serialize(FArchive& Ar, UObject* Owner, int32 Idx=INDEX_NONE, bool bAttemptFileMapping = false, EFileRegionType FileRegionType = EFileRegionType::None)
-	{
-		FBulkData::Serialize(Ar, Owner, bAttemptFileMapping, GetElementSize(), FileRegionType);
-	}
-
-protected:
-
-	FUntypedBulkData()
-	{
-		SerializeElementsCallback = [this](FArchive& Ar, void* Data, int64 Size, EBulkDataFlags InBulkDataFlags)
-		{
-			SerializeBulkData(Ar, Data, Size, InBulkDataFlags);
-		};
-#if !USE_RUNTIME_BULKDATA
-		SerializeBulkDataElements = &SerializeElementsCallback;
-#endif // !USE_RUNTIME_BULKDATA
-	}
-
-	COREUOBJECT_API virtual void SerializeElements(FArchive& Ar, void* Data);
-	
-	virtual void SerializeElement(FArchive& Ar, void* Data, int64 ElementIndex) = 0;
-
-	virtual bool RequiresSingleElementSerialization(FArchive& Ar)
-	{
-		return false;
-	}
-	
-private:
-
-	COREUOBJECT_API void SerializeBulkData(FArchive& Ar, void* Data, int64 BulkDataSize, EBulkDataFlags InBulkDataFlags);
-
-	FSerializeBulkDataElements SerializeElementsCallback;
-};
-
-using FBulkDataInterface UE_DEPRECATED(5.1, "FBulkDataInterface is deprecated. Use FBulkData") = FBulkData;
 using FByteBulkData = TBulkData<uint8>;
 using FWordBulkData = TBulkData<uint16>;
 using FIntBulkData = TBulkData<int32>;
diff --git a/Engine/Source/Runtime/Engine/Private/Rendering/SkeletalMeshLODImporterData.cpp b/Engine/Source/Runtime/Engine/Private/Rendering/SkeletalMeshLODImporterData.cpp
--- a/Engine/Source/Runtime/Engine/Private/Rendering/SkeletalMeshLODImporterData.cpp
+++ b/Engine/Source/Runtime/Engine/Private/Rendering/SkeletalMeshLODImporterData.cpp
@@ -898,7 +898,7 @@ void FRawSkeletalMeshBulkData::LoadRawMesh(FSkeletalMeshImportData& OutMesh)
 		// This allows any thread to be able to deserialize from the RawMesh directly
 		// from disk so we can unload bulk data from memory.
 		bool bHasBeenLoadedFromFileReader = false;
-		if (BulkData.IsAsyncLoadingComplete() && !BulkData.IsBulkDataLoaded())
+		if (!BulkData.IsBulkDataLoaded())
 		{
 			// This can't be called in -game mode because we're not allowed to load bulk data outside of EDL.
 			bHasBeenLoadedFromFileReader = BulkData.LoadBulkDataWithFileReader();
@@ -940,7 +940,7 @@ void FRawSkeletalMeshBulkData::UpdateRawMeshFormat()
 	// This allows any thread to be able to deserialize from the RawMesh directly
 	// from disk so we can unload bulk data from memory.
 	bool bHasBeenLoadedFromFileReader = false;
-	if (BulkData.IsAsyncLoadingComplete() && !BulkData.IsBulkDataLoaded())
+	if (!BulkData.IsBulkDataLoaded())
 	{
 		// This can't be called in -game mode because we're not allowed to load bulk data outside of EDL.
 		bHasBeenLoadedFromFileReader = BulkData.LoadBulkDataWithFileReader();
diff --git a/Engine/Source/Runtime/Engine/Private/SkeletalMeshLODModel.cpp b/Engine/Source/Runtime/Engine/Private/SkeletalMeshLODModel.cpp
--- a/Engine/Source/Runtime/Engine/Private/SkeletalMeshLODModel.cpp
+++ b/Engine/Source/Runtime/Engine/Private/SkeletalMeshLODModel.cpp
@@ -841,7 +841,7 @@ void FSkeletalMeshLODModel::Serialize(FArchive& Ar, UObject* Owner, int32 Idx)
 			RawPointIndices_DEPRECATED.Serialize(Ar, Owner);
 			if (RawPointIndices_DEPRECATED.GetBulkDataSize())
 			{
-				if (RawPointIndices_DEPRECATED.IsAsyncLoadingComplete() && !RawPointIndices_DEPRECATED.IsBulkDataLoaded())
+				if (!RawPointIndices_DEPRECATED.IsBulkDataLoaded())
 				{
 					RawPointIndices_DEPRECATED.LoadBulkDataWithFileReader();
 				}
diff --git a/Engine/Source/Runtime/Engine/Private/TextureDerivedData.cpp b/Engine/Source/Runtime/Engine/Private/TextureDerivedData.cpp
--- a/Engine/Source/Runtime/Engine/Private/TextureDerivedData.cpp
+++ b/Engine/Source/Runtime/Engine/Private/TextureDerivedData.cpp
@@ -2414,14 +2414,6 @@ bool FTexturePlatformData::IsReadyForAsyncPostLoad() const
 	}
 #endif
 
-	for (int32 MipIndex = 0; MipIndex < Mips.Num(); ++MipIndex)
-	{
-		const FTexture2DMipMap& Mip = Mips[MipIndex];
-		if (!Mip.BulkData.IsAsyncLoadingComplete())
-		{
-			return false;
-		}
-	}
 	return true;
 }
 
diff --git a/Engine/Source/Runtime/RawMesh/Private/RawMesh.cpp b/Engine/Source/Runtime/RawMesh/Private/RawMesh.cpp
--- a/Engine/Source/Runtime/RawMesh/Private/RawMesh.cpp
+++ b/Engine/Source/Runtime/RawMesh/Private/RawMesh.cpp
@@ -277,7 +277,7 @@ void FRawMeshBulkData::LoadRawMesh(FRawMesh& OutMesh)
 		// This allows any thread to be able to deserialize from the RawMesh directly
 		// from disk so we can unload bulk data from memory.
 		bool bHasBeenLoadedFromFileReader = false;
-		if (BulkData.IsAsyncLoadingComplete() && !BulkData.IsBulkDataLoaded())
+		if (!BulkData.IsBulkDataLoaded())
 		{
 			// This can't be called in -game mode because we're not allowed to load bulk data outside of EDL.
 			bHasBeenLoadedFromFileReader = BulkData.LoadBulkDataWithFileReader();
