[ZenLoader/Misc]
 - Prepare for EInternalObjectsFlag::AsyncLoading deprecation

#rnx

diff --git a/Engine/Plugins/Experimental/ActorModifierCore/Source/ActorModifierCore/Private/Modifiers/ActorModifierCoreBase.cpp b/Engine/Plugins/Experimental/ActorModifierCore/Source/ActorModifierCore/Private/Modifiers/ActorModifierCoreBase.cpp
--- a/Engine/Plugins/Experimental/ActorModifierCore/Source/ActorModifierCore/Private/Modifiers/ActorModifierCoreBase.cpp
+++ b/Engine/Plugins/Experimental/ActorModifierCore/Source/ActorModifierCore/Private/Modifiers/ActorModifierCoreBase.cpp
@@ -660,7 +660,7 @@ void UActorModifierCoreBase::PostModifierWorldLoad(UWorld* InWorld, ELevelTick I
 	// Check actor is in the world loaded and does not need post load and is not in async loading
 	if (Actor
 		&& InWorld == Actor->GetWorld()
-		&& !Actor->HasAnyInternalFlags(EInternalObjectFlags::AsyncLoading)
+		&& !Actor->HasAnyInternalFlags(EInternalObjectFlags_AsyncLoading)
 		&& !Actor->HasAnyFlags(EObjectFlags::RF_NeedPostLoad)
 		&& !Actor->HasAnyFlags(EObjectFlags::RF_NeedPostLoadSubobjects))
 	{
diff --git a/Engine/Source/Editor/UnrealEd/Private/Kismet2/Kismet2.cpp b/Engine/Source/Editor/UnrealEd/Private/Kismet2/Kismet2.cpp
--- a/Engine/Source/Editor/UnrealEd/Private/Kismet2/Kismet2.cpp
+++ b/Engine/Source/Editor/UnrealEd/Private/Kismet2/Kismet2.cpp
@@ -850,7 +850,7 @@ static void ConformComponentsUtils::ConformRemovedNativeComponents(UObject* BpCd
 		// else, the component has been removed from our native super class
 
 		Component->DestroyComponent(/*bPromoteChildren =*/false);
-		if (Component->HasAnyInternalFlags(EInternalObjectFlags::AsyncLoading))
+		if (Component->HasAnyInternalFlags(EInternalObjectFlags_AsyncLoading))
 		{
 			// Async loading components cannot be pending kill, or the async loading code will assert when trying to postload them.
 			Component->ClearGarbage();
diff --git a/Engine/Source/Runtime/CoreUObject/Private/Blueprint/BlueprintSupport.cpp b/Engine/Source/Runtime/CoreUObject/Private/Blueprint/BlueprintSupport.cpp
--- a/Engine/Source/Runtime/CoreUObject/Private/Blueprint/BlueprintSupport.cpp
+++ b/Engine/Source/Runtime/CoreUObject/Private/Blueprint/BlueprintSupport.cpp
@@ -1832,7 +1832,7 @@ int32 FLinkerLoad::ResolveDependencyPlaceholder(FLinkerPlaceholderBase* Placehol
 	// holding onto objects that are spawned during the process (to ensure 
 	// they're not thrown away prematurely)
 	bool const bIsAsyncLoadRef = (UnresolvedReferences.ExternalReferences.Num() == 1) &&
-		PlaceholderObj->HasAnyInternalFlags(EInternalObjectFlags::AsyncLoading) && (UnresolvedReferences.ExternalReferences[0].Referencer == FGCObject::GGCObjectReferencer);
+		PlaceholderObj->HasAnyInternalFlags(EInternalObjectFlags_AsyncLoading) && (UnresolvedReferences.ExternalReferences[0].Referencer == FGCObject::GGCObjectReferencer);
 
 	DEFERRED_DEPENDENCY_CHECK(!bIsReferenced || bIsAsyncLoadRef);
 #endif // USE_DEFERRED_DEPENDENCY_CHECK_VERIFICATION_TESTS
diff --git a/Engine/Source/Runtime/CoreUObject/Private/Serialization/AsyncLoading.cpp b/Engine/Source/Runtime/CoreUObject/Private/Serialization/AsyncLoading.cpp
--- a/Engine/Source/Runtime/CoreUObject/Private/Serialization/AsyncLoading.cpp
+++ b/Engine/Source/Runtime/CoreUObject/Private/Serialization/AsyncLoading.cpp
@@ -4053,7 +4053,7 @@ void FAsyncPackage::Event_StartPostload()
 			FObjectExport& Export = Linker->ExportMap[LocalExportIndex];
 			UObject* Object = Export.Object;
 			checkSlow(!(Object && !ReferencedObjects.Contains(Object)));
-			if (Object && (Object->HasAnyFlags(RF_NeedPostLoad) || Object->HasAnyInternalFlags(EInternalObjectFlags::AsyncLoading)))
+			if (Object && (Object->HasAnyFlags(RF_NeedPostLoad) || Object->HasAnyInternalFlags(EInternalObjectFlags_AsyncLoading)))
 			{
 				check(Object->IsValidLowLevelFast());
 				LoadContext->AddLoadedObject(Object);
@@ -5563,7 +5563,7 @@ void FAsyncLoadingThread::NotifyConstructedDuringAsyncLoading(UObject* Object, b
 	// finished routing PostLoad to all objects.
 	if (!bSubObject)
 	{
-		Object->SetInternalFlags(EInternalObjectFlags::AsyncLoading);
+		Object->SetInternalFlags(EInternalObjectFlags_AsyncLoading);
 	}
 
 	FAsyncPackage* AsyncPackage = static_cast<FAsyncPackage*>(ThreadContext.AsyncPackage);
@@ -5737,7 +5737,7 @@ void FAsyncPackage::AddObjectReference(UObject* InObject)
 
 void FAsyncPackage::EmptyReferencedObjects()
 {
-	const EInternalObjectFlags AsyncFlags = EInternalObjectFlags::Async | EInternalObjectFlags::AsyncLoading;
+	const EInternalObjectFlags AsyncFlags = EInternalObjectFlags::Async | EInternalObjectFlags_AsyncLoading;
 	FScopeLock ReferencedObjectsLock(&ReferencedObjectsCritical);
 	for (UObject* Obj : ReferencedObjects)
 	{
@@ -7120,7 +7120,7 @@ EAsyncPackageState::Type FAsyncPackage::PostLoadDeferredObjects(double InTickSta
 			UObject* Object = DeferredFinalizeObjects[DeferredFinalizeIndex++];
 			if (Object)
 			{
-				Object->AtomicallyClearInternalFlags(EInternalObjectFlags::AsyncLoading);
+				Object->AtomicallyClearInternalFlags(EInternalObjectFlags_AsyncLoading);
 			}
 
 			// CDO need special handling, no matter if it's listed in DeferredFinalizeObjects
@@ -7132,9 +7132,9 @@ EAsyncPackageState::Type FAsyncPackage::PostLoadDeferredObjects(double InTickSta
 				CDOToHandle->GetDefaultSubobjects(CDODefaultSubobjects);
 				for (UObject* SubObject : CDODefaultSubobjects)
 				{
-					if (SubObject && SubObject->HasAnyInternalFlags(EInternalObjectFlags::AsyncLoading))
+					if (SubObject && SubObject->HasAnyInternalFlags(EInternalObjectFlags_AsyncLoading))
 					{
-						SubObject->AtomicallyClearInternalFlags(EInternalObjectFlags::AsyncLoading);
+						SubObject->AtomicallyClearInternalFlags(EInternalObjectFlags_AsyncLoading);
 					}
 				}
 				CDODefaultSubobjects.Reset();
@@ -7159,8 +7159,8 @@ EAsyncPackageState::Type FAsyncPackage::PostLoadDeferredObjects(double InTickSta
 		{
 			LastObjectWorkWasPerformedOn = LinkerRoot;
 			LastTypeOfWorkPerformed = TEXT("CreateClustersFromPackage");
-			LinkerRoot->AtomicallyClearInternalFlags(EInternalObjectFlags::AsyncLoading);
-			LinkerRoot->MarkAsFullyLoaded();			
+			LinkerRoot->AtomicallyClearInternalFlags(EInternalObjectFlags_AsyncLoading);
+			LinkerRoot->MarkAsFullyLoaded();
 			LinkerRoot->SetLoadTime((float)(FPlatformTime::Seconds() - LoadStartTime));
 
 			if (Linker)
diff --git a/Engine/Source/Runtime/CoreUObject/Private/Serialization/AsyncPackageLoader.cpp b/Engine/Source/Runtime/CoreUObject/Private/Serialization/AsyncPackageLoader.cpp
--- a/Engine/Source/Runtime/CoreUObject/Private/Serialization/AsyncPackageLoader.cpp
+++ b/Engine/Source/Runtime/CoreUObject/Private/Serialization/AsyncPackageLoader.cpp
@@ -670,7 +670,7 @@ void CancelAsyncLoading()
 		CollectGarbage(GARBAGE_COLLECTION_KEEPFLAGS, true);
 	}
 
-	const EInternalObjectFlags AsyncFlags = EInternalObjectFlags::Async | EInternalObjectFlags::AsyncLoading;
+	const EInternalObjectFlags AsyncFlags = EInternalObjectFlags::Async | EInternalObjectFlags_AsyncLoading;
 	for (int32 ObjectIndex = 0; ObjectIndex < GUObjectArray.GetObjectArrayNum(); ++ObjectIndex)
 	{
 		FUObjectItem* ObjectItem = &GUObjectArray.GetObjectItemArrayUnsafe()[ObjectIndex];
diff --git a/Engine/Source/Runtime/CoreUObject/Private/UObject/Obj.cpp b/Engine/Source/Runtime/CoreUObject/Private/UObject/Obj.cpp
--- a/Engine/Source/Runtime/CoreUObject/Private/UObject/Obj.cpp
+++ b/Engine/Source/Runtime/CoreUObject/Private/UObject/Obj.cpp
@@ -1460,7 +1460,7 @@ bool UObject::CanModify() const
 	return
 		!HasAnyFlags(RF_NeedInitialization) && !IsGarbageCollecting() && !GExitPurge && !IsUnreachable() &&
 		// Prevent modification while loading
-		!HasAnyInternalFlags(EInternalObjectFlags::Async | EInternalObjectFlags::AsyncLoading) &&
+		!HasAnyInternalFlags(EInternalObjectFlags::Async | EInternalObjectFlags_AsyncLoading) &&
 		// Only the game-thread should be allowed to touch the transaction buffer at all
 		IsInGameThread();
 }
diff --git a/Engine/Source/Runtime/CoreUObject/Private/UObject/Package.cpp b/Engine/Source/Runtime/CoreUObject/Private/UObject/Package.cpp
--- a/Engine/Source/Runtime/CoreUObject/Private/UObject/Package.cpp
+++ b/Engine/Source/Runtime/CoreUObject/Private/UObject/Package.cpp
@@ -330,7 +330,7 @@ bool UPackage::IsFullyLoaded() const
 		return false;
 	}
 
-	if (HasAnyInternalFlags(EInternalObjectFlags::AsyncLoading))
+	if (HasAnyInternalFlags(EInternalObjectFlags_AsyncLoading))
 	{
 		// If it's in the middle of an async load, don't make any changes and respect the current 'false' value of bHasBeenFullyLoaded
 		return false;
diff --git a/Engine/Source/Runtime/CoreUObject/Private/UObject/ReferenceChainSearch.cpp b/Engine/Source/Runtime/CoreUObject/Private/UObject/ReferenceChainSearch.cpp
--- a/Engine/Source/Runtime/CoreUObject/Private/UObject/ReferenceChainSearch.cpp
+++ b/Engine/Source/Runtime/CoreUObject/Private/UObject/ReferenceChainSearch.cpp
@@ -1323,7 +1323,7 @@ FString FReferenceChainSearch::GetObjectFlags(const FGCObjectInfo& InObject)
 		Flags += TEXT("(async) ");
 	}
 
-	if (InObject.HasAnyInternalFlags(EInternalObjectFlags::AsyncLoading))
+	if (InObject.HasAnyInternalFlags(EInternalObjectFlags_AsyncLoading))
 	{
 		Flags += TEXT("(asyncloading) ");
 	}
diff --git a/Engine/Source/Runtime/CoreUObject/Public/UObject/PersistentObjectPtr.h b/Engine/Source/Runtime/CoreUObject/Public/UObject/PersistentObjectPtr.h
--- a/Engine/Source/Runtime/CoreUObject/Public/UObject/PersistentObjectPtr.h
+++ b/Engine/Source/Runtime/CoreUObject/Public/UObject/PersistentObjectPtr.h
@@ -242,7 +242,7 @@ private:
 	// For example, objects that are in the process of being async loaded may not be cached
 	inline bool CanCacheObjectPointer(const UObject* Ptr) const
 	{
-		if (IsInAsyncLoadingThread() && Ptr && Ptr->HasAnyInternalFlags(EInternalObjectFlags::AsyncLoading))
+		if (IsInAsyncLoadingThread() && Ptr && Ptr->HasAnyInternalFlags(EInternalObjectFlags_AsyncLoading))
 		{
 			return false;
 		}
diff --git a/Engine/Source/Runtime/Engine/Private/LevelStreamingGCHelper.cpp b/Engine/Source/Runtime/Engine/Private/LevelStreamingGCHelper.cpp
--- a/Engine/Source/Runtime/Engine/Private/LevelStreamingGCHelper.cpp
+++ b/Engine/Source/Runtime/Engine/Private/LevelStreamingGCHelper.cpp
@@ -230,7 +230,7 @@ void FLevelStreamingGCHelper::VerifyLevelsGotRemovedByGC()
 					// But disregard package object itself.
 					&& !Object->IsA(UPackage::StaticClass()))
 				{
-					if (bIsAsyncLoading && Object->HasAnyInternalFlags(EInternalObjectFlags::Async | EInternalObjectFlags::AsyncLoading))
+					if (bIsAsyncLoading && Object->HasAnyInternalFlags(EInternalObjectFlags::Async | EInternalObjectFlags_AsyncLoading))
 					{
 						UE_LOG(LogGarbage, Display, TEXT("Level object %s isn't released by async loading yet, "
 							"it will get garbage collected next time instead."),
diff --git a/Engine/Source/Runtime/Engine/Private/ODSC/ODSCManager.cpp b/Engine/Source/Runtime/Engine/Private/ODSC/ODSCManager.cpp
--- a/Engine/Source/Runtime/Engine/Private/ODSC/ODSCManager.cpp
+++ b/Engine/Source/Runtime/Engine/Private/ODSC/ODSCManager.cpp
@@ -203,7 +203,7 @@ void FODSCManager::AddThreadedShaderPipelineRequest(
 
 void FODSCManager::RegisterMaterialInstance(const UMaterialInstance* MaterialInstance)
 {
-	if (IsODSCActive() && MaterialInstance->HasAnyInternalFlags(EInternalObjectFlags::AsyncLoading))
+	if (IsODSCActive() && MaterialInstance->HasAnyInternalFlags(EInternalObjectFlags_AsyncLoading))
 	{
 		FScopeLock Lock(&GODSCManager->MaterialInstancesCachedUniformExpressionsCS);
 		TWeakObjectPtr<const UMaterialInstance>& MaterialInstanceSoftPtr = GODSCManager->MaterialInstancesCachedUniformExpressions.FindOrAdd(MaterialInstance);
@@ -229,7 +229,7 @@ bool FODSCManager::HasAsyncLoadingInstances()
 	{
 		const UMaterialInstance* MI = Iter.Value().Get();
 
-		if (MI == nullptr || !MI->HasAnyInternalFlags(EInternalObjectFlags::AsyncLoading))
+		if (MI == nullptr || !MI->HasAnyInternalFlags(EInternalObjectFlags_AsyncLoading))
 		{
 			Iter.RemoveCurrent();
 			continue;
diff --git a/Engine/Source/Runtime/Engine/Private/ObjectCacheContext.cpp b/Engine/Source/Runtime/Engine/Private/ObjectCacheContext.cpp
--- a/Engine/Source/Runtime/Engine/Private/ObjectCacheContext.cpp
+++ b/Engine/Source/Runtime/Engine/Private/ObjectCacheContext.cpp
@@ -28,7 +28,7 @@ namespace ObjectCacheContextImpl {
 EInternalObjectFlags GetObjectCacheInternalFlagsExclusion()
 {
 	// We never want to return objects that are invalid or still being worked on by other threads
-	return EInternalObjectFlags::Unreachable | EInternalObjectFlags::Garbage | EInternalObjectFlags::AsyncLoading | EInternalObjectFlags::Async;
+	return EInternalObjectFlags::Unreachable | EInternalObjectFlags::Garbage | EInternalObjectFlags_AsyncLoading | EInternalObjectFlags::Async;
 }
 
 } // namespace ObjectCacheContextImpl
diff --git a/Engine/Source/Runtime/Engine/Private/PackageMapClient.cpp b/Engine/Source/Runtime/Engine/Private/PackageMapClient.cpp
--- a/Engine/Source/Runtime/Engine/Private/PackageMapClient.cpp
+++ b/Engine/Source/Runtime/Engine/Private/PackageMapClient.cpp
@@ -3606,7 +3606,7 @@ UObject* FNetGUIDCache::GetObjectFromNetGUID( const FNetworkGUID& NetGUID, const
 		if (!Package->IsFullyLoaded() 
 			&& !Package->HasAnyPackageFlags( TreatAsLoadedFlags )) //TODO: dependencies of CompiledIn could still be loaded asynchronously. Are they necessary at this point??
 		{
-			if (ShouldAsyncLoad() && Package->HasAnyInternalFlags(EInternalObjectFlags::AsyncLoading))
+			if (ShouldAsyncLoad() && Package->HasAnyInternalFlags(EInternalObjectFlags_AsyncLoading))
 			{
 				// Something else is already async loading this package, calling load again will add our callback to the existing load request
 				StartAsyncLoadingPackage(*CacheObjectPtr, NetGUID, true);
diff --git a/Engine/Source/Runtime/Engine/Private/StreamableManager.cpp b/Engine/Source/Runtime/Engine/Private/StreamableManager.cpp
--- a/Engine/Source/Runtime/Engine/Private/StreamableManager.cpp
+++ b/Engine/Source/Runtime/Engine/Private/StreamableManager.cpp
@@ -1655,7 +1655,7 @@ void FStreamableManager::FindInMemory(FSoftObjectPath& InOutTargetName, struct F
 	}
 	checkSlow(Existing->Target == StaticFindObject(UObject::StaticClass(), nullptr, *InOutTargetName.ToString()));
 
-	if (Existing->Target && Existing->Target->HasAnyInternalFlags(EInternalObjectFlags::AsyncLoading))
+	if (Existing->Target && Existing->Target->HasAnyInternalFlags(EInternalObjectFlags_AsyncLoading))
 	{
 		// This can get called from PostLoad on async loaded objects, if it is we do not want to return partially loaded objects and instead want to register for their full load
 		Existing->Target = nullptr;
diff --git a/Engine/Source/Runtime/Engine/Private/UnrealEngine.cpp b/Engine/Source/Runtime/Engine/Private/UnrealEngine.cpp
--- a/Engine/Source/Runtime/Engine/Private/UnrealEngine.cpp
+++ b/Engine/Source/Runtime/Engine/Private/UnrealEngine.cpp
@@ -17696,7 +17696,7 @@ void UEngine::CopyPropertiesForUnrelatedObjects(UObject* OldObject, UObject* New
 
 	// Now notify any tools that aren't already updated via the FArchiveReplaceObjectRef path unless the OldObject is still being async loaded
 	if (Params.bNotifyObjectReplacement && GEngine != nullptr &&
-		!OldObject->HasAnyInternalFlags(EInternalObjectFlags::Async | EInternalObjectFlags::AsyncLoading))
+		!OldObject->HasAnyInternalFlags(EInternalObjectFlags::Async | EInternalObjectFlags_AsyncLoading))
 	{
 		check(IsInGameThread());
 		GEngine->NotifyToolsOfObjectReplacement(ReferenceReplacementMap);
diff --git a/Engine/Source/Runtime/Experimental/Iris/Core/Private/Iris/ReplicationSystem/ObjectReferenceCache.cpp b/Engine/Source/Runtime/Experimental/Iris/Core/Private/Iris/ReplicationSystem/ObjectReferenceCache.cpp
--- a/Engine/Source/Runtime/Experimental/Iris/Core/Private/Iris/ReplicationSystem/ObjectReferenceCache.cpp
+++ b/Engine/Source/Runtime/Experimental/Iris/Core/Private/Iris/ReplicationSystem/ObjectReferenceCache.cpp
@@ -888,7 +888,7 @@ UObject* FObjectReferenceCache::ResolveObjectReferenceHandleInternal(FNetRefHand
 		if (!Package->IsFullyLoaded() 
 			&& !Package->HasAnyPackageFlags(TreatAsLoadedFlags)) //TODO: dependencies of CompiledIn could still be loaded asynchronously. Are they necessary at this point??
 		{
-			if (ShouldAsyncLoad()  && !ResolveContext.bForceSyncLoad && Package->HasAnyInternalFlags(EInternalObjectFlags::AsyncLoading))
+			if (ShouldAsyncLoad() && !ResolveContext.bForceSyncLoad && Package->HasAnyInternalFlags(EInternalObjectFlags_AsyncLoading))
 			{
 				// Something else is already async loading this package, calling load again will add our callback to the existing load request
 				StartAsyncLoadingPackage(*CacheObjectPtr, ObjectPathName, RefHandle, true);
diff --git a/Engine/Source/Runtime/Landscape/Private/LandscapeSplines.cpp b/Engine/Source/Runtime/Landscape/Private/LandscapeSplines.cpp
--- a/Engine/Source/Runtime/Landscape/Private/LandscapeSplines.cpp
+++ b/Engine/Source/Runtime/Landscape/Private/LandscapeSplines.cpp
@@ -2655,7 +2655,7 @@ void ULandscapeSplineSegment::PostInitProperties()
 
 #if WITH_EDITORONLY_DATA
 	if (!HasAnyFlags(RF_ClassDefaultObject | RF_NeedLoad) &&
-		!HasAnyInternalFlags(EInternalObjectFlags::AsyncLoading))
+		!HasAnyInternalFlags(EInternalObjectFlags_AsyncLoading))
 	{
 		// create a new random seed for all new objects
 		RandomSeed = FMath::Rand();
