Remove nodes when a add operation is followed by a remove or the reverse.
Some cleanup of the map usage

#tests Rna Overridable serialization unittests

diff --git a/Engine/Source/Runtime/CoreUObject/Private/UObject/OverriddenPropertySet.cpp b/Engine/Source/Runtime/CoreUObject/Private/UObject/OverriddenPropertySet.cpp
--- a/Engine/Source/Runtime/CoreUObject/Private/UObject/OverriddenPropertySet.cpp
+++ b/Engine/Source/Runtime/CoreUObject/Private/UObject/OverriddenPropertySet.cpp
@@ -162,7 +162,7 @@ FOverriddenPropertyNode& FOverriddenPropertySet::FindOrAddNode(FOverriddenProper
 	FOverriddenPropertyNodeID& SubNodeID = ParentNode.SubPropertyNodeKeys.FindOrAdd(NodeID, FOverriddenPropertyNodeID());
 	if (SubNodeID.IsValid())
 	{
-		FOverriddenPropertyNode* FoundNode = OverriddenPropertyNodes.FindByHash(GetTypeHash(SubNodeID), SubNodeID);
+		FOverriddenPropertyNode* FoundNode = OverriddenPropertyNodes.Find(SubNodeID);
 		checkf(FoundNode, TEXT("Expecting a node"));
 		return *FoundNode;
 	}
@@ -204,7 +204,7 @@ EOverriddenPropertyOperation FOverriddenPropertySet::GetOverriddenPropertyOperat
 		{
 			if (const FOverriddenPropertyNodeID* CurrentPropKey = OverriddenPropertyNode->SubPropertyNodeKeys.Find(CurrentProperty))
 			{
-				CurrentOverriddenPropertyNode = OverriddenPropertyNodes.FindByHash(GetTypeHash(*CurrentPropKey), *CurrentPropKey);
+				CurrentOverriddenPropertyNode = OverriddenPropertyNodes.Find(*CurrentPropKey);
 				checkf(CurrentOverriddenPropertyNode, TEXT("Expecting a node"));
 			}
 		}
@@ -248,7 +248,7 @@ EOverriddenPropertyOperation FOverriddenPropertySet::GetOverriddenPropertyOperat
 							const FOverriddenPropertyNodeID  SubObjectID(*SubObject);
 							if (const FOverriddenPropertyNodeID* CurrentPropKey = CurrentOverriddenPropertyNode->SubPropertyNodeKeys.Find(SubObjectID))
 							{
-								const FOverriddenPropertyNode* SubObjectOverriddenPropertyNode = OverriddenPropertyNodes.FindByHash(GetTypeHash(*CurrentPropKey), *CurrentPropKey);
+								const FOverriddenPropertyNode* SubObjectOverriddenPropertyNode = OverriddenPropertyNodes.Find(*CurrentPropKey);
 								checkf(SubObjectOverriddenPropertyNode, TEXT("Expecting a node"));
 								if (bOutInheritedOperation)
 								{
@@ -291,7 +291,7 @@ EOverriddenPropertyOperation FOverriddenPropertySet::GetOverriddenPropertyOperat
 
 					if (const FOverriddenPropertyNodeID* CurrentPropKey = CurrentOverriddenPropertyNode->SubPropertyNodeKeys.Find(OverriddenKeyID))
 					{
-						const FOverriddenPropertyNode* SubObjectOverriddenPropertyNode = OverriddenPropertyNodes.FindByHash(GetTypeHash(*CurrentPropKey), *CurrentPropKey);
+						const FOverriddenPropertyNode* SubObjectOverriddenPropertyNode = OverriddenPropertyNodes.Find(*CurrentPropKey);
 						checkf(SubObjectOverriddenPropertyNode, TEXT("Expecting a node"));
 						if (bOutInheritedOperation)
 						{
@@ -323,8 +323,8 @@ bool FOverriddenPropertySet::ClearOverriddenProperty(FOverriddenPropertyNode& Pa
 	const FEditPropertyChain::TDoubleLinkedListNode* PropertyIterator = PropertyNode;
 	FOverriddenPropertyNode* OverriddenPropertyNode = &ParentPropertyNode;
 	int32 ArrayIndex = INDEX_NONE;
-	TArray<FOverriddenPropertyNodeID> PropertyNodePath;
-	PropertyNodePath.Push(OverriddenPropertyNode->NodeID);
+	TArray<FOverriddenPropertyNodeID> PropertyNodeIDPath;
+	PropertyNodeIDPath.Push(OverriddenPropertyNode->NodeID);
 	while (PropertyIterator && (!OverriddenPropertyNode || OverriddenPropertyNode->Operation != EOverriddenPropertyOperation::Replace))
 	{
 		ArrayIndex = INDEX_NONE;
@@ -337,9 +337,9 @@ bool FOverriddenPropertySet::ClearOverriddenProperty(FOverriddenPropertyNode& Pa
 		{
 			if (const FOverriddenPropertyNodeID* CurrentPropKey = OverriddenPropertyNode->SubPropertyNodeKeys.Find(CurrentProperty))
 			{
-				CurrentOverriddenPropertyNode = OverriddenPropertyNodes.FindByHash(GetTypeHash(*CurrentPropKey), *CurrentPropKey);
+				CurrentOverriddenPropertyNode = OverriddenPropertyNodes.Find(*CurrentPropKey);
 				checkf(CurrentOverriddenPropertyNode, TEXT("Expecting a node"));
-				PropertyNodePath.Push(CurrentOverriddenPropertyNode->NodeID);
+				PropertyNodeIDPath.Push(CurrentOverriddenPropertyNode->NodeID);
 			}
 		}
 
@@ -398,7 +398,7 @@ bool FOverriddenPropertySet::ClearOverriddenProperty(FOverriddenPropertyNode& Pa
 								FOverriddenPropertyNodeID CurrentPropKey;
 								if (CurrentOverriddenPropertyNode->SubPropertyNodeKeys.RemoveAndCopyValue(SubObjectID, CurrentPropKey))
 								{
-									verifyf(OverriddenPropertyNodes.RemoveByHash(GetTypeHash(CurrentPropKey), CurrentPropKey), TEXT("Expecting a node to be removed"));
+									verifyf(OverriddenPropertyNodes.Remove(CurrentPropKey), TEXT("Expecting a node to be removed"));
 									OverridableManager.ClearInstancedSubObjectOverrides(*Owner, *SubObject);
 									return true;
 								}
@@ -456,7 +456,7 @@ bool FOverriddenPropertySet::ClearOverriddenProperty(FOverriddenPropertyNode& Pa
 				FOverriddenPropertyNodeID CurrentPropKey;
 				if (CurrentOverriddenPropertyNode->SubPropertyNodeKeys.RemoveAndCopyValue(OverriddenKeyID, CurrentPropKey))
 				{
-					verifyf(OverriddenPropertyNodes.RemoveByHash(GetTypeHash(CurrentPropKey), CurrentPropKey), TEXT("Expecting a node to be removed"));
+					verifyf(OverriddenPropertyNodes.Remove(CurrentPropKey), TEXT("Expecting a node to be removed"));
 
 					if (UObject* ValueSubObject = ValueInstancedObjectProperty ? ValueInstancedObjectProperty->GetObjectPropertyValue(MapHelper.GetValuePtr(InternalMapIndex)) : nullptr)
 					{
@@ -473,13 +473,13 @@ bool FOverriddenPropertySet::ClearOverriddenProperty(FOverriddenPropertyNode& Pa
 		PropertyIterator = PropertyIterator->GetNextNode();
 	}
 
-	auto CleanupClearedNodes = [this, &PropertyNodePath]()
+	auto CleanupClearedNodes = [this, &PropertyNodeIDPath]()
 	{
 		// Need to cleanup up the chain of property nodes if they endup empty
 		FOverriddenPropertyNodeID ChildPropertyNodeID;
-		while (FOverriddenPropertyNode* CurrentPropertyNode = !PropertyNodePath.IsEmpty() ? OverriddenPropertyNodes.FindByHash(GetTypeHash(PropertyNodePath.Top()), PropertyNodePath.Top()) : nullptr)
+		while (FOverriddenPropertyNode* CurrentPropertyNode = !PropertyNodeIDPath.IsEmpty() ? OverriddenPropertyNodes.Find(PropertyNodeIDPath.Top()) : nullptr)
 		{
-			PropertyNodePath.Pop();
+			PropertyNodeIDPath.Pop();
 			if (CurrentPropertyNode->SubPropertyNodeKeys.Num() > 1)
 			{
 				// Now need to remove the child from this node
@@ -487,9 +487,9 @@ bool FOverriddenPropertySet::ClearOverriddenProperty(FOverriddenPropertyNode& Pa
 				{
 					const FOverriddenPropertyNodeID* NodeToRemove = CurrentPropertyNode->SubPropertyNodeKeys.FindKey(ChildPropertyNodeID);
 					checkf(NodeToRemove, TEXT("Expecting a node"));
-					CurrentPropertyNode->SubPropertyNodeKeys.RemoveByHash(GetTypeHash(*NodeToRemove),*NodeToRemove);
+					CurrentPropertyNode->SubPropertyNodeKeys.Remove(*NodeToRemove);
 
-					verifyf(OverriddenPropertyNodes.RemoveByHash(GetTypeHash(ChildPropertyNodeID), ChildPropertyNodeID), TEXT("Expecting the node to be removed"));
+					verifyf(OverriddenPropertyNodes.Remove(ChildPropertyNodeID), TEXT("Expecting the node to be removed"));
 				}
 				break;
 			}
@@ -518,7 +518,7 @@ bool FOverriddenPropertySet::ClearOverriddenProperty(FOverriddenPropertyNode& Pa
 	return true;
 }
 
-void FOverriddenPropertySet::NotifyPropertyChange(FOverriddenPropertyNode* ParentPropertyNode, const EPropertyNotificationType Notification, const FPropertyChangedEvent& PropertyEvent, const FEditPropertyChain::TDoubleLinkedListNode* PropertyNode, const void* Data)
+void FOverriddenPropertySet::NotifyPropertyChange(FOverriddenPropertyNode* ParentPropertyNode, const EPropertyNotificationType Notification, const FPropertyChangedEvent& PropertyEvent, const FEditPropertyChain::TDoubleLinkedListNode* PropertyNode, const void* Data, bool& bNeedsCleanup)
 {
 	checkf(IsValid(Owner), TEXT("Expecting a valid overridable owner"));
 
@@ -534,7 +534,7 @@ void FOverriddenPropertySet::NotifyPropertyChange(FOverriddenPropertyNode* Paren
 			ParentPropertyNode->Operation = EOverriddenPropertyOperation::Replace;
 
 			// If we are overriding the root node, need to propagate the overrides to all instanced sub object
-			const FOverriddenPropertyNode* RootNode = OverriddenPropertyNodes.FindByHash(GetTypeHash(RootNodeID), RootNodeID);
+			const FOverriddenPropertyNode* RootNode = OverriddenPropertyNodes.Find(RootNodeID);
 			checkf(RootNode, TEXT("Expecting to always have a "));
 			if (RootNode == ParentPropertyNode)
 			{
@@ -556,6 +556,23 @@ void FOverriddenPropertySet::NotifyPropertyChange(FOverriddenPropertyNode* Paren
 		SubPropertyNode = SubPropertyNodeRef.Operation != EOverriddenPropertyOperation::Replace ? &SubPropertyNodeRef : nullptr;
 	}
 
+	ON_SCOPE_EXIT
+	{
+		if (bNeedsCleanup && ParentPropertyNode && SubPropertyNode && SubPropertyNode->SubPropertyNodeKeys.IsEmpty())
+		{
+			FOverriddenPropertyNodeID RemovedNodeID;
+			if (ParentPropertyNode->SubPropertyNodeKeys.RemoveAndCopyValue(Property, RemovedNodeID))
+			{
+				verifyf(OverriddenPropertyNodes.Remove(RemovedNodeID), TEXT("Expecting the node to be removed"));
+			}
+			if (ParentPropertyNode->Operation == EOverriddenPropertyOperation::Modified && ParentPropertyNode->SubPropertyNodeKeys.IsEmpty())
+			{
+				ParentPropertyNode->Operation = EOverriddenPropertyOperation::None;
+			}
+		}
+	};
+
+
 	if (const FArrayProperty* ArrayProperty = CastField<FArrayProperty>(Property))
 	{
 		// Only special case is instanced subobjects, otherwise we fallback to full array override
@@ -623,7 +640,12 @@ void FOverriddenPropertySet::NotifyPropertyChange(FOverriddenPropertyNode* Paren
 								if(RemovedSubObjectArchetype && !RemovedSubObjectArchetype->HasAnyFlags(RF_ClassDefaultObject))
 								{
 									const FOverriddenPropertyNodeID RemovedSubObjectID(*RemovedSubObjectArchetype);
-									SubPropertyNode->SubPropertyNodeKeys.Remove(RemovedSubObjectID);
+									FOverriddenPropertyNodeID RemovedNodeID;
+									if (SubPropertyNode->SubPropertyNodeKeys.RemoveAndCopyValue(RemovedSubObjectID, RemovedNodeID))
+									{
+										verifyf(OverriddenPropertyNodes.Remove(RemovedNodeID), TEXT("Expecting the node to be removed"));
+										bNeedsCleanup = true;
+									}
 								}
 							}
 						}
@@ -643,7 +665,12 @@ void FOverriddenPropertySet::NotifyPropertyChange(FOverriddenPropertyNode* Paren
 								if (RemovedSubObjectNode.Operation == EOverriddenPropertyOperation::Add)
 								{
 									// An add then a remove becomes no opt
-									SubPropertyNode->SubPropertyNodeKeys.Remove(RemovedSubObjectID);
+									FOverriddenPropertyNodeID RemovedNodeID;
+									if (SubPropertyNode->SubPropertyNodeKeys.RemoveAndCopyValue(RemovedSubObjectID, RemovedNodeID))
+									{
+										verifyf(OverriddenPropertyNodes.Remove(RemovedNodeID), TEXT("Expecting the node to be removed"));
+										bNeedsCleanup = true;
+									}
 								}
 								else
 								{
@@ -852,7 +879,12 @@ void FOverriddenPropertySet::NotifyPropertyChange(FOverriddenPropertyNode* Paren
 					if (RemovedKeyNode.Operation == EOverriddenPropertyOperation::Add)
 					{
 						// @Todo support remove/add/remove
-						SubPropertyNode->SubPropertyNodeKeys.Remove(RemovedKeyID);
+						FOverriddenPropertyNodeID RemovedNodeID;
+						if (SubPropertyNode->SubPropertyNodeKeys.RemoveAndCopyValue(RemovedKeyID, RemovedNodeID))
+						{
+							verifyf(OverriddenPropertyNodes.Remove(RemovedNodeID), TEXT("Expecting the node to be removed"));
+							bNeedsCleanup = true;
+						}
 					}
 					else
 					{
@@ -959,7 +991,7 @@ void FOverriddenPropertySet::NotifyPropertyChange(FOverriddenPropertyNode* Paren
 		}
 		else
 		{
-			NotifyPropertyChange(SubPropertyNode, Notification, PropertyEvent, PropertyNode->GetNextNode(), SubValuePtr);
+			NotifyPropertyChange(SubPropertyNode, Notification, PropertyEvent, PropertyNode->GetNextNode(), SubValuePtr, bNeedsCleanup);
 		}
 		return;
 	}
@@ -990,7 +1022,7 @@ void FOverriddenPropertySet::NotifyPropertyChange(FOverriddenPropertyNode* Paren
 		}
 		else if (OptionalProperty->IsSet(Data))
 		{
-			NotifyPropertyChange(SubPropertyNode, Notification, PropertyEvent, PropertyNode->GetNextNode(), OptionalProperty->GetValuePointerForRead(SubValuePtr));
+			NotifyPropertyChange(SubPropertyNode, Notification, PropertyEvent, PropertyNode->GetNextNode(), OptionalProperty->GetValuePointerForRead(SubValuePtr), bNeedsCleanup);
 		}
 		return;
 	}
@@ -1010,10 +1042,10 @@ void FOverriddenPropertySet::RemoveOverriddenSubProperties(FOverriddenPropertyNo
 {
 	for (const auto& Pair : PropertyNode.SubPropertyNodeKeys)
 	{
-		FOverriddenPropertyNode* RemovedPropertyNode = OverriddenPropertyNodes.FindByHash(GetTypeHash(Pair.Value), Pair.Value);
+		FOverriddenPropertyNode* RemovedPropertyNode = OverriddenPropertyNodes.Find(Pair.Value);
 		checkf(RemovedPropertyNode, TEXT("Expecting a node"));
 		RemoveOverriddenSubProperties(*RemovedPropertyNode);
-		verifyf(OverriddenPropertyNodes.RemoveByHash(GetTypeHash(Pair.Value), Pair.Value), TEXT("Expecting the node to be removed"));
+		verifyf(OverriddenPropertyNodes.Remove(Pair.Value), TEXT("Expecting the node to be removed"));
 	}
 	PropertyNode.Operation = EOverriddenPropertyOperation::None;
 	PropertyNode.SubPropertyNodeKeys.Empty();
@@ -1029,7 +1061,7 @@ void FOverriddenPropertySet::CacheArchetype()
 
 EOverriddenPropertyOperation FOverriddenPropertySet::GetOverriddenPropertyOperation(const FPropertyChangedEvent& PropertyEvent, const FEditPropertyChain::TDoubleLinkedListNode* PropertyNode, bool* bOutInheritedOperation /*= nullptr*/) const
 {
-	if (const FOverriddenPropertyNode* RootNode = OverriddenPropertyNodes.FindByHash(GetTypeHash(RootNodeID), RootNodeID))
+	if (const FOverriddenPropertyNode* RootNode = OverriddenPropertyNodes.Find(RootNodeID))
 	{
 		return GetOverriddenPropertyOperation(*RootNode, PropertyEvent, PropertyNode, bOutInheritedOperation, Owner);
 	}
@@ -1038,7 +1070,7 @@ EOverriddenPropertyOperation FOverriddenPropertySet::GetOverriddenPropertyOperat
 
 bool FOverriddenPropertySet::ClearOverriddenProperty(const FPropertyChangedEvent& PropertyEvent, const FEditPropertyChain::TDoubleLinkedListNode* PropertyNode)
 {
-	if (FOverriddenPropertyNode* RootNode = OverriddenPropertyNodes.FindByHash(GetTypeHash(RootNodeID), RootNodeID))
+	if (FOverriddenPropertyNode* RootNode = OverriddenPropertyNodes.Find(RootNodeID))
 	{
 		return ClearOverriddenProperty(*RootNode, PropertyEvent, PropertyNode, Owner);
 	}
@@ -1047,19 +1079,21 @@ bool FOverriddenPropertySet::ClearOverriddenProperty(const FPropertyChangedEvent
 
 void FOverriddenPropertySet::OverrideProperty(const FPropertyChangedEvent& PropertyEvent, const FEditPropertyChain::TDoubleLinkedListNode* PropertyNode, const void* Data)
 {
-	FOverriddenPropertyNode& RootPropertyNode = OverriddenPropertyNodes.FindOrAddByHash(GetTypeHash(RootNodeID), RootNodeID);
-	NotifyPropertyChange(&RootPropertyNode, EPropertyNotificationType::PreEdit, PropertyEvent, PropertyNode, Data);
-	NotifyPropertyChange(&RootPropertyNode, EPropertyNotificationType::PostEdit, PropertyEvent, PropertyNode, Data);
+	FOverriddenPropertyNode& RootPropertyNode = OverriddenPropertyNodes.FindOrAdd(RootNodeID);
+	bool bNeedsCleanup = false;
+	NotifyPropertyChange(&RootPropertyNode, EPropertyNotificationType::PreEdit, PropertyEvent, PropertyNode, Data, bNeedsCleanup);
+	NotifyPropertyChange(&RootPropertyNode, EPropertyNotificationType::PostEdit, PropertyEvent, PropertyNode, Data, bNeedsCleanup);
 }
 
 void FOverriddenPropertySet::NotifyPropertyChange(const EPropertyNotificationType Notification, const FPropertyChangedEvent& PropertyEvent, const FEditPropertyChain::TDoubleLinkedListNode* PropertyNode, const void* Data)
 {
-	NotifyPropertyChange(&OverriddenPropertyNodes.FindOrAddByHash(GetTypeHash(RootNodeID), RootNodeID), Notification, PropertyEvent, PropertyNode, Data);
+	bool bNeedsCleanup = false;
+	NotifyPropertyChange(&OverriddenPropertyNodes.FindOrAdd(RootNodeID), Notification, PropertyEvent, PropertyNode, Data, bNeedsCleanup);
 }
 
 EOverriddenPropertyOperation FOverriddenPropertySet::GetOverriddenPropertyOperation(const FArchiveSerializedPropertyChain* CurrentPropertyChain, FProperty* Property) const
 {
-	if (const FOverriddenPropertyNode* RootNode = OverriddenPropertyNodes.FindByHash(GetTypeHash(RootNodeID), RootNodeID))
+	if (const FOverriddenPropertyNode* RootNode = OverriddenPropertyNodes.Find(RootNodeID))
 	{
 		return GetOverriddenPropertyOperation(*RootNode, CurrentPropertyChain, Property);
 	}
@@ -1068,12 +1102,12 @@ EOverriddenPropertyOperation FOverriddenPropertySet::GetOverriddenPropertyOperat
 
 FOverriddenPropertyNode* FOverriddenPropertySet::SetOverriddenPropertyOperation(EOverriddenPropertyOperation Operation, const FArchiveSerializedPropertyChain* CurrentPropertyChain, FProperty* Property)
 {
-	return SetOverriddenPropertyOperation(Operation, OverriddenPropertyNodes.FindOrAddByHash(GetTypeHash(RootNodeID), RootNodeID), CurrentPropertyChain, Property);
+	return SetOverriddenPropertyOperation(Operation, OverriddenPropertyNodes.FindOrAdd(RootNodeID), CurrentPropertyChain, Property);
 }
 
 const FOverriddenPropertyNode* FOverriddenPropertySet::GetOverriddenPropertyNode(const FArchiveSerializedPropertyChain* CurrentPropertyChain) const
 {
-	if (const FOverriddenPropertyNode* RootNode = OverriddenPropertyNodes.FindByHash(GetTypeHash(RootNodeID), RootNodeID))
+	if (const FOverriddenPropertyNode* RootNode = OverriddenPropertyNodes.Find(RootNodeID))
 	{
 		return GetOverriddenPropertyNode(*RootNode, CurrentPropertyChain);
 	}
@@ -1104,7 +1138,7 @@ EOverriddenPropertyOperation FOverriddenPropertySet::GetOverriddenPropertyOperat
 		const FProperty* CurrentProperty = (*PropertyIterator);
 		if (const FOverriddenPropertyNodeID* CurrentPropKey = OverriddenPropertyNode->SubPropertyNodeKeys.Find(CurrentProperty))
 		{
-			OverriddenPropertyNode = OverriddenPropertyNodes.FindByHash(GetTypeHash(*CurrentPropKey), *CurrentPropKey);
+			OverriddenPropertyNode = OverriddenPropertyNodes.Find(*CurrentPropKey);
 			checkf(OverriddenPropertyNode, TEXT("Expecting a node"));
 		}
 		else
@@ -1156,7 +1190,7 @@ FOverriddenPropertyNode* FOverriddenPropertySet::SetOverriddenPropertyOperation(
 
 EOverriddenPropertyOperation FOverriddenPropertySet::GetSubPropertyOperation(FOverriddenPropertyNodeID NodeID) const
 {
-	const FOverriddenPropertyNode* OverriddenPropertyNode = OverriddenPropertyNodes.FindByHash(GetTypeHash(NodeID), NodeID);
+	const FOverriddenPropertyNode* OverriddenPropertyNode = OverriddenPropertyNodes.Find(NodeID);
 	return OverriddenPropertyNode ? OverriddenPropertyNode->Operation : EOverriddenPropertyOperation::None;
 }
 
@@ -1220,7 +1254,7 @@ const FOverriddenPropertyNode* FOverriddenPropertySet::GetOverriddenPropertyNode
 		const FProperty* CurrentProperty = (*PropertyIterator);
 		if (const FOverriddenPropertyNodeID* CurrentPropKey = OverriddenPropertyNode->SubPropertyNodeKeys.Find(CurrentProperty))
 		{
-			OverriddenPropertyNode = OverriddenPropertyNodes.FindByHash(GetTypeHash(*CurrentPropKey), *CurrentPropKey);
+			OverriddenPropertyNode = OverriddenPropertyNodes.Find(*CurrentPropKey);
 			checkf(OverriddenPropertyNode, TEXT("Expecting a node"));
 		}
 		else
diff --git a/Engine/Source/Runtime/CoreUObject/Public/UObject/OverriddenPropertySet.h b/Engine/Source/Runtime/CoreUObject/Public/UObject/OverriddenPropertySet.h
--- a/Engine/Source/Runtime/CoreUObject/Public/UObject/OverriddenPropertySet.h
+++ b/Engine/Source/Runtime/CoreUObject/Public/UObject/OverriddenPropertySet.h
@@ -313,7 +313,7 @@ protected:
 
 	EOverriddenPropertyOperation GetOverriddenPropertyOperation(const FOverriddenPropertyNode& ParentPropertyNode, const FPropertyChangedEvent& PropertyEvent, const FEditPropertyChain::TDoubleLinkedListNode* PropertyNode, bool* bOutInheritedOperation, const void* Data) const;
 	bool ClearOverriddenProperty(FOverriddenPropertyNode& ParentPropertyNode, const FPropertyChangedEvent& PropertyEvent, const FEditPropertyChain::TDoubleLinkedListNode* PropertyNode, const void* Data);
-	void NotifyPropertyChange(FOverriddenPropertyNode* ParentPropertyNode, const EPropertyNotificationType Notification, const FPropertyChangedEvent& PropertyEvent, const FEditPropertyChain::TDoubleLinkedListNode* PropertyNode, const void* Data);
+	void NotifyPropertyChange(FOverriddenPropertyNode* ParentPropertyNode, const EPropertyNotificationType Notification, const FPropertyChangedEvent& PropertyEvent, const FEditPropertyChain::TDoubleLinkedListNode* PropertyNode, const void* Data, bool& bNeedsCleanup);
 
 	EOverriddenPropertyOperation GetOverriddenPropertyOperation(const FOverriddenPropertyNode& ParentPropertyNode, const FArchiveSerializedPropertyChain* CurrentPropertyChain, FProperty* Property) const;
 	FOverriddenPropertyNode* SetOverriddenPropertyOperation(EOverriddenPropertyOperation Operation, FOverriddenPropertyNode& ParentPropertyNode, const FArchiveSerializedPropertyChain* CurrentPropertyChain, FProperty* Property);
