Update comment to reflect that we long backed away from replacing Cast with dynamic_cast (experiment began in 2312800)

#rnx

diff --git a/Engine/Source/Runtime/CoreUObject/Public/UObject/ObjectMacros.h b/Engine/Source/Runtime/CoreUObject/Public/UObject/ObjectMacros.h
--- a/Engine/Source/Runtime/CoreUObject/Public/UObject/ObjectMacros.h
+++ b/Engine/Source/Runtime/CoreUObject/Public/UObject/ObjectMacros.h
@@ -777,7 +777,7 @@ namespace UC
 		NotBlueprintable,
 
 		/// This keyword indicates that the class should be accessible outside of it's module, but does not need all methods exported.
-		/// It exports only the autogenerated methods required for dynamic_cast<>, etc... to work.
+		/// It exports only the autogenerated methods required for Cast<>, etc... to work.
 		MinimalAPI,
 
 		/// Prevents automatic generation of the constructor declaration.
@@ -892,7 +892,7 @@ namespace UI
 	enum 
 	{
 		/// This keyword indicates that the interface should be accessible outside of it's module, but does not need all methods exported.
-		/// It exports only the autogenerated methods required for dynamic_cast<>, etc... to work.
+		/// It exports only the autogenerated methods required for Cast<>, etc... to work.
 		MinimalAPI,
 
 		/// Specifies that this interface can be directly implemented by blueprints, this is implied if the interface has any blueprint events.
