FPackageId: Deprecate and replace ValueForDebugging with LexToString.

#rnx

diff --git a/Engine/Source/Developer/IoStoreUtilities/Private/IoStoreUtilities.cpp b/Engine/Source/Developer/IoStoreUtilities/Private/IoStoreUtilities.cpp
--- a/Engine/Source/Developer/IoStoreUtilities/Private/IoStoreUtilities.cpp
+++ b/Engine/Source/Developer/IoStoreUtilities/Private/IoStoreUtilities.cpp
@@ -1122,10 +1122,10 @@ public:
 	{
 		if (OutputArchive.IsValid())
 		{
-			OutputArchive->Logf(TEXT("%d, %s, 0x%llX, %s, %s, %s, %lld, %lld, %lld, %lld, 0x%s, %s, %s"),
+			OutputArchive->Logf(TEXT("%d, %s, 0x%s, %s, %s, %s, %lld, %lld, %lld, %lld, 0x%s, %s, %s"),
 				Index,
 				*LexToString(Info.Id),
-				PackageId.ValueForDebugging(),
+				*LexToString(PackageId),
 				*PackageName,
 				*Info.FileName,
 				*ContainerName,
@@ -6405,8 +6405,8 @@ bool ListContainerBulkData(
 				{
 					Sb.Reset();
 					LexToString(static_cast<EBulkDataFlags>(Entry.Flags), Sb);
-					CsvAr->Logf(TEXT("%s,%s,0x%llX,%lld,%lld,%s"),
-						*Container.Name, *Pkg.Filename, Pkg.Id.Value(), Entry.SerialOffset, Entry.SerialSize, Sb.ToString());
+					CsvAr->Logf(TEXT("%s,%s,0x%s,%lld,%lld,%s"),
+						*Container.Name, *Pkg.Filename, *LexToString(Pkg.Id), Entry.SerialOffset, Entry.SerialSize, Sb.ToString());
 				}
 			}
 		}
@@ -7317,7 +7317,8 @@ namespace DescribeUtils
 							Import.Export = ExportByKeyMap.FindRef(Key);
 							if (!Import.Export)
 							{
-								UE_LOG(LogIoStore, Warning, TEXT("Missing import: 0x%llX in package 0x%llX '%s'"), Import.GlobalImportIndex.Value(), PackageDesc->PackageId.ValueForDebugging(), *PackageDesc->PackageName.ToString());
+								UE_LOG(LogIoStore, Warning, TEXT("Missing import: 0x%llX in package 0x%s '%s'"),
+									Import.GlobalImportIndex.Value(), *LexToString(PackageDesc->PackageId), *PackageDesc->PackageName.ToString());
 							}
 							else
 							{
@@ -7333,7 +7334,7 @@ namespace DescribeUtils
 							}
 							else
 							{
-								UE_LOG(LogIoStore, Warning, TEXT("Missing Script Object for Import: 0x%llX in package 0x%llX '%s'"), Import.GlobalImportIndex.Value(), PackageDesc->PackageId.ValueForDebugging(), *PackageDesc->PackageName.ToString());
+								UE_LOG(LogIoStore, Warning, TEXT("Missing Script Object for Import: 0x%llX in package 0x%s '%s'"), Import.GlobalImportIndex.Value(), *LexToString(PackageDesc->PackageId), *PackageDesc->PackageName.ToString());
 							}
 						}
 					}
@@ -7550,7 +7551,7 @@ int32 Describe(
 						bNeedsHeader = false;
 					}
 					OutputOverride->Logf(ELogVerbosity::Display, TEXT("\t*************************"));
-					OutputOverride->Logf(ELogVerbosity::Display, TEXT("\t\t           Source: 0x%llX '%s'"), LocalizedPackage->PackageId.ValueForDebugging(), *LocalizedPackage->PackageName.ToString());
+					OutputOverride->Logf(ELogVerbosity::Display, TEXT("\t\t           Source: 0x%s '%s'"), *LexToString(LocalizedPackage->PackageId), *LocalizedPackage->PackageName.ToString());
 				}
 			}
 
@@ -7571,8 +7572,8 @@ int32 Describe(
 						bNeedsHeader = false;
 					}
 					OutputOverride->Logf(ELogVerbosity::Display, TEXT("\t*************************"));
-					OutputOverride->Logf(ELogVerbosity::Display, TEXT("\t\t           Source: 0x%llX '%s'"), Redirect.Source->PackageId.ValueForDebugging(), *Redirect.Source->PackageName.ToString());
-					OutputOverride->Logf(ELogVerbosity::Display, TEXT("\t\t           Target: 0x%llX '%s'"), Redirect.Target->PackageId.ValueForDebugging(), *Redirect.Target->PackageName.ToString());
+					OutputOverride->Logf(ELogVerbosity::Display, TEXT("\t\t           Source: 0x%s '%s'"), *LexToString(Redirect.Source->PackageId), *Redirect.Source->PackageName.ToString());
+					OutputOverride->Logf(ELogVerbosity::Display, TEXT("\t\t           Target: 0x%s '%s'"), *LexToString(Redirect.Target->PackageId), *Redirect.Target->PackageName.ToString());
 				}
 			}
 		}
@@ -7583,7 +7584,7 @@ int32 Describe(
 			OutputOverride->Logf(ELogVerbosity::Display, TEXT("Package '%s' Summary"), *PackageDesc->PackageName.ToString());
 			OutputOverride->Logf(ELogVerbosity::Display, TEXT("--------------------------------------------"));
 
-			OutputOverride->Logf(ELogVerbosity::Display, TEXT("\t\t        PackageId: 0x%llX"), PackageDesc->PackageId.ValueForDebugging());
+			OutputOverride->Logf(ELogVerbosity::Display, TEXT("\t\t        PackageId: 0x%s"), *LexToString(PackageDesc->PackageId));
 			OutputOverride->Logf(ELogVerbosity::Display, TEXT("\t\t     PackageFlags: %X"), PackageDesc->PackageFlags);
 			OutputOverride->Logf(ELogVerbosity::Display, TEXT("\t\t        NameCount: %d"), PackageDesc->NameCount);
 			OutputOverride->Logf(ELogVerbosity::Display, TEXT("\t\t      ImportCount: %d"), PackageDesc->Imports.Num());
diff --git a/Engine/Source/Editor/UnrealEd/Private/Cooker/IoStoreCookOnTheFlyRequestManager.cpp b/Engine/Source/Editor/UnrealEd/Private/Cooker/IoStoreCookOnTheFlyRequestManager.cpp
--- a/Engine/Source/Editor/UnrealEd/Private/Cooker/IoStoreCookOnTheFlyRequestManager.cpp
+++ b/Engine/Source/Editor/UnrealEd/Private/Cooker/IoStoreCookOnTheFlyRequestManager.cpp
@@ -205,29 +205,29 @@ private:
 			FName PackageName = Manager.PackageRegistry.Get(PackageId);
 			if (Package.Status == EPackageStatus::Cooked)
 			{
-				UE_LOG(LogCookOnTheFly, Verbose, TEXT("0x%llX was already cooked"), PackageId.ValueForDebugging());
+				UE_LOG(LogCookOnTheFly, Verbose, TEXT("0x%s was already cooked"), *LexToString(PackageId));
 				OutEntry = Package.Entry;
 				return EPackageStoreEntryStatus::Ok;
 			}
 			else if (Package.Status == EPackageStatus::Failed)
 			{
-				UE_LOG(LogCookOnTheFly, Warning, TEXT("Failed to cook package 0x%llX '%s'"), PackageId.ValueForDebugging(), *PackageName.ToString());
+				UE_LOG(LogCookOnTheFly, Warning, TEXT("Failed to cook package 0x%s '%s'"), *LexToString(PackageId), *PackageName.ToString());
 				return EPackageStoreEntryStatus::Missing;
 			}
 			else if (Package.Status == EPackageStatus::Cooking)
 			{
-				UE_LOG(LogCookOnTheFly, Verbose, TEXT("0x%llX was already cooking"), PackageId.ValueForDebugging());
+				UE_LOG(LogCookOnTheFly, Verbose, TEXT("0x%s was already cooking"), *LexToString(PackageId));
 				return EPackageStoreEntryStatus::Pending;
 			}
 			if (PackageName.IsNone())
 			{
-				UE_LOG(LogCookOnTheFly, Warning, TEXT("Received cook request for unknown package 0x%llX"), PackageId.ValueForDebugging());
+				UE_LOG(LogCookOnTheFly, Warning, TEXT("Received cook request for unknown package 0x%s"), *LexToString(PackageId));
 				return EPackageStoreEntryStatus::Missing;
 			}
 			FString Filename;
 			if (FPackageName::TryConvertLongPackageNameToFilename(PackageName.ToString(), Filename))
 			{
-				UE_LOG(LogCookOnTheFly, Verbose, TEXT("Cooking package 0x%llX '%s'"), PackageId.ValueForDebugging(), *PackageName.ToString());
+				UE_LOG(LogCookOnTheFly, Verbose, TEXT("Cooking package 0x%s '%s'"), *LexToString(PackageId), *PackageName.ToString());
 				Package.Status = EPackageStatus::Cooking;
 				const bool bEnqueued = InCookOnTheFlyServer.EnqueueCookRequest(UE::Cook::FCookPackageRequest{ PlatformName, Filename });
 				check(bEnqueued);
@@ -235,7 +235,7 @@ private:
 			}
 			else
 			{
-				UE_LOG(LogCookOnTheFly, Warning, TEXT("Failed to cook package 0x%llX '%s' (File not found)"), PackageId.ValueForDebugging(), *PackageName.ToString());
+				UE_LOG(LogCookOnTheFly, Warning, TEXT("Failed to cook package 0x%s '%s' (File not found)"), *LexToString(PackageId), *PackageName.ToString());
 				Package.Status = EPackageStatus::Failed;
 				return EPackageStoreEntryStatus::Missing;
 			}
@@ -246,20 +246,20 @@ private:
 			FPackage& Package = GetPackage(PackageId);
 			if (Package.Status != EPackageStatus::Cooked && Package.Status != EPackageStatus::Failed)
 			{
-				UE_LOG(LogCookOnTheFly, Verbose, TEXT("Skipping recook of package 0x%llX '%s' that was not cooked"), PackageId.ValueForDebugging(), *PackageName.ToString());
+				UE_LOG(LogCookOnTheFly, Verbose, TEXT("Skipping recook of package 0x%s '%s' that was not cooked"), *LexToString(PackageId), *PackageName.ToString());
 				return;
 			}
 			FString Filename;
 			if (FPackageName::TryConvertLongPackageNameToFilename(PackageName.ToString(), Filename))
 			{
-				UE_LOG(LogCookOnTheFly, Verbose, TEXT("Recooking package 0x%llX '%s'"), PackageId.ValueForDebugging(), *PackageName.ToString());
+				UE_LOG(LogCookOnTheFly, Verbose, TEXT("Recooking package 0x%s '%s'"), *LexToString(PackageId), *PackageName.ToString());
 				Package.Status = EPackageStatus::Cooking;
 				const bool bEnqueued = InCookOnTheFlyServer.EnqueueCookRequest(UE::Cook::FCookPackageRequest{ PlatformName, Filename });
 				check(bEnqueued);
 			}
 			else
 			{
-				UE_LOG(LogCookOnTheFly, Warning, TEXT("Failed to recook package 0x%llX '%s' (File not found)"), PackageId.ValueForDebugging(), *PackageName.ToString());
+				UE_LOG(LogCookOnTheFly, Warning, TEXT("Failed to recook package 0x%s '%s' (File not found)"), *LexToString(PackageId), *PackageName.ToString());
 				Package.Status = EPackageStatus::Failed;
 			}
 		}
@@ -278,7 +278,7 @@ private:
 
 		void MarkAsCooked(FPackageId PackageId, const FPackageStoreEntryResource& Entry, UE::ZenCookOnTheFly::Messaging::FCompletedPackages& OutCompletedPackages)
 		{
-			UE_LOG(LogCookOnTheFly, Verbose, TEXT("0x%llX cooked"), PackageId.ValueForDebugging());
+			UE_LOG(LogCookOnTheFly, Verbose, TEXT("0x%s cooked"), *LexToString(PackageId));
 			FPackage& Package = GetPackage(PackageId);
 			Package.Status = EPackageStatus::Cooked;
 			Package.Entry = Entry;
@@ -407,7 +407,7 @@ private:
 	virtual void OnPackageGenerated(const FName& PackageName) override
 	{
 		FPackageId PackageId = FPackageId::FromName(PackageName);
-		UE_LOG(LogCookOnTheFly, Verbose, TEXT("Package 0x%llX '%s' generated"), PackageId.ValueForDebugging(), *PackageName.ToString());
+		UE_LOG(LogCookOnTheFly, Verbose, TEXT("Package 0x%s '%s' generated"), *LexToString(PackageId), *PackageName.ToString());
 		PackageRegistry.Add(PackageName);
 	}
 
@@ -643,7 +643,7 @@ private:
 
 		for (const FPackageId& PackageId : CookRequest.PackageIds)
 		{
-			UE_LOG(LogCookOnTheFly, Verbose, TEXT("Received cook request 0x%llX"), PackageId.ValueForDebugging());
+			UE_LOG(LogCookOnTheFly, Verbose, TEXT("Received cook request 0x%s"), *LexToString(PackageId));
 
 			FPackageStoreEntryResource Entry;
 			EPackageStoreEntryStatus PackageStatus = EPackageStoreEntryStatus::Pending;
diff --git a/Engine/Source/Runtime/Core/Private/IO/PackageId.cpp b/Engine/Source/Runtime/Core/Private/IO/PackageId.cpp
--- a/Engine/Source/Runtime/Core/Private/IO/PackageId.cpp
+++ b/Engine/Source/Runtime/Core/Private/IO/PackageId.cpp
@@ -84,3 +84,9 @@ void SerializeForLog(FCbWriter& Writer, const FPackageId& Value)
 
 	Writer.EndObject();
 }
+
+FString LexToString(const FPackageId& PackageId)
+{
+	return FString::Printf(TEXT("%llX"), PackageId.Value());
+}
+
diff --git a/Engine/Source/Runtime/Core/Public/IO/PackageId.h b/Engine/Source/Runtime/Core/Public/IO/PackageId.h
--- a/Engine/Source/Runtime/Core/Public/IO/PackageId.h
+++ b/Engine/Source/Runtime/Core/Public/IO/PackageId.h
@@ -3,6 +3,7 @@
 #pragma once
 
 #include "CoreMinimal.h"
+#include "Containers/StringFwd.h"
 #include "HAL/Platform.h"
 #include "Misc/AssertionMacros.h"
 #include "UObject/NameTypes.h"
@@ -38,10 +39,10 @@ public:
 
 	inline uint64 Value() const
 	{
-		check(Id != InvalidId);
 		return Id;
 	}
 
+	UE_DEPRECATED(5.5, "Use LexToString()")
 	inline uint64 ValueForDebugging() const
 	{
 		return Id;
@@ -77,3 +78,6 @@ public:
 	CORE_API FName GetName() const;
 #endif
 };
+
+CORE_API FString LexToString(const FPackageId& PackageId);
+
diff --git a/Engine/Source/Runtime/CoreUObject/Private/Serialization/AsyncLoading.cpp b/Engine/Source/Runtime/CoreUObject/Private/Serialization/AsyncLoading.cpp
--- a/Engine/Source/Runtime/CoreUObject/Private/Serialization/AsyncLoading.cpp
+++ b/Engine/Source/Runtime/CoreUObject/Private/Serialization/AsyncLoading.cpp
@@ -6599,7 +6599,7 @@ EAsyncPackageState::Type FAsyncPackage::LoadImports()
 					PackageLocation = FPackageName::DoesPackageExistEx(ImportedPackagePath, PackageLocation, false /* bMatchCaseOnDisk */, nullptr);
 					if (PackageLocation == FPackageName::EPackageLocationFilter::IoDispatcher)
 					{
-						UE_LOG(LogStreaming, Verbose, TEXT("FAsyncPackage::LoadImports for %s: Loading %s (0x%llX) using IoStore package loader"), *Desc.PackagePath.GetDebugName(), *ImportPackageName, FPackageId::FromName(ImportPackageFName).ValueForDebugging());
+						UE_LOG(LogStreaming, Verbose, TEXT("FAsyncPackage::LoadImports for %s: Loading %s (0x%s) using IoStore package loader"), *Desc.PackagePath.GetDebugName(), *ImportPackageName, *LexToString(FPackageId::FromName(ImportPackageFName)));
 						check(IsInGameThread());
 						FAsyncLoadingThread::LeaveAsyncLoadingTick(AsyncLoadingThread.GetThreadIndex()); // Fix thread check that triggers if we recurse back into the uncooked package loader from here
 						int32 ImportRequestId = IoStorePackageLoader->LoadPackage(ImportedPackagePath, NAME_None, FLoadPackageAsyncDelegate(), PKG_None, INDEX_NONE, 0, nullptr, LOAD_None);
diff --git a/Engine/Source/Runtime/CoreUObject/Private/Serialization/AsyncLoading2.cpp b/Engine/Source/Runtime/CoreUObject/Private/Serialization/AsyncLoading2.cpp
--- a/Engine/Source/Runtime/CoreUObject/Private/Serialization/AsyncLoading2.cpp
+++ b/Engine/Source/Runtime/CoreUObject/Private/Serialization/AsyncLoading2.cpp
@@ -356,11 +356,11 @@ if ((ELogVerbosity::Type(ELogVerbosity::Verbosity & ELogVerbosity::VerbosityMask
 	(GAsyncLoading2_VerboseLogFilter == 2) || \
 	(GAsyncLoading2_VerboseLogFilter == 1 && GAsyncLoading2_VerbosePackageIds.Contains((PackageDesc).UPackageId))) \
 { \
-	UE_LOG(LogStreaming, Verbosity, LogDesc TEXT(": %s (0x%llX) %s (0x%llX) - ") Format, \
+	UE_LOG(LogStreaming, Verbosity, LogDesc TEXT(": %s (0x%s) %s (0x%s) - ") Format, \
 		*(PackageDesc).UPackageName.ToString(), \
-		(PackageDesc).UPackageId.ValueForDebugging(), \
+		*LexToString((PackageDesc).UPackageId), \
 		*(PackageDesc).PackagePathToLoad.GetPackageFName().ToString(), \
-		(PackageDesc).PackageIdToLoad.ValueForDebugging(), \
+		*LexToString((PackageDesc).PackageIdToLoad), \
 		##__VA_ARGS__); \
 }
 
@@ -446,9 +446,9 @@ TRACE_DECLARE_ATOMIC_MEMORY_COUNTER(AsyncLoadingTotalLoaded, TEXT("AsyncLoading/
 FString FormatPackageId(FPackageId PackageId)
 {
 #if WITH_PACKAGEID_NAME_MAP
-	return FString::Printf(TEXT("0x%llX (%s)"), PackageId.ValueForDebugging(), *PackageId.GetName().ToString());
+	return FString::Printf(TEXT("0x%s (%s)"), *LexToString(PackageId), *PackageId.GetName().ToString());
 #else
-	return FString::Printf(TEXT("0x%llX"), PackageId.ValueForDebugging());
+	return FString::Printf(TEXT("0x%s"), *LexToString(PackageId));
 #endif
 }
 
@@ -980,8 +980,8 @@ public:
 	inline FLoadedPackageRef& FindPackageRefChecked(FPackageId PackageId, FName Name = FName())
 	{
 		FLoadedPackageRef* PackageRef = FindPackageRef(PackageId);
-		UE_CLOG(!PackageRef, LogStreaming, Fatal, TEXT("FindPackageRefChecked: Package %s (0x%llX) has been deleted"),
-			*Name.ToString(), PackageId.ValueForDebugging());
+		UE_CLOG(!PackageRef, LogStreaming, Fatal, TEXT("FindPackageRefChecked: Package %s (0x%s) has been deleted"),
+			*Name.ToString(), *LexToString(PackageId));
 		return *PackageRef;
 	}
 
@@ -1042,11 +1042,11 @@ public:
 						if (OldPackageRef)
 						{
 							UE_LOG(LogStreaming, Display,
-								TEXT("FGlobalImportStore:AddPackageRef: Dropping stale reference to package %s (0x%llX) that has been renamed to %s (0x%llX)"),
+								TEXT("FGlobalImportStore:AddPackageRef: Dropping stale reference to package %s (0x%s) that has been renamed to %s (0x%s)"),
 								*OldPackageRef->GetOriginalPackageName().ToString(),
-								OldPackageId.ValueForDebugging(),
+								*LexToString(OldPackageId),
 								*FoundPackage->GetName(),
-								PackageId.ValueForDebugging()
+								*LexToString(PackageId)
 							);
 							
 							check(OldPackageRef->GetRefCount() == 0);
@@ -1059,9 +1059,9 @@ public:
 					{
 						PackageRef.bIsMissing = false;
 						UE_LOG(LogStreaming, Warning,
-							TEXT("FGlobalImportStore:AddPackageRef: Found reference to previously missing package %s (0x%llX)"),
+							TEXT("FGlobalImportStore:AddPackageRef: Found reference to previously missing package %s (0x%s)"),
 							*FoundPackage->GetName(),
-							PackageId.ValueForDebugging()
+							*LexToString(PackageId)
 						);
 					}
 
@@ -1083,10 +1083,10 @@ public:
 				if (Package->IsUnreachable() || PackageRef.GetOriginalPackageName() != Package->GetFName())
 				{
 					UE_CLOG(!Package->IsUnreachable(), LogStreaming, Display,
-						TEXT("FGlobalImportStore:AddPackageRef: Dropping renamed package %s before reloading %s (0x%llX)"),
+						TEXT("FGlobalImportStore:AddPackageRef: Dropping renamed package %s before reloading %s (0x%s)"),
 						*Package->GetName(),
 						*PackageRef.GetOriginalPackageName().ToString(),
-						Package->GetPackageId().ValueForDebugging());
+						*LexToString(Package->GetPackageId()));
 
 					RemoveUnreferencedObsoletePackage(PackageRef);
 				}
@@ -1114,11 +1114,11 @@ public:
 
 #if DO_CHECK
 		ensureMsgf(!PackageRef.bHasBeenLoadedDebug || PackageRef.bAreAllPublicExportsLoaded || PackageRef.bIsMissing || PackageRef.bHasFailed,
-			TEXT("LoadedPackageRef from None (0x%llX) to %s (0x%llX) should not have been released when the package is not complete.")
+			TEXT("LoadedPackageRef from None (0x%s) to %s (0x%s) should not have been released when the package is not complete.")
 			TEXT("RefCount=%d, AreAllExportsLoaded=%d, IsMissing=%d, HasFailed=%d, HasBeenLoaded=%d"),
-			FromPackageId.ValueForDebugging(),
+			*LexToString(FromPackageId),
 			*PackageRef.GetOriginalPackageName().ToString(),
-			PackageId.Value(),
+			*LexToString(PackageId),
 			PackageRef.RefCount,
 			PackageRef.bAreAllPublicExportsLoaded,
 			PackageRef.bIsMissing,
@@ -1295,10 +1295,10 @@ public:
 		FPackageId PackageId = Package->GetPackageId();
 
 		UE_CLOG(PackageRef.GetRefCount() > 0, LogStreaming, Fatal,
-			TEXT("FGlobalImportStore::VerifyPackageForRemoval: %s (0x%llX) - ")
+			TEXT("FGlobalImportStore::VerifyPackageForRemoval: %s (0x%s) - ")
 			TEXT("Package removed while still being referenced, RefCount %d > 0."),
 			*Package->GetName(),
-			PackageId.ValueForDebugging(),
+			*LexToString(PackageId),
 			PackageRef.GetRefCount());
 
 		for (auto It = PackageRef.GetPublicExportObjectIndices(); It; ++It)
@@ -4337,7 +4337,7 @@ bool FAsyncLoadingThread2::CreateAsyncPackagesFromQueue(FAsyncLoadingThreadState
 		if (PendingPackageStatus == EPackageStoreEntryStatus::Ok)
 		{
 			SCOPED_CUSTOM_LOADTIMER(CreateAsyncPackage)
-				ADD_CUSTOM_LOADTIMER_META(CreateAsyncPackage, PackageId, PendingPackage->Desc.UPackageId.ValueForDebugging())
+				ADD_CUSTOM_LOADTIMER_META(CreateAsyncPackage, PackageId, PendingPackage->Desc.UPackageId.Value())
 				ADD_CUSTOM_LOADTIMER_META(CreateAsyncPackage, PackageName, *WriteToString<256>(PendingPackage->Desc.UPackageName));
 			InitializeAsyncPackageFromPackageStore(ThreadState, &IoBatch, PendingPackage, PackageEntry);
 			PendingPackage->StartLoading(ThreadState, IoBatch);
@@ -4480,7 +4480,7 @@ bool FAsyncLoadingThread2::CreateAsyncPackagesFromQueue(FAsyncLoadingThreadState
 					else
 					{
 						SCOPED_CUSTOM_LOADTIMER(CreateAsyncPackage)
-							ADD_CUSTOM_LOADTIMER_META(CreateAsyncPackage, PackageId, PackageDesc.UPackageId.ValueForDebugging())
+							ADD_CUSTOM_LOADTIMER_META(CreateAsyncPackage, PackageId, PackageDesc.UPackageId.Value())
 							ADD_CUSTOM_LOADTIMER_META(CreateAsyncPackage, PackageName, NameBuffer);
 
 						check(PackageStatus == EPackageStoreEntryStatus::Ok);
@@ -5077,7 +5077,7 @@ void FAsyncPackage2::ImportPackagesRecursiveInner(FAsyncLoadingThreadState2& Thr
 			UPackage* UncookedPackage = ImportedPackageRef.GetPackage();
 			if (!ImportedPackageRef.AreAllPublicExportsLoaded())
 			{
-				UE_ASYNC_PACKAGE_LOG(Verbose, Desc, TEXT("ImportPackages: LoadUncookedImport"), TEXT("Loading imported uncooked package '%s' '0x%llX'"), *ImportedPackageNameToLoad.ToString(), ImportedPackageId.ValueForDebugging());
+				UE_ASYNC_PACKAGE_LOG(Verbose, Desc, TEXT("ImportPackages: LoadUncookedImport"), TEXT("Loading imported uncooked package '%s' '0x%s'"), *ImportedPackageNameToLoad.ToString(), *LexToString(ImportedPackageId));
 				check(IsInGameThread());
 				IoBatch.Issue(); // The batch might already contain requests for packages being imported from the uncooked one we're going to load so make sure that those are started before blocking
 				int32 ImportRequestId = AsyncLoadingThread.UncookedPackageLoader->LoadPackage(ImportedPackagePath, NAME_None, FLoadPackageAsyncDelegate(), PKG_None, INDEX_NONE, 0, nullptr, LOAD_None);
@@ -5129,7 +5129,7 @@ void FAsyncPackage2::ImportPackagesRecursiveInner(FAsyncLoadingThreadState2& Thr
 			{
 				ImportedPackageRef.SetHasFailed();
 				UE_ASYNC_PACKAGE_LOG(Warning, Desc, TEXT("ImportPackages: SkipPackage"),
-					TEXT("Failed to load uncooked imported package with id '0x%llX' ('%s')"), ImportedPackageId.ValueForDebugging(), *ImportedPackageNameToLoad.ToString());
+					TEXT("Failed to load uncooked imported package with id '0x%s' ('%s')"), *LexToString(ImportedPackageId), *ImportedPackageNameToLoad.ToString());
 			}
 			continue;
 		}
@@ -5160,7 +5160,7 @@ void FAsyncPackage2::ImportPackagesRecursiveInner(FAsyncLoadingThreadState2& Thr
 			if (!ImportedPackageRef.HasPackage()) // If we found a package it's not actually missing but we can't load it anyway
 			{
 				UE_ASYNC_PACKAGE_CLOG(!ImportedPackageUPackageName.IsNone(), Display, Desc, TEXT("ImportPackages: SkipPackage"),
-					TEXT("Skipping non mounted imported package %s (0x%llX)"), *ImportedPackageNameToLoad.ToString(), ImportedPackageId.ValueForDebugging());
+					TEXT("Skipping non mounted imported package %s (0x%s)"), *ImportedPackageNameToLoad.ToString(), *LexToString(ImportedPackageId));
 				if (!ImportedPackageUPackageName.IsNone())
 				{
 					FName PackagePathToLoadName = Desc.PackagePathToLoad.GetPackageFName();
@@ -9084,8 +9084,8 @@ void FAsyncPackage2::CreateUPackage()
 		else if (ExistingPackage != LinkerRoot)
 		{
 			UE_ASYNC_PACKAGE_LOG(Warning, Desc, TEXT("CreateUPackage: ReplacePackage"),
-				TEXT("Replacing renamed package %s (0x%llX) while being referenced by the loader, RefCount=%d"),
-				*ExistingPackage->GetName(), ExistingPackage->GetPackageId().ValueForDebugging(),
+				TEXT("Replacing renamed package %s (0x%s) while being referenced by the loader, RefCount=%d"),
+				*ExistingPackage->GetName(), *LexToString(ExistingPackage->GetPackageId()),
 				PackageRef.GetRefCount());
 
 			AsyncLoadingThread.GlobalImportStore.ReplaceReferencedRenamedPackage(PackageRef, LinkerRoot);
@@ -9338,8 +9338,8 @@ void FAsyncLoadingThread2::QueueMissingPackage(FAsyncLoadingThreadState2& Thread
 	}
 
 	UE_CLOG(bIssueWarning, LogStreaming, Warning,
-		TEXT("LoadPackage: SkipPackage: %s (0x%llX) - The package to load does not exist on disk or in the loader"),
-		*FailedPackageName.ToString(), PackageDesc.PackageIdToLoad.ValueForDebugging());
+		TEXT("LoadPackage: SkipPackage: %s (0x%s) - The package to load does not exist on disk or in the loader"),
+		*FailedPackageName.ToString(), *LexToString(PackageDesc.PackageIdToLoad));
 
 	if (PackageProgressDelegate.IsValid())
 	{
diff --git a/Engine/Source/Runtime/PakFile/Private/FilePackageStore.cpp b/Engine/Source/Runtime/PakFile/Private/FilePackageStore.cpp
--- a/Engine/Source/Runtime/PakFile/Private/FilePackageStore.cpp
+++ b/Engine/Source/Runtime/PakFile/Private/FilePackageStore.cpp
@@ -361,7 +361,7 @@ bool FFilePackageStoreBackend::GetPackageRedirectInfo(FPackageId PackageId, FNam
 	{
 		OutSourcePackageName = FindRedirect->Get<0>();
 		OutRedirectedToPackageId = FindRedirect->Get<1>();
-		UE_LOG(LogFilePackageStore, Verbose, TEXT("Redirecting from %s to 0x%llx"), *OutSourcePackageName.ToString(), OutRedirectedToPackageId.Value());
+		UE_LOG(LogFilePackageStore, Verbose, TEXT("Redirecting from %s to 0x%s"), *OutSourcePackageName.ToString(), *LexToString(OutRedirectedToPackageId));
 		return true;
 	}
 	
@@ -376,7 +376,7 @@ bool FFilePackageStoreBackend::GetPackageRedirectInfo(FPackageId PackageId, FNam
 			{
 				OutSourcePackageName = *FindLocalizedPackageSourceName;
 				OutRedirectedToPackageId = LocalizedPackageId;
-				UE_LOG(LogFilePackageStore, Verbose, TEXT("Redirecting from localized package %s to 0x%llx"), *OutSourcePackageName.ToString(), OutRedirectedToPackageId.Value());
+				UE_LOG(LogFilePackageStore, Verbose, TEXT("Redirecting from localized package %s to 0x%s"), *OutSourcePackageName.ToString(), *LexToString(OutRedirectedToPackageId));
 				return true;
 			}
 		}
diff --git a/Engine/Source/Runtime/StorageServerClient/Private/CookOnTheFlyPackageStore.cpp b/Engine/Source/Runtime/StorageServerClient/Private/CookOnTheFlyPackageStore.cpp
--- a/Engine/Source/Runtime/StorageServerClient/Private/CookOnTheFlyPackageStore.cpp
+++ b/Engine/Source/Runtime/StorageServerClient/Private/CookOnTheFlyPackageStore.cpp
@@ -167,8 +167,8 @@ void FCookOnTheFlyPackageStoreBackend::AddPackages(TArray<FPackageStoreEntryReso
 				FailedName = FName(TEXT("<Unknown>"));
 			}
 
-			UE_LOG(LogCookOnTheFly, Warning, TEXT("0x%llX [Failed]. Failed to cook package %s."),
-				FailedPackageId.ValueForDebugging(), *FailedName.ToString());
+			UE_LOG(LogCookOnTheFly, Warning, TEXT("0x%s [Failed]. Failed to cook package %s."),
+				*LexToString(FailedPackageId), *FailedName.ToString());
 		}
 		EntryInfo.Status = EPackageStoreEntryStatus::Missing;
 		PackageStats.Failed++;
@@ -187,8 +187,8 @@ void FCookOnTheFlyPackageStoreBackend::AddPackages(TArray<FPackageStoreEntryReso
 		PackageEntries[EntryInfo.EntryIndex] = MoveTemp(Entry);
 		PackageStats.Cooked++;
 
-		UE_LOG(LogCookOnTheFly, Verbose, TEXT("0x%llX '%s' [OK] (Cooked/Failed='%d/%d')"),
-			PackageId.ValueForDebugging(), *PackageEntries[EntryInfo.EntryIndex].PackageName.ToString(), PackageStats.Cooked.Load(), PackageStats.Failed.Load());
+		UE_LOG(LogCookOnTheFly, Verbose, TEXT("0x%s '%s' [OK] (Cooked/Failed='%d/%d')"),
+			*LexToString(PackageId), *PackageEntries[EntryInfo.EntryIndex].PackageName.ToString(), PackageStats.Cooked.Load(), PackageStats.Failed.Load());
 	}
 }
 
@@ -246,7 +246,7 @@ void FCookOnTheFlyPackageStoreBackend::CheckActivity()
 				const FEntryInfo& EntryInfo = KeyValue.Value;
 				if (EntryInfo.Status == EPackageStoreEntryStatus::Pending)
 				{
-					UE_LOG(LogCookOnTheFly, Log, TEXT("0x%llX"), KeyValue.Key.ValueForDebugging());
+					UE_LOG(LogCookOnTheFly, Log, TEXT("0x%s"), *LexToString(KeyValue.Key));
 				}
 			}
 		}
