[ObjectVisibility]
 - Activate object visibility feature so that objects still being initialized and serialized aren't discovered by the game thread during postload and other steps run from inside the loader tick

diff --git a/Engine/Source/Runtime/CoreUObject/Private/UObject/ObjectVisibility.cpp b/Engine/Source/Runtime/CoreUObject/Private/UObject/ObjectVisibility.cpp
--- a/Engine/Source/Runtime/CoreUObject/Private/UObject/ObjectVisibility.cpp
+++ b/Engine/Source/Runtime/CoreUObject/Private/UObject/ObjectVisibility.cpp
@@ -4,7 +4,7 @@
 #include "UObject/UObjectThreadContext.h"
 #include "HAL/IConsoleManager.h"
 
-static bool GUseObjectVisibilityFilterForAsyncLoading = false;
+static bool GUseObjectVisibilityFilterForAsyncLoading = true;
 static FAutoConsoleVariableRef CVarUseObjectVisibilityFilterForAsyncLoading(
 	TEXT("s.UseObjectVisibilityFilterForAsyncLoading"),
 	GUseObjectVisibilityFilterForAsyncLoading,
