Add support for allowing targets to opt into using the engine metadata API even in runtime by defining WITH_METADATA in a project target file.

#tests Tested with and without WITH_METADATA in client/editor, for both uncooked and cooked.
[RN] Add support for allowing targets to opt into using the engine metadata API even in runtime by defining WITH_METADATA in a project target file.

diff --git a/Engine/Source/Runtime/Core/Public/Misc/CoreMiscDefines.h b/Engine/Source/Runtime/Core/Public/Misc/CoreMiscDefines.h
--- a/Engine/Source/Runtime/Core/Public/Misc/CoreMiscDefines.h
+++ b/Engine/Source/Runtime/Core/Public/Misc/CoreMiscDefines.h
@@ -28,7 +28,11 @@
 #endif
 
 /** This controls if metadata for compiled in classes is unpacked and setup at boot time. Meta data is not normally used except by the editor. **/
-#define WITH_METADATA (WITH_EDITORONLY_DATA && WITH_EDITOR)
+#ifndef WITH_METADATA
+#define WITH_METADATA WITH_EDITORONLY_DATA
+#elif WITH_EDITORONLY_DATA && !WITH_METADATA
+#error WITH_EDITORONLY_DATA=1 requires WITH_METADATA=1
+#endif
 
 // Option to check for UE_DISABLE_OPTIMIZATION being submitted
 #ifndef UE_CHECK_DISABLE_OPTIMIZATION
diff --git a/Engine/Source/Runtime/CoreUObject/Private/Blueprint/BlueprintSupport.cpp b/Engine/Source/Runtime/CoreUObject/Private/Blueprint/BlueprintSupport.cpp
--- a/Engine/Source/Runtime/CoreUObject/Private/Blueprint/BlueprintSupport.cpp
+++ b/Engine/Source/Runtime/CoreUObject/Private/Blueprint/BlueprintSupport.cpp
@@ -722,7 +722,7 @@ bool FLinkerLoad::RegenerateBlueprintClass(UClass* LoadClass, UObject* ClassDefa
 		// Blueprint has been completely preloaded.
 		ForcePreloadObject(*ClassSourceObject);
 
-#if WITH_EDITORONLY_DATA
+#if WITH_METADATA
 		// We likely don't need to load meta data here at all! but we have been doing so 
 		// since 2080292 - subtly the code from 2080292 wouldn't assert about missing metadata
 		// but only because UPackage::GetMetaData creates a dummy UMetaData object that is
@@ -738,7 +738,7 @@ bool FLinkerLoad::RegenerateBlueprintClass(UClass* LoadClass, UObject* ClassDefa
 				ForcePreloadObject(*MetadataExport.Object);
 			}
 		}
-#endif
+#endif // WITH_METADATA
 
 		// Flush (ie: create and preload) all remaining exports in the package.
 		//
diff --git a/Engine/Source/Runtime/CoreUObject/Private/Serialization/AsyncLoading.cpp b/Engine/Source/Runtime/CoreUObject/Private/Serialization/AsyncLoading.cpp
--- a/Engine/Source/Runtime/CoreUObject/Private/Serialization/AsyncLoading.cpp
+++ b/Engine/Source/Runtime/CoreUObject/Private/Serialization/AsyncLoading.cpp
@@ -6016,14 +6016,14 @@ EAsyncPackageState::Type FAsyncPackage::TickAsyncPackage(bool InbUseTimeLimit, b
 					LoadingState = CreateImports();
 				}
 
-#if WITH_EDITORONLY_DATA
+#if WITH_METADATA
 				// Create and preload the package meta-data
 				if (LoadingState == EAsyncPackageState::Complete && !bLoadHasFinished)
 				{
 					SCOPED_LOADTIMER(Package_CreateMetaData);
 					LoadingState = CreateMetaData();
 				}
-#endif // WITH_EDITORONLY_DATA
+#endif // WITH_METADATA
 
 				// Create exports from linker export table and also preload them.
 				if (LoadingState == EAsyncPackageState::Complete && !bLoadHasFinished)
@@ -6701,7 +6701,7 @@ EAsyncPackageState::Type FAsyncPackage::CreateImports()
 	return ImportIndex == Linker->ImportMap.Num() ? EAsyncPackageState::Complete : EAsyncPackageState::TimeOut;
 }
 
-#if WITH_EDITORONLY_DATA
+#if WITH_METADATA
 /**
 * Creates and loads meta-data for the package.
 *
@@ -6720,7 +6720,7 @@ EAsyncPackageState::Type FAsyncPackage::CreateMetaData()
 
 	return EAsyncPackageState::Complete;
 }
-#endif // WITH_EDITORONLY_DATA
+#endif // WITH_METADATA
 
 /**
  * Create exports till time limit is exceeded.
@@ -6738,14 +6738,14 @@ EAsyncPackageState::Type FAsyncPackage::CreateExports()
 	// Create exports.
 	while( ExportIndex < Linker->ExportMap.Num() && !IsTimeLimitExceeded() )
 	{
-#if WITH_EDITORONLY_DATA
+#if WITH_METADATA
 		checkf(MetaDataIndex.IsSet(), TEXT("FAsyncPackage::CreateExports called before FAsyncPackage::CreateMetaData!"));
 		if (ExportIndex == MetaDataIndex.GetValue())
 		{
 			++ExportIndex;
 			continue;
 		}
-#endif // WITH_EDITORONLY_DATA
+#endif // WITH_METADATA
 
 		const FObjectExport& Export = Linker->ExportMap[ExportIndex];
 		// Precache data and see whether it's already finished.
diff --git a/Engine/Source/Runtime/CoreUObject/Private/Serialization/AsyncLoading.h b/Engine/Source/Runtime/CoreUObject/Private/Serialization/AsyncLoading.h
--- a/Engine/Source/Runtime/CoreUObject/Private/Serialization/AsyncLoading.h
+++ b/Engine/Source/Runtime/CoreUObject/Private/Serialization/AsyncLoading.h
@@ -638,10 +638,10 @@ private:
 	/** Requests that currently depends on this package. */
 	TSet<int32> DependentRequests;
 
-#if WITH_EDITORONLY_DATA
+#if WITH_METADATA
 	/** Index of the meta-data object within the linkers export table (unset if not yet processed, although may still be INDEX_NONE if there is no meta-data) */
 	TOptional<int32> MetaDataIndex;
-#endif // WITH_EDITORONLY_DATA
+#endif // WITH_METADATA
 	/** Number of times we recursed to load or post load this package. could be split into 2 counters not to use atomics. */
 	std::atomic<int32> ReentryCount;
 	/** List of objects referenced by this package */
@@ -888,14 +888,14 @@ private:
 	 */
 	EAsyncPackageState::Type CreateImports();
 
-#if WITH_EDITORONLY_DATA
+#if WITH_METADATA
 	/**
 	 * Creates and loads meta-data for the package.
 	 *
 	 * @return true if we finished creating meta-data, false otherwise.
 	 */
 	EAsyncPackageState::Type CreateMetaData();
-#endif // WITH_EDITORONLY_DATA
+#endif // WITH_METADATA
 	/**
 	 * Create exports till time limit is exceeded.
 	 *
diff --git a/Engine/Source/Runtime/CoreUObject/Private/UObject/Class.cpp b/Engine/Source/Runtime/CoreUObject/Private/UObject/Class.cpp
--- a/Engine/Source/Runtime/CoreUObject/Private/UObject/Class.cpp
+++ b/Engine/Source/Runtime/CoreUObject/Private/UObject/Class.cpp
@@ -392,6 +392,10 @@ void UField::FormatNativeToolTip(FString& ToolTipString, bool bRemoveExtraSectio
 	}
 }
 
+#endif // WITH_EDITORONLY_DATA
+
+#if WITH_METADATA
+
 /**
  * Determines if the property has any metadata associated with the key
  * 
@@ -548,7 +552,7 @@ void UField::RemoveMetaData(const FName& Key)
 	return Package->GetMetaData()->RemoveValue(this, Key);
 }
 
-#endif // WITH_EDITORONLY_DATA
+#endif // WITH_METADATA
 
 bool UField::HasAnyCastFlags(const uint64 InCastFlags) const
 {
@@ -2470,7 +2474,7 @@ FString UStruct::GetAuthoredNameForField(const FField* Field) const
 	return FString();
 }
 
-#if WITH_EDITORONLY_DATA
+#if WITH_METADATA
 bool UStruct::GetBoolMetaDataHierarchical(const FName& Key) const
 {
 	bool bResult = false;
@@ -2519,7 +2523,7 @@ const UStruct* UStruct::HasMetaDataHierarchical(const FName& Key) const
 	return nullptr;
 }
 
-#endif // WITH_EDITORONLY_DATA
+#endif // WITH_METADATA
 
 #if USE_CIRCULAR_DEPENDENCY_LOAD_DEFERRING
 	/**
@@ -4002,7 +4006,7 @@ static FString GetFieldLocation(const UField* Field)
 	{
 		StructLocation += FString::Printf(TEXT(" File:%s"), *ModuleRelativeIncludePath);
 	}
-#endif
+#endif // WITH_EDITORONLY_DATA
 	return StructLocation;
 };
 
@@ -4204,7 +4208,7 @@ bool FAutomationTestAttemptToFindUninitializedScriptStructMembers::RunTest(const
 	}
 }
 
-#if WITH_EDITORONLY_DATA
+#if WITH_METADATA
 /** 
 * Checks if MetaData value contains short type name 
 * @param MetaDataKey MetaData key name
@@ -4419,7 +4423,7 @@ bool FAutomationTestAttemptToFindShortTypeNamesInMetaData::RunTest(const FString
 	return FStructUtils::AttemptToFindShortTypeNamesInMetaData() == 0;
 }
 
-#endif // WITH_EDITORONLY_DATA
+#endif // WITH_METADATA
 
 // bExactCheck - Check for places where structs serialize a different set of object references to what they declare (Conservative is always an error)
 // otherwise - Check for places where structs serialize object references they don't declare (Conservative allows serializing any reference types)
diff --git a/Engine/Source/Runtime/CoreUObject/Private/UObject/CookedMetaData.cpp b/Engine/Source/Runtime/CoreUObject/Private/UObject/CookedMetaData.cpp
--- a/Engine/Source/Runtime/CoreUObject/Private/UObject/CookedMetaData.cpp
+++ b/Engine/Source/Runtime/CoreUObject/Private/UObject/CookedMetaData.cpp
@@ -26,7 +26,7 @@ void PrepareCookedMetaDataForPurge(UObject* CookedMetaDataPtr)
 template <typename CookedMetaDataOuterType, typename CookedMetaDataType>
 void PostLoadCookedMetaData(CookedMetaDataType* CookedMetaDataPtr)
 {
-#if WITH_EDITORONLY_DATA
+#if WITH_METADATA
 	checkf(CookedMetaDataPtr->GetPackage()->HasAnyPackageFlags(PKG_Cooked), TEXT("Cooked meta-data should only be loaded for a cooked package!"));
 
 	if (CookedMetaDataOuterType* Owner = CastChecked<CookedMetaDataOuterType>(CookedMetaDataPtr->GetOuter()))
@@ -35,7 +35,7 @@ void PostLoadCookedMetaData(CookedMetaDataType* CookedMetaDataPtr)
 		CookedMetaDataPtr->ApplyMetaData(Owner);
 		PrepareCookedMetaDataForPurge(CookedMetaDataPtr);
 	}
-#endif
+#endif // WITH_METADATA
 }
 
 } // namespace CookedMetaDataUtil::Internal
@@ -50,7 +50,7 @@ void FObjectCookedMetaDataStore::CacheMetaData(const UObject* SourceObject)
 {
 	ObjectMetaData.Reset();
 
-#if WITH_EDITORONLY_DATA
+#if WITH_METADATA
 	if (UPackage* SourcePackage = SourceObject->GetPackage())
 	{
 		if (const UMetaData* SourceMetaData = SourcePackage->GetMetaData())
@@ -61,12 +61,12 @@ void FObjectCookedMetaDataStore::CacheMetaData(const UObject* SourceObject)
 			}
 		}
 	}
-#endif
+#endif // WITH_METADATA
 }
 
 void FObjectCookedMetaDataStore::ApplyMetaData(UObject* TargetObject) const
 {
-#if WITH_EDITORONLY_DATA
+#if WITH_METADATA
 	if (UPackage* TargetPackage = TargetObject->GetPackage())
 	{
 		if (UMetaData* TargetMetaData = TargetPackage->GetMetaData())
@@ -74,7 +74,7 @@ void FObjectCookedMetaDataStore::ApplyMetaData(UObject* TargetObject) const
 			TargetMetaData->ObjectMetaDataMap.FindOrAdd(TargetObject).Append(ObjectMetaData);
 		}
 	}
-#endif
+#endif // WITH_METADATA
 }
 
 
@@ -87,19 +87,19 @@ void FFieldCookedMetaDataStore::CacheMetaData(const FField* SourceField)
 {
 	FieldMetaData.Reset();
 
-#if WITH_EDITORONLY_DATA
+#if WITH_METADATA
 	if (const TMap<FName, FString>* SourceFieldMetaData = SourceField->GetMetaDataMap())
 	{
 		FieldMetaData = *SourceFieldMetaData;
 	}
-#endif
+#endif // WITH_METADATA
 }
 
 void FFieldCookedMetaDataStore::ApplyMetaData(FField* TargetField) const
 {
-#if WITH_EDITORONLY_DATA
+#if WITH_METADATA
 	TargetField->AppendMetaData(FieldMetaData);
-#endif
+#endif // WITH_METADATA
 }
 
 
diff --git a/Engine/Source/Runtime/CoreUObject/Private/UObject/Enum.cpp b/Engine/Source/Runtime/CoreUObject/Private/UObject/Enum.cpp
--- a/Engine/Source/Runtime/CoreUObject/Private/UObject/Enum.cpp
+++ b/Engine/Source/Runtime/CoreUObject/Private/UObject/Enum.cpp
@@ -933,7 +933,7 @@ FText UEnum::GetToolTipTextByIndex(int32 NameIndex) const
 
 #endif
 
-#if WITH_EDITORONLY_DATA
+#if WITH_METADATA
 
 bool UEnum::HasMetaData( const TCHAR* Key, int32 NameIndex/*=INDEX_NONE*/ ) const
 {
@@ -1051,7 +1051,7 @@ void UEnum::RemoveMetaData( const TCHAR* Key, int32 NameIndex/*=INDEX_NONE*/) co
 	MetaData->RemoveValue( this, *KeyString );
 }
 
-#endif
+#endif // WITH_METADATA
 
 int64 UEnum::ParseEnum(const TCHAR*& Str)
 {
diff --git a/Engine/Source/Runtime/CoreUObject/Private/UObject/Field.cpp b/Engine/Source/Runtime/CoreUObject/Private/UObject/Field.cpp
--- a/Engine/Source/Runtime/CoreUObject/Private/UObject/Field.cpp
+++ b/Engine/Source/Runtime/CoreUObject/Private/UObject/Field.cpp
@@ -245,7 +245,7 @@
 	}
 }
 
-#if WITH_EDITORONLY_DATA
+#if WITH_METADATA
 bool FFieldVariant::HasMetaData(const FName& Key) const
 {
 	check(Container.Object);
@@ -258,7 +258,7 @@
 		return Container.Field->HasMetaData(Key);
 	}
 }
-#endif // WITH_EDITORONLY_DATA
+#endif // WITH_METADATA
 
 /*-----------------------------------------------------------------------------
 FField implementation.
@@ -280,9 +280,9 @@
 	, Owner((FField*)nullptr)
 	, Next(nullptr)
 	, FlagsPrivate(RF_NoFlags)
-#if WITH_EDITORONLY_DATA
+#if WITH_METADATA
 	, MetaDataMap(nullptr)
-#endif // WITH_EDITORONLY_DATA
+#endif // WITH_METADATA
 {
 }
 
@@ -291,22 +291,22 @@
 	, Next(nullptr)
 	, NamePrivate(InName)
 	, FlagsPrivate(InObjectFlags)
-#if WITH_EDITORONLY_DATA
+#if WITH_METADATA
 	, MetaDataMap(nullptr)
-#endif // WITH_EDITORONLY_DATA
+#endif // WITH_METADATA
 {
 
 }
 
 FField::~FField()
 {
-#if WITH_EDITORONLY_DATA
+#if WITH_METADATA
 	if (MetaDataMap)
 	{
 		delete MetaDataMap;
 		MetaDataMap = nullptr;
 	}
-#endif // WITH_EDITORONLY_DATA
+#endif // WITH_METADATA
 }
 
 #if WITH_EDITORONLY_DATA
@@ -414,7 +414,7 @@
 	Ar << NamePrivate;
 	Ar << (uint32&)FlagsPrivate;
 
-#if WITH_EDITORONLY_DATA	
+#if WITH_METADATA	
 	if (!Ar.IsCooking())
 	{
 		UPackage* Package = GetOutermost();
@@ -441,7 +441,7 @@
 			}			
 		}		
 	}
-#endif
+#endif // WITH_METADATA
 }
 
 void FField::GetPreloadDependencies(TArray<UObject*>& OutDeps)
@@ -626,123 +626,7 @@
 	return Result;
 }
 
-#if WITH_EDITORONLY_DATA
-FString FField::GetFullGroupName(bool bStartWithOuter) const
-{
-	if (bStartWithOuter)
-	{
-		if (Owner.IsValid())
-		{
-			if (Owner.IsUObject())
-			{
-				return Owner.ToUObject()->GetPathName(Owner.ToUObject()->GetOutermost());
-			}
-			else
-			{
-				return Owner.ToField()->GetPathName(GetOutermost());
-			}
-		}
-		else
-		{
-			return FString();
-		}
-	}
-	else
-	{
-		return GetPathName(GetOutermost());
-	}
-}
-
-struct FFieldDisplayNameHelper
-{
-	static FString Get(const FField& Object)
-	{
-		if (const FProperty* Property = CastField<FProperty>(&Object))
-		{
-			if (auto OwnerStruct = Property->GetOwnerStruct())
-			{
-				return OwnerStruct->GetAuthoredNameForField(Property);
-			}
-		}
-
-		return Object.GetName();
-	}
-};
-
-/**
-* Finds the localized display name or native display name as a fallback.
-*
-* @return The display name for this object.
-*/
-FText FField::GetDisplayNameText() const
-{
-	static const FString Namespace = TEXT("UObjectDisplayNames");
-	static const FName NAME_DisplayName(TEXT("DisplayName"));
-
-	const FString Key = GetFullGroupName(false);
-
-	FString NativeDisplayName;
-	if (const FString* FoundMetaData = FindMetaData(NAME_DisplayName))
-	{
-		NativeDisplayName = *FoundMetaData;
-	}
-	else
-	{
-		NativeDisplayName = FName::NameToDisplayString(FFieldDisplayNameHelper::Get(*this), IsA<FBoolProperty>());
-	}
-
-	return FInternationalization::ForUseOnlyByLocMacroAndGraphNodeTextLiterals_CreateText(*NativeDisplayName, *Namespace, *Key);
-}
-
-/**
-* Finds the localized tooltip or native tooltip as a fallback.
-*
-* @return The tooltip for this object.
-*/
-FText FField::GetToolTipText(bool bShortTooltip) const
-{
-	bool bFoundShortTooltip = false;
-	static const FName NAME_Tooltip(TEXT("Tooltip"));
-	static const FName NAME_ShortTooltip(TEXT("ShortTooltip"));
-	FText LocalizedToolTip;
-	FString NativeToolTip;
-
-	if (bShortTooltip)
-	{
-		NativeToolTip = GetMetaData(NAME_ShortTooltip);
-		if (NativeToolTip.IsEmpty())
-		{
-			NativeToolTip = GetMetaData(NAME_Tooltip);
-		}
-		else
-		{
-			bFoundShortTooltip = true;
-		}
-	}
-	else
-	{
-		NativeToolTip = GetMetaData(NAME_Tooltip);
-	}
-
-	const FString Namespace = bFoundShortTooltip ? TEXT("UObjectShortTooltips") : TEXT("UObjectToolTips");
-	const FString Key = GetFullGroupName(false);
-	if (!FText::FindText(Namespace, Key, /*OUT*/LocalizedToolTip, &NativeToolTip))
-	{
-		if (!NativeToolTip.IsEmpty())
-		{
-			static const FString DoxygenSee(TEXT("@see"));
-			static const FString TooltipSee(TEXT("See:"));
-			if (NativeToolTip.ReplaceInline(*DoxygenSee, *TooltipSee) > 0)
-			{
-				NativeToolTip.TrimEndInline();
-			}
-		}
-		LocalizedToolTip = FInternationalization::ForUseOnlyByLocMacroAndGraphNodeTextLiterals_CreateText(*NativeToolTip, *Namespace, *Key);
-	}
-
-	return LocalizedToolTip;
-}
-
+#if WITH_METADATA
 const FString* FField::FindMetaData(const TCHAR* Key) const
 {
 	return FindMetaData(FName(Key, FNAME_Find));
@@ -932,7 +816,7 @@
 	}
 }
 
-#endif // WITH_EDITORONLY_DATA
+#endif // WITH_METADATA
 
 void FField::PostDuplicate(const FField& InField)
 {
@@ -971,6 +855,121 @@
 }
 
 #if WITH_EDITORONLY_DATA
+FString FField::GetFullGroupName(bool bStartWithOuter) const
+{
+	if (bStartWithOuter)
+	{
+		if (Owner.IsValid())
+		{
+			if (Owner.IsUObject())
+			{
+				return Owner.ToUObject()->GetPathName(Owner.ToUObject()->GetOutermost());
+			}
+			else
+			{
+				return Owner.ToField()->GetPathName(GetOutermost());
+			}
+		}
+		else
+		{
+			return FString();
+		}
+	}
+	else
+	{
+		return GetPathName(GetOutermost());
+	}
+}
 
+struct FFieldDisplayNameHelper
+{
+	static FString Get(const FField& Object)
+	{
+		if (const FProperty* Property = CastField<FProperty>(&Object))
+		{
+			if (auto OwnerStruct = Property->GetOwnerStruct())
+			{
+				return OwnerStruct->GetAuthoredNameForField(Property);
+			}
+		}
+
+		return Object.GetName();
+	}
+};
+
+/**
+* Finds the localized display name or native display name as a fallback.
+*
+* @return The display name for this object.
+*/
+FText FField::GetDisplayNameText() const
+{
+	static const FString Namespace = TEXT("UObjectDisplayNames");
+	static const FName NAME_DisplayName(TEXT("DisplayName"));
+
+	const FString Key = GetFullGroupName(false);
+
+	FString NativeDisplayName;
+	if (const FString* FoundMetaData = FindMetaData(NAME_DisplayName))
+	{
+		NativeDisplayName = *FoundMetaData;
+	}
+	else
+	{
+		NativeDisplayName = FName::NameToDisplayString(FFieldDisplayNameHelper::Get(*this), IsA<FBoolProperty>());
+	}
+
+	return FInternationalization::ForUseOnlyByLocMacroAndGraphNodeTextLiterals_CreateText(*NativeDisplayName, *Namespace, *Key);
+}
+
+/**
+* Finds the localized tooltip or native tooltip as a fallback.
+*
+* @return The tooltip for this object.
+*/
+FText FField::GetToolTipText(bool bShortTooltip) const
+{
+	bool bFoundShortTooltip = false;
+	static const FName NAME_Tooltip(TEXT("Tooltip"));
+	static const FName NAME_ShortTooltip(TEXT("ShortTooltip"));
+	FText LocalizedToolTip;
+	FString NativeToolTip;
+
+	if (bShortTooltip)
+	{
+		NativeToolTip = GetMetaData(NAME_ShortTooltip);
+		if (NativeToolTip.IsEmpty())
+		{
+			NativeToolTip = GetMetaData(NAME_Tooltip);
+		}
+		else
+		{
+			bFoundShortTooltip = true;
+		}
+	}
+	else
+	{
+		NativeToolTip = GetMetaData(NAME_Tooltip);
+	}
+
+	const FString Namespace = bFoundShortTooltip ? TEXT("UObjectShortTooltips") : TEXT("UObjectToolTips");
+	const FString Key = GetFullGroupName(false);
+	if (!FText::FindText(Namespace, Key, /*OUT*/LocalizedToolTip, &NativeToolTip))
+	{
+		if (!NativeToolTip.IsEmpty())
+		{
+			static const FString DoxygenSee(TEXT("@see"));
+			static const FString TooltipSee(TEXT("See:"));
+			if (NativeToolTip.ReplaceInline(*DoxygenSee, *TooltipSee) > 0)
+			{
+				NativeToolTip.TrimEndInline();
+			}
+		}
+		LocalizedToolTip = FInternationalization::ForUseOnlyByLocMacroAndGraphNodeTextLiterals_CreateText(*NativeToolTip, *Namespace, *Key);
+	}
+
+	return LocalizedToolTip;
+}
+
 FField::FOnConvertCustomUFieldToFField& FField::GetConvertCustomUFieldToFFieldDelegate()
 {
diff --git a/Engine/Source/Runtime/CoreUObject/Private/UObject/LinkerLoad.cpp b/Engine/Source/Runtime/CoreUObject/Private/UObject/LinkerLoad.cpp
--- a/Engine/Source/Runtime/CoreUObject/Private/UObject/LinkerLoad.cpp
+++ b/Engine/Source/Runtime/CoreUObject/Private/UObject/LinkerLoad.cpp
@@ -3610,14 +3610,14 @@ bool FLinkerLoad::VerifyImportInner(const int32 ImportIndex, FString& WarningSuf
 				Import.SourceLinker = FindExistingLinkerForPackage(Package);
 			}
 		}
-#if WITH_EDITORONLY_DATA
+#if WITH_METADATA
 		if (Import.SourceLinker && !Package->HasAnyFlags(RF_LoadCompleted))
 		{
 			// If we didn't fully load, make sure our metadata is loaded before using this
 			// We need this case for user defined structs due to the LOAD_DeferDependencyLoads code above
 			Import.SourceLinker->LoadMetaDataFromExportMap(false);
 		}
-#endif
+#endif // WITH_METADATA
 		return Package;
 	};
 
@@ -4145,7 +4145,7 @@ UClass* FLinkerLoad::TryCreatePlaceholderClassForExport(int32 ExportIndex)
 }
 #endif
 
-#if WITH_EDITORONLY_DATA
+#if WITH_METADATA
 int32 FLinkerLoad::LoadMetaDataFromExportMap(bool bForcePreload)
 {
 	UMetaData* MetaData = nullptr;
@@ -4195,7 +4195,7 @@ int32 FLinkerLoad::LoadMetaDataFromExportMap(bool bForcePreload)
 
 	return MetaDataIndex;
 }
-#endif
+#endif // WITH_METADATA
 
 /**
  * Loads all objects in package.
@@ -4236,12 +4236,9 @@ void FLinkerLoad::LoadAllObjects(bool bForcePreload)
 	// MetaData object index in this package.
 	int32 MetaDataIndex = INDEX_NONE;
 
-#if WITH_EDITORONLY_DATA
-	if(!FPlatformProperties::RequiresCookedData())
-	{
-		MetaDataIndex = LoadMetaDataFromExportMap(bForcePreload);
-	}
-#endif
+#if WITH_METADATA
+	MetaDataIndex = LoadMetaDataFromExportMap(bForcePreload);
+#endif // WITH_METADATA
 	
 #if USE_STABLE_LOCALIZATION_KEYS
 	if (GIsEditor && (LoadFlags & LOAD_ForDiff))
@@ -7186,11 +7183,11 @@ bool FLinkerLoad::ShouldSkipProperty(const FProperty* InProperty) const
 {
 	if (bSkipKnownProperties)
 	{
-#if WITH_EDITORONLY_DATA
+#if WITH_METADATA
 		static const FName NAME_IsLooseMetadata(ANSITEXTVIEW("IsLoose"));
 		static const FName NAME_ContainsLoosePropertiesMetadata(ANSITEXTVIEW("ContainsLooseProperties"));
 		if (!InProperty->GetBoolMetaData(NAME_IsLooseMetadata) && !InProperty->GetBoolMetaData(NAME_ContainsLoosePropertiesMetadata))
-#endif
+#endif // WITH_METADATA
 		{
 			return true;
 		}
diff --git a/Engine/Source/Runtime/CoreUObject/Private/UObject/Package.cpp b/Engine/Source/Runtime/CoreUObject/Private/UObject/Package.cpp
--- a/Engine/Source/Runtime/CoreUObject/Private/UObject/Package.cpp
+++ b/Engine/Source/Runtime/CoreUObject/Private/UObject/Package.cpp
@@ -75,8 +75,11 @@ void UPackage::PostInitProperties()
 	SetLinkerPackageVersion(GPackageFileUEVersion);
 	SetLinkerLicenseeVersion(GPackageFileLicenseeUEVersion);
 
-#if WITH_EDITORONLY_DATA
+#if WITH_METADATA
 	SetMetaData(nullptr);
+#endif // WITH_METADATA
+
+#if WITH_EDITORONLY_DATA
 	// Always generate a new unique PersistentGuid, required for new disk packages.
 	// For existing disk packages it will be replaced with the existing PersistentGuid when loading the package summary.
 	// For existing script packages it will be replaced in ConstructUPackage with the CRC of the generated code files.
@@ -89,7 +92,7 @@ void UPackage::PostInitProperties()
 	bLoadedByEditorPropertiesOnly = !HasAnyFlags(RF_ClassDefaultObject) && !HasAnyPackageFlags(PKG_CompiledIn) && (IsRunningCommandlet());
 
 	bIsDynamicPIEPackagePending = false;
-#endif
+#endif // WITH_EDITORONLY_DATA
 }
 
 
@@ -230,9 +233,7 @@ TArray<UPackage*> UPackage::GetExternalPackages() const
  */
 UMetaData* UPackage::GetMetaData()
 {
-	checkf(!FPlatformProperties::RequiresCookedData(), TEXT("MetaData is only allowed in the Editor."));
-
-#if WITH_EDITORONLY_DATA
+#if WITH_METADATA
 	PRAGMA_DISABLE_DEPRECATION_WARNINGS
 	UMetaData* LocalMetaData = MetaData;
 	PRAGMA_ENABLE_DEPRECATION_WARNINGS
@@ -261,7 +262,7 @@ UMetaData* UPackage::GetMetaData()
 	return LocalMetaData;
 #else
 	return nullptr;
-#endif
+#endif // WITH_METADATA
 }
 
 /**
@@ -300,14 +301,14 @@ void UPackage::TagSubobjects(EObjectFlags NewFlags)
 {
 	Super::TagSubobjects(NewFlags);
 
-#if WITH_EDITORONLY_DATA
+#if WITH_METADATA
 	PRAGMA_DISABLE_DEPRECATION_WARNINGS
 	if (MetaData)
 	{
 		MetaData->SetFlags(NewFlags);
 	}
 	PRAGMA_ENABLE_DEPRECATION_WARNINGS
-#endif
+#endif // WITH_METADATA
 }
 
 /**
@@ -435,4 +436,4 @@ IMPLEMENT_CORE_INTRINSIC_CLASS(UPackage, UObject,
 	{
 	}
 );
-#endif
+#endif // WITH_EDITORONLY_DATA
diff --git a/Engine/Source/Runtime/CoreUObject/Private/UObject/Property.cpp b/Engine/Source/Runtime/CoreUObject/Private/UObject/Property.cpp
--- a/Engine/Source/Runtime/CoreUObject/Private/UObject/Property.cpp
+++ b/Engine/Source/Runtime/CoreUObject/Private/UObject/Property.cpp
@@ -718,8 +718,8 @@ const TCHAR* FPropertyHelpers::ReadToken( const TCHAR* Buffer, FStringBuilderBas
 	FProperty implementation.
 -----------------------------------------------------------------------------*/
 
-#if UE_GAME && UE_FNAME_OUTLINE_NUMBER
-	static_assert(sizeof(FProperty) <= 104, "FProperty was optimized to reduce its size so most of the classes that inherent from it will fall withing 112 bytes bin of MallocBinned3");
+#if UE_GAME && UE_FNAME_OUTLINE_NUMBER && !WITH_METADATA
+	static_assert(sizeof(FProperty) <= 104, "FProperty was optimized to reduce its size so most of the classes that inherent from it will fall within 112 bytes bin of MallocBinned3");
 #endif
 
 IMPLEMENT_FIELD(FProperty)
diff --git a/Engine/Source/Runtime/CoreUObject/Public/UObject/Class.h b/Engine/Source/Runtime/CoreUObject/Public/UObject/Class.h
--- a/Engine/Source/Runtime/CoreUObject/Public/UObject/Class.h
+++ b/Engine/Source/Runtime/CoreUObject/Public/UObject/Class.h
@@ -213,7 +213,9 @@ class UField : public UObject
 	 * @param bRemoveExtraSections	If true, cut off the comment on first line separator or 2 empty lines in a row
 	 */
 	static COREUOBJECT_API void FormatNativeToolTip(FString& ToolTipString, bool bRemoveExtraSections = true);
+#endif // WITH_EDITORONLY_DATA
 
+#if WITH_METADATA
 	/**
 	 * Determines if the property has any metadata associated with the key
 	 * 
@@ -330,7 +332,7 @@ class UField : public UObject
 	/** Clear any metadata associated with the key */
 	COREUOBJECT_API void RemoveMetaData(const TCHAR* Key);
 	COREUOBJECT_API void RemoveMetaData(const FName& Key);
-#endif // WITH_EDITORONLY_DATA
+#endif // WITH_METADATA
 
 	COREUOBJECT_API bool HasAnyCastFlags(const uint64 InCastFlags) const;
 	COREUOBJECT_API bool HasAllCastFlags(const uint64 InCastFlags) const;
@@ -667,7 +669,7 @@ public:
 	/** Destroys all properties owned by this struct */
 	COREUOBJECT_API void DestroyChildPropertiesAndResetPropertyLinks();
 
-#if WITH_EDITORONLY_DATA
+#if WITH_METADATA
 	/** Try and find boolean metadata with the given key. If not found on this class, work up hierarchy looking for it. */
 	COREUOBJECT_API bool GetBoolMetaDataHierarchical(const FName& Key) const;
 
@@ -681,7 +683,9 @@ public:
 	 * @return pointer to the UStruct that has associated metadata, nullptr if Key is not associated with any UStruct in the hierarchy
 	 */
 	COREUOBJECT_API const UStruct* HasMetaDataHierarchical(const FName& Key) const;
+#endif // WITH_METADATA
 
+#if WITH_EDITORONLY_DATA
 	/* Returns true if this struct has Asset Registry searchable properties */
 	FORCEINLINE bool HasAssetRegistrySearchableProperties() const
 	{
@@ -2421,7 +2425,7 @@ public:
 	COREUOBJECT_API FText GetToolTipTextByIndex(int32 NameIndex) const;
 #endif
 
-#if WITH_EDITORONLY_DATA
+#if WITH_METADATA
 	/**
 	 * Wrapper method for easily determining whether this enum has metadata associated with it.
 	 * 
@@ -2461,7 +2465,7 @@ public:
 	 *
 	 */
 	COREUOBJECT_API void RemoveMetaData( const TCHAR* Key, int32 NameIndex=INDEX_NONE ) const;
-#endif // WITH_EDITORONLY_DATA
+#endif // WITH_METADATA
 	
 	/**
 	 * @param EnumPath         Full enum path.
@@ -4164,10 +4168,10 @@ struct FStructUtils
 	/** Looks for uninitialized script struct pointers. Returns the number found */
 	COREUOBJECT_API static int32 AttemptToFindUninitializedScriptStructMembers();
 
-#if WITH_EDITORONLY_DATA
+#if WITH_METADATA
 	/** Looks for short type names within struct metadata. Returns the number found */
 	COREUOBJECT_API static int32 AttemptToFindShortTypeNamesInMetaData();
-#endif // WITH_EDITORONLY_DATA
+#endif // WITH_METADATA
 #endif // !(UE_BUILD_TEST || UE_BUILD_SHIPPING)
 };
 
diff --git a/Engine/Source/Runtime/CoreUObject/Public/UObject/Field.h b/Engine/Source/Runtime/CoreUObject/Public/UObject/Field.h
--- a/Engine/Source/Runtime/CoreUObject/Public/UObject/Field.h
+++ b/Engine/Source/Runtime/CoreUObject/Public/UObject/Field.h
@@ -435,9 +435,9 @@ public:
 		return Container.Field != Other.Container.Field;
 	}
 
-#if WITH_EDITORONLY_DATA
+#if WITH_METADATA
 	COREUOBJECT_API bool HasMetaData(const FName& Key) const;
-#endif
+#endif // WITH_METADATA
 
 	/** Support comparison functions that make this usable as a KeyValue for a TSet<> */
 	[[nodiscard]] friend uint32 GetTypeHash(const FFieldVariant& InFieldVariant)
@@ -753,38 +753,12 @@ public:
 	{
 	}
 
-#if WITH_EDITORONLY_DATA
-
+#if WITH_METADATA
 private:
 	/** Editor-only meta data map */
 	TMap<FName, FString>* MetaDataMap;
 
 public:
-
-	/**
-	* Walks up the chain of packages until it reaches the top level, which it ignores.
-	*
-	* @param	bStartWithOuter		whether to include this object's name in the returned string
-	* @return	string containing the path name for this object, minus the outermost-package's name
-	*/
-	COREUOBJECT_API FString GetFullGroupName(bool bStartWithOuter) const;
-
-	/**
-	* Finds the localized display name or native display name as a fallback.
-	*
-	* @return The display name for this object.
-	*/
-	COREUOBJECT_API FText GetDisplayNameText() const;
-
-	/**
-	* Finds the localized tooltip or native tooltip as a fallback.
-	*
-	* @param bShortTooltip Look for a shorter version of the tooltip (falls back to the long tooltip if none was specified)
-	*
-	* @return The tooltip for this object.
-	*/
-	COREUOBJECT_API FText GetToolTipText(bool bShortTooltip = false) const;
-
 	/**
 	* Determines if the property has any metadata associated with the key
 	*
@@ -936,6 +910,33 @@ public:
 
 	/** Copies all metadata from Source Field to Dest Field */
 	static COREUOBJECT_API void CopyMetaData(const FField* InSourceField, FField* InDestField);
+#endif // WITH_METADATA
+
+#if WITH_EDITORONLY_DATA
+public:
+	/**
+	* Walks up the chain of packages until it reaches the top level, which it ignores.
+	*
+	* @param	bStartWithOuter		whether to include this object's name in the returned string
+	* @return	string containing the path name for this object, minus the outermost-package's name
+	*/
+	COREUOBJECT_API FString GetFullGroupName(bool bStartWithOuter) const;
+
+	/**
+	* Finds the localized display name or native display name as a fallback.
+	*
+	* @return The display name for this object.
+	*/
+	COREUOBJECT_API FText GetDisplayNameText() const;
+
+	/**
+	* Finds the localized tooltip or native tooltip as a fallback.
+	*
+	* @param bShortTooltip Look for a shorter version of the tooltip (falls back to the long tooltip if none was specified)
+	*
+	* @return The tooltip for this object.
+	*/
+	COREUOBJECT_API FText GetToolTipText(bool bShortTooltip = false) const;
 
 	/** Creates a new FField from existing UField */
 	static COREUOBJECT_API FField* CreateFromUField(UField* InField);
@@ -943,9 +944,9 @@ public:
 	DECLARE_MULTICAST_DELEGATE_ThreeParams(FOnConvertCustomUFieldToFField, FFieldClass*, UField*, FField*&);
 	/** Gets a delegate to convert custom UField types to FFields */
 	static COREUOBJECT_API FOnConvertCustomUFieldToFField& GetConvertCustomUFieldToFFieldDelegate();
-
 #endif // WITH_EDITORONLY_DATA
 
+public:
 	/** Duplicates an FField */
 	static COREUOBJECT_API FField* Duplicate(const FField* InField, FFieldVariant DestOwner, const FName DestName = NAME_None, EObjectFlags FlagMask = RF_AllFlags, EInternalObjectFlags InternalFlagsMask = EInternalObjectFlags_AllFlags);
 
diff --git a/Engine/Source/Runtime/CoreUObject/Public/UObject/LinkerLoad.h b/Engine/Source/Runtime/CoreUObject/Public/UObject/LinkerLoad.h
--- a/Engine/Source/Runtime/CoreUObject/Public/UObject/LinkerLoad.h
+++ b/Engine/Source/Runtime/CoreUObject/Public/UObject/LinkerLoad.h
@@ -894,7 +894,7 @@ private:
 	UClass* TryCreatePlaceholderClassForExport(int32 ExportIndex);
 #endif
 
-#if WITH_EDITORONLY_DATA
+#if WITH_METADATA
 	/** 
 	 * Looks for and loads meta data object from export map.
 	 *
@@ -905,7 +905,7 @@ private:
 	 *         INDEX_NONE otherwise.
 	 */
 	int32 LoadMetaDataFromExportMap(bool bForcePreload);
-#endif
+#endif // WITH_METADATA
 
 	UObject* CreateImport( int32 Index );
 
diff --git a/Engine/Source/Runtime/CoreUObject/Public/UObject/Package.h b/Engine/Source/Runtime/CoreUObject/Public/UObject/Package.h
--- a/Engine/Source/Runtime/CoreUObject/Public/UObject/Package.h
+++ b/Engine/Source/Runtime/CoreUObject/Public/UObject/Package.h
@@ -308,15 +308,17 @@ public:
 	UE_DEPRECATED(5.0, "Use Get/SetFileSize instead")
 	uint64 FileSize;
 
+#if WITH_METADATA
+	// MetaData for the editor, or NULL in the game
+	UE_DEPRECATED(5.0, "Use Get/HasMetaData instead")
+	class UMetaData* MetaData;
+#endif // WITH_METADATA
+
 #if WITH_EDITORONLY_DATA
 	/** Editor only: Thumbnails stored in this package */
 	UE_DEPRECATED(5.0, "Use Get/SetThumbnailMap instead")
 	TUniquePtr< FThumbnailMap > ThumbnailMap;
 
-	// MetaData for the editor, or NULL in the game
-	UE_DEPRECATED(5.0, "Use Get/HasMetaData instead")
-	class UMetaData* MetaData;
-
 	/** Editor only: PIE instance ID this package belongs to, INDEX_NONE otherwise */
 	UE_DEPRECATED(5.0, "Use Get/SetPIEInstanceID instead")
 	int32 PIEInstanceID;
@@ -425,9 +427,9 @@ private:
 
 	void SetMetaData(UMetaData* InMetaData)
 	{
-#if WITH_EDITORONLY_DATA
+#if WITH_METADATA
 		MetaData = InMetaData;
-#endif
+#endif // WITH_METADATA
 	}
 
 public:
@@ -913,11 +915,11 @@ public:
 	 */
 	bool HasMetaData() const
 	{
-#if WITH_EDITORONLY_DATA
+#if WITH_METADATA
 		return MetaData != nullptr;
 #else
 		return false;
-#endif
+#endif // WITH_METADATA
 	}
 
 	/**
diff --git a/Engine/Source/Runtime/FieldNotification/Public/FieldNotificationVariant.h b/Engine/Source/Runtime/FieldNotification/Public/FieldNotificationVariant.h
--- a/Engine/Source/Runtime/FieldNotification/Public/FieldNotificationVariant.h
+++ b/Engine/Source/Runtime/FieldNotification/Public/FieldNotificationVariant.h
@@ -32,9 +32,9 @@ public:
 	using ::FFieldVariant::operator bool;
 	using ::FFieldVariant::operator ==;
 	using ::FFieldVariant::operator !=;
-#if WITH_EDITORONLY_DATA
+#if WITH_METADATA
 	using ::FFieldVariant::HasMetaData;
-#endif
+#endif // WITH_METADATA
 
 	inline bool IsProperty() const
 	{
