PR #11162: Fixed garbage collector stats narrowing conversion

#rnx

- This line of code was missed when the value type of ClassToPurgeCountMap was changed from int32 to uint64.

diff --git a/Engine/Source/Runtime/CoreUObject/Private/UObject/GarbageCollection.cpp b/Engine/Source/Runtime/CoreUObject/Private/UObject/GarbageCollection.cpp
--- a/Engine/Source/Runtime/CoreUObject/Private/UObject/GarbageCollection.cpp
+++ b/Engine/Source/Runtime/CoreUObject/Private/UObject/GarbageCollection.cpp
@@ -541,7 +541,7 @@ namespace UE::GC
 #if PERF_DETAILED_PER_CLASS_GC_STATS
 		// Keep track of how many objects of a certain class we're purging.
 		const FName& ClassName = Object->GetClass()->GetFName();
-		int32 InstanceCount = ClassToPurgeCountMap.FindRef( ClassName );
+		uint64 InstanceCount = ClassToPurgeCountMap.FindRef( ClassName );
 		ClassToPurgeCountMap.Add( ClassName, ++InstanceCount );
 #endif
 	}
