Fix FSparseDelegate to behave correctly under transactional conditions.

diff --git a/Engine/Source/Runtime/CoreUObject/Private/UObject/SparseDelegate.cpp b/Engine/Source/Runtime/CoreUObject/Private/UObject/SparseDelegate.cpp
--- a/Engine/Source/Runtime/CoreUObject/Private/UObject/SparseDelegate.cpp
+++ b/Engine/Source/Runtime/CoreUObject/Private/UObject/SparseDelegate.cpp
@@ -4,9 +4,10 @@
 #include "UObject/Class.h"
 #include "UObject/Package.h"
 #include "HAL/IConsoleManager.h"
+#include "Misc/TransactionallySafeScopeLock.h"
 
 FSparseDelegateStorage::FObjectListener FSparseDelegateStorage::SparseDelegateObjectListener;
-FCriticalSection FSparseDelegateStorage::SparseDelegateMapCritical;
+FTransactionallySafeCriticalSection FSparseDelegateStorage::SparseDelegateMapCritical;
 TMap<const UObjectBase*, FSparseDelegateStorage::FSparseDelegateMap> FSparseDelegateStorage::SparseDelegates;
 TMap<TPair<FName,FName>, size_t> FSparseDelegateStorage::SparseDelegateObjectOffsets;
 
@@ -21,7 +22,7 @@ FSparseDelegateStorage::FObjectListener::~FObjectListener()
 
 void FSparseDelegateStorage::FObjectListener::NotifyUObjectDeleted(const UObjectBase* Object, int32 Index)
 {
-	FScopeLock SparseDelegateMapLock(&FSparseDelegateStorage::SparseDelegateMapCritical);
+	FTransactionallySafeScopeLock SparseDelegateMapLock(&FSparseDelegateStorage::SparseDelegateMapCritical);
 	FSparseDelegateStorage::SparseDelegates.Remove(Object);
 	if (FSparseDelegateStorage::SparseDelegates.Num() == 0)
 	{
@@ -31,7 +32,7 @@ void FSparseDelegateStorage::FObjectListener::NotifyUObjectDeleted(const UObject
 
 void FSparseDelegateStorage::FObjectListener::OnUObjectArrayShutdown()
 {
-	FScopeLock SparseDelegateMapLock(&FSparseDelegateStorage::SparseDelegateMapCritical);
+	FTransactionallySafeScopeLock SparseDelegateMapLock(&FSparseDelegateStorage::SparseDelegateMapCritical);
 	FSparseDelegateStorage::SparseDelegates.Empty();
 	DisableListener();
 }
@@ -83,7 +84,7 @@ UObject* FSparseDelegateStorage::ResolveSparseOwner(const FSparseDelegate& Spars
 
 FMulticastScriptDelegate* FSparseDelegateStorage::GetMulticastDelegate(const UObject* DelegateOwner, const FName DelegateName)
 {
-	FScopeLock SparseDelegateMapLock(&SparseDelegateMapCritical);
+	FTransactionallySafeScopeLock SparseDelegateMapLock(&SparseDelegateMapCritical);
 
 	if (FSparseDelegateMap* DelegateMap = SparseDelegates.Find(DelegateOwner))
 	{
@@ -97,7 +98,7 @@ FMulticastScriptDelegate* FSparseDelegateStorage::GetMulticastDelegate(const UOb
 
 TSharedPtr<FMulticastScriptDelegate> FSparseDelegateStorage::GetSharedMulticastDelegate(const UObject* DelegateOwner, const FName DelegateName)
 {
-	FScopeLock SparseDelegateMapLock(&SparseDelegateMapCritical);
+	FTransactionallySafeScopeLock SparseDelegateMapLock(&SparseDelegateMapCritical);
 
 	TSharedPtr<FMulticastScriptDelegate> Result;
 	if (FSparseDelegateMap* DelegateMap = SparseDelegates.Find(DelegateOwner))
@@ -112,7 +113,7 @@ TSharedPtr<FMulticastScriptDelegate> FSparseDelegateStorage::GetSharedMulticastD
 
 void FSparseDelegateStorage::SetMulticastDelegate(const UObject* DelegateOwner, const FName DelegateName, FMulticastScriptDelegate Delegate)
 {
-	FScopeLock SparseDelegateMapLock(&SparseDelegateMapCritical);
+	FTransactionallySafeScopeLock SparseDelegateMapLock(&SparseDelegateMapCritical);
 
 	if (SparseDelegates.Num() == 0)
 	{
@@ -135,7 +136,7 @@ bool FSparseDelegateStorage::Add(const UObject* DelegateOwner, const FName Deleg
 	bool bDelegateWasBound = false;
 	if (Delegate.IsBound())
 	{
-		FScopeLock SparseDelegateMapLock(&SparseDelegateMapCritical);
+		FTransactionallySafeScopeLock SparseDelegateMapLock(&SparseDelegateMapCritical);
 
 		if (SparseDelegates.Num() == 0)
 		{
@@ -161,7 +162,7 @@ bool FSparseDelegateStorage::AddUnique(const UObject* DelegateOwner, const FName
 	bool bDelegateWasBound = false;
 	if (Delegate.IsBound())
 	{
-		FScopeLock SparseDelegateMapLock(&SparseDelegateMapCritical);
+		FTransactionallySafeScopeLock SparseDelegateMapLock(&SparseDelegateMapCritical);
 
 		if (SparseDelegates.Num() == 0)
 		{
@@ -185,7 +186,7 @@ bool FSparseDelegateStorage::AddUnique(const UObject* DelegateOwner, const FName
 bool FSparseDelegateStorage::Contains(const UObject* DelegateOwner, const FName DelegateName, const FScriptDelegate& Delegate)
 {
 	bool bContainsDelegate = false;
-	FScopeLock SparseDelegateMapLock(&SparseDelegateMapCritical);
+	FTransactionallySafeScopeLock SparseDelegateMapLock(&SparseDelegateMapCritical);
 
 	if (FSparseDelegateMap* DelegateMap = SparseDelegates.Find(DelegateOwner))
 	{
@@ -204,7 +205,7 @@ bool FSparseDelegateStorage::Contains(const UObject* DelegateOwner, const FName
 bool FSparseDelegateStorage::Contains(const UObject* DelegateOwner, const FName DelegateName, const UObject* InObject, FName InFunctionName)
 {
 	bool bContainsDelegate = false;
-	FScopeLock SparseDelegateMapLock(&SparseDelegateMapCritical);
+	FTransactionallySafeScopeLock SparseDelegateMapLock(&SparseDelegateMapCritical);
 
 	if (FSparseDelegateMap* DelegateMap = SparseDelegates.Find(DelegateOwner))
 	{
@@ -222,7 +223,7 @@ bool FSparseDelegateStorage::Contains(const UObject* DelegateOwner, const FName
 bool FSparseDelegateStorage::Remove(const UObject* DelegateOwner, const FName DelegateName, const FScriptDelegate& Delegate)
 {
 	bool bSparseDelegateBound = false;
-	FScopeLock SparseDelegateMapLock(&SparseDelegateMapCritical);
+	FTransactionallySafeScopeLock SparseDelegateMapLock(&SparseDelegateMapCritical);
 
 	if (FSparseDelegateMap* DelegateMap = SparseDelegates.Find(DelegateOwner))
 	{
@@ -258,7 +259,7 @@ bool FSparseDelegateStorage::Remove(const UObject* DelegateOwner, const FName De
 bool FSparseDelegateStorage::Remove(const UObject* DelegateOwner, const FName DelegateName, const UObject* InObject, FName InFunctionName)
 {
 	bool bSparseDelegateBound = false;
-	FScopeLock SparseDelegateMapLock(&SparseDelegateMapCritical);
+	FTransactionallySafeScopeLock SparseDelegateMapLock(&SparseDelegateMapCritical);
 
 	if (FSparseDelegateMap* DelegateMap = SparseDelegates.Find(DelegateOwner))
 	{
@@ -294,7 +295,7 @@ bool FSparseDelegateStorage::Remove(const UObject* DelegateOwner, const FName De
 bool FSparseDelegateStorage::RemoveAll(const UObject* DelegateOwner, const FName DelegateName, const UObject* UserObject)
 {
 	bool bSparseDelegateBound = false;
-	FScopeLock SparseDelegateMapLock(&SparseDelegateMapCritical);
+	FTransactionallySafeScopeLock SparseDelegateMapLock(&SparseDelegateMapCritical);
 
 	if (FSparseDelegateMap* DelegateMap = SparseDelegates.Find(DelegateOwner))
 	{
@@ -328,7 +329,7 @@ bool FSparseDelegateStorage::RemoveAll(const UObject* DelegateOwner, const FName
 
 void FSparseDelegateStorage::Clear(const UObject* DelegateOwner, const FName DelegateName)
 {
-	FScopeLock SparseDelegateMapLock(&SparseDelegateMapCritical);
+	FTransactionallySafeScopeLock SparseDelegateMapLock(&SparseDelegateMapCritical);
 	if (FSparseDelegateMap* DelegateMap = SparseDelegates.Find(DelegateOwner))
 	{
 		if (TSharedPtr<FMulticastScriptDelegate>* MulticastDelegatePtr = DelegateMap->Find(DelegateName))
@@ -415,7 +416,7 @@ void FSparseDelegateStorage::SparseDelegateReport(const TArray<FString>& Args, U
 	uint32 BoundDelegates = 0;
 
 	{
-		FScopeLock SparseDelegateMapLock(&SparseDelegateMapCritical);
+		FTransactionallySafeScopeLock SparseDelegateMapLock(&SparseDelegateMapCritical);
 		for (const TPair<const UObjectBase*, FSparseDelegateMap>& ObjectAnnotation : SparseDelegates)
 		{
 			const UObject* Object = static_cast<const UObject*>(ObjectAnnotation.Key);
diff --git a/Engine/Source/Runtime/CoreUObject/Public/UObject/SparseDelegate.h b/Engine/Source/Runtime/CoreUObject/Public/UObject/SparseDelegate.h
--- a/Engine/Source/Runtime/CoreUObject/Public/UObject/SparseDelegate.h
+++ b/Engine/Source/Runtime/CoreUObject/Public/UObject/SparseDelegate.h
@@ -10,6 +10,7 @@
 #include "HAL/PlatformMath.h"
 #include "Misc/AssertionMacros.h"
 #include "Misc/ScopeLock.h"
+#include "Misc/TransactionallySafeScopeLock.h"
 #include "Templates/SharedPointer.h"
 #include "Templates/UnrealTemplate.h"
 #include "UObject/NameTypes.h"
@@ -103,7 +104,7 @@ private:
 	static COREUOBJECT_API FObjectListener SparseDelegateObjectListener;
 
 	/** Critical Section for locking access to the sparse delegate map */
-	static COREUOBJECT_API FCriticalSection SparseDelegateMapCritical;
+	static COREUOBJECT_API FTransactionallySafeCriticalSection SparseDelegateMapCritical;
 
 	/** Delegate map is a map of Delegate names to a shared pointer of the multicast script delegate */
 	typedef TMap<FName, TSharedPtr<FMulticastScriptDelegate>> FSparseDelegateMap;
diff --git a/Engine/Source/Runtime/Engine/Private/Tests/AutoRTFM/AutoRTFMActorComponentTests.cpp b/Engine/Source/Runtime/Engine/Private/Tests/AutoRTFM/AutoRTFMActorComponentTests.cpp
--- a/Engine/Source/Runtime/Engine/Private/Tests/AutoRTFM/AutoRTFMActorComponentTests.cpp
+++ b/Engine/Source/Runtime/Engine/Private/Tests/AutoRTFM/AutoRTFMActorComponentTests.cpp
@@ -6,6 +6,7 @@
 #include "AutoRTFMTestActor.h"
 #include "AutoRTFMTestBodySetup.h"
 #include "AutoRTFMTestLevel.h"
+#include "AutoRTFMTestObject.h"
 #include "AutoRTFMTestPrimitiveComponent.h"
 #include "PhysicsProxy/SingleParticlePhysicsProxy.h"
 #include "Physics/Experimental/PhysScene_Chaos.h"
@@ -199,6 +200,32 @@ bool FAutoRTFMActorComponentTests::RunTest(const FString & Parameters)
 	TEST_CHECK_TRUE(AutoRTFM::ETransactionResult::Committed == Result);
 	TEST_CHECK_TRUE(!Component->IsWelded());
 
+	TEST_CHECK_TRUE(!Component->IsRegistered());
+	Component->RegisterComponent();
+
+	UAutoRTFMTestObject* const Object = NewObject<UAutoRTFMTestObject>();
+
+	Component->OnComponentPhysicsStateChanged.AddDynamic(Object, &UAutoRTFMTestObject::OnComponentPhysicsStateChanged);
+
+	TEST_CHECK_TRUE(!Object->bHitOnComponentPhysicsStateChanged);
+
+	Result = AutoRTFM::Transact([&]
+		{
+			Component->UnregisterComponent();
+			AutoRTFM::AbortTransaction();
+		});
+
+	TEST_CHECK_TRUE(AutoRTFM::ETransactionResult::AbortedByRequest == Result);
+	TEST_CHECK_TRUE(!Object->bHitOnComponentPhysicsStateChanged);
+
+	Result = AutoRTFM::Transact([&]
+		{
+			Component->UnregisterComponent();
+		});
+
+	TEST_CHECK_TRUE(AutoRTFM::ETransactionResult::Committed == Result);
+	TEST_CHECK_TRUE(Object->bHitOnComponentPhysicsStateChanged);
+
 	return true;
 }
 
diff --git a/Engine/Source/Runtime/Engine/Private/Tests/AutoRTFM/AutoRTFMTestObject.h b/Engine/Source/Runtime/Engine/Private/Tests/AutoRTFM/AutoRTFMTestObject.h
--- a/Engine/Source/Runtime/Engine/Private/Tests/AutoRTFM/AutoRTFMTestObject.h
+++ b/Engine/Source/Runtime/Engine/Private/Tests/AutoRTFM/AutoRTFMTestObject.h
@@ -23,4 +23,12 @@ public:
     }
 
     int Value;
+
+	bool bHitOnComponentPhysicsStateChanged = false;
+
+	UFUNCTION()
+	ENGINE_API void OnComponentPhysicsStateChanged(UPrimitiveComponent* ChangedComponent, EComponentPhysicsStateChange StateChange)
+	{
+		bHitOnComponentPhysicsStateChanged = true;
+	}
 };
