Fix build break

diff --git a/Engine/Source/Runtime/CoreUObject/Private/UObject/OverriddenPropertySet.cpp b/Engine/Source/Runtime/CoreUObject/Private/UObject/OverriddenPropertySet.cpp
--- a/Engine/Source/Runtime/CoreUObject/Private/UObject/OverriddenPropertySet.cpp
+++ b/Engine/Source/Runtime/CoreUObject/Private/UObject/OverriddenPropertySet.cpp
@@ -543,7 +543,7 @@ void FOverriddenPropertySet::NotifyPropertyChange(FOverriddenPropertyNode* Paren
 		if (bNeedsCleanup && ParentPropertyNode && SubPropertyNode && SubPropertyNode->SubPropertyNodeKeys.IsEmpty())
 		{
 			FOverriddenPropertyNodeID RemovedNodeID;
-			if (ParentPropertyNode->SubPropertyNodeKeys.RemoveAndCopyValue(Property, RemovedNodeID))
+			if (ParentPropertyNode->SubPropertyNodeKeys.RemoveAndCopyValue(Property->GetFName(), RemovedNodeID))
 			{
 				verifyf(OverriddenPropertyNodes.Remove(RemovedNodeID), TEXT("Expecting the node to be removed"));
 			}
