Remove placeholder types at reinstancing time prior to reference replacement if class references are also being replaced.

Change summary:
- Added FPropertyBagRepository::RemovePropertyBagPlaceholderType().
- Modified FBlueprintCompileReinstancer::ReplaceInstancesOfClass_Inner() to explicitly remove placeholder types if also replacing references to old class objects. This will remove placeholder type objects from the internal property bag placeholder type registry prior to reference replacement.
- Fixed FPropertyBagRepository::ReassociateObjects() to use the pair's key rather than value in order to remove any placeholder type objects remapped to a new type at reinstancing time. Note: This path still applies when class references are not also being replaced at reinstancing time.
- Modified FPropertyBagRepository::CreatePropertyBagPlaceholderType() to ensure that placeholder type objects assemble the reference token stream at creation time (required for garbage collection).

diff --git a/Engine/Source/Editor/UnrealEd/Private/Kismet2/KismetReinstanceUtilities.cpp b/Engine/Source/Editor/UnrealEd/Private/Kismet2/KismetReinstanceUtilities.cpp
--- a/Engine/Source/Editor/UnrealEd/Private/Kismet2/KismetReinstanceUtilities.cpp
+++ b/Engine/Source/Editor/UnrealEd/Private/Kismet2/KismetReinstanceUtilities.cpp
@@ -2971,6 +2971,12 @@ void FBlueprintCompileReinstancer::ReplaceInstancesOfClass_Inner(const TMap<UCla
 			
 			OldToNewInstanceMap.Add(OldToNew.Key, OldToNew.Value);
 			SourceObjects.Add(OldToNew.Key);
+
+			// Remove references to placeholder types before they get replaced below.
+			if (UE::FPropertyBagRepository::IsPropertyBagPlaceholderType(OldToNew.Key))
+			{
+				UE::FPropertyBagRepository::RemovePropertyBagPlaceholderType(OldToNew.Key);
+			}
 		}
 	}
 
diff --git a/Engine/Source/Runtime/CoreUObject/Internal/UObject/PropertyBagRepository.h b/Engine/Source/Runtime/CoreUObject/Internal/UObject/PropertyBagRepository.h
--- a/Engine/Source/Runtime/CoreUObject/Internal/UObject/PropertyBagRepository.h
+++ b/Engine/Source/Runtime/CoreUObject/Internal/UObject/PropertyBagRepository.h
@@ -181,6 +181,13 @@ public:
 		return Cast<UClass>(CreatePropertyBagPlaceholderType(Outer, Class, Name, Flags, T::StaticClass()));
 	}
 
+	/**
+	 * Remove a placeholder type object from the internal registry.
+	 * 
+	 * @param PlaceholderType	Placeholder type object to remove.
+	 */
+	static COREUOBJECT_API void RemovePropertyBagPlaceholderType(UStruct* PlaceholderType);
+
 private:
 	void Lock() const { CriticalSection.Lock(); }
 	void Unlock() const { CriticalSection.Unlock(); }
diff --git a/Engine/Source/Runtime/CoreUObject/Private/UObject/PropertyBagRepository.cpp b/Engine/Source/Runtime/CoreUObject/Private/UObject/PropertyBagRepository.cpp
--- a/Engine/Source/Runtime/CoreUObject/Private/UObject/PropertyBagRepository.cpp
+++ b/Engine/Source/Runtime/CoreUObject/Private/UObject/PropertyBagRepository.cpp
@@ -181,7 +181,7 @@ void FPropertyBagRepository::ReassociateObjects(const TMap<UObject*, UObject*>&
 			}
 			OldBagData.Destroy();
 		}
-		else if (UStruct* TypeObject = Cast<UStruct>(Pair.Value))
+		else if (UStruct* TypeObject = Cast<UStruct>(Pair.Key))
 		{
 			if (IsPropertyBagPlaceholderType(TypeObject))
 			{
@@ -1052,6 +1052,9 @@ UStruct* FPropertyBagRepository::CreatePropertyBagPlaceholderType(UObject* Outer
 
 		// This class is for internal use and should not be exposed for selection or instancing in the editor.
 		PlaceholderTypeAsClass->ClassFlags |= CLASS_Hidden | CLASS_HideDropDown;
+
+		// Required by garbage collection for class types.
+		PlaceholderTypeAsClass->AssembleReferenceTokenStream();
 	}
 
 	// Use the property bag repository for now to manage property bag placeholder types (e.g. object lifetime).
@@ -1061,6 +1064,12 @@ UStruct* FPropertyBagRepository::CreatePropertyBagPlaceholderType(UObject* Outer
 	return PlaceholderType;
 }
 
+void FPropertyBagRepository::RemovePropertyBagPlaceholderType(UStruct* PlaceholderType)
+{
+	ensure(IsPropertyBagPlaceholderType(PlaceholderType));
+	FPropertyBagPlaceholderTypeRegistry::Get().Remove(PlaceholderType);
+}
+
 } // UE
 
 #endif // WITH_EDITORONLY_DATA
