[ZenLoader/Misc]
 - Deprecate EInternalObjectsFlag::AsyncLoading

#rnx

diff --git a/Engine/Source/Runtime/CoreUObject/Public/UObject/ObjectMacros.h b/Engine/Source/Runtime/CoreUObject/Public/UObject/ObjectMacros.h
--- a/Engine/Source/Runtime/CoreUObject/Public/UObject/ObjectMacros.h
+++ b/Engine/Source/Runtime/CoreUObject/Public/UObject/ObjectMacros.h
@@ -633,9 +633,7 @@ enum class EInternalObjectFlags : int32
 	// !!!!!!!!!!!!!!!!!!!!!!!!!!!!!!!!!!!!!!!!!!!!!!!!!!!!!!!!!!!!!!!!!!!!!!!!!!!!!!!!!!!!!!!!!!!!!!!!!!!!!!!!!!!!!!!!!!!!!!!!!!!!!!!!!!!
 	// DO NOT ADD composite flags to EInternalObjectFlags. Debugger visualisations have trouble displaying composite flag values as text.
 	// !!!!!!!!!!!!!!!!!!!!!!!!!!!!!!!!!!!!!!!!!!!!!!!!!!!!!!!!!!!!!!!!!!!!!!!!!!!!!!!!!!!!!!!!!!!!!!!!!!!!!!!!!!!!!!!!!!!!!!!!!!!!!!!!!!!
-
-	// Required for backward compatibility, will be deprecated soon.
-	AsyncLoading = AsyncLoadingPhase1 | AsyncLoadingPhase2
+	AsyncLoading UE_DEPRECATED(5.5, "Use EInternalObjectFlags_AsyncLoading instead.") = AsyncLoadingPhase1 | AsyncLoadingPhase2
 };
 ENUM_CLASS_FLAGS(EInternalObjectFlags);
 
