[ZenLoader]
 - Make sure exports being recreated after being GCed are loaded correctly

#rnx

diff --git a/Engine/Source/Runtime/CoreUObject/Private/Serialization/AsyncLoading2.cpp b/Engine/Source/Runtime/CoreUObject/Private/Serialization/AsyncLoading2.cpp
--- a/Engine/Source/Runtime/CoreUObject/Private/Serialization/AsyncLoading2.cpp
+++ b/Engine/Source/Runtime/CoreUObject/Private/Serialization/AsyncLoading2.cpp
@@ -5653,33 +5653,48 @@ bool FAsyncPackage2::PreloadLinkerLoadExports(FAsyncLoadingThreadState2& ThreadS
 	// Serialize exports
 	const int32 ExportCount = LinkerLoadState->Linker->ExportMap.Num();
 	check(LinkerLoadState->Linker->ExportMap.Num() == Data.Exports.Num());
-	while (LinkerLoadState->SerializeExportIndex < ExportCount)
+
+	while (LinkerLoadState->SerializeExportIndex < ExportCount + ConstructedObjects.Num())
 	{
 		ThreadState.MarkAsActive();
 
 		const int32 ExportIndex = LinkerLoadState->SerializeExportIndex++;
-		FExportObject& ExportObject = Data.Exports[ExportIndex];
-		FObjectExport& LinkerExport = LinkerLoadState->Linker->ExportMap[ExportIndex];
 
-		// The linker export table can be patched during reinstantiation. We need to adjust our own export table if needed.
-		if (!LinkerExport.bExportLoadFailed)
+		if (ExportIndex < ExportCount)
 		{
-			if (ExportObject.Object != LinkerExport.Object)
+			FExportObject& ExportObject = Data.Exports[ExportIndex];
+			FObjectExport& LinkerExport = LinkerLoadState->Linker->ExportMap[ExportIndex];
+
+			// The linker export table can be patched during reinstantiation. We need to adjust our own export table if needed.
+			if (!LinkerExport.bExportLoadFailed)
 			{
-				UE_ASYNC_PACKAGE_LOG(Verbose, Desc, TEXT("PreloadLinkerLoadExports"), TEXT("Patching export %d: %s -> %s"), ExportIndex, *GetPathNameSafe(ExportObject.Object), *GetPathNameSafe(LinkerExport.Object));
-				ExportObject.Object = LinkerExport.Object;
+				if (ExportObject.Object != LinkerExport.Object)
+				{
+					UE_ASYNC_PACKAGE_LOG(Verbose, Desc, TEXT("PreloadLinkerLoadExports"), TEXT("Patching export %d: %s -> %s"), ExportIndex, *GetPathNameSafe(ExportObject.Object), *GetPathNameSafe(LinkerExport.Object));
+					ExportObject.Object = LinkerExport.Object;
+				}
 			}
-		}
 
-		if (UObject* Object = ExportObject.Object)
+			if (UObject* Object = ExportObject.Object)
+			{
+				if (Object->HasAnyFlags(RF_NeedLoad))
+				{
+					UE_ASYNC_PACKAGE_LOG(VeryVerbose, Desc, TEXT("PreloadLinkerLoadExports"), TEXT("Preloading export %d: %s"), ExportIndex, *Object->GetPathName());
+					UE_TRACK_REFERENCING_PACKAGE_SCOPED(Object, PackageAccessTrackingOps::NAME_PreLoad);
+					LinkerLoadState->Linker->Preload(Object);
+				}
+			}
+		}
+		else if (UObject* Object = ConstructedObjects[ExportIndex - ExportCount])
 		{
 			if (Object->HasAnyFlags(RF_NeedLoad))
 			{
-				UE_ASYNC_PACKAGE_LOG(VeryVerbose, Desc, TEXT("PreloadLinkerLoadExports"), TEXT("Preloading export %d: %s"), ExportIndex, *Object->GetPathName());
+				UE_ASYNC_PACKAGE_LOG(Display, Desc, TEXT("PreloadLinkerLoadExports"), TEXT("Preloading additional constructed object: %s"), *Object->GetPathName());
 				UE_TRACK_REFERENCING_PACKAGE_SCOPED(Object, PackageAccessTrackingOps::NAME_PreLoad);
-				LinkerLoadState->Linker->Preload(Object);
+				Object->GetLinker()->Preload(Object);
 			}
 		}
+
 		if (ThreadState.IsTimeLimitExceeded(TEXT("SerializeLinkerLoadExports")))
 		{
 			return false;
diff --git a/Engine/Source/Runtime/CoreUObject/Private/UObject/LinkerLoad.cpp b/Engine/Source/Runtime/CoreUObject/Private/UObject/LinkerLoad.cpp
--- a/Engine/Source/Runtime/CoreUObject/Private/UObject/LinkerLoad.cpp
+++ b/Engine/Source/Runtime/CoreUObject/Private/UObject/LinkerLoad.cpp
@@ -5064,7 +5064,13 @@ UObject* FLinkerLoad::CreateExport( int32 Index )
 	// Check whether we already loaded the object and if not whether the context flags allow loading it.
 	if( !Export.Object && !FilterExport(Export) ) // for some acceptable position, it was not "not for" 
 	{
-		TGuardValue<void*> GuardThreadContextAsyncPackage(FUObjectThreadContext::Get().AsyncPackage, AsyncRoot);
+		TOptional<TGuardValue<void*>> GuardThreadContextAsyncPackage;
+		// Avoid overriding the current async package to null as it would shadow the package currently being loaded
+		// and would prevent our new object from getting picked up by the loader as part of that other package being loaded higher up the stack.
+		if (AsyncRoot)
+		{
+			GuardThreadContextAsyncPackage.Emplace(FUObjectThreadContext::Get().AsyncPackage, AsyncRoot);
+		}
 		FUObjectSerializeContext* CurrentLoadContext = FUObjectThreadContext::Get().GetSerializeContext();
 		check(!GEventDrivenLoaderEnabled || !bLockoutLegacyOperations || !EVENT_DRIVEN_ASYNC_LOAD_ACTIVE_AT_RUNTIME);
 		check(Export.ObjectName!=NAME_None || !(Export.ObjectFlags&RF_Public));
