Moving freeing of UObject's InternalIndex out of UObject destructor to not lock the global UObjectArray when the destructor runs (otherwise it may result in hangs)

diff --git a/Engine/Plugins/Editor/GeometryMode/Source/BspMode/Private/BspModeModule.cpp b/Engine/Plugins/Editor/GeometryMode/Source/BspMode/Private/BspModeModule.cpp
--- a/Engine/Plugins/Editor/GeometryMode/Source/BspMode/Private/BspModeModule.cpp
+++ b/Engine/Plugins/Editor/GeometryMode/Source/BspMode/Private/BspModeModule.cpp
@@ -97,13 +97,17 @@ void FBspModeModule::UnregisterBspBuilderType( class UClass* InBuilderClass )
 			}
 		}
 	}
-	
-	BspBuilderTypes.RemoveAll( 
-		[InBuilderClass] ( const TSharedPtr<FBspBuilderType>& RemovalCandidate ) -> bool
-		{
-			return (RemovalCandidate->BuilderClass == InBuilderClass);
-		}
-	);
+	// If UObject system is not intialized the blow code can not run because it creates weak object pointers which crashes
+	// the internal UObject array.
+	if (UObjectInitialized())
+	{
+		BspBuilderTypes.RemoveAll(
+			[InBuilderClass](const TSharedPtr<FBspBuilderType>& RemovalCandidate) -> bool
+			{
+				return (RemovalCandidate->BuilderClass == InBuilderClass);
+			}
+		);
+	}
 }
 
 
diff --git a/Engine/Source/Runtime/CoreUObject/Private/UObject/GarbageCollection.cpp b/Engine/Source/Runtime/CoreUObject/Private/UObject/GarbageCollection.cpp
--- a/Engine/Source/Runtime/CoreUObject/Private/UObject/GarbageCollection.cpp
+++ b/Engine/Source/Runtime/CoreUObject/Private/UObject/GarbageCollection.cpp
@@ -100,7 +100,7 @@ static int32 GObjCurrentPurgeObjectIndex = 0;
 static bool GObjCurrentPurgeObjectIndexNeedsReset = true;
 
 /** Contains a list of objects that stayed marked as unreachable after the last reachability analysis */
-static TArray<FUObjectItem*> GUnreachableObjects;
+static TArray<UE::GC::FUnreachableObject> GUnreachableObjects;
 static FCriticalSection GUnreachableObjectsCritical;
 static int32 GUnrechableObjectIndex = 0;
 
@@ -580,7 +580,7 @@ namespace UE::GC::Private
 	/** List of FUObjectItems representing cluster root objects marker as reachable by GC barrier (see UObject::MarkAsReachable()) */
 	static TExpandingChunkedList<FUObjectItem*> GReachableClusters;
 
-	using FGatherUnreachableObjectsState = TThreadedGather<TArray<FUObjectItem*>>;
+	using FGatherUnreachableObjectsState = TThreadedGather<TArray<UE::GC::FUnreachableObject>>;
 	static FGatherUnreachableObjectsState GGatherUnreachableObjectsState;
 
 	static TSet<int32> GRoots;
@@ -734,6 +734,8 @@ void OnDisregardForGCSetDisabled(int32 NumObjects)
  */
 class FObjectPurge
 {
+	/** Current index into the global unreachable objects array (GUnreachableObjects) of the object whose index is being released */
+	int32 ObjCurrentFreeIndexObjectIndex = 0;
 	/** Current index into the global unreachable objects array (GUnreachableObjects) of the object being destroyed */
 	int32 ObjCurrentPurgeObjectIndex = 0;
 	/** Stats for the number of objects destroyed */
@@ -753,6 +755,7 @@ public:
 	void Begin()
 	{
 		check(IsFinished()); // In single-threaded mode we need to be finished or the condition below will hang
+		ObjCurrentFreeIndexObjectIndex = 0;
 		ObjCurrentPurgeObjectIndex = 0;
 		ObjectsDestroyedSinceLastMarkPhase = 0;
 	}
@@ -769,40 +772,76 @@ public:
 		TRACE_CPUPROFILER_EVENT_SCOPE(FObjectPurge::DestroyObjects);
 		constexpr int32 TimeLimitEnforcementGranularityForDeletion = 100;
 		int32 ProcessedObjectsCount = 0;
-		bFinishedDestroyingObjects = true;
 
+		// Global UObject Array needs to be locked only when freeing UObject indices
 		GUObjectArray.LockInternalArray();
-
-		while (ObjCurrentPurgeObjectIndex < GUnreachableObjects.Num())
 		{
-			FUObjectItem* ObjectItem = GUnreachableObjects[ObjCurrentPurgeObjectIndex];
-			check(ObjectItem->IsUnreachable());
+			// This loops replaces every entry in GUnreachableObjects which up untill this point has been an FUObjectItem 
+			// with the actual UObject that the FUObjectItem represented.
+			// This is because in this loop we free UObject indices meaning that they are removed from GUObjectArray and 
+			// the associated FUObjectItem is reset and no longer pointing at UObject.
+			// This approach has the benefit that we don't need to lock GUObjectArray when calling UObject destructors and 
+			// we can reclaim GUObjectArray entires faster.
+			while (ObjCurrentFreeIndexObjectIndex < GUnreachableObjects.Num())
+			{
+				UE::GC::FUnreachableObject& UnreachableObject = GUnreachableObjects[ObjCurrentFreeIndexObjectIndex];
+				FUObjectItem* ObjectItem = UnreachableObject.ObjectItem;
+				check(ObjectItem->IsUnreachable());
 
-			UObject* Object = (UObject*)ObjectItem->Object;
-			check(Object->HasAllFlags(RF_FinishDestroyed | RF_BeginDestroyed));
+				UObject* Object = (UObject*)ObjectItem->Object;
+				check(Object->HasAllFlags(RF_FinishDestroyed | RF_BeginDestroyed));
 
-			Object->~UObject();
-			GUObjectAllocator.FreeUObject(Object);
-			GUnreachableObjects[ObjCurrentPurgeObjectIndex] = nullptr;
+				GUObjectArray.FreeUObjectIndex(Object);
 
-			++ProcessedObjectsCount;
-			++ObjectsDestroyedSinceLastMarkPhase;
-			++ObjCurrentPurgeObjectIndex;
+				// Replace the entry in GUnreachableObjects with the actual UObject so that we can iterate over the same array when 
+				// we call UObject destructors and free their memory in the loop below
+				UnreachableObject.Object = Object;
 
-			// Time slicing when running on the game thread
-			if (bUseTimeLimit && (ProcessedObjectsCount == TimeLimitEnforcementGranularityForDeletion) && (ObjCurrentPurgeObjectIndex < GUnreachableObjects.Num()))
-			{
-				ProcessedObjectsCount = 0;
-				if ((FPlatformTime::Seconds() - StartTime) > TimeLimit)
+				++ProcessedObjectsCount;
+				++ObjCurrentFreeIndexObjectIndex;
+
+				// Time slicing when running on the game thread
+				if (bUseTimeLimit && (ProcessedObjectsCount == TimeLimitEnforcementGranularityForDeletion) && (ObjCurrentFreeIndexObjectIndex < GUnreachableObjects.Num()))
 				{
-					bFinishedDestroyingObjects = false;
-					break;
+					ProcessedObjectsCount = 0;
+					if ((FPlatformTime::Seconds() - StartTime) > TimeLimit)
+					{
+						break;
+					}
 				}
 			}
 		}
-
 		GUObjectArray.UnlockInternalArray();
 
+		if (ObjCurrentFreeIndexObjectIndex == GUnreachableObjects.Num())
+		{
+			// At this point all entires in GUnreachableObjects point at UObject memory instead of FUObjectItems
+			while (ObjCurrentPurgeObjectIndex < GUnreachableObjects.Num())
+			{
+				UE::GC::FUnreachableObject& UnreachableObject = GUnreachableObjects[ObjCurrentPurgeObjectIndex];
+				UObject* Object = UnreachableObject.Object;
+
+				Object->~UObject();
+				GUObjectAllocator.FreeUObject(Object);
+				UnreachableObject.Object = nullptr;
+
+				++ProcessedObjectsCount;
+				++ObjectsDestroyedSinceLastMarkPhase;
+				++ObjCurrentPurgeObjectIndex;
+
+				// Time slicing when running on the game thread
+				if (bUseTimeLimit && (ProcessedObjectsCount == TimeLimitEnforcementGranularityForDeletion) && (ObjCurrentPurgeObjectIndex < GUnreachableObjects.Num()))
+				{
+					ProcessedObjectsCount = 0;
+					if ((FPlatformTime::Seconds() - StartTime) > TimeLimit)
+					{
+						break;
+					}
+				}
+			}
+		}
+
+		bFinishedDestroyingObjects = (ObjCurrentPurgeObjectIndex == GUnreachableObjects.Num());
 		return bFinishedDestroyingObjects;
 	}
 
@@ -820,9 +859,9 @@ public:
 
 	FORCENOINLINE void VerifyAllObjectsDestroyed() const
 	{
-		for (FUObjectItem* ObjectItem : GUnreachableObjects)
+		for (UE::GC::FUnreachableObject UnreachableObject : GUnreachableObjects)
 		{
-			UE_CLOG(ObjectItem, LogGarbage, Fatal, TEXT("Object 0x%016llx has not been destroyed during async purge"), (int64)(PTRINT)ObjectItem->Object);
+			UE_CLOG(UnreachableObject.Object, LogGarbage, Fatal, TEXT("Object 0x%016llx has not been destroyed during purge"), (int64)(PTRINT)UnreachableObject.Object);
 		}
 	}
 } GUObjectPurge;
@@ -4573,7 +4612,7 @@ bool IncrementalDestroyGarbage(bool bUseTimeLimit, double TimeLimit)
 			TRACE_CPUPROFILER_EVENT_SCOPE(ConditionalFinishDestroy);
 			while (GObjCurrentPurgeObjectIndex < GUnreachableObjects.Num())
 			{
-				FUObjectItem* ObjectItem = GUnreachableObjects[GObjCurrentPurgeObjectIndex];
+				FUObjectItem* ObjectItem = GUnreachableObjects[GObjCurrentPurgeObjectIndex].ObjectItem;
 				checkSlow(ObjectItem);
 
 				//@todo UE - A prefetch was removed here. Re-add it. It wasn't right anyway, since it was ten items ahead and the consoles on have 8 prefetch slots
@@ -4994,7 +5033,7 @@ bool GatherUnreachableObjects(UE::GC::EGatherOptions Options, double TimeLimit /
 			{
 				checkf(!ObjectItem->HasAnyFlags(EInternalObjectFlags::ClusterRoot), TEXT("Unreachable cluster root found. Unreachable clusters should have been dissolved in DissolveUnreachableClusters!"));
 				FGCFlags::SetUnreachable(ObjectItem);
-				Iterator.Payload.Add(ObjectItem);
+				Iterator.Payload.Add({ ObjectItem });
 			}
 
 			if (Timer.IsTimeLimitExceeded())
@@ -5034,7 +5073,12 @@ bool GatherUnreachableObjects(UE::GC::EGatherOptions Options, double TimeLimit /
 		{
 			AcquireGCLock();
 		}
-		NotifyUnreachableObjects(GUnreachableObjects);
+		{			
+			TArrayView<UE::GC::FUnreachableObject> UnrachableObjectsView(GUnreachableObjects);
+			// Although the below cast is not pretty, we don't want to expose UE::GC::FUnreachableObject to async loading API 
+			// as it may be confusing what UE::GC::FUnreachableObject currently represents
+			NotifyUnreachableObjects(*reinterpret_cast<TArrayView<FUObjectItem*>*>(&UnrachableObjectsView));
+		}
 		if (bNeedsGCLock)
 		{
 			ReleaseGCLock();
@@ -5860,7 +5904,7 @@ bool UnhashUnreachableObjects(bool bUseTimeLimit, double TimeLimit)
 		{
 			//@todo UE - A prefetch was removed here. Re-add it. It wasn't right anyway, since it was ten items ahead and the consoles on have 8 prefetch slots
 
-			FUObjectItem* ObjectItem = GUnreachableObjects[GUnrechableObjectIndex++];
+			FUObjectItem* ObjectItem = GUnreachableObjects[GUnrechableObjectIndex++].ObjectItem;
 			{
 				UObject* Object = static_cast<UObject*>(ObjectItem->Object);
 				FScopedCBDProfile Profile(Object);
@@ -5975,6 +6019,8 @@ bool TryCollectGarbage(EObjectFlags KeepFlags, bool bPerformFullPurge)
  */
 void PurgeAllUObjectsOnExit()
 {
+	using namespace UE::GC::Private;
+
 	// This can happen when we run into an error early in the init process
 	if (GUObjectArray.IsOpenForDisregardForGC())
 	{
@@ -6056,6 +6102,12 @@ void PurgeAllUObjectsOnExit()
 		IncrementalPurgeGarbage(false);
 	}
 
+	{
+		// Since we just destroyed all objects regardless of their flags, make sure the roots list is empty
+		FScopeLock RootsLock(&GRootsCritical);
+		GRoots.Empty();
+	}
+
 	ReleaseGCLock();
 }
 
diff --git a/Engine/Source/Runtime/CoreUObject/Private/UObject/UObjectArray.cpp b/Engine/Source/Runtime/CoreUObject/Private/UObject/UObjectArray.cpp
--- a/Engine/Source/Runtime/CoreUObject/Private/UObject/UObjectArray.cpp
+++ b/Engine/Source/Runtime/CoreUObject/Private/UObject/UObjectArray.cpp
@@ -338,6 +338,10 @@ void FUObjectArray::FreeUObjectIndex(UObjectBase* Object)
 
 	// No need to call LockInternalArray(); here as it should already be locked by GC
 
+#if UE_WITH_OBJECT_HANDLE_LATE_RESOLVE
+	UE::CoreUObject::Private::FreeObjectHandle(Object);
+#endif 	
+
 	int32 Index = Object->InternalIndex;
 	FUObjectItem* ObjectItem = IndexToObject(Index);
 	UE_CLOG(ObjectItem->Object != Object, LogUObjectArray, Fatal, TEXT("Removing object (0x%016llx) at index %d but the index points to a different object (0x%016llx)!"), (int64)(PTRINT)Object, Index, (int64)(PTRINT)ObjectItem->Object);
@@ -356,6 +360,7 @@ void FUObjectArray::FreeUObjectIndex(UObjectBase* Object)
 	ObjectItem->RefCount = 0;
 	ObjectItem->ClusterRootIndex = 0;
 	ObjectItem->SerialNumber = 0;
+	Object->InternalIndex = INDEX_NONE;
 
 	// You cannot safely recycle indicies in the non-GC range
 	// No point in filling this list when doing exit purge. Nothing should be allocated afterwards anyway.
diff --git a/Engine/Source/Runtime/CoreUObject/Private/UObject/UObjectBase.cpp b/Engine/Source/Runtime/CoreUObject/Private/UObject/UObjectBase.cpp
--- a/Engine/Source/Runtime/CoreUObject/Private/UObject/UObjectBase.cpp
+++ b/Engine/Source/Runtime/CoreUObject/Private/UObject/UObjectBase.cpp
@@ -137,7 +137,7 @@ UObjectBase::UObjectBase(UClass* InClass,
 #if CSV_PROFILER_STATS && CSV_TRACK_UOBJECT_COUNT
 	UObjectStats::IncrementUObjectCount();
 #endif
-		}	
+}	
 
 
 /**
@@ -149,12 +149,9 @@ UObjectBase::~UObjectBase()
 	if( UObjectInitialized() && ClassPrivate && !GIsCriticalError )
 	{
 		// Validate it.
-		check(IsValidLowLevel());
+		check(IsValidLowLevelForDestruction());
 		check(GetFName() == NAME_None);
-#if UE_WITH_OBJECT_HANDLE_LATE_RESOLVE
-		UE::CoreUObject::Private::FreeObjectHandle(this);
-#endif 
-		GUObjectArray.FreeUObjectIndex(this);
+		checkf(InternalIndex == INDEX_NONE, TEXT("Object destroyed outside of GC (InternalIndex=%d, expected %d)"), InternalIndex, INDEX_NONE);	
 	}
 
 #if CSV_PROFILER_STATS && CSV_TRACK_UOBJECT_COUNT
@@ -306,23 +303,28 @@ void UObjectBase::SetClass(UClass* NewClass)
 }
 #endif
 
-/**
- * Checks to see if the object appears to be valid
- * @return true if this appears to be a valid object
- */
-bool UObjectBase::IsValidLowLevel() const
+bool UObjectBase::IsValidLowLevelForDestruction() const
 {
-	if( this == nullptr )
+	if (this == nullptr)
 	{
-		UE_LOG(LogUObjectBase, Warning, TEXT("NULL object") );
+		UE_LOG(LogUObjectBase, Warning, TEXT("NULL object"));
 		return false;
 	}
-	if( !ClassPrivate )
+	if (!ClassPrivate)
 	{
-		UE_LOG(LogUObjectBase, Warning, TEXT("Object is not registered") );
+		UE_LOG(LogUObjectBase, Warning, TEXT("Object is not registered"));
 		return false;
 	}
-	return GUObjectArray.IsValid(this);
+	return true;
+}
+
+/**
+ * Checks to see if the object appears to be valid
+ * @return true if this appears to be a valid object
+ */
+bool UObjectBase::IsValidLowLevel() const
+{
+	return IsValidLowLevelForDestruction() && GUObjectArray.IsValid(this);
 }
 
 bool UObjectBase::IsValidLowLevelFast(bool bRecursive /*= true*/) const
diff --git a/Engine/Source/Runtime/CoreUObject/Private/UObject/UObjectGlobals.cpp b/Engine/Source/Runtime/CoreUObject/Private/UObject/UObjectGlobals.cpp
--- a/Engine/Source/Runtime/CoreUObject/Private/UObject/UObjectGlobals.cpp
+++ b/Engine/Source/Runtime/CoreUObject/Private/UObject/UObjectGlobals.cpp
@@ -3604,10 +3604,11 @@ UObject* StaticAllocateObject
 				// Finish destroying the object.
 				Obj->ConditionalFinishDestroy();
 			}
-			GUObjectArray.LockInternalArray();
 			TGuardValue<bool> _(GUObjectArray.bShouldRecycleObjectIndices, false);
-			Obj->~UObject();
+			GUObjectArray.LockInternalArray();
+			GUObjectArray.FreeUObjectIndex(Obj);
 			GUObjectArray.UnlockInternalArray();
+			Obj->~UObject();
 			bWasConstructedOnOldObject	= true;
 		}
 		else
diff --git a/Engine/Source/Runtime/CoreUObject/Public/UObject/ReachabilityAnalysis.h b/Engine/Source/Runtime/CoreUObject/Public/UObject/ReachabilityAnalysis.h
--- a/Engine/Source/Runtime/CoreUObject/Public/UObject/ReachabilityAnalysis.h
+++ b/Engine/Source/Runtime/CoreUObject/Public/UObject/ReachabilityAnalysis.h
@@ -9,6 +9,7 @@
 
 #include "CoreMinimal.h"
 
+struct FUObjectItem;
 
 /**
  * Enables or disables incremental reachability analysis.
@@ -36,6 +37,21 @@ COREUOBJECT_API void SetReachabilityAnalysisTimeLimit(float TimeLimitSeconds);
  */
 COREUOBJECT_API float GetReachabilityAnalysisTimeLimit();
 
+
+namespace UE::GC
+{
+
+/**
+ * Entry in GUnreachableObjects array that can either be an FUObjectItem (initially) or an UObject after unreachable objects had their InternalIndex freed.
+ */
+union FUnreachableObject
+{
+	FUObjectItem* ObjectItem;
+	UObject* Object;
+};
+
+}
+
 namespace UE::GC::Private
 {
 
diff --git a/Engine/Source/Runtime/CoreUObject/Public/UObject/UObjectBase.h b/Engine/Source/Runtime/CoreUObject/Public/UObject/UObjectBase.h
--- a/Engine/Source/Runtime/CoreUObject/Public/UObject/UObjectBase.h
+++ b/Engine/Source/Runtime/CoreUObject/Public/UObject/UObjectBase.h
@@ -124,6 +124,12 @@ private:
 	 */
 	COREUOBJECT_API void AddObject(FName Name, EInternalObjectFlags InSetInternalFlags, int32 InInternalIndex = -1, int32 InSerialNumber = 0);
 
+	/**
+	 * Checks to see if the object appears to be valid when being destroyed
+	 * @return true if this appears to be a valid object
+	 */
+	COREUOBJECT_API bool IsValidLowLevelForDestruction() const;
+
 public:
 	/**
 	 * Checks to see if the object appears to be valid
diff --git a/Engine/Source/Runtime/SlateRHIRenderer/Private/SlateRHIRendererSettings.cpp b/Engine/Source/Runtime/SlateRHIRenderer/Private/SlateRHIRendererSettings.cpp
--- a/Engine/Source/Runtime/SlateRHIRenderer/Private/SlateRHIRendererSettings.cpp
+++ b/Engine/Source/Runtime/SlateRHIRenderer/Private/SlateRHIRendererSettings.cpp
@@ -43,14 +43,19 @@ USlateRHIRendererSettings::USlateRHIRendererSettings()
 
 USlateRHIRendererSettings::~USlateRHIRendererSettings()
 {
-	for (TPair<ESlatePostRT, FSlatePostSettings>& SlatePostSetting : SlatePostSettings)
+	// No need to remove from root during the exit purge as all objects will destroyed regardless of their flags 
+	// and it's not safe to access any other UObjects from UObject destructors
+	if (!GExitPurge)
 	{
-		FSlatePostSettings& PostSetting = SlatePostSetting.Value;
-
-		UObject* SlatePostBuffer = PostSetting.CachedSlatePostRT;
-		if (SlatePostBuffer)
+		for (TPair<ESlatePostRT, FSlatePostSettings>& SlatePostSetting : SlatePostSettings)
 		{
-			SlatePostBuffer->RemoveFromRoot();
+			FSlatePostSettings& PostSetting = SlatePostSetting.Value;
+
+			UObject* SlatePostBuffer = PostSetting.CachedSlatePostRT;
+			if (SlatePostBuffer)
+			{
+				SlatePostBuffer->RemoveFromRoot();
+			}
 		}
 	}
 }
