FTopLevelAssetPath: Add support for appending to FUtf8StringBuilderBase. Add support for writing FTopLevelAssetPath to FCbWriter. Add support for loading a FTopLevelAssetPath from FCbFieldView.
These changes allow FTopLevelAssetPath to be output via structured log.

#tests Ran System.Core.Misc.TopLevelAssetPath unit test against a cooked build via UnrealFrontend

diff --git a/Engine/Source/Runtime/CoreUObject/Private/UObject/TopLevelAssetPath.cpp b/Engine/Source/Runtime/CoreUObject/Private/UObject/TopLevelAssetPath.cpp
--- a/Engine/Source/Runtime/CoreUObject/Private/UObject/TopLevelAssetPath.cpp
+++ b/Engine/Source/Runtime/CoreUObject/Private/UObject/TopLevelAssetPath.cpp
@@ -8,6 +8,7 @@
 #include "Misc/Optional.h"
 #include "Misc/PackageName.h"
 #include "Misc/RedirectCollector.h"
+#include "Serialization/CompactBinaryWriter.h"
 #include "UObject/CoreRedirects.h"
 #include "UObject/ObjectRedirector.h"
 #include "UObject/Package.h"
@@ -26,7 +27,7 @@ struct TStructOpsTypeTraits<FTopLevelAssetPath> : public TStructOpsTypeTraitsBas
 };
 UE_IMPLEMENT_STRUCT("/Script/CoreUObject", TopLevelAssetPath)
 
-void FTopLevelAssetPath::AppendString(FStringBuilderBase& Builder) const
+void FTopLevelAssetPath::AppendString(FWideStringBuilderBase& Builder) const
 {
 	if (!IsNull())
 	{
@@ -34,7 +35,19 @@ void FTopLevelAssetPath::AppendString(FStringBuilderBase& Builder) const
 		if (!AssetName.IsNone())
 		{
 			Builder << '.' << AssetName;
-		}		
+		}
+	}
+}
+
+void FTopLevelAssetPath::AppendString(FUtf8StringBuilderBase& Builder) const
+{
+	if (!IsNull())
+	{
+		Builder << PackageName;
+		if (!AssetName.IsNone())
+		{
+			Builder << '.' << AssetName;
+		}
 	}
 }
 
@@ -513,11 +526,21 @@ bool FTopLevelAssetPath::SerializeFromMismatchedTag(const FPropertyTag& Tag, FSt
 	return false;
 }
 
+void FTopLevelAssetPath::WriteCompactBinary(FCbWriter& Writer) const
+{
+	Writer << WriteToUtf8String<FName::StringBufferSize>(*this).ToView();
+}
 
+bool LoadFromCompactBinary(FCbFieldView Field, FTopLevelAssetPath& OutPath)
+{
+	OutPath = FTopLevelAssetPath(Field.AsString());
+	return !Field.HasError();
+}
 
 #if WITH_DEV_AUTOMATION_TESTS 
 
 #include "Misc/AutomationTest.h"
+#include "Serialization/CompactBinaryValidation.h"
 
 // Combine import/export tests
 
@@ -547,12 +570,12 @@ bool FTopLevelAssetPathTest::RunTest(const FString& Parameters)
 	TestEqual(TEXT("Empty path to string is empty string"), EmptyPathFromString.ToString(), FString());
 
 	FTopLevelAssetPath PackagePathFromString;
-	TestTrue(TEXT("TrySetPath(PackageName.ToString()) succeeds"), PackagePath.TrySetPath(PackageName.ToString()));
-	TestEqual(TEXT("PackagePath to string is PackageName"), PackagePath.ToString(), PackageName.ToString());
+	TestTrue(TEXT("TrySetPath(PackageName.ToString()) succeeds"), PackagePathFromString.TrySetPath(PackageName.ToString()));
+	TestEqual(TEXT("PackagePathFromString to string is PackageName"), PackagePathFromString.ToString(), PackageName.ToString());
 
 	FTopLevelAssetPath AssetPathFromString;
-	TestTrue(TEXT("TrySetPath(AssetPath) succeeds"), PackagePath.TrySetPath(AssetPathString));
-	TestEqual(TEXT("AssetPathFromString to string is PackageName.AssetName"), PackagePath.ToString(), AssetPathString);
+	TestTrue(TEXT("TrySetPath(AssetPath) succeeds"), AssetPathFromString.TrySetPath(AssetPathString));
+	TestEqual(TEXT("AssetPathFromString to string is PackageName.AssetName"), AssetPathFromString.ToString(), AssetPathString);
 
 	FTopLevelAssetPath FailedPath;
 	//TestFalse(TEXT("TrySetPath with unrooted path string fails"), FailedPath.TrySetPath("UnrootedPackage/Subfolder")); // after ANY_PACKAGE removal this will assert
@@ -584,6 +607,42 @@ bool FTopLevelAssetPathTest::RunTest(const FString& Parameters)
 	TestEqual(TEXT("FromShortExportTextPath ImportTextItem advances buffer"), *ShortExportTextBuffer , '\0');
 	TestEqual(TEXT("FromShortExportTextPath  imports correct path"), FromShortExportTextPath.ToString(), TEXT("/Path/To/Package.Asset"));
 
+	{
+		FCbWriter EmptyPathWriter;
+		EmptyPathWriter << FTopLevelAssetPath();
+		FCbFieldIterator EmptyPathField = EmptyPathWriter.Save();
+		if (TestEqual(TEXT("FCbWriter(TopLevelAssetPath) Validate"), ValidateCompactBinaryRange(EmptyPathField.GetOuterBuffer(), ECbValidateMode::All), ECbValidateError::None))
+		{
+			FTopLevelAssetPath EmptyPathFromCompactBinary;
+			TestTrue(TEXT("EmptyPathFromCompactBinary LoadFromCompactBinary succeeds"), LoadFromCompactBinary(FCbFieldView(*EmptyPathField), EmptyPathFromCompactBinary));
+			TestEqual(TEXT("EmptyPathFromCompactBinary LoadFromCompactBinary correct path"), EmptyPathFromCompactBinary.ToString(), FString());
+		}
+	}
+
+	{
+		FCbWriter PackagePathWriter;
+		PackagePathWriter << FTopLevelAssetPath(PackagePath);
+		FCbFieldIterator PackagePathField = PackagePathWriter.Save();
+		if (TestEqual(TEXT("FCbWriter(TopLevelAssetPath) Validate"), ValidateCompactBinaryRange(PackagePathField.GetOuterBuffer(), ECbValidateMode::All), ECbValidateError::None))
+		{
+			FTopLevelAssetPath PackagePathFromCompactBinary;
+			TestTrue(TEXT("PackagePathFromCompactBinary LoadFromCompactBinary succeeds"), LoadFromCompactBinary(FCbFieldView(*PackagePathField), PackagePathFromCompactBinary));
+			TestEqual(TEXT("PackagePathFromCompactBinary LoadFromCompactBinary correct path"), PackagePathFromCompactBinary.ToString(), PackageName.ToString());
+		}
+	}
+
+	{
+		FCbWriter AssetPathWriter;
+		AssetPathWriter << FTopLevelAssetPath(AssetPathString);
+		FCbFieldIterator AssetPathField = AssetPathWriter.Save();
+		if (TestEqual(TEXT("FCbWriter(TopLevelAssetPath) Validate"), ValidateCompactBinaryRange(AssetPathField.GetOuterBuffer(), ECbValidateMode::All), ECbValidateError::None))
+		{
+			FTopLevelAssetPath AssetPathFromCompactBinary;
+			TestTrue(TEXT("AssetPathFromCompactBinary LoadFromCompactBinary succeeds"), LoadFromCompactBinary(FCbFieldView(*AssetPathField), AssetPathFromCompactBinary));
+			TestEqual(TEXT("AssetPathFromCompactBinary LoadFromCompactBinary correct path"), AssetPathFromCompactBinary.ToString(), AssetPathString);
+		}
+	}
+
 	// Test ExportText path with and without root 
 	// Test starting with . 
 
diff --git a/Engine/Source/Runtime/CoreUObject/Public/UObject/TopLevelAssetPath.h b/Engine/Source/Runtime/CoreUObject/Public/UObject/TopLevelAssetPath.h
--- a/Engine/Source/Runtime/CoreUObject/Public/UObject/TopLevelAssetPath.h
+++ b/Engine/Source/Runtime/CoreUObject/Public/UObject/TopLevelAssetPath.h
@@ -17,6 +17,7 @@
 #include <type_traits>
 
 struct FPropertyTag;
+class FCbWriter;
 class FString;
 class UObject;
 
@@ -85,7 +86,8 @@ struct FTopLevelAssetPath
 	FName GetAssetName() const { return AssetName; }
 
 	/** Append the full asset path (e.g. '/Path/To/Package.AssetName') to the string builder. */
-	COREUOBJECT_API void AppendString(FStringBuilderBase& Builder) const;
+	COREUOBJECT_API void AppendString(FWideStringBuilderBase& Builder) const;
+	COREUOBJECT_API void AppendString(FUtf8StringBuilderBase& Builder) const;
 	/** Append the full asset path (e.g. '/Path/To/Package.AssetName') to the string. */
 	COREUOBJECT_API void AppendString(FString& OutString) const;
 
@@ -171,15 +173,28 @@ struct FTopLevelAssetPath
 	COREUOBJECT_API bool ImportTextItem( const TCHAR*& Buffer, int32 PortFlags, UObject* Parent, FOutputDevice* ErrorText, FArchive* InSerializingArchive = nullptr );
 	COREUOBJECT_API bool SerializeFromMismatchedTag(const FPropertyTag& Tag, FStructuredArchive::FSlot Slot);
 
+	COREUOBJECT_API void WriteCompactBinary(FCbWriter& Writer) const;
 private:
+	friend FCbWriter& operator<<(FCbWriter& Writer, const FTopLevelAssetPath& Path)
+	{
+		Path.WriteCompactBinary(Writer);
+		return Writer;
+	}
+	COREUOBJECT_API friend bool LoadFromCompactBinary(FCbFieldView Field, FTopLevelAssetPath& OutPath);
+
 	/** Name of the package containing the asset e.g. /Path/To/Package */
 	FName PackageName;
 	/** Name of the asset within the package e.g. 'AssetName' */
 	FName AssetName;
 };
 
+inline FWideStringBuilderBase& operator<<(FWideStringBuilderBase& Builder, const FTopLevelAssetPath& Path)
+{
+	Path.AppendString(Builder);
+	return Builder;
+}
 
-inline FStringBuilderBase& operator<<(FStringBuilderBase& Builder, const FTopLevelAssetPath& Path)
+inline FUtf8StringBuilderBase& operator<<(FUtf8StringBuilderBase& Builder, const FTopLevelAssetPath& Path)
 {
 	Path.AppendString(Builder);
 	return Builder;
