Change lambda to accept it's parameters  by reference instead of by value

#rnx

- Minor optimization to prevent's needless copying of containers.

diff --git a/Engine/Source/Runtime/CoreUObject/Private/AssetRegistry/AssetDataTagMap.cpp b/Engine/Source/Runtime/CoreUObject/Private/AssetRegistry/AssetDataTagMap.cpp
--- a/Engine/Source/Runtime/CoreUObject/Private/AssetRegistry/AssetDataTagMap.cpp
+++ b/Engine/Source/Runtime/CoreUObject/Private/AssetRegistry/AssetDataTagMap.cpp
@@ -1172,9 +1172,9 @@ namespace FixedTagPrivate
 		void Save(const FStoreData& Store)
 		{
 			SaveItem(BeginMagic);
-			VisitViews(Store, [&] (auto Array) { SaveItem(Array.Num()); });
+			VisitViews(Store, [&](auto& Array) { SaveItem(Array.Num()); });
 			SaveTextData(MakeArrayView(Store.Texts));
-			VisitViews<EOrder::SkipText>(Store, [&] (auto Array) { SaveViewData(MakeArrayView(Array)); });
+			VisitViews<EOrder::SkipText>(Store, [&] (auto& Array) { SaveViewData(MakeArrayView(Array)); });
 			SaveItem(EndMagic);
 		}
 
