[AR] Memory optimization: Don't include the package ID when saving FText into asset registry tags

#rnx
#tests RR

diff --git a/Engine/Source/Runtime/CoreUObject/Private/AssetRegistry/AssetDataTagMap.cpp b/Engine/Source/Runtime/CoreUObject/Private/AssetRegistry/AssetDataTagMap.cpp
--- a/Engine/Source/Runtime/CoreUObject/Private/AssetRegistry/AssetDataTagMap.cpp
+++ b/Engine/Source/Runtime/CoreUObject/Private/AssetRegistry/AssetDataTagMap.cpp
@@ -138,7 +138,7 @@ uint32 GetTypeHash(const FAssetRegistryExportPath& Path)
 static FString ToComplexString(const FText& In)
 {
 	FString Out;
-	FTextStringHelper::WriteToBuffer(Out, In);
+	FTextStringHelper::WriteToBuffer(Out, In, /*bRequireQuotes*/false, /*bStripPackageNamespace*/true);
 	return Out;
 }
 
diff --git a/Engine/Source/Runtime/CoreUObject/Private/UObject/Obj.cpp b/Engine/Source/Runtime/CoreUObject/Private/UObject/Obj.cpp
--- a/Engine/Source/Runtime/CoreUObject/Private/UObject/Obj.cpp
+++ b/Engine/Source/Runtime/CoreUObject/Private/UObject/Obj.cpp
@@ -31,6 +31,7 @@
 #include "UObject/UObjectIterator.h"
 #include "UObject/Package.h"
 #include "UObject/MetaData.h"
+#include "UObject/TextProperty.h"
 #include "Templates/Casts.h"
 #include "UObject/AssetRegistryTagsContext.h"
 #include "UObject/LazyObjectPtr.h"
@@ -2205,8 +2206,16 @@ static void GetAssetRegistryTagFromProperty(const void* BaseMemoryLocation, cons
 		}
 
 		FString PropertyStr;
-		const uint8* PropertyAddr = Prop->ContainerPtrToValuePtr<uint8>(BaseMemoryLocation);
-		Prop->ExportTextItem_Direct(PropertyStr, PropertyAddr, PropertyAddr, nullptr, PPF_None);
+		if (const FTextProperty* TextProp = CastField<FTextProperty>(Prop))
+		{
+			const FText& TextValue = TextProp->GetPropertyValue_InContainer(BaseMemoryLocation);
+			FTextStringHelper::WriteToBuffer(PropertyStr, TextValue, /*bRequireQuotes*/false, /*bStripPackageNamespace*/true);
+		}
+		else
+		{
+			const uint8* PropertyAddr = Prop->ContainerPtrToValuePtr<uint8>(BaseMemoryLocation);
+			Prop->ExportTextItem_Direct(PropertyStr, PropertyAddr, PropertyAddr, nullptr, PPF_None);
+		}
 
 		Context.AddTag(UObject::FAssetRegistryTag(Prop->GetFName(), MoveTemp(PropertyStr), TagType));
 	}
