Temp fix for self referencing in verse.

#okforversepublic

diff --git a/Engine/Source/Runtime/CoreUObject/Private/UObject/Class.cpp b/Engine/Source/Runtime/CoreUObject/Private/UObject/Class.cpp
--- a/Engine/Source/Runtime/CoreUObject/Private/UObject/Class.cpp
+++ b/Engine/Source/Runtime/CoreUObject/Private/UObject/Class.cpp
@@ -2589,7 +2589,8 @@ void UStruct::InstanceSubobjectTemplates( void* Data, void const* DefaultData, U
 
 	for ( FProperty* Property = RefLink; Property != NULL; Property = Property->NextRef )
 	{
-		if (Property->ContainsInstancedObjectProperty() && (!InstanceGraph || !InstanceGraph->IsPropertyInSubobjectExclusionList(Property)))
+		if ((Property->ContainsInstancedObjectProperty() || Property->HasAnyPropertyFlags(CPF_AllowSelfReference)) && 
+			(!InstanceGraph || !InstanceGraph->IsPropertyInSubobjectExclusionList(Property)))
 		{
 			Property->InstanceSubobjects( Property->ContainerPtrToValuePtr<uint8>(Data), (uint8*)Property->ContainerPtrToValuePtrForDefaults<uint8>(DefaultStruct, DefaultData), Owner, InstanceGraph );
 		}
diff --git a/Engine/Source/Runtime/CoreUObject/Private/UObject/CoreNative.cpp b/Engine/Source/Runtime/CoreUObject/Private/UObject/CoreNative.cpp
--- a/Engine/Source/Runtime/CoreUObject/Private/UObject/CoreNative.cpp
+++ b/Engine/Source/Runtime/CoreUObject/Private/UObject/CoreNative.cpp
@@ -143,10 +143,8 @@ UObject* FObjectInstancingGraph::GetInstancedSubobject( UObject* SourceSubobject
 {
 	checkSlow(SourceSubobject);
 
-	const bool bAreOverridesEnabled = SourceSubobject && FOverridableManager::Get().IsEnabled(*SourceSubobject);
 	const bool bDoNotCreateNewInstance = !!(Flags & EInstancePropertyValueFlags::DoNotCreateNewInstance);
-	const bool bAllowSelfReference     = !!(Flags & EInstancePropertyValueFlags::AllowSelfReference) || bAreOverridesEnabled;
-
+	const bool bAllowSelfReference     = !!(Flags & EInstancePropertyValueFlags::AllowSelfReference);
 
 	UObject* InstancedSubobject = INVALID_OBJECT;
 
@@ -169,6 +167,7 @@ UObject* FObjectInstancingGraph::GetInstancedSubobject( UObject* SourceSubobject
 		if ( bShouldInstance )
 		{
 			// If the CurrentValue is within the SourceRoot, lets use it to instantiate as it must have come from the merge result of the serialization
+			const bool bAreOverridesEnabled = SourceSubobject && FOverridableManager::Get().IsEnabled(*SourceSubobject);
 			if (bAreOverridesEnabled && SourceSubobject != CurrentValue && CurrentValue->IsIn(SourceRoot))
 			{
 				SourceSubobject = CurrentValue;
diff --git a/Engine/Source/Runtime/CoreUObject/Private/UObject/PropertyBaseObject.cpp b/Engine/Source/Runtime/CoreUObject/Private/UObject/PropertyBaseObject.cpp
--- a/Engine/Source/Runtime/CoreUObject/Private/UObject/PropertyBaseObject.cpp
+++ b/Engine/Source/Runtime/CoreUObject/Private/UObject/PropertyBaseObject.cpp
@@ -61,7 +61,16 @@ void FObjectPropertyBase::InstanceSubobjects(void* Data, void const* DefaultData
 		if ( CurrentValue )
 		{
 			UObject *SubobjectTemplate = DefaultData ? GetObjectPropertyValue((uint8*)DefaultData + ArrayIndex * ElementSize): nullptr;
-			UObject* NewValue = InstanceGraph->InstancePropertyValue(SubobjectTemplate, CurrentValue, InOwner, HasAnyPropertyFlags(CPF_InstancedReference) ? EInstancePropertyValueFlags::CausesInstancing : EInstancePropertyValueFlags::None);
+			EInstancePropertyValueFlags Flags = EInstancePropertyValueFlags::None;
+			if (HasAnyPropertyFlags(CPF_InstancedReference))
+			{
+				Flags |= EInstancePropertyValueFlags::CausesInstancing;
+			}
+			if (HasAnyPropertyFlags(CPF_AllowSelfReference))
+			{
+				Flags |= EInstancePropertyValueFlags::AllowSelfReference;
+			}
+			UObject* NewValue = InstanceGraph->InstancePropertyValue(SubobjectTemplate, CurrentValue, InOwner, Flags);
 			SetObjectPropertyValue((uint8*)Data + ArrayIndex * ElementSize, NewValue);
 		}
 	}
diff --git a/Engine/Source/Runtime/CoreUObject/Public/UObject/ObjectMacros.h b/Engine/Source/Runtime/CoreUObject/Public/UObject/ObjectMacros.h
--- a/Engine/Source/Runtime/CoreUObject/Public/UObject/ObjectMacros.h
+++ b/Engine/Source/Runtime/CoreUObject/Public/UObject/ObjectMacros.h
@@ -459,6 +459,7 @@ enum EPropertyFlags : uint64
 	CPF_ExperimentalOverridableLogic	= 0x0200000000000000,	///< ****Experimental*** Property will use different logic to serialize knowing what changes are done against its default use the overridable information provided by the overridable manager on the object
 	CPF_ExperimentalAlwaysOverriden		= 0x0400000000000000,	///< ****Experimental*** Property should never inherit from the parent when using overridable serialization
 	CPF_ExperimentalNeverOverriden		= 0x0800000000000000,	///< ****Experimental*** Property should never be overridden when using overridable serialization
+	CPF_AllowSelfReference				= 0x1000000000000000,	///< Enables the intancing graph self referencing logic, delgates and verse function are already using this
 };
 
 /** All Native Access Specifier flags */
