Temporarily include CPF_PersistentInstance on all IDO loose object reference properties in order to allow override serialization to work.

diff --git a/Engine/Source/Runtime/CoreUObject/Private/UObject/InstanceDataObjectUtils.cpp b/Engine/Source/Runtime/CoreUObject/Private/UObject/InstanceDataObjectUtils.cpp
--- a/Engine/Source/Runtime/CoreUObject/Private/UObject/InstanceDataObjectUtils.cpp
+++ b/Engine/Source/Runtime/CoreUObject/Private/UObject/InstanceDataObjectUtils.cpp
@@ -342,6 +342,21 @@ namespace UE
 				MarkPropertyAsLoose(InnerProperty);
 			}
 		}
+		else if (const FObjectProperty* AsObjectProperty = CastField<FObjectProperty>(Property))
+		{
+			// Hack for now - the assumption is that IDOs are generated only for class types that impose this flag on all object properties.
+			// There is currently an implicit assumption in the serialization logic that all inner properties have this flag set for containers.
+			// Since this is a "loose" property, the underlying type will not explicitly tell us this, and there is no way to know from the tagged
+			// property data stream if this flag was set when it was last serialized for the instance in question. So for now we just always set it.
+			// 
+			// Note that we are not currently including other related flags such as CPF_InstancedReference, CPF_ContainsInstancedReference, etc.
+			// For the most part those have been relegated to object construction and loading paths. We are not instancing IDO types explicitly;
+			// they are instead serving as an editable data archetype for the actual instance, whose type may impose some post-initialization
+			// effects on that data as part of the construction/serialization path.
+			// 
+			// @todo - Remove if/when this flag is no longer required to signal whether this value is to be resolved via a subobject instancing graph.
+			Property->SetPropertyFlags(CPF_PersistentInstance);
+		}
 	}
 
 	// constructs an InstanceDataObject struct by merging the properties in 
