Added impersonation of enum types

diff --git a/Engine/Source/Runtime/CoreUObject/Private/UObject/EnumProperty.cpp b/Engine/Source/Runtime/CoreUObject/Private/UObject/EnumProperty.cpp
--- a/Engine/Source/Runtime/CoreUObject/Private/UObject/EnumProperty.cpp
+++ b/Engine/Source/Runtime/CoreUObject/Private/UObject/EnumProperty.cpp
@@ -302,6 +302,18 @@ void FEnumProperty::SerializeItem(FStructuredArchive::FSlot Slot, void* Value, v
 				if (Enum->IsValidEnumValue(IntValue))
 				{
 					EnumValueName = Enum->GetNameByValue(IntValue);
+
+				#if WITH_EDITORONLY_DATA
+					// Fix up the type name when this property is impersonating another enum type.
+					FUObjectSerializeContext* SerializeContext = FUObjectThreadContext::Get().GetSerializeContext();
+					if (SerializeContext->bImpersonateProperties)
+					{
+						if (UE::FPropertyTypeName OriginalType = UE::FindOriginalType(this); !OriginalType.IsEmpty())
+						{
+							EnumValueName = FName(EnumValueName.ToString().Replace(*Enum->GetName(), *OriginalType.GetName().ToString()));
+						}
+					}
+				#endif
 				}
 			}
 		}
@@ -618,7 +630,16 @@ void FEnumProperty::SaveTypeName(UE::FPropertyTypeNameBuilder& Type) const
 	{
 		check(UnderlyingProp);
 		Type.BeginParameters();
-		Type.AddPath(LocalEnum);
+	#if WITH_EDITORONLY_DATA
+		if (const UE::FPropertyTypeName OriginalType = UE::FindOriginalType(this); !OriginalType.IsEmpty())
+		{
+			Type.AddType(OriginalType);
+		}
+		else
+	#endif // WITH_EDITORONLY_DATA
+		{
+			Type.AddPath(LocalEnum);
+		}
 		UnderlyingProp->SaveTypeName(Type);
 		Type.EndParameters();
 	}
@@ -638,5 +659,17 @@ bool FEnumProperty::CanSerializeFromTypeName(UE::FPropertyTypeName Type) const
 	}
 
 	const FName EnumName = Type.GetParameterName(0);
-	return EnumName == LocalEnum->GetFName();
+	if (EnumName == LocalEnum->GetFName())
+	{
+		return true;
+	}
+
+#if WITH_EDITORONLY_DATA
+	if (const UE::FPropertyTypeName OriginalType = UE::FindOriginalType(this); !OriginalType.IsEmpty())
+	{
+		return EnumName == OriginalType.GetName();
+	}
+#endif // WITH_EDITORONLY_DATA
+
+	return false;
 }
diff --git a/Engine/Source/Runtime/CoreUObject/Private/UObject/InstanceDataObjectUtils.cpp b/Engine/Source/Runtime/CoreUObject/Private/UObject/InstanceDataObjectUtils.cpp
--- a/Engine/Source/Runtime/CoreUObject/Private/UObject/InstanceDataObjectUtils.cpp
+++ b/Engine/Source/Runtime/CoreUObject/Private/UObject/InstanceDataObjectUtils.cpp
@@ -236,20 +236,12 @@ namespace UE
 		{
 			if (!AsStructProperty->Struct->UseNativeSerialization())
 			{
-				//@note: Transfer existing metadata over as we build the InstanceDataObject from the struct or it owner, if any, this is useful for testing purposes
+				//@note: Transfer existing metadata over as we build the InstanceDataObject from the struct or it owners, if any, this is useful for testing purposes
 				FString OriginalName;
-				if (const FString* OriginalType = AsStructProperty->FindMetaData(NAME_OriginalType))
+				if (const FString* OriginalType = FindOriginalTypeName(AsStructProperty))
 				{
 					OriginalName = *OriginalType;
 				}
-				//@note: To support metadata defined on array of struct in UPROPERTY for testing purposes
-				else if (FField* OwnerField = AsStructProperty->Owner.ToField())
-				{
-					if (const FString* OwnerOriginalType = OwnerField->FindMetaData(NAME_OriginalType))
-					{
-						OriginalName = *OwnerOriginalType;
-					}
-				}
 
 				if (OriginalName.IsEmpty())
 				{
diff --git a/Engine/Source/Runtime/CoreUObject/Private/UObject/PropertyByte.cpp b/Engine/Source/Runtime/CoreUObject/Private/UObject/PropertyByte.cpp
--- a/Engine/Source/Runtime/CoreUObject/Private/UObject/PropertyByte.cpp
+++ b/Engine/Source/Runtime/CoreUObject/Private/UObject/PropertyByte.cpp
@@ -4,6 +4,7 @@
 
 #include "Algo/Find.h"
 #include "Hash/Blake3.h"
+#include "UObject/PropertyHelper.h"
 #include "UObject/UnrealTypePrivate.h"
 #include "UObject/UObjectThreadContext.h"
 
@@ -86,6 +87,18 @@ void FByteProperty::SerializeItem( FStructuredArchive::FSlot Slot, void* Value,
 		if ( Enum->IsValidEnumValue(ByteValue) )
 		{
 			EnumValueName = Enum->GetNameByValue(ByteValue);
+
+		#if WITH_EDITORONLY_DATA
+			// Fix up the type name when this property is impersonating another enum type.
+			FUObjectSerializeContext* SerializeContext = FUObjectThreadContext::Get().GetSerializeContext();
+			if (SerializeContext->bImpersonateProperties)
+			{
+				if (UE::FPropertyTypeName OriginalType = UE::FindOriginalType(this); !OriginalType.IsEmpty())
+				{
+					EnumValueName = FName(EnumValueName.ToString().Replace(*Enum->GetName(), *OriginalType.GetName().ToString()));
+				}
+			}
+		#endif
 		}
 		else
 		{
@@ -470,7 +483,16 @@ void FByteProperty::SaveTypeName(UE::FPropertyTypeNameBuilder& Type) const
 	if (const UEnum* LocalEnum = Enum)
 	{
 		Type.BeginParameters();
-		Type.AddPath(LocalEnum);
+	#if WITH_EDITORONLY_DATA
+		if (const UE::FPropertyTypeName OriginalType = UE::FindOriginalType(this); !OriginalType.IsEmpty())
+		{
+			Type.AddType(OriginalType);
+		}
+		else
+	#endif // WITH_EDITORONLY_DATA
+		{
+			Type.AddPath(LocalEnum);
+		}
 		Type.EndParameters();
 	}
 }
@@ -485,7 +507,19 @@ bool FByteProperty::CanSerializeFromTypeName(UE::FPropertyTypeName Type) const
 	const FName EnumName = Type.GetParameterName(0);
 	if (const UEnum* LocalEnum = Enum)
 	{
-		return EnumName == LocalEnum->GetFName();
+		if (EnumName == LocalEnum->GetFName())
+		{
+			return true;
+		}
+
+	#if WITH_EDITORONLY_DATA
+		if (const UE::FPropertyTypeName OriginalType = UE::FindOriginalType(this); !OriginalType.IsEmpty())
+		{
+			return EnumName == OriginalType.GetName();
+		}
+	#endif // WITH_EDITORONLY_DATA
+
+		return false;
 	}
 	return EnumName.IsNone();
 }
diff --git a/Engine/Source/Runtime/CoreUObject/Private/UObject/PropertyHelper.cpp b/Engine/Source/Runtime/CoreUObject/Private/UObject/PropertyHelper.cpp
--- a/Engine/Source/Runtime/CoreUObject/Private/UObject/PropertyHelper.cpp
+++ b/Engine/Source/Runtime/CoreUObject/Private/UObject/PropertyHelper.cpp
@@ -293,27 +293,74 @@ namespace UE
 {
 
 #if WITH_EDITORONLY_DATA
-FPropertyTypeName FindOriginalType(const FStructProperty* Struct)
+const FString* FindOriginalTypeName(const UField* Field)
 {
-	FUObjectSerializeContext* Context = FUObjectThreadContext::Get().GetSerializeContext();
-	if (Context && Context->bImpersonateProperties)
+	if (FUObjectThreadContext::Get().GetSerializeContext()->bImpersonateProperties)
 	{
-		const FString* OriginalType = Struct->FindMetaData(NAME_OriginalType);
-		if (!OriginalType)
+		return Field->FindMetaData(NAME_OriginalType);
+	}
+	return nullptr;
+}
+
+const FString* FindOriginalTypeName(const FProperty* Property)
+{
+	if (FUObjectThreadContext::Get().GetSerializeContext()->bImpersonateProperties)
+	{
+		// Prioritize metadata on the property over metadata on the type.
+		if (const FString* OriginalType = Property->FindMetaData(NAME_OriginalType))
 		{
-			//@note: To support metadata defined on array of struct in UPROPERTY for testing purposes
-			if (FField* OwnerField = Struct->Owner.ToField())
-			{
-				OriginalType = OwnerField->FindMetaData(NAME_OriginalType);
-			}
+			return OriginalType;
 		}
-		if (OriginalType)
+
+		// Search the owner chain to support metadata defined in UPROPERTY on a container for testing purposes.
+		for (FField* OwnerField = Property->Owner.ToField(); OwnerField; OwnerField = OwnerField->Owner.ToField())
 		{
-			if (UE::FPropertyTypeNameBuilder Type; Type.TryParse(*OriginalType))
+			if (const FString* OriginalType = OwnerField->FindMetaData(NAME_OriginalType))
 			{
-				return Type.Build();
+				return OriginalType;
 			}
 		}
+
+		if (const FStructProperty* StructProperty = CastField<FStructProperty>(Property))
+		{
+			return FindOriginalTypeName(StructProperty->Struct);
+		}
+		if (const FClassProperty* ClassProperty = CastField<FClassProperty>(Property))
+		{
+			return FindOriginalTypeName(ClassProperty->PropertyClass);
+		}
+		if (const FEnumProperty* EnumProperty = CastField<FEnumProperty>(Property))
+		{
+			return FindOriginalTypeName(EnumProperty->GetEnum());
+		}
+		if (const FByteProperty* ByteProperty = CastField<FByteProperty>(Property); ByteProperty && ByteProperty->Enum)
+		{
+			return FindOriginalTypeName(ByteProperty->Enum);
+		}
+	}
+	return nullptr;
+}
+
+FPropertyTypeName FindOriginalType(const UField* Field)
+{
+	if (const FString* OriginalType = FindOriginalTypeName(Field))
+	{
+		if (UE::FPropertyTypeNameBuilder Type; Type.TryParse(*OriginalType))
+		{
+			return Type.Build();
+		}
+	}
+	return {};
+}
+
+FPropertyTypeName FindOriginalType(const FProperty* Property)
+{
+	if (const FString* OriginalType = FindOriginalTypeName(Property))
+	{
+		if (UE::FPropertyTypeNameBuilder Type; Type.TryParse(*OriginalType))
+		{
+			return Type.Build();
+		}
 	}
 	return {};
 }
diff --git a/Engine/Source/Runtime/CoreUObject/Private/UObject/PropertyHelper.h b/Engine/Source/Runtime/CoreUObject/Private/UObject/PropertyHelper.h
--- a/Engine/Source/Runtime/CoreUObject/Private/UObject/PropertyHelper.h
+++ b/Engine/Source/Runtime/CoreUObject/Private/UObject/PropertyHelper.h
@@ -10,7 +10,8 @@
 #include "UObject/UObjectThreadContext.h"
 #include "UObject/WeakObjectPtrFwd.h"
 
-class FStructProperty;
+class FProperty;
+class UField;
 class UFunction;
 class UObject;
 
@@ -83,7 +84,10 @@ namespace UE
 
 #if WITH_EDITORONLY_DATA
 inline static const FName NAME_OriginalType(ANSITEXTVIEW("OriginalType"));
-FPropertyTypeName FindOriginalType(const FStructProperty* Struct);
+const FString* FindOriginalTypeName(const UField* Field);
+const FString* FindOriginalTypeName(const FProperty* Property);
+FPropertyTypeName FindOriginalType(const UField* Field);
+FPropertyTypeName FindOriginalType(const FProperty* Property);
 #endif // WITH_EDITORONLY_DATA
 
 /**
diff --git a/Engine/Source/Runtime/CoreUObject/Tests/UObject/InstanceDataObjectUtilsTest.cpp b/Engine/Source/Runtime/CoreUObject/Tests/UObject/InstanceDataObjectUtilsTest.cpp
--- a/Engine/Source/Runtime/CoreUObject/Tests/UObject/InstanceDataObjectUtilsTest.cpp
+++ b/Engine/Source/Runtime/CoreUObject/Tests/UObject/InstanceDataObjectUtilsTest.cpp
@@ -224,6 +224,13 @@ TEST_CASE_NAMED(FTrackUnknownPropertiesTest, "CoreUObject::Serialization::TrackU
 	AltStructData.C = 3;
 	AltStructData.D = 4;
 	AltStructData.E = 5;
+	AltStructData.Bird = TIDOB_Raven;
+	AltStructData.Grain = ETestInstanceDataObjectGrainAlternate::Corn;
+	AltStructData.Fruit = ETestInstanceDataObjectFruitAlternate::Orange;
+	AltStructData.Direction = ETestInstanceDataObjectDirectionAlternate::North | ETestInstanceDataObjectDirectionAlternate::West;
+	AltStructData.Point.U = 1;
+	AltStructData.Point.V = 2;
+	AltStructData.Point.W = 3;
 
 	TArray<uint8> BinaryData;
 	{
@@ -258,6 +265,16 @@ TEST_CASE_NAMED(FTrackUnknownPropertiesTest, "CoreUObject::Serialization::TrackU
 	CHECK(StructData.B == 2);
 	CHECK(StructData.C == 3);
 	CHECK(StructData.D == 4);
+	CHECK(StructData.Bird == TIDOB_Raven);
+	CHECK(StructData.Grain == ETestInstanceDataObjectGrain::Corn);
+	CHECK(StructData.Fruit == ETestInstanceDataObjectFruit::Orange);
+	CHECK(StructData.Direction == (ETestInstanceDataObjectDirection::North | ETestInstanceDataObjectDirection::West));
+	CHECK(StructData.Point.X == 0);
+	CHECK(StructData.Point.Y == 0);
+	CHECK(StructData.Point.Z == 0);
+#if WITH_METADATA
+	CHECK(StructData.Point.W == 3);
+#endif
 
 	FPropertyPathNameTree* Tree = FPropertyBagRepository::Get().FindOrCreateUnknownPropertyTree(Owner);
 	{
@@ -272,6 +289,30 @@ TEST_CASE_NAMED(FTrackUnknownPropertiesTest, "CoreUObject::Serialization::TrackU
 		FSerializedPropertyPathScope Path(SerializeContext, {"E", MakePropertyTypeName(NAME_IntProperty)});
 		CHECK(Tree->Find(SerializeContext->SerializedPropertyPath));
 	}
+	{
+		FProperty* PointProperty = FTestInstanceDataObjectStructAlternate::StaticStruct()->FindPropertyByName("Point");
+		CHECKED_IF(PointProperty)
+		{
+			FPropertyTypeNameBuilder Builder;
+			PointProperty->SaveTypeName(Builder);
+			FSerializedPropertyPathScope Path(SerializeContext, {"Point", Builder.Build()});
+			CHECK(Tree->Find(SerializeContext->SerializedPropertyPath));
+		#if WITH_METADATA
+			{
+				FSerializedPropertyPathScope SubPath(SerializeContext, {"U", MakePropertyTypeName(NAME_IntProperty)});
+				CHECK(Tree->Find(SerializeContext->SerializedPropertyPath));
+			}
+			{
+				FSerializedPropertyPathScope SubPath(SerializeContext, {"V", MakePropertyTypeName(NAME_IntProperty)});
+				CHECK(Tree->Find(SerializeContext->SerializedPropertyPath));
+			}
+			{
+				FSerializedPropertyPathScope SubPath(SerializeContext, {"W", MakePropertyTypeName(NAME_IntProperty)});
+				CHECK_FALSE(Tree->Find(SerializeContext->SerializedPropertyPath));
+			}
+		#endif
+		}
+	}
 	FPropertyBagRepository::Get().DestroyOuterBag(Owner);
 
 #if WITH_TEXT_ARCHIVE_SUPPORT
@@ -288,6 +329,16 @@ TEST_CASE_NAMED(FTrackUnknownPropertiesTest, "CoreUObject::Serialization::TrackU
 	CHECK(StructData.B == 2);
 	CHECK(StructData.C == 3);
 	CHECK(StructData.D == 4);
+	CHECK(StructData.Bird == TIDOB_Raven);
+	CHECK(StructData.Grain == ETestInstanceDataObjectGrain::Corn);
+	CHECK(StructData.Fruit == ETestInstanceDataObjectFruit::Orange);
+	CHECK(StructData.Direction == (ETestInstanceDataObjectDirection::North | ETestInstanceDataObjectDirection::West));
+	CHECK(StructData.Point.X == 0);
+	CHECK(StructData.Point.Y == 0);
+	CHECK(StructData.Point.Z == 0);
+#if WITH_METADATA
+	CHECK(StructData.Point.W == 3);
+#endif
 
 	// Testing of the unknown property tree is skipped because it is not supported by the text format.
 #endif // WITH_TEXT_ARCHIVE_SUPPORT
@@ -336,7 +387,7 @@ TEST_CASE_NAMED(FUnknownEnumNamesTest, "CoreUObject::Serialization::UnknownEnumN
 	FPropertyTypeName FlagsTypeName = []
 	{
 		FPropertyTypeNameBuilder Builder;
-		Builder.AddPath(StaticEnum<ETestInstanceDataObjectFlags>());
+		Builder.AddPath(StaticEnum<ETestInstanceDataObjectDirection>());
 		return Builder.Build();
 	}();
 
@@ -369,13 +420,13 @@ TEST_CASE_NAMED(FUnknownEnumNamesTest, "CoreUObject::Serialization::UnknownEnumN
 
 	// Test a flags enum by enum...
 
-	Repo.AddUnknownEnumName(Owner, StaticEnum<ETestInstanceDataObjectFlags>(), {}, NAME_Up);
+	Repo.AddUnknownEnumName(Owner, StaticEnum<ETestInstanceDataObjectDirection>(), {}, NAME_Up);
 
 	Repo.FindUnknownEnumNames(Owner, FlagsTypeName, Names, bHasFlags);
 	CHECK(Names.Num() == 1);
 	CHECK(bHasFlags);
 
-	Repo.AddUnknownEnumName(Owner, StaticEnum<ETestInstanceDataObjectFlags>(), FlagsTypeName, *FlagsString);
+	Repo.AddUnknownEnumName(Owner, StaticEnum<ETestInstanceDataObjectDirection>(), FlagsTypeName, *FlagsString);
 
 	Repo.FindUnknownEnumNames(Owner, FlagsTypeName, Names, bHasFlags);
 	CHECKED_IF(Names.Num() == 2)
diff --git a/Engine/Source/Runtime/CoreUObject/Tests/UObject/InstanceDataObjectUtilsTest.h b/Engine/Source/Runtime/CoreUObject/Tests/UObject/InstanceDataObjectUtilsTest.h
--- a/Engine/Source/Runtime/CoreUObject/Tests/UObject/InstanceDataObjectUtilsTest.h
+++ b/Engine/Source/Runtime/CoreUObject/Tests/UObject/InstanceDataObjectUtilsTest.h
@@ -3,9 +3,49 @@
 #pragma once
 
 #include "CoreTypes.h"
+#include "Misc/EnumClassFlags.h"
 #include "UObject/Object.h"
 #include "InstanceDataObjectUtilsTest.generated.h"
 
+UENUM()
+enum ETestInstanceDataObjectBird
+{
+	TIDOB_None = 0,
+	TIDOB_Cardinal,
+	TIDOB_Crow,
+	TIDOB_Eagle,
+	TIDOB_Hawk,
+	TIDOB_Owl,
+	TIDOB_Raven,
+};
+
+UENUM()
+namespace ETestInstanceDataObjectGrain
+{
+	enum Type
+	{
+		None = 0,
+		Barley,
+		Corn,
+		Quinoa,
+		Rice,
+		Wheat,
+	};
+}
+
+UENUM()
+namespace ETestInstanceDataObjectGrainAlternate
+{
+	enum Type
+	{
+		None = 0,
+		Corn,
+		Rice,
+		Rye,
+		Wheat,
+	};
+}
+
 UENUM()
 enum class ETestInstanceDataObjectFruit : uint8
 {
@@ -15,8 +55,18 @@ enum class ETestInstanceDataObjectFruit : uint8
 	Orange,
 };
 
+UENUM()
+enum class ETestInstanceDataObjectFruitAlternate : uint8
+{
+	None = 0,
+	Apple,
+	Cherry,
+	Orange,
+	Pear,
+};
+
 UENUM(Flags)
-enum class ETestInstanceDataObjectFlags : uint16
+enum class ETestInstanceDataObjectDirection : uint16
 {
 	None = 0,
 	North = 1 << 0,
@@ -25,6 +75,57 @@ enum class ETestInstanceDataObjectFlags : uint16
 	West = 1 << 3,
 };
 
+ENUM_CLASS_FLAGS(ETestInstanceDataObjectDirection);
+
+UENUM(Flags)
+enum class ETestInstanceDataObjectDirectionAlternate : uint16
+{
+	None = 0,
+	Up = 1 << 0,
+	Down = 1 << 1,
+	North = 1 << 2,
+	East = 1 << 3,
+	South = 1 << 4,
+	West = 1 << 5,
+};
+
+ENUM_CLASS_FLAGS(ETestInstanceDataObjectDirectionAlternate);
+
+USTRUCT()
+struct FTestInstanceDataObjectPoint
+{
+	GENERATED_BODY()
+
+public:
+	UPROPERTY()
+	int32 X = 0;
+
+	UPROPERTY()
+	int32 Y = 0;
+
+	UPROPERTY()
+	int32 Z = 0;
+
+	UPROPERTY()
+	int32 W = 0;
+};
+
+USTRUCT()
+struct FTestInstanceDataObjectPointAlternate
+{
+	GENERATED_BODY()
+
+public:
+	UPROPERTY()
+	int32 U = 0;
+
+	UPROPERTY()
+	int32 V = 0;
+
+	UPROPERTY()
+	int32 W = 0;
+};
+
 USTRUCT()
 struct FTestInstanceDataObjectStruct
 {
@@ -42,9 +143,24 @@ public:
 
 	UPROPERTY()
 	int32 D = -1;
+
+	UPROPERTY()
+	TEnumAsByte<ETestInstanceDataObjectBird> Bird = TIDOB_None;
+
+	UPROPERTY()
+	TEnumAsByte<ETestInstanceDataObjectGrain::Type> Grain = ETestInstanceDataObjectGrain::None;
+
+	UPROPERTY()
+	ETestInstanceDataObjectFruit Fruit = ETestInstanceDataObjectFruit::None;
+
+	UPROPERTY()
+	ETestInstanceDataObjectDirection Direction = ETestInstanceDataObjectDirection::None;
+
+	UPROPERTY()
+	FTestInstanceDataObjectPoint Point;
 };
 
-USTRUCT(meta=(OriginalType="TestInstanceDataObjectStruct"))
+USTRUCT()
 struct FTestInstanceDataObjectStructAlternate
 {
 	GENERATED_BODY()
@@ -61,6 +177,21 @@ public:
 
 	UPROPERTY()
 	int32 E = -1;
+
+	UPROPERTY()
+	TEnumAsByte<ETestInstanceDataObjectBird> Bird = TIDOB_None;
+
+	UPROPERTY(meta=(OriginalType="ETestInstanceDataObjectGrain(/Script/CoreUObject)"))
+	TEnumAsByte<ETestInstanceDataObjectGrainAlternate::Type> Grain = ETestInstanceDataObjectGrainAlternate::None;
+
+	UPROPERTY(meta=(OriginalType="ETestInstanceDataObjectFruit(/Script/CoreUObject)"))
+	ETestInstanceDataObjectFruitAlternate Fruit = ETestInstanceDataObjectFruitAlternate::None;
+
+	UPROPERTY(meta=(OriginalType="ETestInstanceDataObjectDirection(/Script/CoreUObject)"))
+	ETestInstanceDataObjectDirectionAlternate Direction = ETestInstanceDataObjectDirectionAlternate::None;
+
+	UPROPERTY(meta=(OriginalType="TestInstanceDataObjectPoint(/Script/CoreUObject)"))
+	FTestInstanceDataObjectPointAlternate Point;
 };
 
 UCLASS()
