Fixes for enums missing explicit declaration of underlying type.

diff --git a/Engine/Source/Runtime/CoreUObject/Tests/UObject/InstanceDataObjectUtilsTest.h b/Engine/Source/Runtime/CoreUObject/Tests/UObject/InstanceDataObjectUtilsTest.h
--- a/Engine/Source/Runtime/CoreUObject/Tests/UObject/InstanceDataObjectUtilsTest.h
+++ b/Engine/Source/Runtime/CoreUObject/Tests/UObject/InstanceDataObjectUtilsTest.h
@@ -8,7 +8,7 @@
 #include "InstanceDataObjectUtilsTest.generated.h"
 
 UENUM()
-enum ETestInstanceDataObjectBird
+enum ETestInstanceDataObjectBird : uint8
 {
 	TIDOB_None = 0,
 	TIDOB_Cardinal,
@@ -22,7 +22,7 @@ enum ETestInstanceDataObjectBird
 UENUM()
 namespace ETestInstanceDataObjectGrain
 {
-	enum Type
+	enum Type : uint8
 	{
 		None = 0,
 		Barley,
@@ -36,7 +36,7 @@ namespace ETestInstanceDataObjectGrain
 UENUM()
 namespace ETestInstanceDataObjectGrainAlternate
 {
-	enum Type
+	enum Type : uint8
 	{
 		None = 0,
 		Corn,
