Verse native:
- Fix implicit params for functions in V_DECLARE_HAS_METHOD (TNonNullSubclassOf instead of UClass*)
- Support const or non-const native function
- Non-write class functions are suggested as const
#rnx

diff --git a/Engine/Source/Runtime/CoreUObject/Public/Templates/NonNullSubclassOf.h b/Engine/Source/Runtime/CoreUObject/Public/Templates/NonNullSubclassOf.h
--- a/Engine/Source/Runtime/CoreUObject/Public/Templates/NonNullSubclassOf.h
+++ b/Engine/Source/Runtime/CoreUObject/Public/Templates/NonNullSubclassOf.h
@@ -27,7 +27,9 @@ public:
 	}
 
 	/** Constructor that takes a UClass* or FFieldClass* if T is a UObject or a FField type respectively. */
-	template <typename PtrType>
+	template <
+		typename PtrType
+		UE_REQUIRES(std::is_same_v<PtrType, UClass> || std::is_same_v<PtrType, FFieldClass>)>
 	FORCEINLINE TNonNullSubclassOf(PtrType* From)
 		: Super(From)
 	{
