Add basic cooking / validation support when cooking with IDO enabled:
* Disable IDO impersonation when cooking, this result in not including loose property in cooked data if any
* Prevent cooking / saving if placeholder objects are found in the asset as data loss mitigation until placeholder object data can be properly serialized with prefabs

diff --git a/Engine/Source/Runtime/CoreUObject/Internal/UObject/PropertyBagRepository.h b/Engine/Source/Runtime/CoreUObject/Internal/UObject/PropertyBagRepository.h
--- a/Engine/Source/Runtime/CoreUObject/Internal/UObject/PropertyBagRepository.h
+++ b/Engine/Source/Runtime/CoreUObject/Internal/UObject/PropertyBagRepository.h
@@ -209,13 +209,12 @@ struct FScopedIDOSerializationContext
 {
 #if WITH_EDITORONLY_DATA
 	COREUOBJECT_API FScopedIDOSerializationContext(UObject* InObject, FArchive& Archive);
-	explicit COREUOBJECT_API FScopedIDOSerializationContext(UObject* InObject); // assumes save
+	COREUOBJECT_API FScopedIDOSerializationContext(UObject* InObject, bool bImpersonate); // assumes save
 	COREUOBJECT_API ~FScopedIDOSerializationContext();
 
-	bool bHasIDOSupport;
-	bool bCreateIDO;
-	FArchive* Archive;
-	UObject* const Object;
+	bool bCreateIDO = false;
+	FArchive* Archive = nullptr;
+	UObject* const Object = nullptr;
 	const int64 PreSerializeOffset;
 	TOptional<TGuardValue<bool>> ScopedTrackSerializedPropertyPath;
 	TOptional<TGuardValue<bool>> ScopedSerializeUnknownProperty;
@@ -229,7 +228,7 @@ private:
 	void FinishCreatingInstanceDataObject() const;
 #else
 	inline FScopedIDOSerializationContext(UObject* InObject, FArchive& Archive) {}
-	inline explicit FScopedIDOSerializationContext(UObject* InObject) {}
+	inline explicit FScopedIDOSerializationContext(UObject* InObject, bool bImpersonate) {}
 #endif
 
 	FScopedIDOSerializationContext(const FScopedIDOSerializationContext&) = delete;
diff --git a/Engine/Source/Runtime/CoreUObject/Private/UObject/PropertyBagRepository.cpp b/Engine/Source/Runtime/CoreUObject/Private/UObject/PropertyBagRepository.cpp
--- a/Engine/Source/Runtime/CoreUObject/Private/UObject/PropertyBagRepository.cpp
+++ b/Engine/Source/Runtime/CoreUObject/Private/UObject/PropertyBagRepository.cpp
@@ -162,6 +162,11 @@ FPropertyBagRepository& FPropertyBagRepository::Get()
 
 void FPropertyBagRepository::ReassociateObjects(const TMap<UObject*, UObject*>& ReplacedObjects)
 {
+	if (!IsInstanceDataObjectSupportEnabled())
+	{
+		return;
+	}
+
 	FPropertyBagRepositoryLock LockRepo(this);
 	FPropertyBagAssociationData OldBagData;
 	for (const TPair<UObject*, UObject*>& Pair : ReplacedObjects)
@@ -891,7 +896,7 @@ FScopedIDOSerializationContext::FScopedIDOSerializationContext(UObject* InObject
 	, PreSerializeOffset(InArchive.Tell())
 {
 	FUObjectSerializeContext* SerializeContext = FUObjectThreadContext::Get().GetSerializeContext();
-	bHasIDOSupport = FPropertyBagRepository::IsInstanceDataObjectSupportEnabled(Object);
+	bool bHasIDOSupport = FPropertyBagRepository::IsInstanceDataObjectSupportEnabled(Object);
 	
 	bool bHasReinstancedClass = InObject->GetClass()->HasAnyClassFlags(CLASS_NewerVersionExists);
 	bCreateIDO = bHasIDOSupport && !SerializeContext->bImpersonateProperties && Archive->IsLoading() && !bHasReinstancedClass;
@@ -918,15 +923,18 @@ FScopedIDOSerializationContext::FScopedIDOSerializationContext(UObject* InObject
 	}
 }
 
-FScopedIDOSerializationContext::FScopedIDOSerializationContext(UObject* InObject)
+FScopedIDOSerializationContext::FScopedIDOSerializationContext(UObject* InObject, bool bImpersonate)
 	: bCreateIDO(false)
 	, Archive(nullptr)
 	, Object(InObject)
 	, PreSerializeOffset(0)
 {
 	FUObjectSerializeContext* SerializeContext = FUObjectThreadContext::Get().GetSerializeContext();
-	bHasIDOSupport = FPropertyBagRepository::IsInstanceDataObjectSupportEnabled(Object);
-	ScopedImpersonateProperties.Emplace(SerializeContext->bImpersonateProperties, true);
+	bool bHasIDOSupport = FPropertyBagRepository::IsInstanceDataObjectSupportEnabled(Object);
+	if (bHasIDOSupport)
+	{
+		ScopedImpersonateProperties.Emplace(SerializeContext->bImpersonateProperties, bImpersonate);
+	}
 }
 
 FScopedIDOSerializationContext::~FScopedIDOSerializationContext()
diff --git a/Engine/Source/Runtime/CoreUObject/Private/UObject/SavePackage2.cpp b/Engine/Source/Runtime/CoreUObject/Private/UObject/SavePackage2.cpp
--- a/Engine/Source/Runtime/CoreUObject/Private/UObject/SavePackage2.cpp
+++ b/Engine/Source/Runtime/CoreUObject/Private/UObject/SavePackage2.cpp
@@ -619,6 +619,35 @@ ESavePackageResult ValidateExports(FSaveContext& SaveContext)
 		}
 	}
 
+	// Validate there's no placeholder object in the export list
+#if WITH_EDITOR
+	if (UE::FPropertyBagRepository::IsPropertyBagPlaceholderObjectSupportEnabled())
+	{
+		bool bSuccess = true;
+		for (const FTaggedExport& Export : SaveContext.GetExports())
+		{
+			if (UE::FPropertyBagRepository::IsPropertyBagPlaceholderObject(Export.Obj))
+			{
+				UObject* Outer = Export.Obj->GetOuter();
+				FString ErrorMessage = FString::Printf(TEXT("Saving illegal placeholder object in package '%s'. Object Name '%s', Class '%s', Outer '%s', Outer Class '%s', Full Path '%s'. You might have unresolved classes or compilation failures to fix."), *SaveContext.GetPackage()->GetName(), *Export.Obj->GetName(), *Export.Obj->GetClass()->GetName(), Outer ? *Outer->GetName() : TEXT("None"), Outer ? *Outer->GetClass()->GetName() : TEXT("None"), *Export.Obj->GetFullName());
+				if (SaveContext.IsGenerateSaveError())
+				{
+					SaveContext.GetError()->Logf(ELogVerbosity::Warning, TEXT("%s"), *ErrorMessage);
+				}
+				else
+				{
+					UE_LOG(LogSavePackage, Error, TEXT("%s"), *ErrorMessage);
+				}
+				bSuccess = false;
+			}
+		}
+		if (!bSuccess)
+		{
+			return ESavePackageResult::Error;
+		}
+	}
+#endif
+
 	// Validate External Export Rules
 	if (SaveContext.HasExternalExportValidations())
 	{
@@ -2964,7 +2993,7 @@ ESavePackageResult InnerSave(FSaveContext& SaveContext)
 {
 	SaveContext.SetEDLCookChecker(&FEDLCookCheckerThreadState::Get());
 
-	UE::FScopedIDOSerializationContext IDOSaveContext(nullptr);
+	UE::FScopedIDOSerializationContext IDOSaveContext(nullptr, !SaveContext.IsCooking());
 
 	// Create slow task dialog if needed
 	const int32 TotalSaveSteps = 3;
diff --git a/Engine/Source/Runtime/Engine/Private/UnrealEngine.cpp b/Engine/Source/Runtime/Engine/Private/UnrealEngine.cpp
--- a/Engine/Source/Runtime/Engine/Private/UnrealEngine.cpp
+++ b/Engine/Source/Runtime/Engine/Private/UnrealEngine.cpp
@@ -17459,7 +17459,7 @@ void UEngine::CopyPropertiesForUnrelatedObjects(UObject* OldObject, UObject* New
 			OldInstanceMap.Add(OldInstance->GetPathName(OldObject), SavedInstances.Num() - 1);
 			const uint32 AdditionalPortFlags = Params.bCopyDeprecatedProperties ? PPF_UseDeprecatedProperties : PPF_None;
 			
-			UE::FScopedIDOSerializationContext IDOSaveContext(OldInstance);
+			UE::FScopedIDOSerializationContext IDOSaveContext(OldInstance, /*bImpersonate*/ true);
 			
 			FObjectWriter SubObjWriter(OldInstance, pRecord->SavedProperties, true, true, Params.bDoDelta, AdditionalPortFlags);
 		}
