[Package]
 - Reduce size from 128 bytes to 80 bytes at runtime to save memory
 - Move runtime rarely used members into an externally allocated struct
 - Deprecation cleanup

#rnx

diff --git a/Engine/Source/Runtime/CoreUObject/Private/UObject/LinkerLoad.cpp b/Engine/Source/Runtime/CoreUObject/Private/UObject/LinkerLoad.cpp
--- a/Engine/Source/Runtime/CoreUObject/Private/UObject/LinkerLoad.cpp
+++ b/Engine/Source/Runtime/CoreUObject/Private/UObject/LinkerLoad.cpp
@@ -6108,9 +6108,7 @@ void FLinkerLoad::Detach()
 		// This is so that object post loaded in the editor in a package that has no associated linker consider that all package custom versions as latest
 		// (i.e. when duplicating an object in the package)
 		// The runtime *may* use the stored version in the package since there are never any linker associated with it when using iostore
-		PRAGMA_DISABLE_DEPRECATION_WARNINGS
-		LinkerRoot->LinkerCustomVersion.Empty();
-		PRAGMA_ENABLE_DEPRECATION_WARNINGS
+		LinkerRoot->EmptyLinkerCustomVersion();
 		LinkerRoot = nullptr;
 	}
 
diff --git a/Engine/Source/Runtime/CoreUObject/Private/UObject/Package.cpp b/Engine/Source/Runtime/CoreUObject/Private/UObject/Package.cpp
--- a/Engine/Source/Runtime/CoreUObject/Private/UObject/Package.cpp
+++ b/Engine/Source/Runtime/CoreUObject/Private/UObject/Package.cpp
@@ -289,11 +289,6 @@ const FPackagePath& UPackage::GetLoadedPath() const
 void UPackage::SetLoadedPath(const FPackagePath& InPackagePath)
 {
 	LoadedPath = InPackagePath;
-#if !UE_STRIP_DEPRECATED_PROPERTIES
-	PRAGMA_DISABLE_DEPRECATION_WARNINGS
-	FileName = InPackagePath.GetPackageFName();
-	PRAGMA_ENABLE_DEPRECATION_WARNINGS
-#endif
 }
 
 /** Tags generated objects with flags */
@@ -424,13 +419,18 @@ void UPackage::SetSavedHash(const FIoHash& InSavedHash)
 #endif
 
 #if WITH_EDITORONLY_DATA
-PRAGMA_DISABLE_DEPRECATION_WARNINGS
+
+void DeclareIntrinsicUPackageMembers()
+{
+	UE::GC::DeclareIntrinsicMembers(UPackage::StaticClass(), { UE_GC_MEMBER(UPackage, MetaData) });
+}
+
 IMPLEMENT_CORE_INTRINSIC_CLASS(UPackage, UObject,
 	{
-		UE::GC::DeclareIntrinsicMembers(Class, { UE_GC_MEMBER(UPackage, MetaData) });
+		DeclareIntrinsicUPackageMembers();
 	}
 );
-PRAGMA_ENABLE_DEPRECATION_WARNINGS
+
 #else
 IMPLEMENT_CORE_INTRINSIC_CLASS(UPackage, UObject,
 	{
diff --git a/Engine/Source/Runtime/CoreUObject/Public/UObject/Package.h b/Engine/Source/Runtime/CoreUObject/Public/UObject/Package.h
--- a/Engine/Source/Runtime/CoreUObject/Public/UObject/Package.h
+++ b/Engine/Source/Runtime/CoreUObject/Public/UObject/Package.h
@@ -250,21 +250,7 @@ public:
 #endif
 
 private:
-	// @note this should probably be entirely deprecated and removed but certain stat dump function are still using it,
-	// just compile it out in shipping for now
-#if !UE_BUILD_SHIPPING
-	/**
-	 * Time in seconds it took to fully load this package.
-	 * 0 if package is either in process of being loaded or has never been fully loaded.
-	 */
-	float LoadTime;
-#endif
-
 #if WITH_EDITORONLY_DATA
-	/** Indicates which folder to display this package under in the Generic Browser's list of packages. If not specified, package is added to the root level. deprecated	*/
-	UE_DEPRECATED(5.1, "Unused property to be removed. Use something like FPackageName::GetLongPackagePath if requiring something similar.")
-	FName FolderName;
-
 	/** Persistent GUID of package if it was loaded from disk. Persistent across saves. */
 	FGuid PersistentGuid;
 
@@ -286,57 +272,56 @@ private:
 
 	/** The PackagePath this package was loaded from */
 	FPackagePath LoadedPath;
-public:
 
+	// @note this should probably be entirely deprecated and removed but certain stat dump function are still using it,
+	// just compile it out in shipping for now
+#if !UE_BUILD_SHIPPING
+	/**
+	 * Time in seconds it took to fully load this package.
+	 * 0 if package is either in process of being loaded or has never been fully loaded.
+	 */
+	float LoadTime;
+#endif
+
+	struct FAdditionalInfo
+	{
 	/** Linker package version this package has been serialized with. This is mostly used by PostLoad **/
-	UE_DEPRECATED(5.0, "Use Get/SetLinkerPackageVersion instead")
-	FPackageFileVersion LinkerPackageVersion;
+		FPackageFileVersion LinkerPackageVersion = GPackageFileUEVersion;
 
 	/** Linker licensee version this package has been serialized with. This is mostly used by PostLoad **/
-	UE_DEPRECATED(5.0, "Use Get/SetLinkerLicenseeVersion instead")
-	int32 LinkerLicenseeVersion;
+		int32 LinkerLicenseeVersion = GPackageFileLicenseeUEVersion;
 
 	/** Linker custom version container this package has been serialized with. This is mostly used by PostLoad **/
-	UE_DEPRECATED(5.0, "Use Get/SetLinkerCustomVersions instead")
 	FCustomVersionContainer LinkerCustomVersion;
 
 	/** Linker load associated with this package */
-	UE_DEPRECATED(5.0, "Use Get/SetLinker instead")
-	FLinkerLoad* LinkerLoad;
+		FLinkerLoad* LinkerLoad = nullptr;
 
 	/** size of the file for this package; if the package was not loaded from a file or was a forced export in another package, this will be zero */
-	UE_DEPRECATED(5.0, "Use Get/SetFileSize instead")
-	uint64 FileSize;
+		uint64 FileSize = 0;
+
+		// World browser information
+		TUniquePtr< FWorldTileInfo > WorldTileInfo;
+	};
+	
+	/** Contains additional information if they differ from the defaults. */
+	TUniquePtr<FAdditionalInfo> AdditionalInfo;
 
 #if WITH_METADATA
 	// MetaData for the editor, or NULL in the game
-	UE_DEPRECATED(5.0, "Use Get/HasMetaData instead")
 	class UMetaData* MetaData;
 #endif // WITH_METADATA
 
 #if WITH_EDITORONLY_DATA
 	/** Editor only: Thumbnails stored in this package */
-	UE_DEPRECATED(5.0, "Use Get/SetThumbnailMap instead")
 	TUniquePtr< FThumbnailMap > ThumbnailMap;
 
 	/** Editor only: PIE instance ID this package belongs to, INDEX_NONE otherwise */
-	UE_DEPRECATED(5.0, "Use Get/SetPIEInstanceID instead")
 	int32 PIEInstanceID;
 #endif
 
-#if !UE_STRIP_DEPRECATED_PROPERTIES
-	/** The name of the file that this package was loaded from */
-	UE_DEPRECATED(5.0, "Use GetLoadedPath instead")
-	FName FileName;
-#endif
-
-	// World browser information
-	UE_DEPRECATED(5.0, "Use Get/SetWorldTileInfo instead")
-	TUniquePtr< FWorldTileInfo > WorldTileInfo;
-
 #if WITH_RELOAD
 	/** Link list of delegates registered to the package.  The next pointer chain can't be used for this. */
-	UE_DEPRECATED(5.0, "Use Get/SetReloadDelegates instead")
 	TArray<UFunction*> Delegates;
 #endif
 
@@ -402,27 +387,72 @@ private:
 
 	void SetLinker(FLinkerLoad* InLinker)
 	{
-		LinkerLoad = InLinker;
+		LLM_SCOPE_BYNAME(TEXT("Package/AdditionalInfo"));
+		if (AdditionalInfo.IsValid())
+		{
+			AdditionalInfo->LinkerLoad = InLinker;
+		}
+		else if (InLinker != nullptr)
+		{
+			AdditionalInfo = MakeUnique<FAdditionalInfo>();
+			AdditionalInfo->LinkerLoad = InLinker;
+		}
 	}
 
 	void SetLinkerPackageVersion(FPackageFileVersion InVersion)
 	{
-		LinkerPackageVersion = MoveTemp(InVersion);
+		LLM_SCOPE_BYNAME(TEXT("Package/AdditionalInfo"));
+		if (AdditionalInfo.IsValid())
+		{
+			AdditionalInfo->LinkerPackageVersion = MoveTemp(InVersion);
+		}
+		else if (InVersion != GPackageFileUEVersion)
+		{
+			AdditionalInfo = MakeUnique<FAdditionalInfo>();
+			AdditionalInfo->LinkerPackageVersion = MoveTemp(InVersion);
+		}
 	}
 
 	void SetLinkerLicenseeVersion(int32 InVersion)
 	{
-		LinkerLicenseeVersion = InVersion;
+		LLM_SCOPE_BYNAME(TEXT("Package/AdditionalInfo"));
+		if (AdditionalInfo.IsValid())
+		{
+			AdditionalInfo->LinkerLicenseeVersion = InVersion;
+		}
+		else if (InVersion != GPackageFileLicenseeUEVersion)
+		{
+			AdditionalInfo = MakeUnique<FAdditionalInfo>();
+			AdditionalInfo->LinkerLicenseeVersion = InVersion;
+		}
 	}
 
 	void SetLinkerCustomVersions(FCustomVersionContainer InVersions)
 	{
-		LinkerCustomVersion = MoveTemp(InVersions);
+		LLM_SCOPE_BYNAME(TEXT("Package/AdditionalInfo"));
+		if (AdditionalInfo.IsValid())
+		{
+			AdditionalInfo->LinkerCustomVersion = MoveTemp(InVersions);
+		}
+		else if (!InVersions.GetAllVersions().IsEmpty())
+		{
+			AdditionalInfo = MakeUnique<FAdditionalInfo>();
+			AdditionalInfo->LinkerCustomVersion = MoveTemp(InVersions);
+		}
 	}
 
 	void SetFileSize(int64 InFileSize)
 	{
-		FileSize = InFileSize;
+		LLM_SCOPE_BYNAME(TEXT("Package/AdditionalInfo"));
+		if (AdditionalInfo.IsValid())
+		{
+			AdditionalInfo->FileSize = InFileSize;
+		}
+		else if (InFileSize != 0)
+		{
+			AdditionalInfo = MakeUnique<FAdditionalInfo>();
+			AdditionalInfo->FileSize = MoveTemp(InFileSize);
+		}
 	}
 
 	void SetMetaData(UMetaData* InMetaData)
@@ -459,24 +489,35 @@ public:
 
 	FLinkerLoad* GetLinker() const
 	{
-		return LinkerLoad;
+		return AdditionalInfo.IsValid() ? AdditionalInfo->LinkerLoad : nullptr;
 	}
 
 	const FPackageFileVersion& GetLinkerPackageVersion() const
 	{
-		return LinkerPackageVersion;
+		return AdditionalInfo.IsValid() ? AdditionalInfo->LinkerPackageVersion : GPackageFileUEVersion;
 	}
 
 	int32 GetLinkerLicenseeVersion() const
 	{
-		return LinkerLicenseeVersion;
+		return AdditionalInfo.IsValid() ? AdditionalInfo->LinkerLicenseeVersion : GPackageFileLicenseeUEVersion;
 	}
 
 	const FCustomVersionContainer& GetLinkerCustomVersions() const
 	{
-		return LinkerCustomVersion;
+		static FCustomVersionContainer EmptyVersions;
+		return AdditionalInfo.IsValid() ? AdditionalInfo->LinkerCustomVersion : EmptyVersions;
+	}
+
+private:
+	void EmptyLinkerCustomVersion()
+	{
+		if (AdditionalInfo.IsValid())
+		{
+			AdditionalInfo->LinkerCustomVersion.Empty();
+		}
 	}
 
+public:
 	/**
 	 * Sets the time it took to load this package.
 	 */
@@ -501,27 +542,6 @@ public:
 #endif
 	}
 
-#if WITH_EDITORONLY_DATA
-	/**
-	* Get the package's folder name
-	* @return		Folder name
-	*/
-	UE_DEPRECATED(5.1, "Unused property to be removed. Use something like FPackageName::GetLongPackagePath if requiring something similar.")
-	FName GetFolderName() const
-	{
-		return FolderName;
-	}
-
-	/**
-	* Set the package's folder name
-	*/
-	UE_DEPRECATED(5.1, "Unused property to be removed. Use something like FPackageName::GetLongPackagePath if requiring something similar.")
-	void SetFolderName (FName name)
-	{
-		FolderName = name;
-	}
-#endif
-
 	/**
 	 * Clear the package dirty flag without any transaction tracking
 	 */
@@ -806,17 +826,25 @@ public:
 	}
 #endif
 
-
 	/** Get the world tile info if any*/
 	FWorldTileInfo* GetWorldTileInfo() const
 	{
-		return WorldTileInfo.Get();
+		return AdditionalInfo.IsValid() ? AdditionalInfo->WorldTileInfo.Get() : nullptr;
 	}
 
 	/** Set the world tile info */
 	void SetWorldTileInfo(TUniquePtr<FWorldTileInfo> InWorldTileInfo)
 	{
-		WorldTileInfo = MoveTemp(InWorldTileInfo);
+		LLM_SCOPE_BYNAME(TEXT("Package/AdditionalInfo"));
+		if (AdditionalInfo.IsValid())
+		{
+			AdditionalInfo->WorldTileInfo = MoveTemp(InWorldTileInfo);
+		}
+		else if (InWorldTileInfo.IsValid())
+		{
+			AdditionalInfo = MakeUnique<FAdditionalInfo>();
+			AdditionalInfo->WorldTileInfo = MoveTemp(InWorldTileInfo);
+		}
 	}
 
 #if WITH_EDITORONLY_DATA
@@ -854,7 +882,7 @@ public:
 	/** returns our FileSize */
 	FORCEINLINE int64 GetFileSize() const
 	{
-		return FileSize;
+		return AdditionalInfo.IsValid() ? AdditionalInfo->FileSize : 0;
 	}
 
 	/** returns our ChunkIDs */
@@ -989,6 +1017,7 @@ public:
 	COREUOBJECT_API static bool IsEmptyPackage(UPackage* Package, const UObject* LastReferencer = NULL);
 
 private:
+	friend void DeclareIntrinsicUPackageMembers();
 	static FSavePackageResultStruct Save2(UPackage* InPackage, UObject* InAsset, const TCHAR* InFilename, const FSavePackageArgs& SaveArgs);
 };
 PRAGMA_ENABLE_DEPRECATION_WARNINGS
