Add a structure for locking a FBulkData object in a given scope.

#rnx

- Code was provided by a third party on UDN who waived their rights to the code and turned down credit.
- Minor bug fixing/cleanup has been applied.

diff --git a/Engine/Source/Runtime/CoreUObject/Public/Serialization/BulkDataScopedLock.h b/Engine/Source/Runtime/CoreUObject/Public/Serialization/BulkDataScopedLock.h
new file mode 100644
--- /dev/null
+++ b/Engine/Source/Runtime/CoreUObject/Public/Serialization/BulkDataScopedLock.h
@@ -0,0 +1,109 @@
+#pragma once
+
+#include "Containers/ArrayView.h"
+#include "Containers/ContainersFwd.h"
+#include "Serialization/BulkData.h"
+
+// Note that these lock structures only work with the templated versions of FBulkData, such as
+// FByteBulkData or FFloatBulkData because they need to know the format of the bulkdata payload
+// which the base class FBulkData cannot provide.
+
+/** Locks the given bulkdata object for read access */
+template<typename ElementType>
+class TBulkDataScopedReadLock
+{
+public:
+	UE_NONCOPYABLE(TBulkDataScopedReadLock);
+
+	TBulkDataScopedReadLock(const TBulkData<ElementType>& InBulkData)
+		: BulkData(InBulkData)
+	{
+		Data = reinterpret_cast<const ElementType*>(InBulkData.LockReadOnly());
+	}
+
+	~TBulkDataScopedReadLock()
+	{
+		BulkData.Unlock();
+	}
+
+	/** Returns the number of elements in the bulkdata payload */
+	int64 Num() const
+	{
+		return BulkData.GetElementCount();
+	}
+
+	/** Returns the length of the bulkdata payload in bytes */
+	int64 GetAllocatedSize() const
+	{
+		return BulkData.GetBulkDataSize();
+	}
+
+	/** Returns a raw pointer to the bulkdata payload */
+	const ElementType* GetData() const
+	{
+		return Data;
+	}
+
+	/** Returns the bulkdata payload wrapped in a TConstArrayView64 for added safety */
+	TConstArrayView64<ElementType> GetView() const
+	{
+		return TConstArrayView64<ElementType>(Data, Num());
+	}
+
+private:
+	const TBulkData<ElementType>& BulkData;
+	const ElementType* Data;
+};
+
+/**
+ * Locks the given bulkdata object for read/write access.
+ * Note that performing actions on the original bulkdata object while this scope
+ * is active (such as calls to realloc) may cause problems. Only use this object
+ * for bulkdata operations while it is active.
+ */
+template<typename ElementType>
+class TBulkDataScopedWriteLock
+{
+public:
+	UE_NONCOPYABLE(TBulkDataScopedWriteLock);
+
+	TBulkDataScopedWriteLock(TBulkData<ElementType>& InBulkData)
+		: BulkData(InBulkData)
+	{
+		Data = reinterpret_cast<ElementType*>(InBulkData.Lock(LOCK_READ_WRITE));
+	}
+
+	~TBulkDataScopedWriteLock()
+	{
+		BulkData.Unlock();
+	}
+
+	/** Returns the number of elements in the bulkdata payload */
+	int64 Num() const
+	{
+		return BulkData.GetElementCount();
+	}
+
+	/** Returns the length of the bulkdata payload in bytes */
+	int64 GetAllocatedSize() const
+	{
+		return BulkData.GetBulkDataSize();
+	}
+
+	/** Returns a raw pointer to the bulkdata payload */
+	ElementType* GetData() const
+	{
+		return Data;
+	}
+
+	/** Returns the bulkdata payload wrapped in a TArrayView64 for added safety */
+	TArrayView64<ElementType> GetView() const
+	{
+		return TArrayView64<ElementType>(Data, Num());
+	}
+
+private:
+
+	TBulkData<ElementType>& BulkData;
+	ElementType* Data;
+};
diff --git a/Engine/Source/Runtime/CoreUObject/Tests/Serialization/BulkDataTests.cpp b/Engine/Source/Runtime/CoreUObject/Tests/Serialization/BulkDataTests.cpp
--- a/Engine/Source/Runtime/CoreUObject/Tests/Serialization/BulkDataTests.cpp
+++ b/Engine/Source/Runtime/CoreUObject/Tests/Serialization/BulkDataTests.cpp
@@ -7,11 +7,12 @@
 #include "TestMacros/Assertions.h"
 #include "Memory/SharedBuffer.h"
 #include "Serialization/BulkData.h"
+#include "Serialization/BulkDataScopedLock.h"
 #include "Serialization/LargeMemoryWriter.h"
 #include "Serialization/MemoryReader.h"
 #include <catch2/generators/catch_generators.hpp>
 
-namespace UE::Serialization::BulkDataTest
+namespace UE
 {
 
 static FUniqueBuffer CreatePayload(uint64 Size, uint64 Seed = 0)
@@ -375,6 +376,70 @@ TEST_CASE("CoreUObject::Serialization::FBulkData::Serialize", "[CoreUObject][Ser
 
 #endif // WITH_EDITORONLY_DATA
 
-} // namespace UE::Serialization::BulkDataTest
+TEST_CASE("CoreUObject::Serialization::FBulkData::LockScope", "[CoreUObject][Serialization]")
+{
+	const int32 NumElements = 32;
+	FByteBulkData BulkData;
+
+	// Lock for write
+	{
+		BulkData.Lock(LOCK_READ_WRITE);
+		BulkData.Realloc(32);
+		BulkData.Unlock();
+
+		{
+			TBulkDataScopedWriteLock WriteLock(BulkData);
+			REQUIRE(BulkData.IsLocked());
+
+			uint8* DataPtr = WriteLock.GetData();
+
+			REQUIRE(DataPtr != nullptr);
+
+			REQUIRE(WriteLock.Num() == NumElements);
+			REQUIRE(WriteLock.Num() == BulkData.GetElementCount());
+			REQUIRE(WriteLock.GetAllocatedSize() == BulkData.GetBulkDataSize());
+
+			REQUIRE(WriteLock.GetView().GetData() != nullptr);
+			REQUIRE(WriteLock.GetView().GetData() == DataPtr);
+			REQUIRE(WriteLock.GetView().Num() == BulkData.GetElementCount());
+
+			for (int32 Index = 0; Index < NumElements; ++Index)
+			{
+				DataPtr[Index] = static_cast<uint8>(Index);
+			}
+		}
+
+		REQUIRE_FALSE(BulkData.IsLocked());
+	}
+
+	// Lock for read
+	{
+		{
+			TBulkDataScopedReadLock ReadLock(BulkData);
+			REQUIRE(BulkData.IsLocked());
+
+			const uint8* DataPtr = ReadLock.GetData();
+
+			REQUIRE(DataPtr != nullptr);
+
+			REQUIRE(ReadLock.Num() == NumElements);
+			REQUIRE(ReadLock.Num() == BulkData.GetElementCount());
+			REQUIRE(ReadLock.GetAllocatedSize() == BulkData.GetBulkDataSize());
+
+			REQUIRE(ReadLock.GetView().GetData() != nullptr);
+			REQUIRE(ReadLock.GetView().GetData() == DataPtr);
+			REQUIRE(ReadLock.GetView().Num() == BulkData.GetElementCount());
+
+			for (int32 Index = 0; Index < NumElements; ++Index)
+			{
+				REQUIRE(DataPtr[Index] == Index);
+			}
+		}
+
+		REQUIRE_FALSE(BulkData.IsLocked());
+	}
+}
+
+} // namespace UE
 
 #endif // WITH_LOW_LEVEL_TESTS
