Add missing boilerplate code

#rnx

diff --git a/Engine/Source/Runtime/CoreUObject/Public/Serialization/BulkDataScopedLock.h b/Engine/Source/Runtime/CoreUObject/Public/Serialization/BulkDataScopedLock.h
--- a/Engine/Source/Runtime/CoreUObject/Public/Serialization/BulkDataScopedLock.h
+++ b/Engine/Source/Runtime/CoreUObject/Public/Serialization/BulkDataScopedLock.h
@@ -1,3 +1,5 @@
+// Copyright Epic Games, Inc. All Rights Reserved.
+
 #pragma once
 
 #include "Containers/ArrayView.h"
