- Drop CoreRedirect INI sections from config files since we only need to initialize the redirects once.

#rnx

diff --git a/Engine/Source/Runtime/CoreUObject/Private/UObject/CoreRedirects.cpp b/Engine/Source/Runtime/CoreUObject/Private/UObject/CoreRedirects.cpp
--- a/Engine/Source/Runtime/CoreUObject/Private/UObject/CoreRedirects.cpp
+++ b/Engine/Source/Runtime/CoreUObject/Private/UObject/CoreRedirects.cpp
@@ -2121,7 +2121,8 @@ bool FCoreRedirects::ReadRedirectsFromIni(const FString& IniName)
 
 	if (GConfig)
 	{
-		const FConfigSection* RedirectSection = GConfig->GetSection(TEXT("CoreRedirects"), false, IniName);
+		const TCHAR* RedirectSectionName = TEXT("CoreRedirects");
+		const FConfigSection* RedirectSection = GConfig->GetSection(RedirectSectionName, false, IniName);
 		if (RedirectSection)
 		{
 			TArray<FCoreRedirect> NewRedirects;
@@ -2247,6 +2248,9 @@ bool FCoreRedirects::ReadRedirectsFromIni(const FString& IniName)
 				}
 			}
 
+			// We no longer need the redirect config data in memory, so drop it
+			GConfig->EmptySection(RedirectSectionName, IniName);
+
 			return AddRedirectList(NewRedirects, IniName);
 		}
 		else
