SerializeForLog implementations

#rnx

diff --git a/Engine/Source/Runtime/CoreUObject/Private/AssetRegistry/AssetData.cpp b/Engine/Source/Runtime/CoreUObject/Private/AssetRegistry/AssetData.cpp
--- a/Engine/Source/Runtime/CoreUObject/Private/AssetRegistry/AssetData.cpp
+++ b/Engine/Source/Runtime/CoreUObject/Private/AssetRegistry/AssetData.cpp
@@ -79,6 +79,20 @@ bool LoadFromCompactBinary(FCbFieldView Field, FAssetIdentifier& Identifier)
 	return true;
 }
 
+void SerializeForLog(FCbWriter& Writer, const FAssetIdentifier& Value)
+{
+	Writer.BeginObject();
+	Writer.AddString(ANSITEXTVIEW("$type"), ANSITEXTVIEW("AssetIdentifier"));
+	TStringBuilder<256> Text;
+	Value.AppendString(Text);
+	Writer.AddString(ANSITEXTVIEW("$text"), Text);
+	Writer.AddString(ANSITEXTVIEW("PackageName"), WriteToUtf8String<256>(Value.PackageName));
+	Writer.AddString(ANSITEXTVIEW("PrimaryAssetType"), WriteToUtf8String<256>(Value.PrimaryAssetType.GetName()));
+	Writer.AddString(ANSITEXTVIEW("ObjectName"), WriteToUtf8String<256>(Value.ObjectName));
+	Writer.AddString(ANSITEXTVIEW("ValueName"), WriteToUtf8String<256>(Value.ValueName));
+	Writer.EndObject();
+}
+
 namespace UE::AssetRegistry::Private
 {
 	FAssetPathParts SplitIntoOuterPathAndAssetName(FStringView InObjectPath)
diff --git a/Engine/Source/Runtime/CoreUObject/Private/UObject/PrimaryAssetId.cpp b/Engine/Source/Runtime/CoreUObject/Private/UObject/PrimaryAssetId.cpp
--- a/Engine/Source/Runtime/CoreUObject/Private/UObject/PrimaryAssetId.cpp
+++ b/Engine/Source/Runtime/CoreUObject/Private/UObject/PrimaryAssetId.cpp
@@ -2,6 +2,7 @@
 
 #include "UObject/PrimaryAssetId.h"
 #include "Misc/StringBuilder.h"
+#include "Serialization/CompactBinaryWriter.h"
 #include "UObject/ObjectRedirector.h"
 #include "UObject/PropertyPortFlags.h"
 #include "UObject/UnrealType.h"
@@ -134,3 +135,13 @@ FStringBuilderBase& operator<<(FStringBuilderBase& Builder, const FPrimaryAssetI
 {
 	return Builder << Id.PrimaryAssetType.GetName() << ":" << Id.PrimaryAssetName;
 }
+
+void SerializeForLog(FCbWriter& Writer, const FPrimaryAssetId& Value)
+{
+	Writer.BeginObject();
+	Writer.AddString(ANSITEXTVIEW("$type"), ANSITEXTVIEW("PrimaryAssetId"));
+	Writer.AddString(ANSITEXTVIEW("$text"), WriteToUtf8String<256>(Value.PrimaryAssetType.GetName(), ':', Value.PrimaryAssetName));
+	Writer.AddString(ANSITEXTVIEW("PrimaryAssetType"), WriteToUtf8String<256>(Value.PrimaryAssetType));
+	Writer.AddString(ANSITEXTVIEW("PrimaryAssetName"), WriteToUtf8String<256>(Value.PrimaryAssetName));
+	Writer.EndObject();
+}
\ No newline at end of file
diff --git a/Engine/Source/Runtime/CoreUObject/Private/UObject/SoftObjectPath.cpp b/Engine/Source/Runtime/CoreUObject/Private/UObject/SoftObjectPath.cpp
--- a/Engine/Source/Runtime/CoreUObject/Private/UObject/SoftObjectPath.cpp
+++ b/Engine/Source/Runtime/CoreUObject/Private/UObject/SoftObjectPath.cpp
@@ -7,6 +7,7 @@
 #include "Misc/AsciiSet.h"
 #include "Misc/PackageName.h"
 #include "Misc/StringBuilder.h"
+#include "Serialization/CompactBinaryWriter.h"
 #include "UObject/LinkerLoad.h"
 #include "UObject/UObjectThreadContext.h"
 #include "UObject/CoreRedirects.h"
@@ -93,6 +94,11 @@ void FSoftObjectPath::ToString(FStringBuilderBase& Builder) const
 	AppendString(Builder);
 }
 
+void FSoftObjectPath::ToString(FUtf8StringBuilderBase& Builder) const
+{
+	AppendString(Builder);
+}
+
 void FSoftObjectPath::AppendString(FStringBuilderBase& Builder) const
 {
 	if (AssetPath.IsNull())
@@ -107,6 +113,20 @@ void FSoftObjectPath::AppendString(FStringBuilderBase& Builder) const
 	}
 }
 
+void FSoftObjectPath::AppendString(FUtf8StringBuilderBase& Builder) const
+{
+	if (AssetPath.IsNull())
+	{
+		return;
+	}
+
+	Builder << AssetPath;
+	if (SubPathString.Len() > 0)
+	{
+		Builder << SUBOBJECT_DELIMITER_CHAR_ANSI << SubPathString;
+	}
+}
+
 void FSoftObjectPath::AppendString(FString& Builder) const
 {
 	if (AssetPath.IsNull())
@@ -926,6 +946,17 @@ bool FSoftObjectPathThreadContext::GetSerializationOptions(FName& OutPackageName
 
 TSet<FName> FSoftObjectPath::PIEPackageNames;
 
+void SerializeForLog(FCbWriter& Writer, const FSoftObjectPath& Value)
+{
+	Writer.BeginObject();
+	Writer.AddString(ANSITEXTVIEW("$type"), ANSITEXTVIEW("SoftObjectPath"));
+	Writer.AddString(ANSITEXTVIEW("$text"), WriteToUtf8String<256>(Value));
+	Writer.AddString(ANSITEXTVIEW("PackageName"), WriteToUtf8String<256>(Value.GetLongPackageFName()));
+	Writer.AddString(ANSITEXTVIEW("AssetName"), WriteToUtf8String<256>(Value.GetAssetFName()));
+	Writer.AddString(ANSITEXTVIEW("SubPath"), Value.GetSubPathString());
+	Writer.EndObject();
+}
+
 #if WITH_LOW_LEVEL_TESTS
 
 #include "TestHarness.h"
diff --git a/Engine/Source/Runtime/CoreUObject/Private/UObject/TopLevelAssetPath.cpp b/Engine/Source/Runtime/CoreUObject/Private/UObject/TopLevelAssetPath.cpp
--- a/Engine/Source/Runtime/CoreUObject/Private/UObject/TopLevelAssetPath.cpp
+++ b/Engine/Source/Runtime/CoreUObject/Private/UObject/TopLevelAssetPath.cpp
@@ -526,6 +526,16 @@ bool FTopLevelAssetPath::SerializeFromMismatchedTag(const FPropertyTag& Tag, FSt
 	return false;
 }
 
+void SerializeForLog(FCbWriter& Writer, const FTopLevelAssetPath& Value)
+{
+	Writer.BeginObject();
+	Writer.AddString(ANSITEXTVIEW("$type"), ANSITEXTVIEW("TopLevelAssetPath"));
+	Writer.AddString(ANSITEXTVIEW("$text"), WriteToUtf8String<256>(Value));
+	Writer.AddString(ANSITEXTVIEW("PackageName"), WriteToUtf8String<256>(Value.GetPackageName()));
+	Writer.AddString(ANSITEXTVIEW("AssetName"), WriteToUtf8String<256>(Value.GetAssetName()));
+	Writer.EndObject();
+}
+
 void FTopLevelAssetPath::WriteCompactBinary(FCbWriter& Writer) const
 {
 	Writer << WriteToUtf8String<FName::StringBufferSize>(*this).ToView();
diff --git a/Engine/Source/Runtime/CoreUObject/Public/AssetRegistry/AssetIdentifier.h b/Engine/Source/Runtime/CoreUObject/Public/AssetRegistry/AssetIdentifier.h
--- a/Engine/Source/Runtime/CoreUObject/Public/AssetRegistry/AssetIdentifier.h
+++ b/Engine/Source/Runtime/CoreUObject/Public/AssetRegistry/AssetIdentifier.h
@@ -234,5 +234,6 @@ private:
 	}
 	// Load Cannot be inline because we need to hide implementation and copy-by-value is invalid without definition
 	COREUOBJECT_API friend bool LoadFromCompactBinary(FCbFieldView Field, FAssetIdentifier& Identifier);
+	COREUOBJECT_API friend void SerializeForLog(FCbWriter& Writer, const FAssetIdentifier& Value);
 };
 
diff --git a/Engine/Source/Runtime/CoreUObject/Public/UObject/PrimaryAssetId.h b/Engine/Source/Runtime/CoreUObject/Public/UObject/PrimaryAssetId.h
--- a/Engine/Source/Runtime/CoreUObject/Public/UObject/PrimaryAssetId.h
+++ b/Engine/Source/Runtime/CoreUObject/Public/UObject/PrimaryAssetId.h
@@ -12,6 +12,7 @@
 #include "UObject/NameTypes.h"
 #include "UObject/UnrealNames.h"
 
+class FCbWriter;
 class FOutputDevice;
 class UObject;
 
@@ -200,6 +201,9 @@ struct FPrimaryAssetId
 	}
 
 	friend struct Z_Construct_UScriptStruct_FPrimaryAssetId_Statics;
+
+private:
+	 COREUOBJECT_API void SerializeForLog(FCbWriter& Writer, const FPrimaryAssetId& Value);
 };
 
 COREUOBJECT_API FStringBuilderBase& operator<<(FStringBuilderBase& Builder, const FPrimaryAssetId& Id);
diff --git a/Engine/Source/Runtime/CoreUObject/Public/UObject/SoftObjectPath.h b/Engine/Source/Runtime/CoreUObject/Public/UObject/SoftObjectPath.h
--- a/Engine/Source/Runtime/CoreUObject/Public/UObject/SoftObjectPath.h
+++ b/Engine/Source/Runtime/CoreUObject/Public/UObject/SoftObjectPath.h
@@ -28,6 +28,7 @@
 #include "UObject/UObjectHash.h"
 
 class FArchive;
+class FCbWriter;
 class FOutputDevice;
 struct FPropertyTag;
 struct FUObjectSerializeContext;
@@ -141,10 +142,12 @@ PRAGMA_ENABLE_DEPRECATION_WARNINGS;
 
 	/** Append string representation of reference, in form /package/path.assetname[:subpath] */
 	COREUOBJECT_API void ToString(FStringBuilderBase& Builder) const;
+	COREUOBJECT_API void ToString(FUtf8StringBuilderBase& Builder) const;
 
 	/** Append string representation of reference, in form /package/path.assetname[:subpath] */
 	COREUOBJECT_API void AppendString(FString& Builder) const;
 	COREUOBJECT_API void AppendString(FStringBuilderBase& Builder) const;
+	COREUOBJECT_API void AppendString(FUtf8StringBuilderBase& Builder) const;
 
 	/** Returns the top-level asset part of this path, without the subobject path. */
 	FTopLevelAssetPath GetAssetPath() const
@@ -379,6 +382,7 @@ private:
 	COREUOBJECT_API UObject* ResolveObjectInternal() const;
 	COREUOBJECT_API UObject* ResolveObjectInternal(const TCHAR* PathString) const;
 
+	COREUOBJECT_API friend void SerializeForLog(FCbWriter& Writer, const FSoftObjectPath& Value);
 	friend struct Z_Construct_UScriptStruct_FSoftObjectPath_Statics;
 };
 
@@ -405,6 +409,11 @@ inline FStringBuilderBase& operator<<(FStringBuilderBase& Builder, const FSoftOb
 	Path.ToString(Builder);
 	return Builder;
 }
+inline FUtf8StringBuilderBase& operator<<(FUtf8StringBuilderBase& Builder, const FSoftObjectPath& Path)
+{
+	Path.ToString(Builder);
+	return Builder;
+}
 
 /**
  * A struct that contains a string reference to a class, can be used to make soft references to classes
diff --git a/Engine/Source/Runtime/CoreUObject/Public/UObject/TopLevelAssetPath.h b/Engine/Source/Runtime/CoreUObject/Public/UObject/TopLevelAssetPath.h
--- a/Engine/Source/Runtime/CoreUObject/Public/UObject/TopLevelAssetPath.h
+++ b/Engine/Source/Runtime/CoreUObject/Public/UObject/TopLevelAssetPath.h
@@ -16,6 +16,7 @@
 
 #include <type_traits>
 
+class FCbWriter;
 struct FPropertyTag;
 class FCbWriter;
 class FString;
@@ -186,6 +187,8 @@ private:
 	FName PackageName;
 	/** Name of the asset within the package e.g. 'AssetName' */
 	FName AssetName;
+
+	COREUOBJECT_API friend void SerializeForLog(FCbWriter& Writer, const FTopLevelAssetPath& Value);
 };
 
 inline FWideStringBuilderBase& operator<<(FWideStringBuilderBase& Builder, const FTopLevelAssetPath& Path)
diff --git a/Engine/Source/Runtime/Engine/Classes/Engine/AssetManagerTypes.h b/Engine/Source/Runtime/Engine/Classes/Engine/AssetManagerTypes.h
--- a/Engine/Source/Runtime/Engine/Classes/Engine/AssetManagerTypes.h
+++ b/Engine/Source/Runtime/Engine/Classes/Engine/AssetManagerTypes.h
@@ -21,6 +21,8 @@
 
 #include "AssetManagerTypes.generated.h"
 
+class FCbWriter;
+
 /** Rule about when to cook/ship a primary asset */
 UENUM()
 enum class EPrimaryAssetCookRule : uint8
@@ -381,4 +383,6 @@ private:
 
 	FAssignmentInfo InclusionByLevel[(int32)EPrimaryAssetProductionLevel::Count];
 	FAssignmentInfo ExclusionByLevel[(int32)EPrimaryAssetProductionLevel::Count];
+
+	ENGINE_API friend void SerializeForLog(FCbWriter& Writer, const FPrimaryAssetRules& Value);
 };
diff --git a/Engine/Source/Runtime/Engine/Private/AssetManagerTypes.cpp b/Engine/Source/Runtime/Engine/Private/AssetManagerTypes.cpp
--- a/Engine/Source/Runtime/Engine/Private/AssetManagerTypes.cpp
+++ b/Engine/Source/Runtime/Engine/Private/AssetManagerTypes.cpp
@@ -4,6 +4,7 @@
 #include "Engine/AssetManager.h"
 #include "Engine/AssetManagerSettings.h"
 #include "Engine/DeveloperSettings.h"
+#include "Serialization/CompactBinaryWriter.h"
 #include "UObject/UnrealType.h"
 
 #include UE_INLINE_GENERATED_CPP_BY_NAME(AssetManagerTypes)
@@ -197,6 +198,18 @@ void UAssetManagerSettings::PostReloadConfig(FProperty* PropertyThatWasLoaded)
 	UAssetManager::Get().LoadRedirectorMaps();
 }
 
+void SerializeForLog(FCbWriter& Writer, const FPrimaryAssetRules& Value)
+{
+	Writer.BeginObject();
+	Writer.AddString(ANSITEXTVIEW("$type"), ANSITEXTVIEW("PrimaryAssetRules"));
+	// no text field
+	Writer.AddInteger(ANSITEXTVIEW("Priority"), Value.Priority);
+	Writer.AddInteger(ANSITEXTVIEW("ChunkId"), Value.ChunkId);
+	Writer.AddBool(ANSITEXTVIEW("Recursive"), Value.bApplyRecursively);
+	Writer.AddString(ANSITEXTVIEW("CookRule"), UEnum::GetValueAsString(Value.CookRule));
+	Writer.EndObject();
+}
+
 #if WITH_EDITOR
 void UAssetManagerSettings::PostInitProperties()
 {
