Update FTimecode in NoExportTypes.h to match the updated structure in FTimecode.h

diff --git a/Engine/Source/Runtime/CoreUObject/Public/UObject/NoExportTypes.h b/Engine/Source/Runtime/CoreUObject/Public/UObject/NoExportTypes.h
--- a/Engine/Source/Runtime/CoreUObject/Public/UObject/NoExportTypes.h
+++ b/Engine/Source/Runtime/CoreUObject/Public/UObject/NoExportTypes.h
@@ -2208,6 +2208,9 @@ struct FTimecode
 	UPROPERTY(BlueprintReadWrite, Category=Timecode)
 	int32 Frames;
 
+	UPROPERTY(BlueprintReadWrite, Category=Timecode)
+	float Subframe;
+
 	/** If true, this Timecode represents a Drop Frame timecode used to account for fractional frame rates in NTSC play rates. */
 	UPROPERTY(BlueprintReadWrite, Category= Timecode)
 	bool bDropFrameFormat;
