- Removed an unnecessary assert on missing Section. Was causing crashes in the editor during an undo operation

diff --git a/Engine/Source/Runtime/CoreUObject/Private/UObject/Obj.cpp b/Engine/Source/Runtime/CoreUObject/Private/UObject/Obj.cpp
--- a/Engine/Source/Runtime/CoreUObject/Private/UObject/Obj.cpp
+++ b/Engine/Source/Runtime/CoreUObject/Private/UObject/Obj.cpp
@@ -3418,8 +3418,6 @@ void UObject::SaveConfig(uint64 RequiredPropertyFlags, const TCHAR* InFilename,
 
 				if (!bPropDeprecated && (!bShouldCheckIfIdenticalBeforeAdding || !Property->Identical_InContainer(this, SuperClassDefaultObject)))
 				{
-					check(Sec);
-
 					if (Array)
 					{
 						FScriptArrayHelper_InContainer ArrayHelper(Array, this);
