Fix root cause of an unintentional disabling of editable struct fields in the Property Editor when IDO features are enabled.

diff --git a/Engine/Source/Runtime/CoreUObject/Private/UObject/InstanceDataObjectUtils.cpp b/Engine/Source/Runtime/CoreUObject/Private/UObject/InstanceDataObjectUtils.cpp
--- a/Engine/Source/Runtime/CoreUObject/Private/UObject/InstanceDataObjectUtils.cpp
+++ b/Engine/Source/Runtime/CoreUObject/Private/UObject/InstanceDataObjectUtils.cpp
@@ -443,8 +443,8 @@ namespace UE
 							// skip loose types that have been explicitly excluded from IDOs
 							continue;
 						}
-						MarkPropertyAsLoose(Property);
 						ConvertToInstanceDataObjectProperty(Property, Type, Result, It.GetNode().GetSubTree());
+						MarkPropertyAsLoose(Property);	// note: make sure not to mark until AFTER conversion, as this can mutate property flags on nested struct fields
 						LooseInstanceDataObjectProperties.Add(Property);
 						continue;
 					}
