Core redirects now remove [CoreRedirect] config data from all layers in the config system after loading into the runtime representation to save memory. Saves 3MB in 31.00 config memory

#rnx

diff --git a/Engine/Source/Runtime/CoreUObject/Private/UObject/CoreRedirects.cpp b/Engine/Source/Runtime/CoreUObject/Private/UObject/CoreRedirects.cpp
--- a/Engine/Source/Runtime/CoreUObject/Private/UObject/CoreRedirects.cpp
+++ b/Engine/Source/Runtime/CoreUObject/Private/UObject/CoreRedirects.cpp
@@ -2248,8 +2248,8 @@ bool FCoreRedirects::ReadRedirectsFromIni(const FString& IniName)
 				}
 			}
 
-			// We no longer need the redirect config data in memory, so drop it
-			GConfig->EmptySection(RedirectSectionName, IniName);
+			// We no longer need the redirect config data in memory so remove it entirely
+			GConfig->RemoveSectionFromBranch(RedirectSectionName, *IniName);
 
 			return AddRedirectList(NewRedirects, IniName);
 		}
