Fixed instantiation involving placeholder objects to prevent them to be resolved in incompatible object property
added DecayTraits assertion on non resolved TObject
Also added helper to check if a class is an IDO class

#rnx

diff --git a/Engine/Source/Runtime/CoreUObject/Internal/UObject/InstanceDataObjectUtils.h b/Engine/Source/Runtime/CoreUObject/Internal/UObject/InstanceDataObjectUtils.h
--- a/Engine/Source/Runtime/CoreUObject/Internal/UObject/InstanceDataObjectUtils.h
+++ b/Engine/Source/Runtime/CoreUObject/Internal/UObject/InstanceDataObjectUtils.h
@@ -27,6 +27,9 @@ namespace UE
 bool IsInstanceDataObjectSupportEnabled(const UObject* Object = nullptr);
 bool StructContainsLooseProperties(const UStruct* Struct);
 
+/** Helper to check if a class is an IDO class. */
+bool IsClassOfInstanceDataObjectClass(UStruct* Class);
+
 /** Generate a UClass that contains the union of the properties of PropertyTree and OwnerClass. */
 UClass* CreateInstanceDataObjectClass(const FPropertyPathNameTree* PropertyTree, UClass* OwnerClass, UObject* Outer);
 
diff --git a/Engine/Source/Runtime/CoreUObject/Private/UObject/InstanceDataObjectUtils.cpp b/Engine/Source/Runtime/CoreUObject/Private/UObject/InstanceDataObjectUtils.cpp
--- a/Engine/Source/Runtime/CoreUObject/Private/UObject/InstanceDataObjectUtils.cpp
+++ b/Engine/Source/Runtime/CoreUObject/Private/UObject/InstanceDataObjectUtils.cpp
@@ -182,6 +182,13 @@ namespace UE
 		return bIsEnabled;
 	}
 
+
+	bool IsClassOfInstanceDataObjectClass(UStruct* Class)
+	{
+		return Class->IsA(UInstanceDataObjectClass::StaticClass()) || Class->IsA(UInstanceDataObjectStruct::StaticClass());
+	}
+
+
 	bool StructContainsLooseProperties(const UStruct* Struct)
 	{
 		return Struct->GetBoolMetaData(NAME_ContainsLoosePropertiesMetadata);
@@ -347,7 +354,7 @@ namespace UE
 			// effects on that data as part of the construction/serialization path.
 			// 
 			// @todo - Remove if/when this flag is no longer required to signal whether this value is to be resolved via a subobject instancing graph.
-			Property->SetPropertyFlags(CPF_PersistentInstance);
+			Property->SetPropertyFlags(CPF_PersistentInstance| CPF_InstancedReference);
 		}
 	}
 
diff --git a/Engine/Source/Runtime/CoreUObject/Private/UObject/PropertyBaseObject.cpp b/Engine/Source/Runtime/CoreUObject/Private/UObject/PropertyBaseObject.cpp
--- a/Engine/Source/Runtime/CoreUObject/Private/UObject/PropertyBaseObject.cpp
+++ b/Engine/Source/Runtime/CoreUObject/Private/UObject/PropertyBaseObject.cpp
@@ -57,10 +57,11 @@ void FObjectPropertyBase::InstanceSubobjects(void* Data, void const* DefaultData
 {
 	for ( int32 ArrayIndex = 0; ArrayIndex < ArrayDim; ArrayIndex++ )
 	{
-		UObject* CurrentValue = GetObjectPropertyValue((uint8*)Data + ArrayIndex * ElementSize);
-		if ( CurrentValue )
+		TObjectPtr<UObject> CurrentObjectPtr = GetObjectPtrPropertyValue((uint8*)Data + ArrayIndex * ElementSize);
+		UObject* CurrentValue = CurrentObjectPtr.Get();
+		if (CurrentObjectPtr.IsResolved() && CurrentValue)
 		{
-			UObject *SubobjectTemplate = DefaultData ? GetObjectPropertyValue((uint8*)DefaultData + ArrayIndex * ElementSize): nullptr;
+			TObjectPtr<UObject> SubobjectTemplate = DefaultData ? GetObjectPtrPropertyValue((uint8*)DefaultData + ArrayIndex * ElementSize): nullptr;
 			EInstancePropertyValueFlags Flags = EInstancePropertyValueFlags::None;
 			if (HasAnyPropertyFlags(CPF_InstancedReference))
 			{
diff --git a/Engine/Source/Runtime/CoreUObject/Public/UObject/ObjectPtr.h b/Engine/Source/Runtime/CoreUObject/Public/UObject/ObjectPtr.h
--- a/Engine/Source/Runtime/CoreUObject/Public/UObject/ObjectPtr.h
+++ b/Engine/Source/Runtime/CoreUObject/Public/UObject/ObjectPtr.h
@@ -1267,6 +1267,7 @@ namespace UE::Core::Private // private facilities; not for direct use
 		{
 #if UE_WITH_OBJECT_HANDLE_LATE_RESOLVE || UE_WITH_OBJECT_HANDLE_TRACKING
 			Value.Get();
+			check(Value.IsResolved());
 #endif
 		}
 	};
