CastChecked with `DO_CHECK` false can return nullptr from a method that (used to) promise it didn't!

diff --git a/Engine/Source/Runtime/CoreUObject/Public/Templates/Casts.h b/Engine/Source/Runtime/CoreUObject/Public/Templates/Casts.h
--- a/Engine/Source/Runtime/CoreUObject/Public/Templates/Casts.h
+++ b/Engine/Source/Runtime/CoreUObject/Public/Templates/Casts.h
@@ -244,14 +244,13 @@ FORCEINLINE T* ExactCast( UObject* Src )
 #else
 
 	template <typename To, typename From>
-	FUNCTION_NON_NULL_RETURN_START
 		FORCEINLINE To* CastChecked(From* Src)
-	FUNCTION_NON_NULL_RETURN_END
 	{
 		static_assert(sizeof(From) > 0 && sizeof(To) > 0, "Attempting to cast between incomplete types");
 
 		if (!Src)
 		{
+			// TODO: should this instead be an unreachable?
 			return nullptr;
 		}
 
