[Bugfix] non-overridden properties of prefab instances get reset to CDO defaults on reinstance instead of the prefab default

diff --git a/Engine/Source/Runtime/CoreUObject/Private/UObject/PropertyBagRepository.cpp b/Engine/Source/Runtime/CoreUObject/Private/UObject/PropertyBagRepository.cpp
--- a/Engine/Source/Runtime/CoreUObject/Private/UObject/PropertyBagRepository.cpp
+++ b/Engine/Source/Runtime/CoreUObject/Private/UObject/PropertyBagRepository.cpp
@@ -866,7 +866,13 @@ void FPropertyBagRepository::CreateInstanceDataObjectUnsafe(UObject* Owner, FPro
 	if (Archive && Archive != Linker)
 	{
 		// re-deserialize Owner but redirect it into the IDO instead using impersonation
-		Owner->Serialize(*Archive);
+		{
+			FGuardValue_Bitfield(Archive->ArMergeOverrides, true);
+			Owner->Serialize(*Archive);
+		}
+		
+		// copy data from owner to IDO
+		CopyTaggedProperties(Owner, BagData.InstanceDataObject);
 	}
 	else if (Linker)
 	{
