[CrashFix] A duplicate property change notification was being sent for both IDO properties and Instance properties. This caused an issue where non-idempotent Overridable Serialization code was reregistering objects that were removed from an array. To fix this, Overridable serialization now ignores IDO property change notifications

diff --git a/Engine/Source/Runtime/CoreUObject/Private/UObject/OverridableManager.cpp b/Engine/Source/Runtime/CoreUObject/Private/UObject/OverridableManager.cpp
--- a/Engine/Source/Runtime/CoreUObject/Private/UObject/OverridableManager.cpp
+++ b/Engine/Source/Runtime/CoreUObject/Private/UObject/OverridableManager.cpp
@@ -2,6 +2,7 @@
 
 #include "UObject/OverridableManager.h"
 #include "InstancedReferenceSubobjectHelper.h"
+#include "UObject/InstanceDataObjectUtils.h"
 #include "UObject/UObjectGlobals.h"
 
 /*
@@ -212,26 +213,25 @@ bool FOverridableManager::ClearOverriddenProperty(UObject& Object, const FProper
 void FOverridableManager::PreOverrideProperty(UObject& Object, const FEditPropertyChain& PropertyChain)
 {
 #if WITH_EDITORONLY_DATA
-	if (FOverriddenPropertySet* ThisObjectOverriddenProperties = OverriddenObjectAnnotations.Find(Object))
-	{
-		ThisObjectOverriddenProperties->NotifyPropertyChange(EPropertyNotificationType::PreEdit, FPropertyChangedEvent(nullptr), PropertyChain.GetActiveMemberNode() ? PropertyChain.GetActiveMemberNode() : PropertyChain.GetHead(), &Object);
-	}
+	NotifyPropertyChange(EPropertyNotificationType::PreEdit, Object, FPropertyChangedEvent(nullptr), PropertyChain.GetActiveMemberNode() ? PropertyChain.GetActiveMemberNode() : PropertyChain.GetHead());
 #endif // WITH_EDITORONLY_DATA
 }
 
 void FOverridableManager::PostOverrideProperty(UObject& Object, const FPropertyChangedEvent& PropertyEvent, const FEditPropertyChain& PropertyChain)
 {
 #if WITH_EDITORONLY_DATA
-	if (FOverriddenPropertySet* ThisObjectOverriddenProperties = OverriddenObjectAnnotations.Find(Object))
-	{
-		ThisObjectOverriddenProperties->NotifyPropertyChange(EPropertyNotificationType::PostEdit, PropertyEvent, PropertyChain.GetActiveMemberNode() ? PropertyChain.GetActiveMemberNode() : PropertyChain.GetHead(), &Object);
-	}
+	NotifyPropertyChange(EPropertyNotificationType::PostEdit, Object, PropertyEvent, PropertyChain.GetActiveMemberNode() ? PropertyChain.GetActiveMemberNode() : PropertyChain.GetHead());
 #endif // WITH_EDITORONLY_DATA
 }
 
 void FOverridableManager::NotifyPropertyChange(const EPropertyNotificationType Notification, UObject& Object, const FPropertyChangedEvent& PropertyEvent, const FEditPropertyChain::TDoubleLinkedListNode* PropertyNode)
 {
 #if WITH_EDITORONLY_DATA
+	if (UE::IsClassOfInstanceDataObjectClass(Object.GetClass()))
+	{
+		// IDOs share override keys with their associated instances so override events are handled by the instance
+		return;
+	}
 	if (FOverriddenPropertySet* ThisObjectOverriddenProperties = OverriddenObjectAnnotations.Find(Object))
 	{
 		ThisObjectOverriddenProperties->NotifyPropertyChange(Notification, PropertyEvent, PropertyNode, &Object);
