[Bugfix] Made FPropertyTypeName use the original struct name and path for dead verse structs

diff --git a/Engine/Source/Runtime/CoreUObject/Private/UObject/InstanceDataObjectUtils.cpp b/Engine/Source/Runtime/CoreUObject/Private/UObject/InstanceDataObjectUtils.cpp
--- a/Engine/Source/Runtime/CoreUObject/Private/UObject/InstanceDataObjectUtils.cpp
+++ b/Engine/Source/Runtime/CoreUObject/Private/UObject/InstanceDataObjectUtils.cpp
@@ -265,6 +265,7 @@ namespace UE
 				AsStructProperty->Struct = Struct;
 				AsStructProperty->SetMetaData(NAME_OriginalType, *OriginalName);
 				AsStructProperty->SetMetaData(NAME_PresentAsTypeMetadata, *OriginalName);
+				AsStructProperty->Struct->SetMetaData(NAME_OriginalType, *OriginalName);
 				AsStructProperty->Struct->SetMetaData(NAME_PresentAsTypeMetadata, *OriginalName);
 
 				TrySetContainsLooseProperties(AsStructProperty, AsStructProperty->Struct);
@@ -421,6 +422,12 @@ namespace UE
 		UStruct* Result = NewObject<UStruct>(Outer, StructClass, MakeUniqueObjectName(nullptr, StructClass, InstanceDataObjectName));
 		Result->SetSuperStruct(Super);
 
+		// inherit ContainsLooseProperties metadata
+		if (Super && Super->GetBoolMetaData(NAME_ContainsLoosePropertiesMetadata))
+		{
+			Result->SetMetaData(NAME_ContainsLoosePropertiesMetadata, TEXT("True"));
+		}
+
 		TSet<FString> ExcludedLoosePropertyTypes = GetExcludedLoosePropertyTypes();
 
 		// Gather "loose" properties for child Struct
diff --git a/Engine/Source/Runtime/CoreUObject/Private/UObject/PropertyHelper.cpp b/Engine/Source/Runtime/CoreUObject/Private/UObject/PropertyHelper.cpp
--- a/Engine/Source/Runtime/CoreUObject/Private/UObject/PropertyHelper.cpp
+++ b/Engine/Source/Runtime/CoreUObject/Private/UObject/PropertyHelper.cpp
@@ -304,39 +304,36 @@ const FString* FindOriginalTypeName(const UField* Field)
 
 const FString* FindOriginalTypeName(const FProperty* Property)
 {
-	if (FUObjectThreadContext::Get().GetSerializeContext()->bImpersonateProperties)
+	// Prioritize metadata on the property over metadata on the type.
+	if (const FString* OriginalType = Property->FindMetaData(NAME_OriginalType))
 	{
-		// Prioritize metadata on the property over metadata on the type.
-		if (const FString* OriginalType = Property->FindMetaData(NAME_OriginalType))
-		{
-			return OriginalType;
-		}
+		return OriginalType;
+	}
 
-		// Search the owner chain to support metadata defined in UPROPERTY on a container for testing purposes.
-		for (FField* OwnerField = Property->Owner.ToField(); OwnerField; OwnerField = OwnerField->Owner.ToField())
+	// Search the owner chain to support metadata defined in UPROPERTY on a container for testing purposes.
+	for (FField* OwnerField = Property->Owner.ToField(); OwnerField; OwnerField = OwnerField->Owner.ToField())
+	{
+		if (const FString* OriginalType = OwnerField->FindMetaData(NAME_OriginalType))
 		{
-			if (const FString* OriginalType = OwnerField->FindMetaData(NAME_OriginalType))
-			{
-				return OriginalType;
-			}
+			return OriginalType;
 		}
+	}
 
-		if (const FStructProperty* StructProperty = CastField<FStructProperty>(Property))
-		{
-			return FindOriginalTypeName(StructProperty->Struct);
-		}
-		if (const FClassProperty* ClassProperty = CastField<FClassProperty>(Property))
-		{
-			return FindOriginalTypeName(ClassProperty->PropertyClass);
-		}
-		if (const FEnumProperty* EnumProperty = CastField<FEnumProperty>(Property))
-		{
-			return FindOriginalTypeName(EnumProperty->GetEnum());
-		}
-		if (const FByteProperty* ByteProperty = CastField<FByteProperty>(Property); ByteProperty && ByteProperty->Enum)
-		{
-			return FindOriginalTypeName(ByteProperty->Enum);
-		}
+	if (const FStructProperty* StructProperty = CastField<FStructProperty>(Property))
+	{
+		return FindOriginalTypeName(StructProperty->Struct);
+	}
+	if (const FClassProperty* ClassProperty = CastField<FClassProperty>(Property))
+	{
+		return FindOriginalTypeName(ClassProperty->PropertyClass);
+	}
+	if (const FEnumProperty* EnumProperty = CastField<FEnumProperty>(Property))
+	{
+		return FindOriginalTypeName(EnumProperty->GetEnum());
+	}
+	if (const FByteProperty* ByteProperty = CastField<FByteProperty>(Property); ByteProperty && ByteProperty->Enum)
+	{
+		return FindOriginalTypeName(ByteProperty->Enum);
 	}
 	return nullptr;
 }
@@ -469,12 +466,15 @@ static bool FindRedirectForProperty(FPropertyTypeName OldType, FPropertyTypeName
 		// Compare against the original type of an impersonated struct.
 		if constexpr (std::is_same_v<PropertyType, FStructProperty>)
 		{
-			if (const FPropertyTypeName OriginalType = FindOriginalType(Property); !OriginalType.IsEmpty())
+			if (FUObjectThreadContext::Get().GetSerializeContext()->bImpersonateProperties)
 			{
-				FieldName = OriginalType.GetName();
-				if (FieldName == OldNameRedirect.ObjectName)
+				if (const FPropertyTypeName OriginalType = FindOriginalType(Property); !OriginalType.IsEmpty())
 				{
-					return false;
+					FieldName = OriginalType.GetName();
+					if (FieldName == OldNameRedirect.ObjectName)
+					{
+						return false;
+					}
 				}
 			}
 		}
