Engine (FScriptSetHelper);

Adding FindLogicalIndex() to the set helper so we can get the logical index from the internal index. The internal index can have invalid entry indices kept in a continuous list for faster adding after an entry have been removed, while the logical index is all of the valid entry indices.

-- Code --
UnrealType:
FScriptSetHelper::FindLogicalIndex() - Added this function so we can get the logical index by passing in the internal index.

#tests Editor, used this for another change where we're fixing map and set entries being added the next missing entry where the key is an enum.

diff --git a/Engine/Source/Runtime/CoreUObject/Public/UObject/UnrealType.h b/Engine/Source/Runtime/CoreUObject/Public/UObject/UnrealType.h
--- a/Engine/Source/Runtime/CoreUObject/Public/UObject/UnrealType.h
+++ b/Engine/Source/Runtime/CoreUObject/Public/UObject/UnrealType.h
@@ -5660,7 +5660,7 @@ public:
 	COREUOBJECT_API void Rehash();
 
 	/**
-	 * Maps have gaps in their indices, so this function translates a logical index (ie. Nth element)
+	 * Sets have gaps in their indices, so this function translates a logical index (ie. Nth element)
 	 * to an internal index that can be used for the other functions in this class.
 	 * NOTE: This is slow, do not use this for iteration! Use CreateIterator() instead.
 	 */
@@ -5691,7 +5691,37 @@ public:
 		}
 		return INDEX_NONE;
 	}
-	
+
+	/**
+	 * Sets have gaps in their indices, so this function translates a internal index
+	 * to an logical index (ie. Nth element).
+	 * NOTE: This is slow, do not use this for iteration!
+	 */
+	int32 FindLogicalIndex(int32 InternalIdx) const
+	{
+		if (!IsValidIndex(InternalIdx))
+		{
+			return INDEX_NONE;
+		}
+
+		// if set is compact, use random access
+		if (GetMaxIndex() == Num())
+		{
+			return  InternalIdx;
+		}
+
+		int32 LogicalIndex = InternalIdx;
+		for (int i = 0; i < InternalIdx; ++i)
+		{
+			if (!IsValidIndex(i))
+			{
+				LogicalIndex--;
+			}
+		}
+
+		return LogicalIndex;
+	}
+
 	/**
 	* Finds the index of an element in a set
 	*
