Improve diagnostic messages for "Graph Linked to External Private Object" errors.

diff --git a/Engine/Source/Runtime/CoreUObject/Private/UObject/SavePackage/SavePackageUtilities.cpp b/Engine/Source/Runtime/CoreUObject/Private/UObject/SavePackage/SavePackageUtilities.cpp
--- a/Engine/Source/Runtime/CoreUObject/Private/UObject/SavePackage/SavePackageUtilities.cpp
+++ b/Engine/Source/Runtime/CoreUObject/Private/UObject/SavePackage/SavePackageUtilities.cpp
@@ -316,16 +316,28 @@ EObjectMark GetExcludedObjectMarksForTargetPlatform(const class ITargetPlatform*
 }
 
 /**
- * Find most likely culprit that caused the objects in the passed in array to be considered for saving.
+ * Find the most likely culprit that caused the objects in the passed in array to be considered illegal for saving.
  *
- * @param	BadObjects	array of objects that are considered "bad" (e.g. non- RF_Public, in different map package, ...)
- * @return	UObject that is considered the most likely culprit causing them to be referenced or NULL
+ * @param	BadObjects				Array of objects that are considered "bad" (e.g. non- RF_Public, in different map package, ...)
+ * @param	OutMostLikelyCulprit	UObject that is considered the most likely culprit causing the "bad" objects to be referenced or NULL
+ * @param	OutReferencer			UObject referencing the most likely culprit
+ * @param	OutReferencerProperty	Property (belonging to referencer) storing the offending reference
+ * @param	OutIsCulpritArchetype	Is the most likely culprit an archetype object
+ * @param	InOptionalSaveContext	Optional save context
  */
-void FindMostLikelyCulprit(const TArray<UObject*>& BadObjects, UObject*& MostLikelyCulprit, FString& OutReferencer, FSaveContext* InOptionalSaveContext)
+void FindMostLikelyCulprit(const TArray<UObject*>& BadObjects, 
+						   UObject*& OutMostLikelyCulprit, UObject*& OutReferencer, const FProperty*& OutReferencerProperty, bool& OutIsCulpritArchetype,
+						   FSaveContext* InOptionalSaveContext)
 {
 	UObject* ArchetypeCulprit = nullptr;
 	UObject* ReferencedCulprit = nullptr;
-	const FProperty* ReferencedCulpritReferencer = nullptr;
+	const FProperty* CulpritReferencerProperty = nullptr;
+	UObject* CulpritReferencer = nullptr;
+
+	OutMostLikelyCulprit = nullptr;
+	OutReferencer = nullptr;
+	OutReferencerProperty = nullptr;
+	OutIsCulpritArchetype = false;
 
 	auto IsObjectIncluded = [InOptionalSaveContext](UObject* InObject)
 	{
@@ -387,12 +399,14 @@ void FindMostLikelyCulprit(const TArray<UObject*>& BadObjects, UObject*& MostLik
 							continue;
 						}
 
+						CulpritReferencer = RefObj;
+
 						UE_LOG(LogSavePackage, Warning, TEXT("\t%s (%i refs)"), *RefObj->GetFullName(), Refs.ExternalReferences[i].TotalReferences);
 						for (int32 j = 0; j < Refs.ExternalReferences[i].ReferencingProperties.Num(); j++)
 						{
 							const FProperty* Prop = Refs.ExternalReferences[i].ReferencingProperties[j];
 							UE_LOG(LogSavePackage, Warning, TEXT("\t\t%i) %s"), j, *Prop->GetFullName());
-							ReferencedCulpritReferencer = Prop;
+							CulpritReferencerProperty = Prop;
 						}
 
 						// Later ReferencedCulprits are higher priority than earlier culprits. TODO: Not sure if this is an intentional behavior or if they choice was arbitrary.
@@ -406,30 +420,24 @@ void FindMostLikelyCulprit(const TArray<UObject*>& BadObjects, UObject*& MostLik
 	if (ArchetypeCulprit)
 	{
 		// ArchetypeCulprits are the most likely to be the problem; they are definitely a problem
-		MostLikelyCulprit = ArchetypeCulprit;
-		OutReferencer = TEXT("Referenced because it is an archetype object");
+		OutMostLikelyCulprit = ArchetypeCulprit;
+		OutIsCulpritArchetype = true;
 	}
 	else
 	{
-		MostLikelyCulprit = ReferencedCulprit; // Might be null, in which case we didn't find one
-		if (ReferencedCulpritReferencer)
-		{
-			OutReferencer = *ReferencedCulpritReferencer->GetName();
-		}
-		else
-		{
-			OutReferencer = TEXT("Unknown property");
-		}
+		OutMostLikelyCulprit = ReferencedCulprit; // Might be null, in which case we didn't find one
+		OutReferencer = CulpritReferencer;
+		OutReferencerProperty = CulpritReferencerProperty;
 	}
 
-	if (MostLikelyCulprit == nullptr)
+	if (OutMostLikelyCulprit == nullptr)
 	{
 		// Make sure we report something
 		for (UObject* BadObject : BadObjects)
 		{
 			if (BadObject)
 			{
-				MostLikelyCulprit = BadObject;
+				OutMostLikelyCulprit = BadObject;
 				break;
 			}
 		}
diff --git a/Engine/Source/Runtime/CoreUObject/Private/UObject/SavePackage/SavePackageUtilities.h b/Engine/Source/Runtime/CoreUObject/Private/UObject/SavePackage/SavePackageUtilities.h
--- a/Engine/Source/Runtime/CoreUObject/Private/UObject/SavePackage/SavePackageUtilities.h
+++ b/Engine/Source/Runtime/CoreUObject/Private/UObject/SavePackage/SavePackageUtilities.h
@@ -350,7 +350,7 @@ ESavePackageResult CreatePayloadSidecarFile(FLinkerSave& Linker, const FPackageP
 	
 void SaveWorldLevelInfo(UPackage* InOuter, FLinkerSave* Linker, FStructuredArchive::FRecord Record);
 EObjectMark GetExcludedObjectMarksForTargetPlatform(const class ITargetPlatform* TargetPlatform);
-void FindMostLikelyCulprit(const TArray<UObject*>& BadObjects, UObject*& MostLikelyCulprit, FString& OutReferencer, FSaveContext* InOptionalSaveContext = nullptr);
+void FindMostLikelyCulprit(const TArray<UObject*>& BadObjects, UObject*& OutMostLikelyCulprit, UObject*& OutReferencer, const FProperty*& OutReferencerProperty, bool& OutIsCulpritArchetype, FSaveContext* InOptionalSaveContext = nullptr);
 	
 /** 
 	* Search 'OutputFiles' for output files that were saved to the temp directory and move those files
diff --git a/Engine/Source/Runtime/CoreUObject/Private/UObject/SavePackage2.cpp b/Engine/Source/Runtime/CoreUObject/Private/UObject/SavePackage2.cpp
--- a/Engine/Source/Runtime/CoreUObject/Private/UObject/SavePackage2.cpp
+++ b/Engine/Source/Runtime/CoreUObject/Private/UObject/SavePackage2.cpp
@@ -711,66 +711,73 @@ ESavePackageResult ValidateIllegalReferences(FSaveContext& SaveContext, TArray<U
 {
 	FFormatNamedArguments Args;
 
+	TArray<UObject*>* ObjectsToCheck = nullptr;
+	FString ErrorPrologue;
+	FString ErrorEpilogue;
+
 	// Illegal objects in other map warning
 	if (ObjectsInOtherMaps.Num() > 0)
+	{
+		ObjectsToCheck = &ObjectsInOtherMaps;
+		ErrorPrologue = FString::Printf(TEXT("Can't save '%s': Illegal reference to private object: "), SaveContext.GetFilename());
+		ErrorEpilogue = FString(TEXT(" (private object belongs to an external map)."));
+	}
+	else if (PrivateObjects.Num() > 0)
+	{
+		ObjectsToCheck = &PrivateObjects;
+		ErrorPrologue = FString::Printf(TEXT("Can't save '%s': Illegal reference to private object: "), SaveContext.GetFilename());
+		ErrorEpilogue = FString(TEXT("."));
+	}
+	else if (PrivateContentObjects.Num() > 0)
+	{
+		ObjectsToCheck = &PrivateContentObjects;
+		ErrorPrologue = FString::Printf(TEXT("Can't save package: PKG_NotExternallyReferenceable: Package '%s' imports object: "), SaveContext.GetFilename());
+		ErrorEpilogue = FString(TEXT(", which is in a different mount point and its package is marked as PKG_NotExternallyReferenceable."));
+	}
+
+	if (ObjectsToCheck)
 	{
 		UObject* MostLikelyCulprit = nullptr;
-		FString CulpritString = TEXT("Unknown");
-		FString Referencer;
-		UE::SavePackageUtilities::FindMostLikelyCulprit(ObjectsInOtherMaps, MostLikelyCulprit, Referencer, &SaveContext);
-		if (MostLikelyCulprit != nullptr)
-		{
-			CulpritString = FString::Printf(TEXT("%s (%s)"), *MostLikelyCulprit->GetFullName(), *Referencer);
-		}
+		UObject* Referencer = nullptr;
+		const FProperty* ReferencerProperty = nullptr;
+		bool IsCulpritArchetype;
+		UE::SavePackageUtilities::FindMostLikelyCulprit(*ObjectsToCheck, MostLikelyCulprit, Referencer, ReferencerProperty, IsCulpritArchetype, &SaveContext);
 
-		FString ErrorMessage = FString::Printf(TEXT("Can't save %s: Graph is linked to object %s in external map"), SaveContext.GetFilename(), *CulpritString);
-		if (SaveContext.IsGenerateSaveError())
+		UObject* ReferencerOuter = Referencer ? Referencer->GetOuter() : nullptr;
+
+		FString ReferencerStr = Referencer ? Referencer->GetName() : TEXT("Unknown referencer");
+		FString ReferencerPropertyStr = ReferencerProperty ? *ReferencerProperty->GetName() : TEXT("Unknown property");
+		FString CulpritStr = MostLikelyCulprit ? *MostLikelyCulprit->GetFullName() : TEXT("Unknown");
+		FString ReferencerOuterStr = ReferencerOuter ? ReferencerOuter->GetPathName() : TEXT("Unknown owner");
+
+		FString ReferenceDescription;
+
+		if (IsCulpritArchetype)
 		{
-			SaveContext.GetError()->Logf(ELogVerbosity::Warning, TEXT("%s"), *ErrorMessage);
+			ReferenceDescription = FString::Printf(TEXT("'%s' referenced because it is an archetype object"), *CulpritStr);
 		}
 		else
 		{
-			UE_LOG(LogSavePackage, Error, TEXT("%s"), *ErrorMessage);
+			ReferenceDescription = FString::Printf(TEXT("'%s' referenced by '%s' (at '%s') in its '%s' property"), 
+				*CulpritStr, *ReferencerStr, *ReferencerOuterStr, *ReferencerPropertyStr);
 		}
-		return ESavePackageResult::Error;
-	}
 
-	if (PrivateObjects.Num() > 0)
-	{
-		UObject* MostLikelyCulprit = nullptr;
-		FString CulpritString = TEXT("Unknown");
-		FString Referencer;
-		UE::SavePackageUtilities::FindMostLikelyCulprit(PrivateObjects, MostLikelyCulprit, Referencer, &SaveContext);
-		CulpritString = FString::Printf(TEXT("%s (%s)"),
-			(MostLikelyCulprit != nullptr) ? *MostLikelyCulprit->GetFullName() : TEXT("(unknown culprit)"),
-			*Referencer);
+		FString ErrorMessage = FString::Printf(TEXT("%s%s%s"), *ErrorPrologue, *ReferenceDescription, *ErrorEpilogue);
 
 		if (SaveContext.IsGenerateSaveError())
 		{
-			SaveContext.GetError()->Logf(ELogVerbosity::Warning, TEXT("Can't save %s: Graph is linked to external private object %s"), SaveContext.GetFilename(), *CulpritString);
+			SaveContext.GetError()->Logf(ELogVerbosity::Warning, TEXT("%s"), *ErrorMessage);
 		}
-		return ESavePackageResult::Error;
-	}
-
-	if (PrivateContentObjects.Num() > 0)
-	{
-		UObject* MostLikelyCulprit = nullptr;
-		FString CulpritString = TEXT("Unknown");
-		FString Referencer;
-		UE::SavePackageUtilities::FindMostLikelyCulprit(PrivateContentObjects, MostLikelyCulprit, Referencer, &SaveContext);
-		CulpritString = FString::Printf(TEXT("%s (%s)"),
-			(MostLikelyCulprit != nullptr) ? *MostLikelyCulprit->GetFullName() : TEXT("(unknown culprit)"),
-			*Referencer);
-
-		if (SaveContext.IsGenerateSaveError())
+		else
 		{
-			SaveContext.GetError()->Logf(ELogVerbosity::Warning,
-				TEXT("Can't save package: PKG_NotExternallyReferenceable: Package %s imports object %s, which is in a different mount point and its package is marked as PKG_NotExternallyReferenceable."),
-				SaveContext.GetFilename(), *CulpritString);
+			UE_LOG(LogSavePackage, Error, TEXT("%s"), *ErrorMessage);
 		}
 		return ESavePackageResult::Error;
 	}
-	return ReturnSuccessOrCancel();
+	else
+	{
+		return ReturnSuccessOrCancel();
+	}
 }
 
 ESavePackageResult ValidateImports(FSaveContext& SaveContext)
