Fix reloading of UObjects while iterating

#rnx

diff --git a/Engine/Source/Runtime/CoreUObject/Private/UObject/UObjectBaseUtility.cpp b/Engine/Source/Runtime/CoreUObject/Private/UObject/UObjectBaseUtility.cpp
--- a/Engine/Source/Runtime/CoreUObject/Private/UObject/UObjectBaseUtility.cpp
+++ b/Engine/Source/Runtime/CoreUObject/Private/UObject/UObjectBaseUtility.cpp
@@ -617,7 +617,15 @@ void UObjectBaseUtility::ReloadObjectsFromModifiedConfigSections(const FConfigMo
 	// Now that we have a list of classes to update, we can iterate objects and reload
 	for (const UClass* ClassToReload : ClassesToReload)
 	{
-		ForEachObjectOfClass(ClassToReload, ReloadObjectImpl, true, RF_NoFlags);
+		TArray<UObject*> ClassObjectsToReload;
+		GetObjectsOfClass(ClassToReload, ClassObjectsToReload, true, RF_NoFlags);
+		for (UObject* ReloadObject : ClassObjectsToReload)
+		{
+			if (IsValid(ReloadObject))
+			{
+				ReloadObjectImpl(ReloadObject);
+			}
+		}
 	}
 
 	// Reload any PerObjectConfig objects that were affected
