Added tracking of unknown enum names

diff --git a/Engine/Source/Runtime/CoreUObject/Internal/UObject/PropertyBagRepository.h b/Engine/Source/Runtime/CoreUObject/Internal/UObject/PropertyBagRepository.h
--- a/Engine/Source/Runtime/CoreUObject/Internal/UObject/PropertyBagRepository.h
+++ b/Engine/Source/Runtime/CoreUObject/Internal/UObject/PropertyBagRepository.h
@@ -158,8 +158,6 @@ public:
 	static COREUOBJECT_API bool IsPropertyBagPlaceholderObject(const UObject* Object);
 	// query for whether or not creating property bag placeholder objects should be allowed
 	static COREUOBJECT_API bool IsPropertyBagPlaceholderObjectSupportEnabled();
-	// query whether an object supports IDO generation
-	static COREUOBJECT_API bool IsInstanceDataObjectSupportEnabled(UObject* InObject = nullptr);
 
 	/**
 	 * Create a new placeholder type object to swap in for a missing class/struct. An object of
@@ -217,7 +215,8 @@ struct FScopedIDOSerializationContext
 	UObject* const Object = nullptr;
 	const int64 PreSerializeOffset;
 	TOptional<TGuardValue<bool>> ScopedTrackSerializedPropertyPath;
-	TOptional<TGuardValue<bool>> ScopedSerializeUnknownProperty;
+	TOptional<TGuardValue<bool>> ScopedSerializeUnknownProperties;
+	TOptional<TGuardValue<bool>> ScopedSerializeUnknownEnumNames;
 	TOptional<TGuardValue<bool>> ScopedImpersonateProperties;
 	TOptional<TGuardValue<bool>> ScopedTrackInitializedProperties;
 	TOptional<TGuardValue<bool>> ScopedTrackSerializedProperties;
diff --git a/Engine/Source/Runtime/CoreUObject/Private/UObject/Enum.cpp b/Engine/Source/Runtime/CoreUObject/Private/UObject/Enum.cpp
--- a/Engine/Source/Runtime/CoreUObject/Private/UObject/Enum.cpp
+++ b/Engine/Source/Runtime/CoreUObject/Private/UObject/Enum.cpp
@@ -7,6 +7,7 @@
 #include "UObject/Class.h"
 #include "UObject/Package.h"
 #include "UObject/MetaData.h"
+#include "UObject/PropertyHelper.h"
 #include "UObject/PropertyPortFlags.h"
 #include "UObject/UObjectThreadContext.h"
 #include "UObject/DevObjectVersion.h"
@@ -14,7 +15,7 @@
 #include "UObject/CoreRedirects.h"
 #include "UObject/LinkerLoad.h"
 
-DEFINE_LOG_CATEGORY_STATIC(LogEnum, Log, All);
+DEFINE_LOG_CATEGORY(LogEnum);
 
 /*-----------------------------------------------------------------------------
 	UEnum implementation.
diff --git a/Engine/Source/Runtime/CoreUObject/Private/UObject/EnumProperty.cpp b/Engine/Source/Runtime/CoreUObject/Private/UObject/EnumProperty.cpp
--- a/Engine/Source/Runtime/CoreUObject/Private/UObject/EnumProperty.cpp
+++ b/Engine/Source/Runtime/CoreUObject/Private/UObject/EnumProperty.cpp
@@ -8,6 +8,10 @@
 #include "UObject/UnrealTypePrivate.h"
 #include "UObject/UObjectThreadContext.h"
 
+#if WITH_EDITORONLY_DATA
+#include "UObject/PropertyBagRepository.h"
+#endif
+
 namespace UE::CoreUObject::Private
 {
 	template <typename OldIntType>
@@ -166,6 +170,48 @@ namespace UE::CoreUObject::Private
 	}
 }
 
+bool TryLoadEnumValueByName(FStructuredArchive::FSlot Slot, FArchive& UnderlyingArchive, UEnum* Enum, FName& OutEnumValueName, int64& OutEnumValue)
+{
+	Slot << OutEnumValueName;
+
+	if (Enum)
+	{
+		// Make sure enum is properly populated
+		if (Enum->HasAnyFlags(RF_NeedLoad))
+		{
+			UnderlyingArchive.Preload(Enum);
+		}
+
+		if (Enum->HasAnyEnumFlags(EEnumFlags::Flags))
+		{
+			if (OutEnumValueName != NAME_None)
+			{
+				OutEnumValue = Enum->GetValueOrBitfieldFromString(*OutEnumValueName.ToString());
+				return OutEnumValue != INDEX_NONE;
+			}
+		}
+		else
+		{
+			// There's no guarantee EnumValueName is still present in Enum, in which case Value will be set to the enum's max value.
+			// On save, it will then be serialized as NAME_None.
+			const int32 EnumIndex = Enum->GetIndexByName(OutEnumValueName, EGetByNameFlags::ErrorIfNotFound);
+			if (EnumIndex == INDEX_NONE)
+			{
+				OutEnumValue = Enum->GetMaxEnumValue();
+				return false;
+			}
+			else
+			{
+				OutEnumValue = Enum->GetValueByIndex(EnumIndex);
+				return true;
+			}
+		}
+	}
+
+	OutEnumValue = 0;
+	return false;
+}
+
 IMPLEMENT_FIELD(FEnumProperty)
 
 FEnumProperty::FEnumProperty(FFieldVariant InOwner, const FName& InName, EObjectFlags InObjectFlags)
@@ -245,42 +291,18 @@ void FEnumProperty::SerializeItem(FStructuredArchive::FSlot Slot, void* Value, v
 	if (UnderlyingArchive.IsLoading())
 	{
 		FName EnumValueName;
-		Slot << EnumValueName;
-
-		int64 NewEnumValue = 0;
-
-		if (Enum)
+		int64 EnumValue = 0;
+		if (!TryLoadEnumValueByName(Slot, UnderlyingArchive, Enum, EnumValueName, EnumValue))
 		{
-			// Make sure enum is properly populated
-			if (Enum->HasAnyFlags(RF_NeedLoad))
-			{
-				UnderlyingArchive.Preload(Enum);
-			}
-
-			if (Enum->HasAnyEnumFlags(EEnumFlags::Flags))
+		#if WITH_EDITORONLY_DATA
+			FUObjectSerializeContext* SerializeContext = FUObjectThreadContext::Get().GetSerializeContext();
+			if (UNLIKELY(SerializeContext->bTrackUnknownEnumNames))
 			{
-				if (EnumValueName != NAME_None)
-				{
-					NewEnumValue = Enum->GetValueOrBitfieldFromString(*EnumValueName.ToString());
-				}
-			}
-			else
-			{
-				// There's no guarantee EnumValueName is still present in Enum, in which case Value will be set to the enum's max value.
-				// On save, it will then be serialized as NAME_None.
-				const int32 EnumIndex = Enum->GetIndexByName(EnumValueName, EGetByNameFlags::ErrorIfNotFound);
-				if (EnumIndex == INDEX_NONE)
-				{
-					NewEnumValue = Enum->GetMaxEnumValue();
-				}
-				else
-				{
-					NewEnumValue = Enum->GetValueByIndex(EnumIndex);
-				}
+				UE::FPropertyBagRepository::Get().AddUnknownEnumName(SerializeContext->SerializedObject, Enum, UE::FindOriginalType(this), EnumValueName);
 			}
+		#endif
 		}
-
-		UnderlyingProp->SetIntPropertyValue(Value, NewEnumValue);
+		UnderlyingProp->SetIntPropertyValue(Value, EnumValue);
 	}
 	// Saving
 	else if (UnderlyingArchive.IsSaving())
@@ -470,7 +492,31 @@ bool FEnumProperty::SameType(const FProperty* Other) const
 EConvertFromTypeResult FEnumProperty::ConvertFromType(const FPropertyTag& Tag, FStructuredArchive::FSlot Slot , uint8* Data, UStruct* DefaultsStruct, const uint8* Defaults)
 {
 	const EName* TagType = Tag.Type.ToEName();
-	if (LIKELY(!TagType || *TagType == NAME_EnumProperty || Tag.Type.GetNumber() || !Enum || !UnderlyingProp))
+	if (!TagType || Tag.Type.GetNumber())
+	{
+		return EConvertFromTypeResult::UseSerializeItem;
+	}
+
+	if (*TagType == NAME_EnumProperty)
+	{
+	#if WITH_EDITORONLY_DATA
+		FUObjectSerializeContext* SerializeContext = FUObjectThreadContext::Get().GetSerializeContext();
+		if (UNLIKELY(SerializeContext->bTrackUnknownEnumNames && !CanSerializeFromTypeName(Tag.GetType())))
+		{
+			FName EnumValueName;
+			int64 EnumValue = 0;
+			TryLoadEnumValueByName(Slot, Slot.GetUnderlyingArchive(), Enum, EnumValueName, EnumValue);
+
+			check(UnderlyingProp);
+			UnderlyingProp->SetIntPropertyValue(ContainerPtrToValuePtr<void>(Data, Tag.ArrayIndex), EnumValue);
+
+			return EConvertFromTypeResult::Converted;
+		}
+	#endif
+		return EConvertFromTypeResult::UseSerializeItem;
+	}
+
+	if (!Enum || !UnderlyingProp)
 	{
 		return EConvertFromTypeResult::UseSerializeItem;
 	}
@@ -604,7 +650,12 @@ bool FEnumProperty::LoadTypeName(UE::FPropertyTypeName Type, const FPropertyTag*
 	UEnum* LocalEnum = FindFirstObject<UEnum>(*WriteToString<256>(EnumName), EFindFirstObjectOptions::NativeFirst);
 	if (!LocalEnum)
 	{
+	#if WITH_EDITORONLY_DATA
+		LocalEnum = StaticEnum<EFallbackEnum>();
+		SetMetaData(UE::NAME_OriginalType, *WriteToString<256>(Type.GetParameter(0)));
+	#else
 		return false;
+	#endif
 	}
 
 	const UE::FPropertyTypeName UnderlyingType = Type.GetParameter(1);
diff --git a/Engine/Source/Runtime/CoreUObject/Private/UObject/Obj.cpp b/Engine/Source/Runtime/CoreUObject/Private/UObject/Obj.cpp
--- a/Engine/Source/Runtime/CoreUObject/Private/UObject/Obj.cpp
+++ b/Engine/Source/Runtime/CoreUObject/Private/UObject/Obj.cpp
@@ -72,6 +72,7 @@
 #include "ProfilingDebugging/AssetMetadataTrace.h"
 #include "Containers/VersePath.h"
 #include "Serialization/LoadTimeTracePrivate.h"
+#include "UObject/InstanceDataObjectUtils.h"
 #include "UObject/PropertyBagRepository.h"
 #include "AutoRTFM/AutoRTFM.h"
 
@@ -1284,11 +1285,9 @@ void UObject::ConditionalPostLoad()
 
 #if WITH_EDITORONLY_DATA
 		// Object has been deserialized, if IDO is enabled, generate it
-		UE::FPropertyBagRepository& PropertyBagRepository = UE::FPropertyBagRepository::Get();
-		const FUObjectSerializeContext* LoadContext = FUObjectThreadContext::Get().GetSerializeContext();
-		const bool bIDOEnabled = PropertyBagRepository.IsInstanceDataObjectSupportEnabled(this) && !LoadContext->bImpersonateProperties;
-		if (bIDOEnabled)
+		if (UE::IsInstanceDataObjectSupportEnabled(this) && !FUObjectThreadContext::Get().GetSerializeContext()->bImpersonateProperties)
 		{
+			UE::FPropertyBagRepository& PropertyBagRepository = UE::FPropertyBagRepository::Get();
 			PropertyBagRepository.PostLoadInstanceDataObject(this);
 		}
 #endif
diff --git a/Engine/Source/Runtime/CoreUObject/Private/UObject/PropertyBagRepository.cpp b/Engine/Source/Runtime/CoreUObject/Private/UObject/PropertyBagRepository.cpp
--- a/Engine/Source/Runtime/CoreUObject/Private/UObject/PropertyBagRepository.cpp
+++ b/Engine/Source/Runtime/CoreUObject/Private/UObject/PropertyBagRepository.cpp
@@ -570,7 +570,14 @@ void FPropertyBagRepository::AddUnknownEnumName(const UObject* Owner, const UEnu
 	TStringBuilder<128> EnumValueString(InPlace, EnumValueName);
 	if (String::FindFirstChar(EnumValueString, TEXT('|')) == INDEX_NONE)
 	{
-		Info.Names.Add(EnumValueName);
+		if (int32 ColonIndex = String::FindFirst(EnumValueString, TEXTVIEW("::")); ColonIndex != INDEX_NONE)
+		{
+			Info.Names.Emplace(EnumValueString.ToView().RightChop(ColonIndex + TEXTVIEW("::").Len()));
+		}
+		else
+		{
+			Info.Names.Add(EnumValueName);
+		}
 	}
 	else
 	{
@@ -825,7 +832,7 @@ void FPropertyBagRepository::CreateInstanceDataObjectUnsafe(UObject* Owner, FPro
 	// TODO: should we put the InstanceDataObject or it's class in a package?
 	const UClass* InstanceDataObjectClass = CreateInstanceDataObjectClass(PropertyTree, Owner->GetClass(), GetTransientPackage());
 
-	BagData.bNeedsFixup = UE::StructContainsLooseProperties(InstanceDataObjectClass);
+	BagData.bNeedsFixup = StructContainsLooseProperties(InstanceDataObjectClass);
 
 	TObjectPtr<UObject>* OuterPtr = nullptr;
 	if (FPropertyBagAssociationData* OuterData = AssociatedData.Find(Owner->GetOuter()))
@@ -902,8 +909,8 @@ FScopedIDOSerializationContext::FScopedIDOSerializationContext(UObject* InObject
 	, PreSerializeOffset(InArchive.Tell())
 {
 	FUObjectSerializeContext* SerializeContext = FUObjectThreadContext::Get().GetSerializeContext();
-	bool bHasIDOSupport = FPropertyBagRepository::IsInstanceDataObjectSupportEnabled(Object);
-	
+	bool bHasIDOSupport = IsInstanceDataObjectSupportEnabled(Object);
+
 	bool bHasReinstancedClass = InObject->GetClass()->HasAnyClassFlags(CLASS_NewerVersionExists);
 	bCreateIDO = bHasIDOSupport && !SerializeContext->bImpersonateProperties && Archive->IsLoading() && !bHasReinstancedClass;
 
@@ -914,7 +921,8 @@ FScopedIDOSerializationContext::FScopedIDOSerializationContext(UObject* InObject
 			// Enable creation of a property path name tree to track any property that does not match the current class schema,
 			// except when impersonation is enabled because that implies we are deserializing an IDO.
 			ScopedTrackSerializedPropertyPath.Emplace(SerializeContext->bTrackSerializedPropertyPath, bCreateIDO);
-			ScopedSerializeUnknownProperty.Emplace(SerializeContext->bTrackUnknownProperties, bCreateIDO);
+			ScopedSerializeUnknownProperties.Emplace(SerializeContext->bTrackUnknownProperties, bCreateIDO);
+			ScopedSerializeUnknownEnumNames.Emplace(SerializeContext->bTrackUnknownEnumNames, bCreateIDO);
 			ScopedSerializedObject.Emplace(SerializeContext->SerializedObject, Object);
 
 			// Enable tracking of initialized properties when loading an IDO, which is implied by impersonation being enabled.
@@ -936,7 +944,7 @@ FScopedIDOSerializationContext::FScopedIDOSerializationContext(UObject* InObject
 	, PreSerializeOffset(0)
 {
 	FUObjectSerializeContext* SerializeContext = FUObjectThreadContext::Get().GetSerializeContext();
-	bool bHasIDOSupport = FPropertyBagRepository::IsInstanceDataObjectSupportEnabled(Object);
+	bool bHasIDOSupport = IsInstanceDataObjectSupportEnabled(Object);
 	if (bHasIDOSupport)
 	{
 		ScopedImpersonateProperties.Emplace(SerializeContext->bImpersonateProperties, bImpersonate);
@@ -1040,15 +1048,6 @@ bool FPropertyBagRepository::IsPropertyBagPlaceholderObjectSupportEnabled()
 #endif
 }
 
-bool FPropertyBagRepository::IsInstanceDataObjectSupportEnabled(UObject* InObject)
-{
-#if WITH_EDITOR
-	return UE::IsInstanceDataObjectSupportEnabled(InObject);
-#else
-	return false;
-#endif
-}
-
 UStruct* FPropertyBagRepository::CreatePropertyBagPlaceholderType(UObject* Outer, UClass* Class, FName Name, EObjectFlags Flags, UStruct* SuperStruct)
 {
 	// Generate and link a new type object using the given SuperStruct as its base.
diff --git a/Engine/Source/Runtime/CoreUObject/Private/UObject/PropertyByte.cpp b/Engine/Source/Runtime/CoreUObject/Private/UObject/PropertyByte.cpp
--- a/Engine/Source/Runtime/CoreUObject/Private/UObject/PropertyByte.cpp
+++ b/Engine/Source/Runtime/CoreUObject/Private/UObject/PropertyByte.cpp
@@ -8,6 +8,13 @@
 #include "UObject/UnrealTypePrivate.h"
 #include "UObject/UObjectThreadContext.h"
 
+#if WITH_EDITORONLY_DATA
+#include "UObject/PropertyBagRepository.h"
+#endif
+
+// Implemented in EnumProperty.cpp
+bool TryLoadEnumValueByName(FStructuredArchive::FSlot Slot, FArchive& UnderlyingArchive, UEnum* Enum, FName& OutEnumValueName, int64& OutEnumValue);
+
 /*-----------------------------------------------------------------------------
 	FByteProperty.
 -----------------------------------------------------------------------------*/
@@ -57,24 +64,18 @@ void FByteProperty::SerializeItem( FStructuredArchive::FSlot Slot, void* Value,
 	else if (UnderlyingArchive.IsLoading())
 	{
 		FName EnumValueName;
-		Slot << EnumValueName;
-		// Make sure enum is properly populated
-		if( Enum->HasAnyFlags(RF_NeedLoad) )
+		int64 EnumValue = 0;
+		if (!TryLoadEnumValueByName(Slot, UnderlyingArchive, Enum, EnumValueName, EnumValue))
 		{
-			UnderlyingArchive.Preload(Enum);
-		}
-
-		// There's no guarantee EnumValueName is still present in Enum, in which case Value will be set to the enum's max value.
-		// On save, it will then be serialized as NAME_None.
-		int32 EnumIndex = Enum->GetIndexByName(EnumValueName, EGetByNameFlags::ErrorIfNotFound);
-		if (EnumIndex == INDEX_NONE)
-		{
-			*(uint8*)Value = IntCastChecked<uint8>(Enum->GetMaxEnumValue());
-		}
-		else
-		{
-			*(uint8*)Value = IntCastChecked<uint8>(Enum->GetValueByIndex(EnumIndex));
+		#if WITH_EDITORONLY_DATA
+			FUObjectSerializeContext* SerializeContext = FUObjectThreadContext::Get().GetSerializeContext();
+			if (UNLIKELY(SerializeContext->bTrackUnknownEnumNames))
+			{
+				UE::FPropertyBagRepository::Get().AddUnknownEnumName(SerializeContext->SerializedObject, Enum, UE::FindOriginalType(this), EnumValueName);
+			}
+		#endif
 		}
+		*(uint8*)Value = IntCastChecked<uint8>(EnumValue);
 	}
 	// Saving
 	else
@@ -473,7 +474,13 @@ bool FByteProperty::LoadTypeName(UE::FPropertyTypeName Type, const FPropertyTag*
 		return true;
 	}
 
+#if WITH_EDITORONLY_DATA
+	Enum = StaticEnum<EFallbackEnum>();
+	SetMetaData(UE::NAME_OriginalType, *WriteToString<256>(Type.GetParameter(0)));
+	return true;
+#else
 	return false;
+#endif
 }
 
 void FByteProperty::SaveTypeName(UE::FPropertyTypeNameBuilder& Type) const
diff --git a/Engine/Source/Runtime/CoreUObject/Private/UObject/PropertyHelper.cpp b/Engine/Source/Runtime/CoreUObject/Private/UObject/PropertyHelper.cpp
--- a/Engine/Source/Runtime/CoreUObject/Private/UObject/PropertyHelper.cpp
+++ b/Engine/Source/Runtime/CoreUObject/Private/UObject/PropertyHelper.cpp
@@ -295,7 +295,7 @@ namespace UE
 #if WITH_EDITORONLY_DATA
 const FString* FindOriginalTypeName(const UField* Field)
 {
-	if (FUObjectThreadContext::Get().GetSerializeContext()->bImpersonateProperties)
+	if (Field && FUObjectThreadContext::Get().GetSerializeContext()->bImpersonateProperties)
 	{
 		return Field->FindMetaData(NAME_OriginalType);
 	}
@@ -331,7 +331,7 @@ const FString* FindOriginalTypeName(const FProperty* Property)
 	{
 		return FindOriginalTypeName(EnumProperty->GetEnum());
 	}
-	if (const FByteProperty* ByteProperty = CastField<FByteProperty>(Property); ByteProperty && ByteProperty->Enum)
+	if (const FByteProperty* ByteProperty = CastField<FByteProperty>(Property))
 	{
 		return FindOriginalTypeName(ByteProperty->Enum);
 	}
diff --git a/Engine/Source/Runtime/CoreUObject/Private/UObject/PropertyHelper.h b/Engine/Source/Runtime/CoreUObject/Private/UObject/PropertyHelper.h
--- a/Engine/Source/Runtime/CoreUObject/Private/UObject/PropertyHelper.h
+++ b/Engine/Source/Runtime/CoreUObject/Private/UObject/PropertyHelper.h
@@ -10,13 +10,21 @@
 #include "UObject/UObjectThreadContext.h"
 #include "UObject/WeakObjectPtrFwd.h"
 
+#include "PropertyHelper.generated.h"
+
 class FProperty;
 class UField;
 class UFunction;
 class UObject;
 
+DECLARE_LOG_CATEGORY_EXTERN(LogEnum, Log, All);
 DECLARE_LOG_CATEGORY_EXTERN(LogProperty, Log, All);
 
+UENUM()
+enum class EFallbackEnum : uint8
+{
+};
+
 /**
  * Advances the character pointer past any spaces or tabs.
  * 
diff --git a/Engine/Source/Runtime/CoreUObject/Private/UObject/PropertyNumeric.cpp b/Engine/Source/Runtime/CoreUObject/Private/UObject/PropertyNumeric.cpp
--- a/Engine/Source/Runtime/CoreUObject/Private/UObject/PropertyNumeric.cpp
+++ b/Engine/Source/Runtime/CoreUObject/Private/UObject/PropertyNumeric.cpp
@@ -27,7 +27,11 @@ int64 FNumericProperty::ReadEnumAsInt64(FStructuredArchive::FSlot Slot, UStruct*
 
 	FName EnumTypeName = Tag.GetType().GetParameterName(0);
 
-	UEnum* Enum = FindUField<UEnum>(dynamic_cast<UClass*>(DefaultsStruct) ? static_cast<UClass*>(DefaultsStruct) : DefaultsStruct->GetTypedOuter<UClass>(), EnumTypeName);
+	UEnum* Enum = nullptr;
+	if (DefaultsStruct)
+	{
+		Enum = FindUField<UEnum>(DefaultsStruct->IsA<UClass>() ? (UClass*)DefaultsStruct : DefaultsStruct->GetTypedOuter<UClass>(), EnumTypeName);
+	}
 	if (!Enum)
 	{
 		// Enums (at least native) are stored as short names (for now) so find the Tag enum by name
diff --git a/Engine/Source/Runtime/CoreUObject/Private/UObject/UObjectGlobals.cpp b/Engine/Source/Runtime/CoreUObject/Private/UObject/UObjectGlobals.cpp
--- a/Engine/Source/Runtime/CoreUObject/Private/UObject/UObjectGlobals.cpp
+++ b/Engine/Source/Runtime/CoreUObject/Private/UObject/UObjectGlobals.cpp
@@ -34,6 +34,7 @@
 #include "UObject/Class.h"
 #include "UObject/CoreRedirects.h"
 #include "UObject/FastReferenceCollector.h"
+#include "UObject/InstanceDataObjectUtils.h"
 #include "UObject/OverridableManager.h"
 #include "UObject/UObjectIterator.h"
 #include "UObject/Package.h"
@@ -3253,9 +3254,9 @@ UObject* StaticDuplicateObjectEx( FObjectDuplicationParameters& Parameters )
 
 #if WITH_EDITORONLY_DATA
 	// if Source has an IDO, make one for Dest and copy it
-	UE::FPropertyBagRepository& PropertyBagRepository = UE::FPropertyBagRepository::Get();
-	if (UE::FPropertyBagRepository::IsInstanceDataObjectSupportEnabled(Parameters.SourceObject))
+	if (UE::IsInstanceDataObjectSupportEnabled(Parameters.SourceObject))
 	{
+		UE::FPropertyBagRepository& PropertyBagRepository = UE::FPropertyBagRepository::Get();
 		PropertyBagRepository.DuplicateInstanceDataObject(Parameters.SourceObject, DupRootObject);
 	}
 #endif
diff --git a/Engine/Source/Runtime/CoreUObject/Private/UObject/UObjectThreadContext.cpp b/Engine/Source/Runtime/CoreUObject/Private/UObject/UObjectThreadContext.cpp
--- a/Engine/Source/Runtime/CoreUObject/Private/UObject/UObjectThreadContext.cpp
+++ b/Engine/Source/Runtime/CoreUObject/Private/UObject/UObjectThreadContext.cpp
@@ -49,6 +49,7 @@ FUObjectSerializeContext::FUObjectSerializeContext()
 	, bTrackInitializedProperties(false)
 	, bTrackSerializedProperties(false)
 	, bTrackUnknownProperties(false)
+	, bTrackUnknownEnumNames(false)
 	, bImpersonateProperties(false)
 #endif
 {
diff --git a/Engine/Source/Runtime/CoreUObject/Public/UObject/UObjectThreadContext.h b/Engine/Source/Runtime/CoreUObject/Public/UObject/UObjectThreadContext.h
--- a/Engine/Source/Runtime/CoreUObject/Public/UObject/UObjectThreadContext.h
+++ b/Engine/Source/Runtime/CoreUObject/Public/UObject/UObjectThreadContext.h
@@ -178,8 +178,10 @@ public:
 	UE_INTERNAL bool bTrackInitializedProperties;
 	/** True when serialized properties will be tracked for the serialized object. */
 	UE_INTERNAL bool bTrackSerializedProperties;
-	/** True when unknown properties will tracked for the serialized object. */
+	/** True when unknown properties will be tracked for the serialized object. */
 	UE_INTERNAL bool bTrackUnknownProperties;
+	/** True when unknown enum names will be tracked for the serialized object. */
+	UE_INTERNAL bool bTrackUnknownEnumNames;
 	/** True when the SerializedObject properties are being impersonated. */
 	UE_INTERNAL bool bImpersonateProperties;
 #endif
diff --git a/Engine/Source/Runtime/CoreUObject/Tests/UObject/InstanceDataObjectUtilsTest.cpp b/Engine/Source/Runtime/CoreUObject/Tests/UObject/InstanceDataObjectUtilsTest.cpp
--- a/Engine/Source/Runtime/CoreUObject/Tests/UObject/InstanceDataObjectUtilsTest.cpp
+++ b/Engine/Source/Runtime/CoreUObject/Tests/UObject/InstanceDataObjectUtilsTest.cpp
@@ -210,6 +210,13 @@ TEST_CASE_NAMED(FTrackUnknownPropertiesTest, "CoreUObject::Serialization::TrackU
 		return Builder.Build();
 	};
 
+	const auto SavePropertyTypeName = [](const FProperty* Property)
+	{
+		FPropertyTypeNameBuilder Builder;
+		Property->SaveTypeName(Builder);
+		return Builder.Build();
+	};
+
 	UObject* Owner = NewObject<UTestInstanceDataObjectClass>();
 
 	FUObjectSerializeContext* SerializeContext = FUObjectThreadContext::Get().GetSerializeContext();
@@ -228,6 +235,10 @@ TEST_CASE_NAMED(FTrackUnknownPropertiesTest, "CoreUObject::Serialization::TrackU
 	AltStructData.Grain = ETestInstanceDataObjectGrainAlternate::Corn;
 	AltStructData.Fruit = ETestInstanceDataObjectFruitAlternate::Orange;
 	AltStructData.Direction = ETestInstanceDataObjectDirectionAlternate::North | ETestInstanceDataObjectDirectionAlternate::West;
+	AltStructData.GrainTypeChange = ETestInstanceDataObjectGrainAlternate::Corn;
+	AltStructData.FruitTypeChange = ETestInstanceDataObjectFruitAlternate::Orange;
+	AltStructData.GrainTypeAndPropertyChange = ETestInstanceDataObjectGrainAlternateEnumClass::Corn;
+	AltStructData.FruitTypeAndPropertyChange = ETestInstanceDataObjectFruitAlternateNamespace::Orange;
 	AltStructData.Point.U = 1;
 	AltStructData.Point.V = 2;
 	AltStructData.Point.W = 3;
@@ -269,6 +280,11 @@ TEST_CASE_NAMED(FTrackUnknownPropertiesTest, "CoreUObject::Serialization::TrackU
 	CHECK(StructData.Grain == ETestInstanceDataObjectGrain::Corn);
 	CHECK(StructData.Fruit == ETestInstanceDataObjectFruit::Orange);
 	CHECK(StructData.Direction == (ETestInstanceDataObjectDirection::North | ETestInstanceDataObjectDirection::West));
+	CHECK(StructData.GrainTypeChange == ETestInstanceDataObjectGrain::Corn);
+	CHECK(StructData.FruitTypeChange == ETestInstanceDataObjectFruit::Orange);
+	// TODO: ConvertFromType incorrectly assigns numeric values based on the saved enum type.
+	//CHECK(StructData.GrainTypeAndPropertyChange == ETestInstanceDataObjectGrain::Corn);
+	//CHECK(StructData.FruitTypeAndPropertyChange == ETestInstanceDataObjectFruit::Orange);
 	CHECK(StructData.Point.X == 0);
 	CHECK(StructData.Point.Y == 0);
 	CHECK(StructData.Point.Z == 0);
@@ -313,6 +329,23 @@ TEST_CASE_NAMED(FTrackUnknownPropertiesTest, "CoreUObject::Serialization::TrackU
 		#endif
 		}
 	}
+	// TODO: Enum type name changes are not recorded as a conversion.
+	//{
+	//	FSerializedPropertyPathScope Path(SerializeContext, {"GrainTypeChange", SavePropertyTypeName(FindFProperty<FProperty>(StaticStruct<FTestInstanceDataObjectStructAlternate>(), TEXT("GrainTypeChange")))});
+	//	CHECK(Tree->Find(SerializeContext->SerializedPropertyPath));
+	//}
+	//{
+	//	FSerializedPropertyPathScope Path(SerializeContext, {"FruitTypeChange", SavePropertyTypeName(FindFProperty<FProperty>(StaticStruct<FTestInstanceDataObjectStructAlternate>(), TEXT("FruitTypeChange")))});
+	//	CHECK(Tree->Find(SerializeContext->SerializedPropertyPath));
+	//}
+	{
+		FSerializedPropertyPathScope Path(SerializeContext, {"GrainTypeAndPropertyChange", SavePropertyTypeName(FindFProperty<FProperty>(StaticStruct<FTestInstanceDataObjectStructAlternate>(), TEXT("GrainTypeAndPropertyChange")))});
+		CHECK(Tree->Find(SerializeContext->SerializedPropertyPath));
+	}
+	{
+		FSerializedPropertyPathScope Path(SerializeContext, {"FruitTypeAndPropertyChange", SavePropertyTypeName(FindFProperty<FProperty>(StaticStruct<FTestInstanceDataObjectStructAlternate>(), TEXT("FruitTypeAndPropertyChange")))});
+		CHECK(Tree->Find(SerializeContext->SerializedPropertyPath));
+	}
 	FPropertyBagRepository::Get().DestroyOuterBag(Owner);
 
 #if WITH_TEXT_ARCHIVE_SUPPORT
@@ -333,6 +366,11 @@ TEST_CASE_NAMED(FTrackUnknownPropertiesTest, "CoreUObject::Serialization::TrackU
 	CHECK(StructData.Grain == ETestInstanceDataObjectGrain::Corn);
 	CHECK(StructData.Fruit == ETestInstanceDataObjectFruit::Orange);
 	CHECK(StructData.Direction == (ETestInstanceDataObjectDirection::North | ETestInstanceDataObjectDirection::West));
+	CHECK(StructData.GrainTypeChange == ETestInstanceDataObjectGrain::Corn);
+	CHECK(StructData.FruitTypeChange == ETestInstanceDataObjectFruit::Orange);
+	// TODO: ConvertFromType incorrectly assigns numeric values based on the saved enum type.
+	//CHECK(StructData.GrainTypeAndPropertyChange == ETestInstanceDataObjectGrain::Corn);
+	//CHECK(StructData.FruitTypeAndPropertyChange == ETestInstanceDataObjectFruit::Orange);
 	CHECK(StructData.Point.X == 0);
 	CHECK(StructData.Point.Y == 0);
 	CHECK(StructData.Point.Z == 0);
@@ -344,6 +382,189 @@ TEST_CASE_NAMED(FTrackUnknownPropertiesTest, "CoreUObject::Serialization::TrackU
 #endif // WITH_TEXT_ARCHIVE_SUPPORT
 }
 
+TEST_CASE_NAMED(FTrackUnknownEnumNamesTest, "CoreUObject::Serialization::TrackUnknownEnumNames", "[CoreUObject][EngineFilter]")
+{
+	const auto MakePropertyTypeName = [](const UEnum* Enum)
+	{
+		FPropertyTypeNameBuilder Builder;
+		Builder.AddPath(Enum);
+		return Builder.Build();
+	};
+
+	const auto ParsePropertyTypeName = [](const TCHAR* Name) -> FPropertyTypeName
+	{
+		FPropertyTypeNameBuilder Builder;
+		CHECK(Builder.TryParse(Name));
+		return Builder.Build();
+	};
+
+	UObject* Owner = NewObject<UTestInstanceDataObjectClass>();
+
+	FUObjectSerializeContext* SerializeContext = FUObjectThreadContext::Get().GetSerializeContext();
+	TGuardValue<UObject*> SerializedObjectScope(SerializeContext->SerializedObject, Owner);
+	TGuardValue<bool> TrackSerializedPropertyPathScope(SerializeContext->bTrackSerializedPropertyPath, true);
+	TGuardValue<bool> TrackUnknownPropertiesScope(SerializeContext->bTrackUnknownProperties, true);
+	TGuardValue<bool> TrackUnknownEnumNamesScope(SerializeContext->bTrackUnknownEnumNames, true);
+	TGuardValue<bool> TrackImpersonatePropertiesScope(SerializeContext->bImpersonateProperties, true);
+	FSerializedPropertyPathScope SerializedObjectPath(SerializeContext, {"Struct"});
+
+	FTestInstanceDataObjectStructAlternate AltStructData;
+	AltStructData.Grain = ETestInstanceDataObjectGrainAlternate::Rye;
+	AltStructData.Fruit = ETestInstanceDataObjectFruitAlternate::Cherry;
+	AltStructData.Direction = ETestInstanceDataObjectDirectionAlternate::North | ETestInstanceDataObjectDirectionAlternate::West |
+		ETestInstanceDataObjectDirectionAlternate::Up | ETestInstanceDataObjectDirectionAlternate::Down;
+	AltStructData.GrainFromEnumClass = ETestInstanceDataObjectGrainAlternateEnumClass::Corn;
+	AltStructData.FruitFromNamespace = ETestInstanceDataObjectFruitAlternateNamespace::Orange;
+	AltStructData.GrainTypeChange = ETestInstanceDataObjectGrainAlternate::Corn;
+	AltStructData.FruitTypeChange = ETestInstanceDataObjectFruitAlternate::Orange;
+	AltStructData.DeletedGrain = ETestInstanceDataObjectGrainAlternate::Rice;
+	AltStructData.DeletedFruit = ETestInstanceDataObjectFruitAlternate::Apple;
+	AltStructData.DeletedDirection = ETestInstanceDataObjectDirectionAlternate::South | ETestInstanceDataObjectDirectionAlternate::Up;
+
+	TArray<uint8> BinaryData;
+	{
+		FMemoryWriter Ar(BinaryData, /*bIsPersistent*/ true);
+		FBinaryArchiveFormatter Formatter(Ar);
+		FStructuredArchive StructuredAr(Formatter);
+		FTestInstanceDataObjectStructAlternate::StaticStruct()->SerializeTaggedProperties(StructuredAr.Open(), (uint8*)&AltStructData, nullptr, nullptr);
+	}
+
+#if WITH_TEXT_ARCHIVE_SUPPORT
+	TArray<uint8> JsonData;
+	{
+		FMemoryWriter Ar(JsonData, /*bIsPersistent*/ true);
+		FJsonArchiveOutputFormatter Formatter(Ar);
+		FStructuredArchive StructuredAr(Formatter);
+		FTestInstanceDataObjectStructAlternate::StaticStruct()->SerializeTaggedProperties(StructuredAr.Open(), (uint8*)&AltStructData, nullptr, nullptr);
+	}
+#endif // WITH_TEXT_ARCHIVE_SUPPORT
+
+	LOG_SCOPE_VERBOSITY_OVERRIDE(LogClass, ELogVerbosity::Error);
+	LOG_SCOPE_VERBOSITY_OVERRIDE(LogEnum, ELogVerbosity::Error);
+
+	const FPropertyTypeName GrainTypeName = MakePropertyTypeName(StaticEnum<ETestInstanceDataObjectGrain::Type>());
+	const FPropertyTypeName FruitTypeName = MakePropertyTypeName(StaticEnum<ETestInstanceDataObjectFruit>());
+	const FPropertyTypeName DirectionTypeName = MakePropertyTypeName(StaticEnum<ETestInstanceDataObjectDirection>());
+
+	FPropertyBagRepository& Repo = FPropertyBagRepository::Get();
+	TArray<FName> Names{NAME_None};
+	bool bHasFlags = false;
+
+	FTestInstanceDataObjectStruct StructData;
+
+	{
+		FMemoryReader Ar(BinaryData, /*bIsPersistent*/ true);
+		FBinaryArchiveFormatter Formatter(Ar);
+		FStructuredArchive StructuredAr(Formatter);
+		FTestInstanceDataObjectStruct::StaticStruct()->SerializeTaggedProperties(StructuredAr.Open(), (uint8*)&StructData, nullptr, nullptr);
+	}
+
+	CHECK(StructData.Grain == (ETestInstanceDataObjectGrain::Type)((uint8)ETestInstanceDataObjectGrain::Wheat + 1));
+	CHECK(StructData.Fruit == (ETestInstanceDataObjectFruit)((uint8)ETestInstanceDataObjectFruit::Orange + 1));
+	CHECK(StructData.Direction == (ETestInstanceDataObjectDirection)MAX_uint16);
+	// TODO: ConvertFromType incorrectly assigns numeric values based on the saved enum type.
+	//CHECK(StructData.GrainFromEnumClass == ETestInstanceDataObjectGrain::Corn);
+	//CHECK(StructData.FruitFromNamespace == ETestInstanceDataObjectFruit::Orange);
+	CHECK(StructData.GrainTypeChange == ETestInstanceDataObjectGrain::Corn);
+	CHECK(StructData.FruitTypeChange == ETestInstanceDataObjectFruit::Orange);
+
+#if WITH_METADATA
+	Repo.FindUnknownEnumNames(Owner, GrainTypeName, Names, bHasFlags);
+	CHECKED_IF(Names.Num() == 1)
+	{
+		CHECK(Names[0] == "Rye");
+	}
+	CHECK_FALSE(bHasFlags);
+
+	Repo.FindUnknownEnumNames(Owner, FruitTypeName, Names, bHasFlags);
+	CHECKED_IF(Names.Num() == 1)
+	{
+		CHECK(Names[0] == "Cherry");
+	}
+	CHECK_FALSE(bHasFlags);
+
+	Repo.FindUnknownEnumNames(Owner, DirectionTypeName, Names, bHasFlags);
+	CHECKED_IF(Names.Num() == 2)
+	{
+		CHECK(Names[0] == "Up");
+		CHECK(Names[1] == "Down");
+	}
+	CHECK(bHasFlags);
+
+	Repo.FindUnknownEnumNames(Owner, ParsePropertyTypeName(TEXT("ETestInstanceDataObjectDeletedGrain(/Script/CoreUObject)")), Names, bHasFlags);
+	CHECKED_IF(Names.Num() == 1)
+	{
+		CHECK(Names[0] == "Rice");
+	}
+	CHECK_FALSE(bHasFlags);
+
+	Repo.FindUnknownEnumNames(Owner, ParsePropertyTypeName(TEXT("ETestInstanceDataObjectDeletedFruit(/Script/CoreUObject)")), Names, bHasFlags);
+	CHECKED_IF(Names.Num() == 1)
+	{
+		CHECK(Names[0] == "Apple");
+	}
+	CHECK_FALSE(bHasFlags);
+
+	Repo.FindUnknownEnumNames(Owner, ParsePropertyTypeName(TEXT("ETestInstanceDataObjectDeletedDirection(/Script/CoreUObject)")), Names, bHasFlags);
+	CHECKED_IF(Names.Num() == 2)
+	{
+		CHECK(Names[0] == "Up");
+		CHECK(Names[1] == "South");
+	}
+	CHECK(bHasFlags);
+#endif // WITH_METADATA
+
+	FPropertyBagRepository::Get().DestroyOuterBag(Owner);
+
+#if WITH_TEXT_ARCHIVE_SUPPORT
+	StructData = {};
+
+	{
+		FMemoryReader Ar(JsonData, /*bIsPersistent*/ true);
+		FJsonArchiveInputFormatter Formatter(Ar);
+		FStructuredArchive StructuredAr(Formatter);
+		FTestInstanceDataObjectStruct::StaticStruct()->SerializeTaggedProperties(StructuredAr.Open(), (uint8*)&StructData, nullptr, nullptr);
+	}
+
+	CHECK(StructData.Grain == (ETestInstanceDataObjectGrain::Type)((uint8)ETestInstanceDataObjectGrain::Wheat + 1));
+	CHECK(StructData.Fruit == (ETestInstanceDataObjectFruit)((uint8)ETestInstanceDataObjectFruit::Orange + 1));
+	CHECK(StructData.Direction == (ETestInstanceDataObjectDirection)MAX_uint16);
+	// TODO: ConvertFromType incorrectly assigns numeric values based on the saved enum type.
+	//CHECK(StructData.GrainFromEnumClass == ETestInstanceDataObjectGrain::Corn);
+	//CHECK(StructData.FruitFromNamespace == ETestInstanceDataObjectFruit::Orange);
+	CHECK(StructData.GrainTypeChange == ETestInstanceDataObjectGrain::Corn);
+	CHECK(StructData.FruitTypeChange == ETestInstanceDataObjectFruit::Orange);
+
+#if WITH_METADATA
+	Repo.FindUnknownEnumNames(Owner, GrainTypeName, Names, bHasFlags);
+	CHECKED_IF(Names.Num() == 1)
+	{
+		CHECK(Names[0] == "Rye");
+	}
+	CHECK_FALSE(bHasFlags);
+
+	Repo.FindUnknownEnumNames(Owner, FruitTypeName, Names, bHasFlags);
+	CHECKED_IF(Names.Num() == 1)
+	{
+		CHECK(Names[0] == "Cherry");
+	}
+	CHECK_FALSE(bHasFlags);
+
+	Repo.FindUnknownEnumNames(Owner, DirectionTypeName, Names, bHasFlags);
+	CHECKED_IF(Names.Num() == 2)
+	{
+		CHECK(Names[0] == "Up");
+		CHECK(Names[1] == "Down");
+	}
+	CHECK(bHasFlags);
+
+	// Testing of the unknown property tree is skipped because it is not supported by the text format.
+#endif // WITH_METADATA
+
+	FPropertyBagRepository::Get().DestroyOuterBag(Owner);
+#endif // WITH_TEXT_ARCHIVE_SUPPORT
+}
+
 TEST_CASE_NAMED(FUnknownEnumNamesTest, "CoreUObject::Serialization::UnknownEnumNames", "[CoreUObject][EngineFilter]")
 {
 	UObject* Owner = NewObject<UTestInstanceDataObjectClass>();
diff --git a/Engine/Source/Runtime/CoreUObject/Tests/UObject/InstanceDataObjectUtilsTest.h b/Engine/Source/Runtime/CoreUObject/Tests/UObject/InstanceDataObjectUtilsTest.h
--- a/Engine/Source/Runtime/CoreUObject/Tests/UObject/InstanceDataObjectUtilsTest.h
+++ b/Engine/Source/Runtime/CoreUObject/Tests/UObject/InstanceDataObjectUtilsTest.h
@@ -46,12 +46,26 @@ namespace ETestInstanceDataObjectGrainAlternate
 	};
 }
 
+UENUM()
+enum class ETestInstanceDataObjectGrainAlternateEnumClass : uint8
+{
+	None = 0,
+	Corn,
+	Rice,
+	Rye,
+	Wheat,
+};
+
+static_assert((uint8)ETestInstanceDataObjectGrain::Corn != (uint8)ETestInstanceDataObjectGrainAlternate::Corn);
+static_assert((uint8)ETestInstanceDataObjectGrain::Corn != (uint8)ETestInstanceDataObjectGrainAlternateEnumClass::Corn);
+
 UENUM()
 enum class ETestInstanceDataObjectFruit : uint8
 {
 	None = 0,
 	Apple,
 	Banana,
+	Lemon,
 	Orange,
 };
 
@@ -65,6 +79,22 @@ enum class ETestInstanceDataObjectFruitAlternate : uint8
 	Pear,
 };
 
+UENUM()
+namespace ETestInstanceDataObjectFruitAlternateNamespace
+{
+	enum Type : uint8
+	{
+		None = 0,
+		Apple,
+		Cherry,
+		Orange,
+		Pear,
+	};
+}
+
+static_assert((uint8)ETestInstanceDataObjectFruit::Orange != (uint8)ETestInstanceDataObjectFruitAlternate::Orange);
+static_assert((uint8)ETestInstanceDataObjectFruit::Orange != (uint8)ETestInstanceDataObjectFruitAlternateNamespace::Orange);
+
 UENUM(Flags)
 enum class ETestInstanceDataObjectDirection : uint16
 {
@@ -156,6 +186,24 @@ public:
 	UPROPERTY()
 	ETestInstanceDataObjectDirection Direction = ETestInstanceDataObjectDirection::None;
 
+	UPROPERTY()
+	TEnumAsByte<ETestInstanceDataObjectGrain::Type> GrainFromEnumClass = ETestInstanceDataObjectGrain::None;
+
+	UPROPERTY()
+	ETestInstanceDataObjectFruit FruitFromNamespace = ETestInstanceDataObjectFruit::None;
+
+	UPROPERTY()
+	TEnumAsByte<ETestInstanceDataObjectGrain::Type> GrainTypeChange = ETestInstanceDataObjectGrain::None;
+
+	UPROPERTY()
+	ETestInstanceDataObjectFruit FruitTypeChange = ETestInstanceDataObjectFruit::None;
+
+	UPROPERTY()
+	TEnumAsByte<ETestInstanceDataObjectGrain::Type> GrainTypeAndPropertyChange = ETestInstanceDataObjectGrain::None;
+
+	UPROPERTY()
+	ETestInstanceDataObjectFruit FruitTypeAndPropertyChange = ETestInstanceDataObjectFruit::None;
+
 	UPROPERTY()
 	FTestInstanceDataObjectPoint Point;
 };
@@ -190,6 +238,33 @@ public:
 	UPROPERTY(meta=(OriginalType="ETestInstanceDataObjectDirection(/Script/CoreUObject)"))
 	ETestInstanceDataObjectDirectionAlternate Direction = ETestInstanceDataObjectDirectionAlternate::None;
 
+	UPROPERTY(meta=(OriginalType="ETestInstanceDataObjectGrain(/Script/CoreUObject)"))
+	ETestInstanceDataObjectGrainAlternateEnumClass GrainFromEnumClass = ETestInstanceDataObjectGrainAlternateEnumClass::None;
+
+	UPROPERTY(meta=(OriginalType="ETestInstanceDataObjectFruit(/Script/CoreUObject)"))
+	TEnumAsByte<ETestInstanceDataObjectFruitAlternateNamespace::Type> FruitFromNamespace = ETestInstanceDataObjectFruitAlternateNamespace::None;
+
+	UPROPERTY()
+	TEnumAsByte<ETestInstanceDataObjectGrainAlternate::Type> GrainTypeChange = ETestInstanceDataObjectGrainAlternate::None;
+
+	UPROPERTY()
+	ETestInstanceDataObjectFruitAlternate FruitTypeChange = ETestInstanceDataObjectFruitAlternate::None;
+
+	UPROPERTY()
+	ETestInstanceDataObjectGrainAlternateEnumClass GrainTypeAndPropertyChange = ETestInstanceDataObjectGrainAlternateEnumClass::None;
+
+	UPROPERTY()
+	TEnumAsByte<ETestInstanceDataObjectFruitAlternateNamespace::Type> FruitTypeAndPropertyChange = ETestInstanceDataObjectFruitAlternateNamespace::None;
+
+	UPROPERTY(meta=(OriginalType="ETestInstanceDataObjectDeletedGrain(/Script/CoreUObject)"))
+	TEnumAsByte<ETestInstanceDataObjectGrainAlternate::Type> DeletedGrain = ETestInstanceDataObjectGrainAlternate::None;
+
+	UPROPERTY(meta=(OriginalType="ETestInstanceDataObjectDeletedFruit(/Script/CoreUObject)"))
+	ETestInstanceDataObjectFruitAlternate DeletedFruit = ETestInstanceDataObjectFruitAlternate::None;
+
+	UPROPERTY(meta=(OriginalType="ETestInstanceDataObjectDeletedDirection(/Script/CoreUObject)"))
+	ETestInstanceDataObjectDirectionAlternate DeletedDirection = ETestInstanceDataObjectDirectionAlternate::None;
+
 	UPROPERTY(meta=(OriginalType="TestInstanceDataObjectPoint(/Script/CoreUObject)"))
 	FTestInstanceDataObjectPointAlternate Point;
 };
