Suppressed LogEnum warnings in the test TrackUnknownProperties

diff --git a/Engine/Source/Runtime/CoreUObject/Tests/UObject/InstanceDataObjectUtilsTest.cpp b/Engine/Source/Runtime/CoreUObject/Tests/UObject/InstanceDataObjectUtilsTest.cpp
--- a/Engine/Source/Runtime/CoreUObject/Tests/UObject/InstanceDataObjectUtilsTest.cpp
+++ b/Engine/Source/Runtime/CoreUObject/Tests/UObject/InstanceDataObjectUtilsTest.cpp
@@ -262,6 +262,7 @@ TEST_CASE_NAMED(FTrackUnknownPropertiesTest, "CoreUObject::Serialization::TrackU
 #endif // WITH_TEXT_ARCHIVE_SUPPORT
 
 	LOG_SCOPE_VERBOSITY_OVERRIDE(LogClass, ELogVerbosity::Error);
+	LOG_SCOPE_VERBOSITY_OVERRIDE(LogEnum, ELogVerbosity::Error);
 
 	FTestInstanceDataObjectStruct StructData;
 
