[Partial Rollback] 35155046 naively changed the behavior of UE::FindOriginalTypeName. Rolling back and instead targeting the behavior change for exclusively Overridable serialization where the change was needed

diff --git a/Engine/Source/Runtime/CoreUObject/Private/UObject/OverriddenPropertySet.cpp b/Engine/Source/Runtime/CoreUObject/Private/UObject/OverriddenPropertySet.cpp
--- a/Engine/Source/Runtime/CoreUObject/Private/UObject/OverriddenPropertySet.cpp
+++ b/Engine/Source/Runtime/CoreUObject/Private/UObject/OverriddenPropertySet.cpp
@@ -7,6 +7,7 @@
 #include "UObject/PropertyOptional.h"
 #include "Misc/ScopeExit.h"
 #include "Serialization/StructuredArchiveNameHelpers.h"
+#include "UObject/UObjectThreadContext.h"
 
 /*
  *************************************************************************************
@@ -98,7 +99,14 @@ FOverriddenPropertyNodeID::FOverriddenPropertyNodeID(const FProperty* Property)
 	{
 		// append typename to the end of the property ID
 		UE::FPropertyTypeNameBuilder TypeNameBuilder;
-		Property->SaveTypeName(TypeNameBuilder);
+#if WITH_EDITORONLY_DATA
+		{
+			// use property impersonation for SaveTypeName so that keys don't change when classes die
+			FUObjectSerializeContext* SerializeContext = FUObjectThreadContext::Get().GetSerializeContext();
+            TGuardValue<bool> ScopedImpersonateProperties(SerializeContext->bImpersonateProperties, true);
+            Property->SaveTypeName(TypeNameBuilder);
+		}
+#endif
 		UE::FPropertyTypeName TypeName = TypeNameBuilder.Build();
 		TStringBuilder<256> StringBuilder;
 		StringBuilder << Property->GetFName();
diff --git a/Engine/Source/Runtime/CoreUObject/Private/UObject/PropertyHelper.cpp b/Engine/Source/Runtime/CoreUObject/Private/UObject/PropertyHelper.cpp
--- a/Engine/Source/Runtime/CoreUObject/Private/UObject/PropertyHelper.cpp
+++ b/Engine/Source/Runtime/CoreUObject/Private/UObject/PropertyHelper.cpp
@@ -304,6 +304,11 @@ const FString* FindOriginalTypeName(const UField* Field)
 
 const FString* FindOriginalTypeName(const FProperty* Property)
 {
+	if (!FUObjectThreadContext::Get().GetSerializeContext()->bImpersonateProperties)
+	{
+		return nullptr;
+	}
+	
 	// Prioritize metadata on the property over metadata on the type.
 	if (const FString* OriginalType = Property->FindMetaData(NAME_OriginalType))
 	{
@@ -466,15 +471,12 @@ static bool FindRedirectForProperty(FPropertyTypeName OldType, FPropertyTypeName
 		// Compare against the original type of an impersonated struct.
 		if constexpr (std::is_same_v<PropertyType, FStructProperty>)
 		{
-			if (FUObjectThreadContext::Get().GetSerializeContext()->bImpersonateProperties)
+			if (const FPropertyTypeName OriginalType = FindOriginalType(Property); !OriginalType.IsEmpty())
 			{
-				if (const FPropertyTypeName OriginalType = FindOriginalType(Property); !OriginalType.IsEmpty())
+				FieldName = OriginalType.GetName();
+				if (FieldName == OldNameRedirect.ObjectName)
 				{
-					FieldName = OriginalType.GetName();
-					if (FieldName == OldNameRedirect.ObjectName)
-					{
-						return false;
-					}
+					return false;
 				}
 			}
 		}
