- Restores the previous FPublicExportMap implementation to be used when AGGRESSIVE_MEMORY_SAVING is enabled as it allocates optimally for it's initial load and trades performance when storing and removing from the map. Saves 8MB  of runtime memory in 31.00
- Codepaths iterating over values in the map have been commonized to use a new FValueIterator which allows for iterating over either implementation's values

#rnx

diff --git a/Engine/Source/Runtime/CoreUObject/Private/Serialization/AsyncLoading2.cpp b/Engine/Source/Runtime/CoreUObject/Private/Serialization/AsyncLoading2.cpp
--- a/Engine/Source/Runtime/CoreUObject/Private/Serialization/AsyncLoading2.cpp
+++ b/Engine/Source/Runtime/CoreUObject/Private/Serialization/AsyncLoading2.cpp
@@ -685,6 +685,268 @@ private:
 
 	class FPublicExportMap
 	{
+// The AGGRESSIVE_MEMORY_SAVING implementation allocates only as many keys and values as necessary to load a package's exports
+// which is space efficient but results in more memory traffic as exports are Store()'d or Remove()'d.
+// The !AGGRESSIVE_MEMORY_SAVING implementation uses more memory but provides O(1) Store() and Remove()
+#if AGGRESSIVE_MEMORY_SAVING
+	public:
+		static const int32 InvalidValue = -1;
+
+		FPublicExportMap() = default;
+
+		FPublicExportMap(const FPublicExportMap&) = delete;
+
+		FPublicExportMap(FPublicExportMap&& Other)
+		{
+			Allocation = Other.Allocation;
+			Count = Other.Count;
+			SingleItemValue = Other.SingleItemValue;
+			Other.Allocation = nullptr;
+			Other.Count = 0;
+		};
+
+		FPublicExportMap& operator=(const FPublicExportMap&) = delete;
+
+		FPublicExportMap& operator=(FPublicExportMap&& Other)
+		{
+			if (Count > 1)
+			{
+				FMemory::Free(Allocation);
+			}
+			Allocation = Other.Allocation;
+			Count = Other.Count;
+			SingleItemValue = Other.SingleItemValue;
+			Other.Allocation = nullptr;
+			Other.Count = 0;
+			return *this;
+		}
+
+		~FPublicExportMap()
+		{
+			if (Count > 1)
+			{
+				FMemory::Free(Allocation);
+			}
+		}
+		void Reserve(int32 NewCount)
+		{
+			if (NewCount <= Count)
+			{
+				return;
+			}
+			if (NewCount > 1)
+			{
+				TArrayView<uint64> OldKeys = GetKeys();
+				TArrayView<int32> OldValues = GetValues();
+				const uint64 OldKeysSize = Count * sizeof(uint64);
+				const uint64 NewKeysSize = NewCount * sizeof(uint64);
+				const uint64 OldValuesSize = Count * sizeof(int32);
+				const uint64 NewValuesSize = NewCount * sizeof(int32);
+				const uint64 KeysToAddSize = NewKeysSize - OldKeysSize;
+				const uint64 ValuesToAddSize = NewValuesSize - OldValuesSize;
+
+				uint8* NewAllocation = reinterpret_cast<uint8*>(FMemory::Malloc(NewKeysSize + NewValuesSize));
+				FMemory::Memzero(NewAllocation, KeysToAddSize); // Insert new keys initialized to zero
+				FMemory::Memcpy(NewAllocation + KeysToAddSize, OldKeys.GetData(), OldKeysSize); // Copy old keys
+				FMemory::Memset(NewAllocation + NewKeysSize, 0xFF, ValuesToAddSize); // Insert new values initialized to InvalidValue
+				FMemory::Memcpy(NewAllocation + NewKeysSize + ValuesToAddSize, OldValues.GetData(), OldValuesSize); // Copy old values
+				if (Count > 1)
+				{
+					FMemory::Free(Allocation);
+				}
+				Allocation = NewAllocation;
+			}
+			Count = NewCount;
+		}
+
+		void Store(uint64 ExportHash, UObject* Object)
+		{
+			TArrayView<uint64> Keys = GetKeys();
+			TArrayView<int32> Values = GetValues();
+			int32 Index = Algo::LowerBound(Keys, ExportHash);
+			if (Index < Count && Keys[Index] == ExportHash)
+			{
+				// Slot already exists so reuse it
+				Values[Index] = GUObjectArray.ObjectToIndex(Object);
+				return;
+			}
+			if (Count == 0 || Keys[0] != 0)
+			{
+				// No free slots so we need to add one (will be inserted at the beginning of the array)
+				Reserve(Count + 1);
+				Keys = GetKeys();
+				Values = GetValues();
+			}
+			else
+			{
+				--Index; // Update insertion index to one before the lower bound item
+			}
+			if (Index > 0)
+			{
+				// Move items down
+				FMemory::Memmove(Keys.GetData(), Keys.GetData() + 1, Index * sizeof(uint64));
+				FMemory::Memmove(Values.GetData(), Values.GetData() + 1, Index * sizeof(int32));
+			}
+			Keys[Index] = ExportHash;
+			Values[Index] = GUObjectArray.ObjectToIndex(Object);
+		}
+
+		bool Remove(uint64 ExportHash)
+		{
+			TArrayView<uint64> Keys = GetKeys();
+			int32 Index = Algo::LowerBound(Keys, ExportHash);
+			if (Index < Count && Keys[Index] == ExportHash)
+			{
+				TArrayView<int32> Values = GetValues();
+				Values[Index] = InvalidValue;
+				return true;
+			}
+
+			return false;
+		}
+
+		UObject* Find(uint64 ExportHash)
+		{
+			TArrayView<uint64> Keys = GetKeys();
+			int32 Index = Algo::LowerBound(Keys, ExportHash);
+			if (Index < Count && Keys[Index] == ExportHash)
+			{
+				TArrayView<int32> Values = GetValues();
+				int32 ObjectIndex = Values[Index];
+				if (ObjectIndex >= 0)
+				{
+					return static_cast<UObject*>(GUObjectArray.IndexToObject(ObjectIndex)->Object);
+				}
+			}
+			return nullptr;
+		}
+
+		[[nodiscard]] bool PinForGC(TArray<int32>& OutUnreachableObjectIndices)
+		{
+			OutUnreachableObjectIndices.Reset();
+			for (int32& ObjectIndex : GetValues())
+			{
+				if (ObjectIndex >= 0)
+				{
+					FUObjectItem* ObjectItem = GUObjectArray.IndexToObject(ObjectIndex);
+					if (!ObjectItem->IsUnreachable())
+					{
+						UObject* Object = static_cast<UObject*>(ObjectItem->Object);
+						checkf(!ObjectItem->HasAnyFlags(EInternalObjectFlags::LoaderImport), TEXT("%s"), *Object->GetFullName());
+						ObjectItem->SetFlags(EInternalObjectFlags::LoaderImport);
+					}
+					else
+					{
+						OutUnreachableObjectIndices.Reserve(Count);
+						OutUnreachableObjectIndices.Add(ObjectIndex);
+						ObjectIndex = InvalidValue;
+					}
+				}
+			}
+			return OutUnreachableObjectIndices.Num() == 0;
+		}
+
+		void UnpinForGC()
+		{
+			for (int32 ObjectIndex : GetValues())
+			{
+				if (ObjectIndex >= 0)
+				{
+					UObject* Object = static_cast<UObject*>(GUObjectArray.IndexToObject(ObjectIndex)->Object);
+					checkf(Object->HasAnyInternalFlags(EInternalObjectFlags::LoaderImport), TEXT("%s"), *Object->GetFullName());
+					Object->AtomicallyClearInternalFlags(EInternalObjectFlags::LoaderImport);
+				}
+			}
+		}
+
+		TArrayView<uint64> GetKeys()
+		{
+			if (Count == 1)
+			{
+				return MakeArrayView(&SingleItemKey, 1);
+			}
+			else
+			{
+				return MakeArrayView(reinterpret_cast<uint64*>(Allocation), Count);
+			}
+		}
+
+		TArrayView<int32> GetValues()
+		{
+			if (Count == 1)
+			{
+				return MakeArrayView(&SingleItemValue, 1);
+			}
+			else
+			{
+				return MakeArrayView(reinterpret_cast<int32*>(Allocation + Count * sizeof(uint64)), Count);
+			}
+		}
+
+		class FValueIterator
+		{
+			TArrayView<int32> Array;
+			int32 Pos;
+		public:
+			FValueIterator(const TArrayView<int32>& Data) : Array(Data), Pos(0) { }
+
+			FORCEINLINE bool operator== (const FValueIterator& Other) const
+			{
+				return Pos == Other.Pos && Array.GetData() == Other.Array.GetData();
+			}
+
+			FORCEINLINE bool operator!= (const FValueIterator& Other) const
+			{
+				return !(Pos == Other.Pos && Array.GetData() == Other.Array.GetData());
+			}
+
+			FORCEINLINE explicit operator bool() const
+			{
+				return Pos < Array.Num();
+			}
+
+			FORCEINLINE FValueIterator& operator++()
+			{
+				const int32 ArraySize = Array.Num();
+				do
+				{
+					++Pos;
+				} while (Pos < ArraySize && Array[Pos] != FPublicExportMap::InvalidValue);
+				return *this;
+			}
+
+			FORCEINLINE int32 operator*()
+			{
+				return Array[Pos];
+			}
+
+			FORCEINLINE void RemoveCurrent()
+			{
+				Array[Pos] = FPublicExportMap::InvalidValue;
+			}
+		};
+
+		FValueIterator CreateValueIterator()
+		{
+			if (Count == 1)
+			{
+				return FValueIterator{ MakeArrayView(&SingleItemValue, 1) };
+			}
+			else
+			{
+				return FValueIterator{ MakeArrayView(reinterpret_cast<int32*>(Allocation + Count * sizeof(uint64)), Count) };
+			}
+		}
+
+	private:
+		union
+		{
+			uint8* Allocation = nullptr;
+			uint64 SingleItemKey;
+		};
+		int32 Count = 0;
+		int32 SingleItemValue = -1;
+#else
 	public:
 		typedef TMap<uint64, int32> FObjectIndexMap;
 
@@ -716,9 +978,9 @@ private:
 		[[nodiscard]] bool PinForGC(TArray<int32>& OutUnreachableObjectIndices)
 		{
 			OutUnreachableObjectIndices.Reset();
-			for (FObjectIndexMap::TIterator It(Data); It; ++It)
+			for (FValueIterator It(Data); It; ++It)
 			{
-				FUObjectItem* ObjectItem = GUObjectArray.IndexToObject(It->Value);
+				FUObjectItem* ObjectItem = GUObjectArray.IndexToObject(*It);
 				if (!ObjectItem->IsUnreachable())
 				{
 					UObject* Object = static_cast<UObject*>(ObjectItem->Object);
@@ -728,7 +990,7 @@ private:
 				else
 				{
 					OutUnreachableObjectIndices.Reserve(Data.Num());
-					OutUnreachableObjectIndices.Add(It->Value);
+					OutUnreachableObjectIndices.Add(*It);
 					It.RemoveCurrent();
 				}
 			}
@@ -738,21 +1000,60 @@ private:
 
 		void UnpinForGC()
 		{
-			for (FObjectIndexMap::TConstIterator It(Data); It; ++It)
+			for (FValueIterator It(Data); It; ++It)
 			{
-				UObject* Object = static_cast<UObject*>(GUObjectArray.IndexToObject(It->Value)->Object);
+				UObject* Object = static_cast<UObject*>(GUObjectArray.IndexToObject(*It)->Object);
 				checkf(Object->HasAnyInternalFlags(EInternalObjectFlags::LoaderImport), TEXT("%s"), *Object->GetFullName());
 				Object->AtomicallyClearInternalFlags(EInternalObjectFlags::LoaderImport);
 			}
 		}
 
-		FObjectIndexMap::TConstIterator ConstIterator() const
+		class FValueIterator
 		{
-			return FObjectIndexMap::TConstIterator(Data);
+			FObjectIndexMap::TIterator It;
+		public:
+			FValueIterator(FObjectIndexMap& Data) : It(Data.CreateIterator()) { }
+
+			FORCEINLINE bool operator== (const FObjectIndexMap::TIterator& Other) const
+			{
+				return It == Other;
+			}
+
+			FORCEINLINE bool operator!= (const FObjectIndexMap::TIterator& Other) const
+			{
+				return !(It == Other);
+			}
+
+			FORCEINLINE explicit operator bool() const
+			{
+				return It.operator bool();
+			}
+
+			FORCEINLINE FValueIterator& operator++()
+			{
+				++It;
+				return *this;
+			}
+
+			FORCEINLINE int32 operator*()
+			{
+				return It.Value();
+			}
+
+			FORCEINLINE void RemoveCurrent()
+			{
+				It.RemoveCurrent();
+			}
+		};
+
+		FValueIterator CreateValueIterator()
+		{
+			return FValueIterator { Data };
 		}
 
 	private:
 		FObjectIndexMap Data;
+#endif
 	};
 
 	FPublicExportMap PublicExportMap;
@@ -869,9 +1170,9 @@ public:
 		bHasFailed = true;
 	}
 
-	FPublicExportMap::FObjectIndexMap::TConstIterator GetPublicExportObjectIndices() const
+	FPublicExportMap::FValueIterator GetPublicExportObjectIndices()
 	{
-		return PublicExportMap.ConstIterator();
+		return PublicExportMap.CreateValueIterator();
 	}
 
 	void ReserveSpaceForPublicExports(int32 PublicExportCount)
@@ -1166,9 +1467,10 @@ public:
 			VerifyPackageForRemoval(PackageRef);
 		}
 
-		for (auto It = PackageRef.PublicExportMap.ConstIterator(); It; ++It)
+		for (auto It = PackageRef.PublicExportMap.CreateValueIterator(); It; ++It)
 		{
-			ObjectIndexToPublicExport.Remove(It->Value);
+			const int32 ObjectIndex = *It;
+			ObjectIndexToPublicExport.Remove(ObjectIndex);
 		}
 
 		PackageRef.RemoveUnreferencedObsoletePackage();
@@ -1209,9 +1511,10 @@ public:
 		bool bRemoved = Packages.RemoveAndCopyValue(PackageId, PackageRef);
 		if (bRemoved)
 		{
-			for (auto It = PackageRef.PublicExportMap.ConstIterator(); It; ++It)
+			for (auto It = PackageRef.PublicExportMap.CreateValueIterator(); It; ++It)
 			{
-				ObjectIndexToPublicExport.Remove(It->Value);
+				const int32 ObjectIndex = *It;
+				ObjectIndexToPublicExport.Remove(ObjectIndex);
 			}
 		}
 	}
@@ -1303,7 +1606,7 @@ public:
 
 		for (auto It = PackageRef.GetPublicExportObjectIndices(); It; ++It)
 		{
-			const int32 ObjectIndex = It->Value;
+			const int32 ObjectIndex = *It;
 			UObject* Object = static_cast<UObject*>(GUObjectArray.IndexToObject(ObjectIndex)->Object);
 			ensureMsgf(!Object->HasAnyInternalFlags(EInternalObjectFlags::LoaderImport) || GUObjectArray.IsDisregardForGC(Object),
 					TEXT("FGlobalImportStore::VerifyPackageForRemoval: The loaded public export object '%s' with flags (ObjectFlags=%x, InternalObjectFlags=%x) and id %s is probably still referenced by the loader."),
@@ -5793,7 +6096,6 @@ bool FAsyncPackage2::CreateLinkerLoadExports(FAsyncLoadingThreadState2& ThreadSt
 
 	// Create exports
 	const int32 ExportCount = LinkerLoadState->Linker->ExportMap.Num();
-	FLoadedPackageRef& PackageRef = ImportStore.GlobalImportStore.FindPackageRefChecked(Desc.UPackageId, Desc.UPackageName);
 	while (LinkerLoadState->CreateExportIndex < ExportCount)
 	{
 		const int32 ExportIndex = LinkerLoadState->CreateExportIndex++;
@@ -5803,7 +6105,6 @@ bool FAsyncPackage2::CreateLinkerLoadExports(FAsyncLoadingThreadState2& ThreadSt
 			continue;
 		}
 #endif
-
 		FObjectExport& LinkerExport = LinkerLoadState->Linker->ExportMap[ExportIndex];
 		FExportObject& ExportObject = Data.Exports[ExportIndex];
 		if (UObject* Object = LinkerLoadState->Linker->CreateExport(ExportIndex))
