Split IsInstanceDataObjectSupportEnabled into generic and per-object queries to avoid incorrect reporting with object that were nulled due to unrecoverable serialization errors.

diff --git a/Engine/Source/Runtime/CoreUObject/Internal/UObject/InstanceDataObjectUtils.h b/Engine/Source/Runtime/CoreUObject/Internal/UObject/InstanceDataObjectUtils.h
--- a/Engine/Source/Runtime/CoreUObject/Internal/UObject/InstanceDataObjectUtils.h
+++ b/Engine/Source/Runtime/CoreUObject/Internal/UObject/InstanceDataObjectUtils.h
@@ -19,12 +19,10 @@ namespace UE { class FPropertyPathNameTree; }
 namespace UE
 {
 
-/**
- * Query if InstanceDataObject support is enabled for a specific object.
- *
- * Pass nullptr to query if the system is enabled.
- */
-bool IsInstanceDataObjectSupportEnabled(const UObject* Object = nullptr);
+/** Query if InstanceDataObject support is available generally. */
+bool IsInstanceDataObjectSupportEnabled();
+/** Query if InstanceDataObject support is enabled for a specific object. */
+bool IsInstanceDataObjectSupportEnabled(const UObject* Object);
 bool StructContainsLooseProperties(const UStruct* Struct);
 
 /** Helper to check if a class is an IDO class. */
diff --git a/Engine/Source/Runtime/CoreUObject/Private/UObject/InstanceDataObjectUtils.cpp b/Engine/Source/Runtime/CoreUObject/Private/UObject/InstanceDataObjectUtils.cpp
--- a/Engine/Source/Runtime/CoreUObject/Private/UObject/InstanceDataObjectUtils.cpp
+++ b/Engine/Source/Runtime/CoreUObject/Private/UObject/InstanceDataObjectUtils.cpp
@@ -157,29 +157,32 @@ namespace UE
 		return TSet<FString>(Result);
 	}
 
+	bool IsInstanceDataObjectSupportEnabled()
+	{
+		return bEnableIDOSupport;
+	}
+
 	bool IsInstanceDataObjectSupportEnabled(const UObject* InObject)
 	{
-		// Note: NULL is a valid (default) input here; in that case we just return the enable flag.
-		bool bIsEnabled = bEnableIDOSupport;
-		if (bIsEnabled && InObject)
+		if(!InObject)
 		{
-			// Property bag placeholder objects are always enabled for IDO support
-			if (UE::FPropertyBagRepository::IsPropertyBagPlaceholderObject(InObject))
-			{
-				return true;
-			}
-
-			//@todo FH: change to check trait when available or use config object
-			const UClass* ObjClass = InObject->GetClass();
-			while (ObjClass && ObjClass->GetClass()->GetFName() != NAME_VerseClass)
-			{
-				ObjClass = ObjClass->GetSuperClass();
-			}
+			return false;
+		}
+		
+		// Property bag placeholder objects are always enabled for IDO support
+		if (UE::FPropertyBagRepository::IsPropertyBagPlaceholderObject(InObject))
+		{
+			return true;
+		}
 
-			bIsEnabled = !!ObjClass;
+		//@todo FH: change to check trait when available or use config object
+		const UClass* ObjClass = InObject->GetClass();
+		while (ObjClass && ObjClass->GetClass()->GetFName() != NAME_VerseClass)
+		{
+			ObjClass = ObjClass->GetSuperClass();
 		}
 
-		return bIsEnabled;
+		return !!ObjClass;
 	}
 
 
