- Restores the previous FPublicExportMap implementation to be used when AGGRESSIVE_MEMORY_SAVING is enabled as it allocates optimally for it's initial load and trades performance when storing and removing from the map. Saves 8MB  of runtime memory in 31.00
- Codepaths iterating over values in the map have been commonized to use a new FValueIterator which allows for iterating over either implementation's values

#rnx

diff --git a/Engine/Source/Runtime/CoreUObject/Private/Serialization/AsyncLoading2.cpp b/Engine/Source/Runtime/CoreUObject/Private/Serialization/AsyncLoading2.cpp
--- a/Engine/Source/Runtime/CoreUObject/Private/Serialization/AsyncLoading2.cpp
+++ b/Engine/Source/Runtime/CoreUObject/Private/Serialization/AsyncLoading2.cpp
@@ -888,7 +888,16 @@ private:
 			TArrayView<int32> Array;
 			int32 Pos;
 		public:
-			FValueIterator(const TArrayView<int32>& Data) : Array(Data), Pos(0) { }
+			FValueIterator(const TArrayView<int32>& Data) : Array(Data), Pos(0)
+			{ 
+				const int32 ArraySize = Array.Num();
+
+				// Move to the first valid element
+				while (Pos < ArraySize && Array[Pos] == FPublicExportMap::InvalidValue)
+				{
+					++Pos;
+				}
+			}
 
 			FORCEINLINE bool operator== (const FValueIterator& Other) const
 			{
@@ -911,7 +920,7 @@ private:
 				do
 				{
 					++Pos;
-				} while (Pos < ArraySize && Array[Pos] != FPublicExportMap::InvalidValue);
+				} while (Pos < ArraySize && Array[Pos] == FPublicExportMap::InvalidValue);
 				return *this;
 			}
 
