[Linker]
 - Parallelize loading the bulk data in memory in ResetLoaders
 - 1m52s to 7.2s when doing a second PIE in AncientWorld after going in the dark world on first PIE

#rnx

diff --git a/Engine/Source/Runtime/CoreUObject/Private/UObject/LinkerManager.cpp b/Engine/Source/Runtime/CoreUObject/Private/UObject/LinkerManager.cpp
--- a/Engine/Source/Runtime/CoreUObject/Private/UObject/LinkerManager.cpp
+++ b/Engine/Source/Runtime/CoreUObject/Private/UObject/LinkerManager.cpp
@@ -9,6 +9,7 @@
 #include "UObject/ObjectResource.h"
 #include "UObject/LinkerLoad.h"
 #include "UObject/UObjectThreadContext.h"
+#include "Async/ParallelFor.h"
 #include "ProfilingDebugging/CsvProfiler.h"
 
 FLinkerManager& FLinkerManager::Get()
@@ -222,10 +223,20 @@ void FLinkerManager::ResetLoaders(const TSet<FLinkerLoad*>& InLinkerLoads)
 			}
 		}
 	}
+
+	// Parallelize loading the bulk data since that part is thread-safe
+	TArray<FLinkerLoad*> LinkerLoadsArray = InLinkerLoads.Array();
+	ParallelFor(LinkerLoadsArray.Num(),
+		[&LinkerLoadsArray](int32 Index)
+		{
+			LinkerLoadsArray[Index]->LoadAndDetachAllBulkData();
+		},
+		EParallelForFlags::Unbalanced
+	);
+
 	for (FLinkerLoad* LinkerToReset : InLinkerLoads)
 	{
 		// Detach linker, also removes from array and sets LinkerRoot to NULL.
-		LinkerToReset->LoadAndDetachAllBulkData();
 		LinkerToReset->Detach();
 	}
 	// Remove all linkers in the specified set
@@ -233,7 +244,7 @@ void FLinkerManager::ResetLoaders(const TSet<FLinkerLoad*>& InLinkerLoads)
 #if THREADSAFE_UOBJECTS
 		FScopeLock PendingCleanupListLock(&PendingCleanupListCritical);
 #endif
-		PendingCleanupList.Append(InLinkerLoads.Array());
+		PendingCleanupList.Append(LinkerLoadsArray);
 		bHasPendingCleanup = true;
 	}
 }
