Restored the 'non-null-return' attributes of non-debug CastChecked() and removed the explicit handling of null, which is not a legal return value.

diff --git a/Engine/Source/Runtime/CoreUObject/Public/Templates/Casts.h b/Engine/Source/Runtime/CoreUObject/Public/Templates/Casts.h
--- a/Engine/Source/Runtime/CoreUObject/Public/Templates/Casts.h
+++ b/Engine/Source/Runtime/CoreUObject/Public/Templates/Casts.h
@@ -244,16 +244,12 @@ FORCEINLINE T* ExactCast( UObject* Src )
 #else
 
 	template <typename To, typename From>
+	FUNCTION_NON_NULL_RETURN_START
 		FORCEINLINE To* CastChecked(From* Src)
+	FUNCTION_NON_NULL_RETURN_END
 	{
 		static_assert(sizeof(From) > 0 && sizeof(To) > 0, "Attempting to cast between incomplete types");
 
-		if (!Src)
-		{
-			// TODO: should this instead be an unreachable?
-			return nullptr;
-		}
-
 		if constexpr (TIsIInterface<From>::Value)
 		{
 			UObject* Obj = Src->_getUObject();
