Fixed TSoftObjectPtr being constructible/assignable from any UObject pointer type, not just convertible ones.

diff --git a/Engine/Source/Runtime/CoreUObject/Public/UObject/SoftObjectPtr.h b/Engine/Source/Runtime/CoreUObject/Public/UObject/SoftObjectPtr.h
--- a/Engine/Source/Runtime/CoreUObject/Public/UObject/SoftObjectPtr.h
+++ b/Engine/Source/Runtime/CoreUObject/Public/UObject/SoftObjectPtr.h
@@ -10,9 +10,20 @@
 #include "Templates/Casts.h"
 #include "UObject/PersistentObjectPtr.h"
 #include "UObject/SoftObjectPath.h"
+#include "Templates/Models.h"
+#include "Concepts/EqualityComparable.h"
 
 #include <type_traits>
 
+#ifndef UE_DEPRECATE_SOFTOBJECTPTR_CONVERSIONS
+	#define UE_DEPRECATE_SOFTOBJECTPTR_CONVERSIONS 1
+#endif
+#if UE_DEPRECATE_SOFTOBJECTPTR_CONVERSIONS
+	#define UE_SOFTOBJECTPTR_CONVERSION_DEPRECATED(Version, Message) UE_DEPRECATED(Version, Message)
+#else
+	#define UE_SOFTOBJECTPTR_CONVERSION_DEPRECATED(Version, Message)
+#endif
+
 /**
  * TIsSoftObjectPointerType
  * Trait for recognizing 'soft' (path-based) object pointer types
@@ -138,6 +149,15 @@ public:
 		: SoftObjectPtr(Other.SoftObjectPtr)
 	{
 	}
+	template <
+		class U
+		UE_REQUIRES(!std::is_convertible_v<U*, T*>)
+	>
+	UE_SOFTOBJECTPTR_CONVERSION_DEPRECATED(5.5, "Constructing TSoftObjectPtr from an incompatible pointer type has been deprecated.")
+	FORCEINLINE TSoftObjectPtr(const TSoftObjectPtr<U>& Other)
+		: SoftObjectPtr(Other.SoftObjectPtr)
+	{
+	}
 
 	/** Construct from a moveable soft pointer */
 	template <
@@ -148,16 +168,49 @@ public:
 		: SoftObjectPtr(MoveTemp(Other.SoftObjectPtr))
 	{
 	}
+	template <
+		class U
+		UE_REQUIRES(!std::is_convertible_v<U*, T*>)
+	>
+	UE_SOFTOBJECTPTR_CONVERSION_DEPRECATED(5.5, "Constructing TSoftObjectPtr from an incompatible pointer type has been deprecated.")
+	FORCEINLINE TSoftObjectPtr(TSoftObjectPtr<U>&& Other)
+		: SoftObjectPtr(MoveTemp(Other.SoftObjectPtr))
+	{
+	}
 
 	/** Construct from an object already in memory */
-	template <typename U>
-	FORCEINLINE TSoftObjectPtr(const U* Object)
+	template <
+		typename U
+		UE_REQUIRES(std::is_convertible_v<U*, T*>)
+	>
+	FORCEINLINE TSoftObjectPtr(U* Object)
+		: SoftObjectPtr(Object)
+	{
+	}
+	template <
+		typename U
+		UE_REQUIRES(!std::is_convertible_v<U*, T*>)
+	>
+	UE_SOFTOBJECTPTR_CONVERSION_DEPRECATED(5.5, "Constructing TSoftObjectPtr from an incompatible pointer type has been deprecated.")
+	FORCEINLINE TSoftObjectPtr(U* Object)
 		: SoftObjectPtr(Object)
 	{
 	}
 
 	/** Construct from a TObjectPtr<U> which may or may not be in memory. */
-	template <typename U>
+	template <
+		typename U
+		UE_REQUIRES(std::is_convertible_v<U*, T*>)
+	>
+	FORCEINLINE TSoftObjectPtr(const TObjectPtr<U> Object)
+		: SoftObjectPtr(Object.Get())
+	{
+	}
+	template <
+		typename U
+		UE_REQUIRES(!std::is_convertible_v<U*, T*>)
+	>
+	UE_SOFTOBJECTPTR_CONVERSION_DEPRECATED(5.5, "Constructing TSoftObjectPtr from an incompatible pointer type has been deprecated.")
 	FORCEINLINE TSoftObjectPtr(const TObjectPtr<U> Object)
 		: SoftObjectPtr(Object.Get())
 	{
@@ -170,8 +223,23 @@ public:
 	}
 
 	/** Construct from a soft object path */
-	explicit FORCEINLINE TSoftObjectPtr(FSoftObjectPath ObjectPath)
+	template <
+		typename SoftObjectPathType
+		UE_REQUIRES(std::is_same_v<SoftObjectPathType, FSoftObjectPath>)
+	>
+	explicit FORCEINLINE TSoftObjectPtr(SoftObjectPathType ObjectPath)
 		: SoftObjectPtr(MoveTemp(ObjectPath))
+	{
+		// The reason for the strange templated/constrained signature is to prevent this type hole:
+		//
+		// const UThing* ConstPtr = ...;
+		// TSoftObjectPtr<UThing*> NonConstSoftPtr(ConstPtr); // implicitly constructs an FSoftObjectPath from ConstPtr, then explicitly constructs the TSoftObjectPtr from that
+		//
+		// The real fix is to make the FSoftObjectPath(const UObject*) constructor explicit, but that will involve substantially more changes.
+	}
+	UE_SOFTOBJECTPTR_CONVERSION_DEPRECATED(5.5, "Constructing TSoftObjectPtr from an FString has been deprecated - instead, explicitly construct a FSoftObjectPath.")
+	explicit TSoftObjectPtr(const FString& Path)
+		: SoftObjectPtr(FSoftObjectPath(Path))
 	{
 	}
 
@@ -188,15 +256,41 @@ public:
 	}
 
 	/** Copy from an object already in memory */
-	template <typename U>
-	FORCEINLINE TSoftObjectPtr& operator=(const U* Object)
+	template <
+		typename U
+		UE_REQUIRES(std::is_convertible_v<U*, T*>)
+	>
+	FORCEINLINE TSoftObjectPtr& operator=(U* Object)
+	{
+		SoftObjectPtr = Object;
+		return *this;
+	}
+	template <
+		typename U
+		UE_REQUIRES(!std::is_convertible_v<U*, T*>)
+	>
+	UE_SOFTOBJECTPTR_CONVERSION_DEPRECATED(5.5, "Assigning TSoftObjectPtr from an incompatible pointer type has been deprecated.")
+	FORCEINLINE TSoftObjectPtr& operator=(U* Object)
 	{
 		SoftObjectPtr = Object;
 		return *this;
 	}
 
 	/** Copy from a TObjectPtr<U> which may or may not be in memory. */
-	template <typename U>
+	template <
+		typename U
+		UE_REQUIRES(std::is_convertible_v<U*, T*>)
+	>
+	FORCEINLINE TSoftObjectPtr& operator=(const TObjectPtr<U> Object)
+	{
+		SoftObjectPtr = Object.Get();
+		return *this;
+	}
+	template <
+		typename U
+		UE_REQUIRES(!std::is_convertible_v<U*, T*>)
+	>
+	UE_SOFTOBJECTPTR_CONVERSION_DEPRECATED(5.5, "Assigning TSoftObjectPtr from an incompatible pointer type has been deprecated.")
 	FORCEINLINE TSoftObjectPtr& operator=(const TObjectPtr<U> Object)
 	{
 		SoftObjectPtr = Object.Get();
@@ -227,6 +321,16 @@ public:
 		SoftObjectPtr = Other;
 		return *this;
 	}
+	template <
+		class U
+		UE_REQUIRES(!std::is_convertible_v<U*, T*>)
+	>
+	UE_SOFTOBJECTPTR_CONVERSION_DEPRECATED(5.5, "Assigning TSoftObjectPtr from an incompatible pointer type has been deprecated.")
+	FORCEINLINE TSoftObjectPtr& operator=(const TWeakObjectPtr<U>& Other)
+	{
+		SoftObjectPtr = Other;
+		return *this;
+	}
 
 	/** Copy from another soft pointer */
 	template <
@@ -238,6 +342,16 @@ public:
 		SoftObjectPtr = MoveTemp(Other.SoftObjectPtr);
 		return *this;
 	}
+	template <
+		class U
+		UE_REQUIRES(!std::is_convertible_v<U*, T*>)
+	>
+	UE_SOFTOBJECTPTR_CONVERSION_DEPRECATED(5.5, "Assigning TSoftObjectPtr from an incompatible pointer type has been deprecated.")
+	FORCEINLINE TSoftObjectPtr& operator=(TSoftObjectPtr<U> Other)
+	{
+		SoftObjectPtr = MoveTemp(Other.SoftObjectPtr);
+		return *this;
+	}
 
 	/**
 	 * Compare soft pointers for equality
@@ -249,41 +363,44 @@ public:
 	{
 		return SoftObjectPtr == Rhs.SoftObjectPtr;
 	}
-
-	/**
-	 * Compare soft pointers for equality
-	 * Caution: Two soft pointers might not be equal to each other, but they both might return nullptr
-	 *
-	 * @param Other soft pointer to compare to
-	 */
-	FORCEINLINE bool operator==(TYPE_OF_NULLPTR) const
+	template <
+		typename U
+		UE_REQUIRES(TModels_V<CEqualityComparableWith, T*, U*>)
+	>
+	FORCEINLINE bool operator==(const TSoftObjectPtr<U>& Rhs) const
 	{
-		return SoftObjectPtr == nullptr;
+		return SoftObjectPtr == Rhs.SoftObjectPtr;
 	}
-
-#if !PLATFORM_COMPILER_HAS_GENERATED_COMPARISON_OPERATORS
-	/**
-	 * Compare soft pointers for inequality
-	 * Caution: Two soft pointers might not be equal to each other, but they both might return nullptr
-	 *
-	 * @param Other soft pointer to compare to
-	 */
-	FORCEINLINE bool operator!=(const TSoftObjectPtr& Rhs) const
+	template <
+		typename U
+		UE_REQUIRES(!TModels_V<CEqualityComparableWith, T*, U*>)
+	>
+	UE_SOFTOBJECTPTR_CONVERSION_DEPRECATED(5.5, "Comparing TSoftObjectPtrs of incompatible pointer types has been deprecated.")
+	FORCEINLINE bool operator==(const TSoftObjectPtr<U>& Rhs) const
 	{
-		return SoftObjectPtr != Rhs.SoftObjectPtr;
+		return SoftObjectPtr == Rhs.SoftObjectPtr;
 	}
-
-	/**
-	 * Compare soft pointers for inequality
-	 * Caution: Two soft pointers might not be equal to each other, but they both might return nullptr
-	 *
-	 * @param Other soft pointer to compare to
-	 */
-	FORCEINLINE bool operator!=(TYPE_OF_NULLPTR) const
+	template <
+		typename U
+		UE_REQUIRES(TModels_V<CEqualityComparableWith, T*, U*>)
+	>
+	FORCEINLINE bool operator==(U* Rhs) const
 	{
-		return SoftObjectPtr != nullptr;
+		return SoftObjectPtr == TSoftObjectPtr<U>(Rhs).SoftObjectPtr;
+	}
+	template <
+		typename U
+		UE_REQUIRES(!TModels_V<CEqualityComparableWith, T*, U*>)
+	>
+	UE_SOFTOBJECTPTR_CONVERSION_DEPRECATED(5.5, "Comparing TSoftObjectPtrs of incompatible pointer types has been deprecated.")
+	FORCEINLINE bool operator==(U* Rhs) const
+	{
+		return SoftObjectPtr == TSoftObjectPtr<U>(Rhs).SoftObjectPtr;
+	}
+	FORCEINLINE bool operator==(TYPE_OF_NULLPTR) const
+	{
+		return SoftObjectPtr == nullptr;
 	}
-#endif
 
 	/**
 	 * Dereference the soft pointer.
@@ -393,29 +510,100 @@ private:
 	FSoftObjectPtr SoftObjectPtr;
 };
 
+template <typename T>
+TSoftObjectPtr(T*) -> TSoftObjectPtr<T>;
+
 #if !PLATFORM_COMPILER_HAS_GENERATED_COMPARISON_OPERATORS
-template<class T>
-FORCEINLINE bool operator==(TYPE_OF_NULLPTR, const TSoftObjectPtr<T>& Rhs)
+template <
+	typename LhsType,
+	typename RhsType
+	UE_REQUIRES(TModels_V<CEqualityComparableWith, LhsType*, RhsType*>)
+>
+FORCEINLINE bool operator==(LhsType* Lhs, const TSoftObjectPtr<RhsType>& Rhs)
 {
-	return Rhs == nullptr;
+	return Rhs == Lhs;
 }
-
-template<class T>
-FORCEINLINE bool operator!=(TYPE_OF_NULLPTR, const TSoftObjectPtr<T>& Rhs)
+template <
+	typename LhsType,
+	typename RhsType
+	UE_REQUIRES(!TModels_V<CEqualityComparableWith, LhsType*, RhsType*>)
+>
+UE_SOFTOBJECTPTR_CONVERSION_DEPRECATED(5.5, "Comparing TSoftObjectPtrs of incompatible pointer types has been deprecated.")
+FORCEINLINE bool operator==(LhsType* Lhs, const TSoftObjectPtr<RhsType>& Rhs)
 {
-	return Rhs != nullptr;
+	return Rhs == Lhs;
 }
-
-template<class T>
-FORCEINLINE bool operator==(const T* Lhs, const TSoftObjectPtr<T>& Rhs)
+template <typename RhsType>
+FORCEINLINE bool operator==(TYPE_OF_NULLPTR Lhs, const TSoftObjectPtr<RhsType>& Rhs)
 {
 	return Rhs == Lhs;
 }
-
-template<class T>
-FORCEINLINE bool operator!=(const T* Lhs, const TSoftObjectPtr<T>& Rhs)
+template <
+	typename LhsType,
+	typename RhsType
+	UE_REQUIRES(TModels_V<CEqualityComparableWith, LhsType*, RhsType*>)
+>
+FORCEINLINE bool operator!=(const TSoftObjectPtr<LhsType>& Lhs, const TSoftObjectPtr<RhsType>& Rhs)
 {
-	return Rhs != Lhs;
+	return !(Lhs == Rhs);
+}
+template <
+	typename LhsType,
+	typename RhsType
+	UE_REQUIRES(!TModels_V<CEqualityComparableWith, LhsType*, RhsType*>)
+>
+UE_SOFTOBJECTPTR_CONVERSION_DEPRECATED(5.5, "Comparing TSoftObjectPtrs of incompatible pointer types has been deprecated.")
+FORCEINLINE bool operator!=(const TSoftObjectPtr<LhsType>& Lhs, const TSoftObjectPtr<RhsType>& Rhs)
+{
+	return !(Lhs == Rhs);
+}
+template <
+	typename LhsType,
+	typename RhsType
+	UE_REQUIRES(TModels_V<CEqualityComparableWith, LhsType*, RhsType*>)
+>
+FORCEINLINE bool operator!=(const TSoftObjectPtr<LhsType>& Lhs, RhsType* Rhs)
+{
+	return !(Lhs == Rhs);
+}
+template <
+	typename LhsType,
+	typename RhsType
+	UE_REQUIRES(!TModels_V<CEqualityComparableWith, LhsType*, RhsType*>)
+>
+UE_SOFTOBJECTPTR_CONVERSION_DEPRECATED(5.5, "Comparing TSoftObjectPtrs of incompatible pointer types has been deprecated.")
+FORCEINLINE bool operator!=(const TSoftObjectPtr<LhsType>& Lhs, RhsType* Rhs)
+{
+	return !(Lhs == Rhs);
+}
+template <typename LhsType>
+FORCEINLINE bool operator!=(const TSoftObjectPtr<LhsType>& Lhs, TYPE_OF_NULLPTR Rhs)
+{
+	return !(Lhs == Rhs);
+}
+template <
+	typename LhsType,
+	typename RhsType
+	UE_REQUIRES(TModels_V<CEqualityComparableWith, LhsType*, RhsType*>)
+>
+FORCEINLINE bool operator!=(LhsType* Lhs, const TSoftObjectPtr<RhsType>& Rhs)
+{
+	return !(Lhs == Rhs);
+}
+template <
+	typename LhsType,
+	typename RhsType
+	UE_REQUIRES(!TModels_V<CEqualityComparableWith, LhsType*, RhsType*>)
+>
+UE_SOFTOBJECTPTR_CONVERSION_DEPRECATED(5.5, "Comparing TSoftObjectPtrs of incompatible pointer types has been deprecated.")
+FORCEINLINE bool operator!=(LhsType* Lhs, const TSoftObjectPtr<RhsType>& Rhs)
+{
+	return !(Lhs == Rhs);
+}
+template <typename RhsType>
+FORCEINLINE bool operator!=(TYPE_OF_NULLPTR Lhs, const TSoftObjectPtr<RhsType>& Rhs)
+{
+	return !(Lhs == Rhs);
 }
 #endif
 
@@ -446,17 +634,17 @@ struct TCallTraits<TSoftObjectPtr<T>> : public TCallTraitsBase<TSoftObjectPtr<T>
 
 /** Utility to create a TSoftObjectPtr without specifying the type */
 template <class T>
-TSoftObjectPtr<std::remove_cv_t<T>> MakeSoftObjectPtr(T* Object)
+TSoftObjectPtr<T> MakeSoftObjectPtr(T* Object)
 {
 	static_assert(std::is_base_of_v<UObject, T>, "Type must derive from UObject");
-	return TSoftObjectPtr<std::remove_cv_t<T>>(Object);
+	return TSoftObjectPtr<T>(Object);
 }
 
 template <class T>
-TSoftObjectPtr<std::remove_cv_t<T>> MakeSoftObjectPtr(TObjectPtr<T> Object)
+TSoftObjectPtr<T> MakeSoftObjectPtr(TObjectPtr<T> Object)
 {
 	static_assert(std::is_base_of_v<UObject, T>, "Type must derive from UObject");
-	return TSoftObjectPtr<std::remove_cv_t<T>>(ToRawPtr(Object));
+	return TSoftObjectPtr<T>(ToRawPtr(Object));
 }
 
 /**
