FSoftObjectPathFixupArchive - Fix SubPath getting wiped when performing fixup on a SoftObjectPath

diff --git a/Engine/Source/Runtime/CoreUObject/Public/UObject/SoftObjectPath.h b/Engine/Source/Runtime/CoreUObject/Public/UObject/SoftObjectPath.h
--- a/Engine/Source/Runtime/CoreUObject/Public/UObject/SoftObjectPath.h
+++ b/Engine/Source/Runtime/CoreUObject/Public/UObject/SoftObjectPath.h
@@ -591,11 +591,11 @@ struct FSoftObjectPathFixupArchive : public FArchiveUObject
 	}
 
 	FSoftObjectPathFixupArchive(const FString& InOldAssetPathString, const FString& InNewAssetPathString)
-		: FSoftObjectPathFixupArchive([OldAssetPathString = InOldAssetPathString, NewAssetPathString = InNewAssetPathString](FSoftObjectPath& Value)
+		: FSoftObjectPathFixupArchive([OldAssetPathString = InOldAssetPathString, NewAssetPath = FTopLevelAssetPath(InNewAssetPathString)](FSoftObjectPath& Value)
 		{
 			if (!Value.IsNull() && Value.GetAssetPathString().Equals(OldAssetPathString, ESearchCase::IgnoreCase))
 			{
-				Value = FSoftObjectPath(NewAssetPathString);
+				Value = FSoftObjectPath(NewAssetPath, Value.GetSubPathString());
 			}
 		})
 	{
