Using structured logging for missing package errors.

diff --git a/Engine/Source/Runtime/CoreUObject/Private/Serialization/AsyncLoading.cpp b/Engine/Source/Runtime/CoreUObject/Private/Serialization/AsyncLoading.cpp
--- a/Engine/Source/Runtime/CoreUObject/Private/Serialization/AsyncLoading.cpp
+++ b/Engine/Source/Runtime/CoreUObject/Private/Serialization/AsyncLoading.cpp
@@ -15,6 +15,7 @@
 #include "Stats/StatsMisc.h"
 #include "Misc/CoreStats.h"
 #include "HAL/IConsoleManager.h"
+#include "Logging/StructuredLog.h"
 #include "Misc/CoreDelegates.h"
 #include "Misc/CommandLine.h"
 #include "Misc/App.h"
@@ -6261,7 +6262,7 @@ EAsyncPackageState::Type FAsyncPackage::CreateLinker()
 
 					if (!FLinkerLoad::IsKnownMissingPackage(NameToLoadFName))
 					{
-						UE_LOG(LogStreaming, Error, TEXT("Couldn't find file for package %s requested by async loading code. NameToLoad: %s"), *Desc.Name.ToString(), *Desc.PackagePath.GetDebugName());
+						UE_LOGFMT(LogStreaming, Error, "Couldn't find file for package {Package} requested by async loading code. NameToLoad: {NameToLoad}", UE::FAssetLog(*Desc.Name.ToString()), UE::FAssetLog(Desc.PackagePath));
 
 #if !WITH_EDITORONLY_DATA
 						UE_CLOG(bUseTimeLimit, LogStreaming, Error, TEXT("This will hitch streaming because it ends up searching the disk instead of finding the file in the pak file."));
