[ZenLoader]
 - Make sure we reload missing exports after a GC

#rnx

diff --git a/Engine/Source/Runtime/CoreUObject/Private/Serialization/AsyncLoading2.cpp b/Engine/Source/Runtime/CoreUObject/Private/Serialization/AsyncLoading2.cpp
--- a/Engine/Source/Runtime/CoreUObject/Private/Serialization/AsyncLoading2.cpp
+++ b/Engine/Source/Runtime/CoreUObject/Private/Serialization/AsyncLoading2.cpp
@@ -418,6 +418,19 @@ static FAutoConsoleVariableRef CVarGOnlyProcessRequiredPackagesWhenSyncLoading(
 	ECVF_Default
 );
 
+#if WITH_EDITOR
+// This is important for the editor because the linker can end up recreating missing exports from its CreateImport function.
+// If we don't put the package in the zenloader loading queue, we can be left with objects in a RF_NeedLoad state that will never be actually loaded.
+// For the runtime, this should not happen since LinkerLoad is not involved, and we don't want to pay the memory and performance taxes involved in recreating missing exports.
+static bool GReloadPackagesWithGCedExports = true;
+static FAutoConsoleVariableRef CVarGReloadPackagesWithGCedExports(
+	TEXT("s.ReloadPackagesWithGCedExports"),
+	GReloadPackagesWithGCedExports,
+	TEXT("When active, packages with exports that have been garbage collected will go throught loading again even if they are currently in memory"),
+	ECVF_Default
+);
+#endif
+
 static float GStallDetectorTimeout = 120.0f;
 static FAutoConsoleVariableRef CVarGStallDetectorTimeout(
 	TEXT("s.StallDetectorTimeout"),
@@ -688,7 +701,7 @@ private:
 // The AGGRESSIVE_MEMORY_SAVING implementation allocates only as many keys and values as necessary to load a package's exports
 // which is space efficient but results in more memory traffic as exports are Store()'d or Remove()'d.
 // The !AGGRESSIVE_MEMORY_SAVING implementation uses more memory but provides O(1) Store() and Remove()
-#if AGGRESSIVE_MEMORY_SAVING
+#if AGGRESSIVE_MEMORY_SAVING && !WITH_EDITOR
 	public:
 		static const int32 InvalidValue = -1;
 
@@ -974,6 +987,11 @@ private:
 			return !!Data.Remove(ExportHash);
 		}
 
+		int32 Num() const
+		{
+			return Data.Num();
+		}
+
 		UObject* Find(uint64 ExportHash)
 		{
 			int32* ObjectIndex = Data.Find(ExportHash);
@@ -1069,6 +1087,7 @@ private:
 	FName OriginalPackageName;
 	int32 PackageObjectIndex = -1;
 	int32 RefCount = 0;
+	int32 ExportCount = -1;
 	bool bAreAllPublicExportsLoaded = false;
 	bool bIsMissing = false;
 	bool bHasFailed = false;
@@ -1150,7 +1169,7 @@ public:
 		return bAreAllPublicExportsLoaded && OriginalPackageName == GetPackage()->GetFName();
 	}
 
-	inline void SetAllPublicExportsLoaded()
+	inline void SetAllPublicExportsLoaded(bool bSnapshotExportCount)
 	{
 		check(!bIsMissing);
 		check(!bHasFailed);
@@ -1158,6 +1177,12 @@ public:
 		bIsMissing = false;
 		bAreAllPublicExportsLoaded = true;
 		bHasBeenLoadedDebug = true;
+#if WITH_EDITOR
+		if (bSnapshotExportCount)
+		{
+			ExportCount = PublicExportMap.Num();
+		}
+#endif
 	}
 
 	inline void SetIsMissingPackage()
@@ -1194,12 +1219,13 @@ public:
 		PublicExportMap.Store(ExportHash, Object);
 	}
 
-	void RemovePublicExport(uint64 ExportHash)
+	void RemovePublicExport(uint64 ExportHash, FName ObjectName = NAME_None)
 	{
 		check(!bIsMissing);
 		check(HasPackage());
 		if (PublicExportMap.Remove(ExportHash))
 		{
+			UE_LOG(LogStreaming, VeryVerbose, TEXT("Package %s got its export %X removed %s"), *GetPackage()->GetPathName(), ExportHash, *ObjectName.ToString());
 			bAreAllPublicExportsLoaded = false;
 		}
 	}
@@ -1219,6 +1245,14 @@ public:
 			return;
 		}
 		bAreAllPublicExportsLoaded = PublicExportMap.PinForGC(OutUnreachableObjectIndices);
+#if WITH_EDITOR
+		// Only reload if we ever got snapshotted after a proper load and the export count has shrinked since the package was last loaded
+		if (GReloadPackagesWithGCedExports && bAreAllPublicExportsLoaded && ExportCount > PublicExportMap.Num())
+		{
+			UE_LOG(LogStreaming, Log, TEXT("Reloading %s because %d on %d exports were GCed since it was loaded"), *Package->GetPathName(), ExportCount - PublicExportMap.Num(), ExportCount);
+			bAreAllPublicExportsLoaded = false;
+		}
+#endif
 		checkf(!Package->HasAnyInternalFlags(EInternalObjectFlags::LoaderImport), TEXT("%s"), *Package->GetFullName());
 		Package->SetInternalFlags(EInternalObjectFlags::LoaderImport);
 	}
@@ -1384,7 +1418,8 @@ public:
 			{
 				if (PackageRef.GetPackage()->bHasBeenFullyLoaded)
 				{
-					PackageRef.SetAllPublicExportsLoaded();
+					const bool bSnapshotExportCount = false;
+					PackageRef.SetAllPublicExportsLoaded(bSnapshotExportCount);
 				}
 			}
 #endif
@@ -1551,7 +1586,7 @@ public:
 				}
 				if (PackageRef)
 				{
-					PackageRef->RemovePublicExport(PublicExportKey.GetExportHash());
+					PackageRef->RemovePublicExport(PublicExportKey.GetExportHash(), Item.ObjectName);
 				}
 			}
 		}
@@ -5401,7 +5436,8 @@ void FAsyncPackage2::ImportPackagesRecursiveInner(FAsyncLoadingThreadState2& Thr
 					UncookedPackage->SetCanBeImportedFlag(true);
 					UncookedPackage->SetPackageId(ImportedPackageId);
 					UncookedPackage->SetInternalFlags(EInternalObjectFlags::LoaderImport);
-					ImportedPackageRef.SetAllPublicExportsLoaded();
+					const bool bSnapshotExportCount = true;
+					ImportedPackageRef.SetAllPublicExportsLoaded(bSnapshotExportCount);
 				}
 			}
 			if (UncookedPackage)
@@ -7354,7 +7390,8 @@ EEventLoadNodeExecutionResult FAsyncPackage2::Event_ExportsDone(FAsyncLoadingThr
 	if (!Package->bLoadHasFailed && Package->Desc.bCanBeImported)
 	{
 		FLoadedPackageRef& PackageRef = Package->AsyncLoadingThread.GlobalImportStore.FindPackageRefChecked(Package->Desc.UPackageId, Package->Desc.UPackageName);
-		PackageRef.SetAllPublicExportsLoaded();
+		const bool bSnapshotExportCount = true;
+		PackageRef.SetAllPublicExportsLoaded(bSnapshotExportCount);
 	}
 
 	if (!Package->Data.ShaderMapHashes.IsEmpty())
