[ZenLoader]
 - Fix deadlock at boot when a package and CDO have circular dependencies between them

#rnx

diff --git a/Engine/Source/Runtime/CoreUObject/Private/Serialization/AsyncLoading2.cpp b/Engine/Source/Runtime/CoreUObject/Private/Serialization/AsyncLoading2.cpp
--- a/Engine/Source/Runtime/CoreUObject/Private/Serialization/AsyncLoading2.cpp
+++ b/Engine/Source/Runtime/CoreUObject/Private/Serialization/AsyncLoading2.cpp
@@ -3917,12 +3917,20 @@ public:
 
 	void AddPendingCDOs(FAsyncPackage2* Package, TArray<UClass*, TInlineAllocator<8>>& Classes)
 	{
+		check(IsInGameThread()); // this is accessing PendingCDOsRecursiveStack which is GT only
 		for (UClass* Class : Classes)
 		{
-			TArray<FEventLoadNode2*>& Nodes = PendingCDOs.FindOrAdd(Class);
-			FEventLoadNode2& Node = Package->GetPackageNode(Package_DependenciesReady);
-			Node.AddBarrier();
-			Nodes.Add(&Node);
+			// Don't add a dependency on a CDO that is currently being created on the stack as this would cause a deadlock.
+			// StaticFind will be able to find the CDO and a second call to CreateDefaultObject would return the pointer
+			// so the only risk left is if the package referencing the CDO needs it to be fully constructed, which we 
+			// can't possibly satisfy because of the circular dependency.
+			if (!PendingCDOsRecursiveStack.Contains(Class))
+			{
+				TArray<FEventLoadNode2*>& Nodes = PendingCDOs.FindOrAdd(Class);
+				FEventLoadNode2& Node = Package->GetPackageNode(Package_DependenciesReady);
+				Node.AddBarrier();
+				Nodes.Add(&Node);
+			}
 		}
 	}
 
