Removed the Modified list, not needed anymore from the overridable Array
Some small refactor and cleanup left over from release branches

diff --git a/Engine/Source/Runtime/CoreUObject/Private/UObject/OverriddenPropertySet.cpp b/Engine/Source/Runtime/CoreUObject/Private/UObject/OverriddenPropertySet.cpp
--- a/Engine/Source/Runtime/CoreUObject/Private/UObject/OverriddenPropertySet.cpp
+++ b/Engine/Source/Runtime/CoreUObject/Private/UObject/OverriddenPropertySet.cpp
@@ -123,9 +123,9 @@ FOverriddenPropertyNodeID FOverriddenPropertyNodeID::RootNodeId()
 	return Result;
 }
 
-FOverriddenPropertyNodeID FOverriddenPropertyNodeID::ForMapKey(const FProperty* KeyProperty, const void* KeyData)
+FOverriddenPropertyNodeID FOverriddenPropertyNodeID::FromMapKey(const FProperty* KeyProperty, const void* KeyData)
 {
-	if (const FObjectProperty* KeyObjectProperty = CastField<FObjectProperty>(KeyProperty))
+	if (const FObjectPropertyBase* KeyObjectProperty = CastField<FObjectPropertyBase>(KeyProperty))
 	{
 		if (const UObject* Object = KeyObjectProperty->GetObjectPropertyValue(KeyData))
 		{
@@ -145,6 +145,40 @@ FOverriddenPropertyNodeID FOverriddenPropertyNodeID::ForMapKey(const FProperty*
 	return FOverriddenPropertyNodeID();
 }
 
+int32 FOverriddenPropertyNodeID::ToMapInternalIndex(FScriptMapHelper& MapHelper) const
+{
+	// Special case for object we didn't use the pointer to create the key
+	if (const FObjectPropertyBase* KeyObjectProperty = CastField<FObjectPropertyBase>(MapHelper.KeyProp))
+	{
+		for (FScriptMapHelper::FIterator It(MapHelper); It; ++It) 
+		{
+			if (UObject* CurrentObject = KeyObjectProperty->GetObjectPropertyValue(MapHelper.GetKeyPtr(It)))
+			{
+				if ((*this) == FOverriddenPropertyNodeID(*CurrentObject))
+				{
+					return It.GetInternalIndex();
+				}
+			}
+		}
+	}
+	else
+	{
+		// Default case, just import the text as key value for comparison
+		void* TempKeyValueStorage = FMemory_Alloca(MapHelper.MapLayout.SetLayout.Size);
+		MapHelper.KeyProp->InitializeValue(TempKeyValueStorage);
+
+		FString KeyToFind(ToString());
+		MapHelper.KeyProp->ImportText_Direct(*KeyToFind, TempKeyValueStorage, nullptr, PPF_None);
+
+		const int32 InternalIndex = MapHelper.FindMapPairIndexFromHash(TempKeyValueStorage);
+
+		MapHelper.KeyProp->DestroyValue(TempKeyValueStorage);
+
+		return InternalIndex;
+	}
+	return INDEX_NONE;
+}
+
 //----------------------------------------------------------------------//
 // FOverriddenPropertyNodeID
 //----------------------------------------------------------------------//
@@ -219,7 +253,7 @@ EOverriddenPropertyOperation FOverriddenPropertySet::GetOverriddenPropertyOperat
 
 
 		// Special handling for for instanced subobjects 
-		if (const FObjectProperty* ObjectProperty = CastField<FObjectProperty>(CurrentProperty))
+		if (const FObjectPropertyBase* ObjectProperty = CastField<FObjectPropertyBase>(CurrentProperty))
 		{
 			if (PropertyIterator->GetNextNode())
 			{
@@ -238,7 +272,7 @@ EOverriddenPropertyOperation FOverriddenPropertySet::GetOverriddenPropertyOperat
 
 			// Only special case is instanced subobjects, otherwise we fallback to full array override
 			checkf(ArrayProperty->Inner, TEXT("Expecting an inner type for Arrays"));
-			if (const FObjectProperty* InnerObjectProperty = ArrayProperty->Inner->HasAnyPropertyFlags(CPF_PersistentInstance) ? CastField<FObjectProperty>(ArrayProperty->Inner) : nullptr)
+			if (const FObjectPropertyBase* InnerObjectProperty = ArrayProperty->Inner->HasAnyPropertyFlags(CPF_PersistentInstance) ? CastField<FObjectPropertyBase>(ArrayProperty->Inner) : nullptr)
 			{
 				FScriptArrayHelper ArrayHelper(ArrayProperty, SubValuePtr);
 				if(ArrayHelper.IsValidIndex(ArrayIndex))
@@ -283,7 +317,7 @@ EOverriddenPropertyOperation FOverriddenPropertySet::GetOverriddenPropertyOperat
 				if (PropertyIterator->GetNextNode())
 				{
 					// Forward any sub queries to the subobject
-					if (const FObjectProperty* ValueInstancedObjectProperty = MapProperty->ValueProp->HasAnyPropertyFlags(CPF_PersistentInstance) ? CastField<FObjectProperty>(MapProperty->ValueProp) : nullptr)
+					if (const FObjectPropertyBase* ValueInstancedObjectProperty = MapProperty->ValueProp->HasAnyPropertyFlags(CPF_PersistentInstance) ? CastField<FObjectPropertyBase>(MapProperty->ValueProp) : nullptr)
 					{
 						if (UObject* ValueSubObject = ValueInstancedObjectProperty->GetObjectPropertyValue(MapHelper.GetValuePtr(InternalMapIndex)))
 						{
@@ -295,7 +329,7 @@ EOverriddenPropertyOperation FOverriddenPropertySet::GetOverriddenPropertyOperat
 				{
 					// Caller wants to know about any override state on the reference of the map pair itself
 					checkf(MapProperty->KeyProp, TEXT("Expecting a key type for Maps"));
-					FOverriddenPropertyNodeID OverriddenKeyID = FOverriddenPropertyNodeID::ForMapKey(MapProperty->KeyProp, MapHelper.GetKeyPtr(InternalMapIndex));
+					FOverriddenPropertyNodeID OverriddenKeyID = FOverriddenPropertyNodeID::FromMapKey(MapProperty->KeyProp, MapHelper.GetKeyPtr(InternalMapIndex));
 
 					if (const FOverriddenPropertyNodeID* CurrentPropKey = CurrentOverriddenPropertyNode->SubPropertyNodeKeys.Find(OverriddenKeyID))
 					{
@@ -352,7 +386,7 @@ bool FOverriddenPropertySet::ClearOverriddenProperty(FOverriddenPropertyNode& Pa
 		}
 
 		// Special handling for for instanced subobjects 
-		if (const FObjectProperty* ObjectProperty = CastField<FObjectProperty>(CurrentProperty))
+		if (const FObjectPropertyBase* ObjectProperty = CastField<FObjectPropertyBase>(CurrentProperty))
 		{
 			if (UObject* SubObject = ObjectProperty->GetObjectPropertyValue(SubValuePtr))
 			{
@@ -373,7 +407,7 @@ bool FOverriddenPropertySet::ClearOverriddenProperty(FOverriddenPropertyNode& Pa
 			ArrayIndex = PropertyEvent.GetArrayIndex(CurrentProperty->GetName());
 
 			// Only special case is instanced subobjects, otherwise we fallback to full array override
-			if (FObjectProperty* InnerObjectProperty = CastField<FObjectProperty>(ArrayProperty->Inner))
+			if (FObjectPropertyBase* InnerObjectProperty = CastField<FObjectPropertyBase>(ArrayProperty->Inner))
 			{
 				if (InnerObjectProperty->HasAnyPropertyFlags(CPF_PersistentInstance))
 				{
@@ -424,7 +458,7 @@ bool FOverriddenPropertySet::ClearOverriddenProperty(FOverriddenPropertyNode& Pa
 			FScriptMapHelper MapHelper(MapProperty, SubValuePtr);
 
 			const int32 InternalMapIndex = ArrayIndex != INDEX_NONE ? MapHelper.FindInternalIndex(ArrayIndex) : INDEX_NONE;
-			const FObjectProperty* ValueInstancedObjectProperty = MapProperty->ValueProp->HasAnyPropertyFlags(CPF_PersistentInstance) ? CastField<FObjectProperty>(MapProperty->ValueProp) : nullptr;
+			const FObjectPropertyBase* ValueInstancedObjectProperty = MapProperty->ValueProp->HasAnyPropertyFlags(CPF_PersistentInstance) ? CastField<FObjectPropertyBase>(MapProperty->ValueProp) : nullptr;
 
 			// If there is a next node, it is probably because the map value is holding a instanced subobject and the user is changing value on it.
 			// So forward the call to the instanced subobject
@@ -459,7 +493,7 @@ bool FOverriddenPropertySet::ClearOverriddenProperty(FOverriddenPropertyNode& Pa
 			else if (MapHelper.IsValidIndex(InternalMapIndex) && CurrentOverriddenPropertyNode)
 			{
 				checkf(MapProperty->KeyProp, TEXT("Expecting a key type for Maps"));
-				FOverriddenPropertyNodeID OverriddenKeyID = FOverriddenPropertyNodeID::ForMapKey(MapProperty->KeyProp, MapHelper.GetKeyPtr(InternalMapIndex));
+				FOverriddenPropertyNodeID OverriddenKeyID = FOverriddenPropertyNodeID::FromMapKey(MapProperty->KeyProp, MapHelper.GetKeyPtr(InternalMapIndex));
 
 				FOverriddenPropertyNodeID CurrentPropKey;
 				if (CurrentOverriddenPropertyNode->SubPropertyNodeKeys.RemoveAndCopyValue(OverriddenKeyID, CurrentPropKey))
@@ -584,7 +618,7 @@ void FOverriddenPropertySet::NotifyPropertyChange(FOverriddenPropertyNode* Paren
 	if (const FArrayProperty* ArrayProperty = CastField<FArrayProperty>(Property))
 	{
 		// Only special case is instanced subobjects, otherwise we fallback to full array override
-		if (FObjectProperty* InnerObjectProperty = CastField<FObjectProperty>(ArrayProperty->Inner))
+		if (FObjectPropertyBase* InnerObjectProperty = CastField<FObjectPropertyBase>(ArrayProperty->Inner))
 		{
 			if (InnerObjectProperty->HasAnyPropertyFlags(CPF_PersistentInstance))
 			{
@@ -778,13 +812,13 @@ void FOverriddenPropertySet::NotifyPropertyChange(FOverriddenPropertyNode* Paren
 	{
 		// Special handling of instanced subobjects
 		checkf(MapProperty->KeyProp, TEXT("Expecting a key type for Maps"));
-		FObjectProperty* KeyObjectProperty = CastField<FObjectProperty>(MapProperty->KeyProp);
+		FObjectPropertyBase* KeyObjectProperty = CastField<FObjectPropertyBase>(MapProperty->KeyProp);
 
 		// SubObjects
 		checkf(!KeyObjectProperty || !MapProperty->KeyProp->HasAnyPropertyFlags(CPF_PersistentInstance) || CastField<FClassProperty>(MapProperty->KeyProp), TEXT("Keys as a instanced subobject is not supported yet"));
 
 		checkf(MapProperty->ValueProp, TEXT("Expecting a value type for Maps"));
-		FObjectProperty* ValueInstancedObjectProperty = MapProperty->ValueProp->HasAnyPropertyFlags(CPF_PersistentInstance) ? CastField<FObjectProperty>(MapProperty->ValueProp) : nullptr;
+		FObjectPropertyBase* ValueInstancedObjectProperty = MapProperty->ValueProp->HasAnyPropertyFlags(CPF_PersistentInstance) ? CastField<FObjectPropertyBase>(MapProperty->ValueProp) : nullptr;
 
 		FScriptMapHelper MapHelper(MapProperty, SubValuePtr);
 		int32 LogicalMapIndex = PropertyEvent.GetArrayIndex(Property->GetName());
@@ -847,7 +881,7 @@ void FOverriddenPropertySet::NotifyPropertyChange(FOverriddenPropertyNode* Paren
 				{
 					if(SubPropertyNode)
 					{
-						FOverriddenPropertyNodeID OverriddenKeyID = FOverriddenPropertyNodeID::ForMapKey(MapProperty->KeyProp, MapHelper.GetKeyPtr(It.GetInternalIndex()));
+						FOverriddenPropertyNodeID OverriddenKeyID = FOverriddenPropertyNodeID::FromMapKey(MapProperty->KeyProp, MapHelper.GetKeyPtr(It.GetInternalIndex()));
 						FOverriddenPropertyNode& OverriddenKeyNode = FindOrAddNode(*SubPropertyNode, OverriddenKeyID);
 						OverriddenKeyNode.Operation = EOverriddenPropertyOperation::Replace;
 					}
@@ -870,7 +904,7 @@ void FOverriddenPropertySet::NotifyPropertyChange(FOverriddenPropertyNode* Paren
 
 				if(SubPropertyNode)
 				{
-					FOverriddenPropertyNodeID AddedKeyID = FOverriddenPropertyNodeID::ForMapKey(MapProperty->KeyProp, MapHelper.GetKeyPtr(InternalMapIndex));
+					FOverriddenPropertyNodeID AddedKeyID = FOverriddenPropertyNodeID::FromMapKey(MapProperty->KeyProp, MapHelper.GetKeyPtr(InternalMapIndex));
 					FOverriddenPropertyNode& AddedKeyNode = FindOrAddNode(*SubPropertyNode, AddedKeyID);
 					AddedKeyNode.Operation = EOverriddenPropertyOperation::Add;
 				}
@@ -882,7 +916,7 @@ void FOverriddenPropertySet::NotifyPropertyChange(FOverriddenPropertyNode* Paren
 
 				if(SubPropertyNode)
 				{
-					FOverriddenPropertyNodeID RemovedKeyID = FOverriddenPropertyNodeID::ForMapKey(MapProperty->KeyProp, PreEditMapHelper.GetKeyPtr(InternalPreEditMapIndex));
+					FOverriddenPropertyNodeID RemovedKeyID = FOverriddenPropertyNodeID::FromMapKey(MapProperty->KeyProp, PreEditMapHelper.GetKeyPtr(InternalPreEditMapIndex));
 					FOverriddenPropertyNode& RemovedKeyNode = FindOrAddNode(*SubPropertyNode, RemovedKeyID);
 					if (RemovedKeyNode.Operation == EOverriddenPropertyOperation::Add)
 					{
@@ -1003,7 +1037,7 @@ void FOverriddenPropertySet::NotifyPropertyChange(FOverriddenPropertyNode* Paren
 		}
 		return;
 	}
-	else if (const FObjectProperty* ObjectProperty = CastField<FObjectProperty>(Property))
+	else if (const FObjectPropertyBase* ObjectProperty = CastField<FObjectPropertyBase>(Property))
 	{
 		if (!PropertyNode->GetNextNode())
 		{
@@ -1082,7 +1116,7 @@ bool FOverriddenPropertySet::ClearOverriddenProperty(const FPropertyChangedEvent
 	{
 		return ClearOverriddenProperty(*RootNode, PropertyEvent, PropertyNode, Owner);
 	}
-	return false;
+	return true;
 }
 
 void FOverriddenPropertySet::OverrideProperty(const FPropertyChangedEvent& PropertyEvent, const FEditPropertyChain::TDoubleLinkedListNode* PropertyNode, const void* Data)
@@ -1274,4 +1308,4 @@ const FOverriddenPropertyNode* FOverriddenPropertySet::GetOverriddenPropertyNode
 	}
 
 	return OverriddenPropertyNode;
-}
+}
\ No newline at end of file
diff --git a/Engine/Source/Runtime/CoreUObject/Private/UObject/PropertyArray.cpp b/Engine/Source/Runtime/CoreUObject/Private/UObject/PropertyArray.cpp
--- a/Engine/Source/Runtime/CoreUObject/Private/UObject/PropertyArray.cpp
+++ b/Engine/Source/Runtime/CoreUObject/Private/UObject/PropertyArray.cpp
@@ -251,7 +251,7 @@ void FArrayProperty::SerializeItem(FStructuredArchive::FSlot Slot, void* Value,
 			else
 			{
 				// Only Array of Instanced subobject are handled here as sort of a set where the matching key is done using the archetype
-				const FObjectProperty* InnerObjectProperty = CastFieldChecked<FObjectProperty>(Inner);
+				const FObjectPropertyBase* InnerObjectProperty = CastFieldChecked<FObjectPropertyBase>(Inner);
 				checkf(InnerObjectProperty->HasAnyPropertyFlags(CPF_PersistentInstance), TEXT("Only supported code path here is the instanced subobjects"));
 
 				FOverriddenPropertySet* OverriddenProperties = FOverridableSerializationLogic::GetOverriddenProperties();
@@ -352,6 +352,7 @@ void FArrayProperty::SerializeItem(FStructuredArchive::FSlot Slot, void* Value,
 					}
 				}
 
+				// @Todo can be removed with an engine version, not being saved anymore on new objects
 				int32 NumModified = 0;
 				FStructuredArchive::FArray ModifiedArray = Record.EnterArray(TEXT("Modified"), NumModified);
 				if (NumModified != 0)
@@ -424,7 +425,6 @@ void FArrayProperty::SerializeItem(FStructuredArchive::FSlot Slot, void* Value,
 		{
 			// Container for temporarily tracking some indices
 			TArray<int32> RemovedIndices;
-			TArray<int32> ModifiedIndices;
 			TArray<int32> AddedIndices;
 
 			bool bReplaceArray = false;
@@ -434,7 +434,7 @@ void FArrayProperty::SerializeItem(FStructuredArchive::FSlot Slot, void* Value,
 			}
 			else
 			{
-				const FObjectProperty* InnerObjectProperty = CastField<FObjectProperty>(Inner);
+				const FObjectPropertyBase* InnerObjectProperty = CastField<FObjectPropertyBase>(Inner);
 				EOverriddenPropertyOperation ArrayOverrideOp = EOverriddenPropertyOperation::None;
 				FOverriddenPropertySet* OverriddenProperties = FOverridableSerializationLogic::GetOverriddenProperties();
 				if (OverriddenProperties)
@@ -452,13 +452,6 @@ void FArrayProperty::SerializeItem(FStructuredArchive::FSlot Slot, void* Value,
 					// Only array of instanced subobjects are handled here as sort of a set where the matching key is done using the archetype
 					checkf(InnerObjectProperty&& InnerObjectProperty->HasAnyPropertyFlags(CPF_PersistentInstance), TEXT("Expecting only arrays of instanced subobjects"));
 
-					// We need to always serialize instanced subobjects to know if they have overridden values.
-					const int32 ArrayNum = ArrayHelper.Num();
-					for (int i = 0; i < ArrayNum; i++)
-					{
-						ModifiedIndices.Add(i);
-					}
-
 					if (OverriddenProperties && ArrayOverrideOp != EOverriddenPropertyOperation::None)
 					{
 
@@ -527,7 +520,6 @@ void FArrayProperty::SerializeItem(FStructuredArchive::FSlot Slot, void* Value,
 										if (Index != INDEX_NONE)
 										{
 											AddedIndices.Add(Index);
-											ModifiedIndices.Remove(Index);
 										}
 										else
 										{
@@ -567,12 +559,9 @@ void FArrayProperty::SerializeItem(FStructuredArchive::FSlot Slot, void* Value,
 					SerializeContainerItem(RemovedArray.EnterElement(), DefaultsArrayHelper.GetRawPtr(i));
 				}
 
-				int32 NumModified = ModifiedIndices.Num();
+				// @Todo can be removed with an engine version, not being saved anymore on new objects
+				int32 NumModified = 0;
 				FStructuredArchive::FArray ModifiedArray = Record.EnterArray(TEXT("Modified"), NumModified);
-				for (int32 i : ModifiedIndices)
-				{
-					SerializeContainerItem(ModifiedArray.EnterElement(), ArrayHelper.GetRawPtr(i));
-				}
 
 				int32 NumAdded = AddedIndices.Num();
 				FStructuredArchive::FArray AddedArray = Record.EnterArray(TEXT("Added"), NumAdded);
diff --git a/Engine/Source/Runtime/CoreUObject/Private/UObject/PropertyMap.cpp b/Engine/Source/Runtime/CoreUObject/Private/UObject/PropertyMap.cpp
--- a/Engine/Source/Runtime/CoreUObject/Private/UObject/PropertyMap.cpp
+++ b/Engine/Source/Runtime/CoreUObject/Private/UObject/PropertyMap.cpp
@@ -308,11 +308,6 @@ void FMapProperty::SerializeItem(FStructuredArchive::FSlot Slot, void* Value, co
 		{
 			checkf(!UnderlyingArchive.ArUseCustomPropertyList, TEXT("Using custom property list is not supported by overridable serialization"));
 
-			auto GetIDFromKey = [&](uint8* KeyData) -> FOverriddenPropertyNodeID
-			{
-				return FOverriddenPropertyNodeID::ForMapKey(KeyProp, KeyData);
-			};
-
 			if (UnderlyingArchive.IsLoading())
 			{
 				int32 NumReplaced = 0;
@@ -378,7 +373,7 @@ void FMapProperty::SerializeItem(FStructuredArchive::FSlot Slot, void* Value, co
 							if (FOverriddenPropertyNode* MapOverriddenPropertyNode = OverriddenProperties ? OverriddenProperties->SetOverriddenPropertyOperation(EOverriddenPropertyOperation::Modified, UnderlyingArchive.GetSerializedPropertyChain(), /*Property*/nullptr) : nullptr)
 							{
 								// Rebuild the overridden info
-								FOverriddenPropertyNodeID RemovedKeyID = GetIDFromKey(TempKeyValueStorage);
+								FOverriddenPropertyNodeID RemovedKeyID = FOverriddenPropertyNodeID::FromMapKey(KeyProp, TempKeyValueStorage);
 								OverriddenProperties->SetSubPropertyOperation(EOverriddenPropertyOperation::Remove, *MapOverriddenPropertyNode, RemovedKeyID);
 							}
 						}
@@ -414,6 +409,14 @@ void FMapProperty::SerializeItem(FStructuredArchive::FSlot Slot, void* Value, co
 								FSerializedPropertyScope SerializedProperty(UnderlyingArchive, ValueProp, this);
 								ValueProp->SerializeItem(EntryRecord.EnterField(TEXT("Value")), ValuePtr);
 							}
+
+							// Need to fetch the MapOverriddenPropertyNode every loop as the previous might have reallocated the node.
+							if (FOverriddenPropertyNode* MapOverriddenPropertyNode = OverriddenProperties ? OverriddenProperties->SetOverriddenPropertyOperation(EOverriddenPropertyOperation::Modified, UnderlyingArchive.GetSerializedPropertyChain(), /*Property*/nullptr) : nullptr)
+							{
+								// Rebuild the overridden info
+								FOverriddenPropertyNodeID ModifiedKeyID = FOverriddenPropertyNodeID::FromMapKey(KeyProp, TempKeyValueStorage);
+								OverriddenProperties->SetSubPropertyOperation(EOverriddenPropertyOperation::Modified, *MapOverriddenPropertyNode, ModifiedKeyID);
+							}
 						}
 					}
 
@@ -452,7 +455,7 @@ void FMapProperty::SerializeItem(FStructuredArchive::FSlot Slot, void* Value, co
 							if (FOverriddenPropertyNode* MapOverriddenPropertyNode = OverriddenProperties ? OverriddenProperties->SetOverriddenPropertyOperation(EOverriddenPropertyOperation::Modified, UnderlyingArchive.GetSerializedPropertyChain(), /*Property*/nullptr) : nullptr)
 							{
 								// Rebuild the overridden info
-								FOverriddenPropertyNodeID AddedKeyID = GetIDFromKey(TempKeyValueStorage);
+								FOverriddenPropertyNodeID AddedKeyID = FOverriddenPropertyNodeID::FromMapKey(KeyProp, TempKeyValueStorage);
 								OverriddenProperties->SetSubPropertyOperation(EOverriddenPropertyOperation::Add, *MapOverriddenPropertyNode, AddedKeyID);
 							}
 						}
@@ -461,39 +464,6 @@ void FMapProperty::SerializeItem(FStructuredArchive::FSlot Slot, void* Value, co
 			}
 			else
 			{
-				auto FindKeyInternalIndex = [this](const FOverriddenPropertyNodeID& KeyIDToFind, FScriptMapHelper& MapHelper) -> int32
-				{
-					if (const FObjectProperty* KeyObjectProperty = CastField<FObjectProperty>(KeyProp))
-					{
-						for (FScriptMapHelper::FIterator It(MapHelper); It; ++It) 
-						{
-							if (UObject* CurrentObject = KeyObjectProperty->GetObjectPropertyValue(MapHelper.GetKeyPtr(It.GetInternalIndex())))
-							{
-								if (KeyIDToFind == FOverriddenPropertyNodeID(*CurrentObject))
-								{
-									return It.GetInternalIndex();
-								}
-							}
-						}
-					}
-					else
-					{
-						void* TempKeyValueStorage = FMemory::Malloc(MapLayout.SetLayout.Size);
-						KeyProp->InitializeValue(TempKeyValueStorage);
-
-						FString KeyToFind(KeyIDToFind.ToString());
-						KeyProp->ImportText_Direct(*KeyToFind, TempKeyValueStorage, nullptr, PPF_None);
-
-						const int32 InternalIndex = MapHelper.FindMapPairIndexFromHash(TempKeyValueStorage);
-
-						KeyProp->DestroyValue(TempKeyValueStorage);
-						FMemory::Free(TempKeyValueStorage);
-
-						return InternalIndex;
-					}
-					return INDEX_NONE;
-				};
-
 				// Container for temporarily tracking some indices
 				TArray<int32> RemovedIndices;
 				TArray<int32> AddedIndices;
@@ -523,16 +493,6 @@ void FMapProperty::SerializeItem(FStructuredArchive::FSlot Slot, void* Value, co
 					{
 						checkf(!KeyProp->HasAnyPropertyFlags(CPF_PersistentInstance) || CastField<FClassProperty>(KeyProp) || !CastField<FObjectProperty>(KeyProp), TEXT("The key as an instanced sub object is NYI"));
 
-						// For instanced subobject, the overridable is handled per object base, not only here, so we need to serialize the object ptr no matter what.
-						const bool bAreValuesInstancedSubObjects = ValueProp->HasAnyPropertyFlags(CPF_PersistentInstance) && CastField<FObjectProperty>(ValueProp);
-						if(bAreValuesInstancedSubObjects)
-						{
-							for (FScriptMapHelper::FIterator It(MapHelper); It; ++It)
-							{
-								ModifiedIndices.Add(It.GetInternalIndex());
-							}
-						}
-
 						if (OverriddenProperties && MapOverrideOp != EOverriddenPropertyOperation::None)
 						{
 							checkf(Defaults, TEXT("Expecting overridable serialization to have defaults to compare to"));
@@ -546,35 +506,33 @@ void FMapProperty::SerializeItem(FStructuredArchive::FSlot Slot, void* Value, co
 									const EOverriddenPropertyOperation OverrideOp = OverriddenProperties->GetSubPropertyOperation(Pair.Value);
 									switch (OverrideOp)
 									{
-									case EOverriddenPropertyOperation::Remove:
+										case EOverriddenPropertyOperation::Remove:
 										{
-											const int32 InternalIndex = FindKeyInternalIndex(Pair.Key, DefaultsMapHelper);
+											const int32 InternalIndex = Pair.Key.ToMapInternalIndex(DefaultsMapHelper);
 											if (InternalIndex != INDEX_NONE)
 											{
 												RemovedIndices.Add(InternalIndex);
 											}
 											break;
 										}
-									case EOverriddenPropertyOperation::Add:
+										case EOverriddenPropertyOperation::Add:
 										{
-											const int32 InternalIndex = FindKeyInternalIndex(Pair.Key, MapHelper);
+											const int32 InternalIndex = Pair.Key.ToMapInternalIndex(MapHelper);
 											if (InternalIndex != INDEX_NONE)
 											{
 												AddedIndices.Add(InternalIndex);
-												ModifiedIndices.Remove(InternalIndex);
 											}
 											break;
 										}
-									case EOverriddenPropertyOperation::Modified:
-										if (!bAreValuesInstancedSubObjects)
+										case EOverriddenPropertyOperation::Modified:
 										{
-											const int32 InternalIndex = FindKeyInternalIndex(Pair.Key, MapHelper);
+											const int32 InternalIndex = Pair.Key.ToMapInternalIndex(MapHelper);
 											if (InternalIndex != INDEX_NONE)
 											{
 												ModifiedIndices.Add(InternalIndex);
 											}
+											break;
 										}
-										break;
 									default:
 										checkf(false, TEXT("Unsupported map operation"));
 										break;
diff --git a/Engine/Source/Runtime/CoreUObject/Public/UObject/OverriddenPropertySet.h b/Engine/Source/Runtime/CoreUObject/Public/UObject/OverriddenPropertySet.h
--- a/Engine/Source/Runtime/CoreUObject/Public/UObject/OverriddenPropertySet.h
+++ b/Engine/Source/Runtime/CoreUObject/Public/UObject/OverriddenPropertySet.h
@@ -121,7 +121,10 @@ struct FOverriddenPropertyNodeID
 	}
 
 	static COREUOBJECT_API FOverriddenPropertyNodeID RootNodeId();
-	static COREUOBJECT_API FOverriddenPropertyNodeID ForMapKey(const FProperty* KeyProperty, const void* KeyData);
+
+	// Overridden property node map helpers
+	static COREUOBJECT_API FOverriddenPropertyNodeID FromMapKey(const FProperty* KeyProperty, const void* KeyData);
+	int32 ToMapInternalIndex(FScriptMapHelper& MapHelper) const;
 
 	bool operator==(const FOverriddenPropertyNodeID& Other) const
 	{
