Updated SetEditorProperty
This makes the function useful for modifying blueprint CDOs and having the modification propagate to subclass CDOs and loaded instances. In addition, the PreEditChange/PostEditChange callbacks fired from SetEditorProperty will now receive a change event with correct ChangeType == ValueSet, rather than Unspecified.

diff --git a/Engine/Plugins/Experimental/PythonScriptPlugin/Source/PythonScriptPlugin/Private/PyConversion.cpp b/Engine/Plugins/Experimental/PythonScriptPlugin/Source/PythonScriptPlugin/Private/PyConversion.cpp
--- a/Engine/Plugins/Experimental/PythonScriptPlugin/Source/PythonScriptPlugin/Private/PyConversion.cpp
+++ b/Engine/Plugins/Experimental/PythonScriptPlugin/Source/PythonScriptPlugin/Private/PyConversion.cpp
@@ -657,7 +657,7 @@ PyObject* PythonizeEnumEntry(const int64 Val, const UEnum* EnumType, const ESetE
 	return Obj;
 }
 
-FPyConversionResult NativizeProperty(PyObject* PyObj, const FProperty* Prop, void* ValueAddr, const FPropertyAccessChangeNotify* InChangeNotify, const ESetErrorState SetErrorState)
+FPyConversionResult NativizeProperty(PyObject* PyObj, const FProperty* Prop, void* ValueAddr, const TConstArrayView<void*>& InArchetypeInstValueAddrs, const FPropertyAccessChangeNotify* InChangeNotify, const ESetErrorState SetErrorState)
 {
 #define PYCONVERSION_PROPERTY_RETURN(RESULT) \
 	PYCONVERSION_RETURN(RESULT, TEXT("NativizeProperty"), *FString::Printf(TEXT("Cannot nativize '%s' as '%s' (%s)"), *PyUtil::GetFriendlyTypename(PyObj), *Prop->GetName(), *Prop->GetClass()->GetName()))
@@ -672,6 +672,10 @@ FPyConversionResult NativizeProperty(PyObject* PyObj, const FProperty* Prop, voi
 				const int32 ArrSize = FMath::Min(Prop->ArrayDim, PyFixedArray->ArrayProp->ArrayDim);
 				for (int32 ArrIndex = 0; ArrIndex < ArrSize; ++ArrIndex)
 				{
+					for (void* ArchInstValueAddr : InArchetypeInstValueAddrs)
+					{
+						Prop->CopySingleValue(static_cast<uint8*>(ArchInstValueAddr) + (Prop->ElementSize * ArrIndex), FPyWrapperFixedArray::GetItemPtr(PyFixedArray, ArrIndex));
+					}
 					Prop->CopySingleValue(static_cast<uint8*>(ValueAddr) + (Prop->ElementSize * ArrIndex), FPyWrapperFixedArray::GetItemPtr(PyFixedArray, ArrIndex));
 				}
 			});
@@ -681,7 +685,7 @@ FPyConversionResult NativizeProperty(PyObject* PyObj, const FProperty* Prop, voi
 		PYCONVERSION_PROPERTY_RETURN(FPyConversionResult::Failure());
 	}
 
-	return NativizeProperty_Direct(PyObj, Prop, ValueAddr, InChangeNotify, SetErrorState);
+	return NativizeProperty_Direct(PyObj, Prop, ValueAddr, InArchetypeInstValueAddrs, InChangeNotify, SetErrorState);
 
 #undef PYCONVERSION_PROPERTY_RETURN
 }
@@ -697,7 +701,7 @@ FPyConversionResult PythonizeProperty(const FProperty* Prop, const void* ValueAd
 	return PythonizeProperty_Direct(Prop, ValueAddr, OutPyObj, ConversionMethod, OwnerPyObj, SetErrorState);
 }
 
-FPyConversionResult NativizeProperty_Direct(PyObject* PyObj, const FProperty* Prop, void* ValueAddr, const FPropertyAccessChangeNotify* InChangeNotify, const ESetErrorState SetErrorState)
+FPyConversionResult NativizeProperty_Direct(PyObject* PyObj, const FProperty* Prop, void* ValueAddr, const TConstArrayView<void*>& InArchetypeInstValueAddrs, const FPropertyAccessChangeNotify* InChangeNotify, const ESetErrorState SetErrorState)
 {
 #define PYCONVERSION_PROPERTY_RETURN(RESULT) \
 	PYCONVERSION_RETURN(RESULT, TEXT("NativizeProperty"), *FString::Printf(TEXT("Cannot nativize '%s' as '%s' (%s)"), *PyUtil::GetFriendlyTypename(PyObj), *Prop->GetName(), *Prop->GetClass()->GetName()))
@@ -713,6 +717,10 @@ FPyConversionResult NativizeProperty_Direct(PyObject* PyObj, const FProperty* Pr
 			EmitPropertyChangeNotifications(InChangeNotify,							\
 				OldValue == NewValue, [&]()											\
 			{																		\
+				for (void* ArchInstValueAddr : InArchetypeInstValueAddrs)			\
+				{																	\
+					CastProp->SetPropertyValue(ArchInstValueAddr, NewValue);		\
+				}																	\
 				CastProp->SetPropertyValue(ValueAddr, NewValue);					\
 			});																		\
 		}																			\
@@ -730,6 +738,11 @@ FPyConversionResult NativizeProperty_Direct(PyObject* PyObj, const FProperty* Pr
 			EmitPropertyChangeNotifications(InChangeNotify,							\
 				CastProp->Identical(ValuePtr, &NewValue, PPF_None), [&]()			\
 			{																		\
+				for (void* ArchInstValueAddr : InArchetypeInstValueAddrs)			\
+				{																	\
+					auto* ArchInstValuePtr = static_cast<PROPTYPE::TCppType*>(ArchInstValueAddr);	\
+					*ArchInstValuePtr = NewValue;									\
+				}																	\
 				*ValuePtr = MoveTemp(NewValue);										\
 			});																		\
 		}																			\
@@ -780,6 +793,11 @@ FPyConversionResult NativizeProperty_Direct(PyObject* PyObj, const FProperty* Pr
 			auto* ValuePtr = static_cast<uint8*>(ValueAddr);
 			EmitPropertyChangeNotifications(InChangeNotify, *ValuePtr == NewValue, [&]()
 			{
+				for (void* ArchInstValueAddr : InArchetypeInstValueAddrs)
+				{
+					auto* InstValuePtr = static_cast<uint8*>(ArchInstValueAddr);
+					*InstValuePtr = NewValue;
+				}
 				*ValuePtr = NewValue;
 			});
 		}
@@ -807,6 +825,10 @@ FPyConversionResult NativizeProperty_Direct(PyObject* PyObj, const FProperty* Pr
 				const int64 OldValue = EnumInternalProp->GetSignedIntPropertyValue(ValueAddr);
 				EmitPropertyChangeNotifications(InChangeNotify, OldValue == NewValue, [&]()
 				{
+					for (void* ArchInstValueAddr : InArchetypeInstValueAddrs)
+					{
+						EnumInternalProp->SetIntPropertyValue(ArchInstValueAddr, NewValue);
+					}
 					EnumInternalProp->SetIntPropertyValue(ValueAddr, NewValue);
 				});
 			}
@@ -824,6 +846,10 @@ FPyConversionResult NativizeProperty_Direct(PyObject* PyObj, const FProperty* Pr
 			UObject* OldValue = CastProp->GetObjectPropertyValue(ValueAddr);
 			EmitPropertyChangeNotifications(InChangeNotify, OldValue == NewValue, [&]()
 			{
+				for (void* ArchInstValueAddr : InArchetypeInstValueAddrs)
+				{
+					CastProp->SetPropertyValue(ArchInstValueAddr, NewValue);
+				}
 				CastProp->SetObjectPropertyValue(ValueAddr, NewValue);
 			});
 		}
@@ -839,6 +865,10 @@ FPyConversionResult NativizeProperty_Direct(PyObject* PyObj, const FProperty* Pr
 			UObject* OldValue = CastProp->GetObjectPropertyValue(ValueAddr);
 			EmitPropertyChangeNotifications(InChangeNotify, OldValue == NewValue, [&]()
 			{
+				for (void* ArchInstValueAddr : InArchetypeInstValueAddrs)
+				{
+					CastProp->SetObjectPropertyValue(ArchInstValueAddr, NewValue);
+				}
 				CastProp->SetObjectPropertyValue(ValueAddr, NewValue);
 			});
 		}
@@ -854,6 +884,10 @@ FPyConversionResult NativizeProperty_Direct(PyObject* PyObj, const FProperty* Pr
 			UObject* OldValue = CastProp->GetObjectPropertyValue(ValueAddr);
 			EmitPropertyChangeNotifications(InChangeNotify, OldValue == NewValue, [&]()
 			{
+				for (void* ArchInstValueAddr : InArchetypeInstValueAddrs)
+				{
+					CastProp->SetObjectPropertyValue(ArchInstValueAddr, NewValue);
+				}
 				CastProp->SetObjectPropertyValue(ValueAddr, NewValue);
 			});
 		}
@@ -869,7 +903,12 @@ FPyConversionResult NativizeProperty_Direct(PyObject* PyObj, const FProperty* Pr
 			UObject* OldValue = CastProp->GetPropertyValue(ValueAddr).GetObject();
 			EmitPropertyChangeNotifications(InChangeNotify, OldValue == NewValue, [&]()
 			{
-				CastProp->SetPropertyValue(ValueAddr, FScriptInterface(NewValue, NewValue ? NewValue->GetInterfaceAddress(CastProp->InterfaceClass) : nullptr));
+				const FScriptInterface ScriptInterface = FScriptInterface(NewValue, NewValue ? NewValue->GetInterfaceAddress(CastProp->InterfaceClass) : nullptr);
+				for (void* ArchInstValueAddr : InArchetypeInstValueAddrs)
+				{
+					CastProp->SetPropertyValue(ArchInstValueAddr, ScriptInterface);
+				}
+				CastProp->SetPropertyValue(ValueAddr, ScriptInterface);
 			});
 		}
 		PYCONVERSION_PROPERTY_RETURN(Result);
@@ -884,6 +923,10 @@ FPyConversionResult NativizeProperty_Direct(PyObject* PyObj, const FProperty* Pr
 		{
 			EmitPropertyChangeNotifications(InChangeNotify, CastProp->Identical(ValueAddr, PyStruct->StructInstance, PPF_None), [&]()
 			{
+				for (void* ArchInstValueAddr : InArchetypeInstValueAddrs)
+				{
+					CastProp->Struct->CopyScriptStruct(ArchInstValueAddr, PyStruct->StructInstance);
+				}
 				CastProp->Struct->CopyScriptStruct(ValueAddr, PyStruct->StructInstance);
 			});
 		}
@@ -907,6 +950,10 @@ FPyConversionResult NativizeProperty_Direct(PyObject* PyObj, const FProperty* Pr
 		{
 			EmitPropertyChangeNotifications(InChangeNotify, CastProp->Identical(ValueAddr, PyDelegate->DelegateInstance, PPF_None), [&]()
 			{
+				for (void* ArchInstValueAddr : InArchetypeInstValueAddrs)
+				{
+					CastProp->SetPropertyValue(ArchInstValueAddr, *PyDelegate->DelegateInstance);
+				}
 				CastProp->SetPropertyValue(ValueAddr, *PyDelegate->DelegateInstance);
 			});
 		}
@@ -922,6 +969,10 @@ FPyConversionResult NativizeProperty_Direct(PyObject* PyObj, const FProperty* Pr
 		{
 			EmitPropertyChangeNotifications(InChangeNotify, CastProp->Identical(ValueAddr, PyDelegate->DelegateInstance, PPF_None), [&]()
 			{
+				for (void* ArchInstValueAddr : InArchetypeInstValueAddrs)
+				{
+					CastProp->SetMulticastDelegate(ArchInstValueAddr, *PyDelegate->DelegateInstance);
+				}
 				CastProp->SetMulticastDelegate(ValueAddr, *PyDelegate->DelegateInstance);
 			});
 		}
@@ -936,6 +987,10 @@ FPyConversionResult NativizeProperty_Direct(PyObject* PyObj, const FProperty* Pr
 		{
 			EmitPropertyChangeNotifications(InChangeNotify, CastProp->Identical(ValueAddr, PyArray->ArrayInstance, PPF_None), [&]()
 			{
+				for (void* ArchInstValueAddr : InArchetypeInstValueAddrs)
+				{
+					CastProp->CopySingleValue(ArchInstValueAddr, PyArray->ArrayInstance);
+				}
 				CastProp->CopySingleValue(ValueAddr, PyArray->ArrayInstance);
 			});
 		}
@@ -950,6 +1005,10 @@ FPyConversionResult NativizeProperty_Direct(PyObject* PyObj, const FProperty* Pr
 		{
 			EmitPropertyChangeNotifications(InChangeNotify, CastProp->Identical(ValueAddr, PySet->SetInstance, PPF_None), [&]()
 			{
+				for (void* ArchInstValueAddr : InArchetypeInstValueAddrs)
+				{
+					CastProp->CopySingleValue(ArchInstValueAddr, PySet->SetInstance);
+				}
 				CastProp->CopySingleValue(ValueAddr, PySet->SetInstance);
 			});
 		}
@@ -964,6 +1023,10 @@ FPyConversionResult NativizeProperty_Direct(PyObject* PyObj, const FProperty* Pr
 		{
 			EmitPropertyChangeNotifications(InChangeNotify, CastProp->Identical(ValueAddr, PyMap->MapInstance, PPF_None), [&]()
 			{
+				for (void* ArchInstValueAddr : InArchetypeInstValueAddrs)
+				{
+					CastProp->CopySingleValue(ArchInstValueAddr, PyMap->MapInstance);
+				}
 				CastProp->CopySingleValue(ValueAddr, PyMap->MapInstance);
 			});
 		}
@@ -1113,10 +1176,17 @@ FPyConversionResult PythonizeProperty_Direct(const FProperty* Prop, const void*
 #undef PYCONVERSION_PROPERTY_RETURN
 }
 
-FPyConversionResult NativizeProperty_InContainer(PyObject* PyObj, const FProperty* Prop, void* BaseAddr, const int32 ArrayIndex, const FPropertyAccessChangeNotify* InChangeNotify, const ESetErrorState SetErrorState)
+FPyConversionResult NativizeProperty_InContainer(PyObject* PyObj, const FProperty* Prop, void* BaseAddr, const int32 ArrayIndex, const TConstArrayView<void*>& InArchetypeInstBaseAddrs, const FPropertyAccessChangeNotify* InChangeNotify, const ESetErrorState SetErrorState)
 {
+	TArray<void*> ArchetypeInstValueAddrs;
+	ArchetypeInstValueAddrs.Reserve(InArchetypeInstBaseAddrs.Num());
+	for (void* ContainerAddr : InArchetypeInstBaseAddrs)
+	{
+		ArchetypeInstValueAddrs.Add(Prop->ContainerPtrToValuePtr<void>(ContainerAddr, ArrayIndex));
+	}
+
 	check(ArrayIndex < Prop->ArrayDim);
-	return NativizeProperty(PyObj, Prop, Prop->ContainerPtrToValuePtr<void>(BaseAddr, ArrayIndex), InChangeNotify, SetErrorState);
+	return NativizeProperty(PyObj, Prop, Prop->ContainerPtrToValuePtr<void>(BaseAddr, ArrayIndex), ArchetypeInstValueAddrs, InChangeNotify, SetErrorState);
 }
 
 FPyConversionResult PythonizeProperty_InContainer(const FProperty* Prop, const void* BaseAddr, const int32 ArrayIndex, PyObject*& OutPyObj, const EPyConversionMethod ConversionMethod, PyObject* OwnerPyObj, const ESetErrorState SetErrorState)
diff --git a/Engine/Plugins/Experimental/PythonScriptPlugin/Source/PythonScriptPlugin/Private/PyConversion.h b/Engine/Plugins/Experimental/PythonScriptPlugin/Source/PythonScriptPlugin/Private/PyConversion.h
--- a/Engine/Plugins/Experimental/PythonScriptPlugin/Source/PythonScriptPlugin/Private/PyConversion.h
+++ b/Engine/Plugins/Experimental/PythonScriptPlugin/Source/PythonScriptPlugin/Private/PyConversion.h
@@ -222,15 +222,15 @@ namespace PyConversion
 	}
 
 	/** Conversion for property instances (including fixed arrays) - ValueAddr should point to the property data */
-	FPyConversionResult NativizeProperty(PyObject* PyObj, const FProperty* Prop, void* ValueAddr, const FPropertyAccessChangeNotify* InChangeNotify = nullptr, const ESetErrorState SetErrorState = ESetErrorState::Yes);
+	FPyConversionResult NativizeProperty(PyObject* PyObj, const FProperty* Prop, void* ValueAddr, const TConstArrayView<void*>& ArchetypeInstValueAddrs = TArray<void*>(), const FPropertyAccessChangeNotify* InChangeNotify = nullptr, const ESetErrorState SetErrorState = ESetErrorState::Yes);
 	FPyConversionResult PythonizeProperty(const FProperty* Prop, const void* ValueAddr, PyObject*& OutPyObj, const EPyConversionMethod ConversionMethod = EPyConversionMethod::Copy, PyObject* OwnerPyObj = nullptr, const ESetErrorState SetErrorState = ESetErrorState::Yes);
 
 	/** Conversion for single property instances - ValueAddr should point to the property data */
-	FPyConversionResult NativizeProperty_Direct(PyObject* PyObj, const FProperty* Prop, void* ValueAddr, const FPropertyAccessChangeNotify* InChangeNotify = nullptr, const ESetErrorState SetErrorState = ESetErrorState::Yes);
+	FPyConversionResult NativizeProperty_Direct(PyObject* PyObj, const FProperty* Prop, void* ValueAddr, const TConstArrayView<void*>& ArchetypeInstValueAddrs = TArray<void*>(), const FPropertyAccessChangeNotify* InChangeNotify = nullptr, const ESetErrorState SetErrorState = ESetErrorState::Yes);
 	FPyConversionResult PythonizeProperty_Direct(const FProperty* Prop, const void* ValueAddr, PyObject*& OutPyObj, const EPyConversionMethod ConversionMethod = EPyConversionMethod::Copy, PyObject* OwnerPyObj = nullptr, const ESetErrorState SetErrorState = ESetErrorState::Yes);
 
 	/** Conversion for property instances within a structure (including fixed arrays) - BaseAddr should point to the structure data */
-	FPyConversionResult NativizeProperty_InContainer(PyObject* PyObj, const FProperty* Prop, void* BaseAddr, const int32 ArrayIndex, const FPropertyAccessChangeNotify* InChangeNotify = nullptr, const ESetErrorState SetErrorState = ESetErrorState::Yes);
+	FPyConversionResult NativizeProperty_InContainer(PyObject* PyObj, const FProperty* Prop, void* BaseAddr, const int32 ArrayIndex, const TConstArrayView<void*>& ArchetypeInstBaseAddrs = TArray<void*>(), const FPropertyAccessChangeNotify* InChangeNotify = nullptr, const ESetErrorState SetErrorState = ESetErrorState::Yes);
 	FPyConversionResult PythonizeProperty_InContainer(const FProperty* Prop, const void* BaseAddr, const int32 ArrayIndex, PyObject*& OutPyObj, const EPyConversionMethod ConversionMethod = EPyConversionMethod::Copy, PyObject* OwnerPyObj = nullptr, const ESetErrorState SetErrorState = ESetErrorState::Yes);
 
 	/**
diff --git a/Engine/Plugins/Experimental/PythonScriptPlugin/Source/PythonScriptPlugin/Private/PyGenUtil.cpp b/Engine/Plugins/Experimental/PythonScriptPlugin/Source/PythonScriptPlugin/Private/PyGenUtil.cpp
--- a/Engine/Plugins/Experimental/PythonScriptPlugin/Source/PythonScriptPlugin/Private/PyGenUtil.cpp
+++ b/Engine/Plugins/Experimental/PythonScriptPlugin/Source/PythonScriptPlugin/Private/PyGenUtil.cpp
@@ -1322,7 +1322,7 @@ PyObject* GetPropertyValue(const UStruct* InStruct, const void* InStructData, co
 	return PyUtil::GetPropertyValue(InStruct, InStructData, InPropDef.Prop, InAttributeName, InOwnerPyObject, InErrorCtxt);
 }
 
-int SetPropertyValue(const UStruct* InStruct, void* InStructData, PyObject* InValue, const FGeneratedWrappedProperty& InPropDef, const char *InAttributeName, const FPropertyAccessChangeNotify* InChangeNotify, const uint64 InReadOnlyFlags, const bool InOwnerIsTemplate, const TCHAR* InErrorCtxt)
+int SetPropertyValue(const UStruct* InStruct, void* InStructData, PyObject* InValue, const FGeneratedWrappedProperty& InPropDef, const char *InAttributeName, const FPropertyAccessChangeNotify* InChangeNotify, const uint64 InReadOnlyFlags, const bool InOwnerIsTemplate, const TCHAR* InErrorCtxt, const TConstArrayView<void*>& InArchetypeInstContainers)
 {
 	// Has this property been deprecated?
 	if (InStruct && InPropDef.Prop && InPropDef.DeprecationMessage.IsSet())
@@ -1344,7 +1344,7 @@ int SetPropertyValue(const UStruct* InStruct, void* InStructData, PyObject* InVa
 		}
 	}
 
-	return PyUtil::SetPropertyValue(InStruct, InStructData, InValue, InPropDef.Prop, InAttributeName, InChangeNotify, InReadOnlyFlags, InOwnerIsTemplate, InErrorCtxt);
+	return PyUtil::SetPropertyValue(InStruct, InStructData, InValue, InPropDef.Prop, InAttributeName, InChangeNotify, InReadOnlyFlags, InOwnerIsTemplate, InErrorCtxt, InArchetypeInstContainers);
 }
 
 FString BuildFunctionDocString(const UFunction* InFunc, const FString& InFuncPythonName, const TArray<FGeneratedWrappedMethodParameter>& InInputParams, const TArray<FGeneratedWrappedMethodParameter>& InOutputParams, const bool* InStaticOverride)
diff --git a/Engine/Plugins/Experimental/PythonScriptPlugin/Source/PythonScriptPlugin/Private/PyGenUtil.h b/Engine/Plugins/Experimental/PythonScriptPlugin/Source/PythonScriptPlugin/Private/PyGenUtil.h
--- a/Engine/Plugins/Experimental/PythonScriptPlugin/Source/PythonScriptPlugin/Private/PyGenUtil.h
+++ b/Engine/Plugins/Experimental/PythonScriptPlugin/Source/PythonScriptPlugin/Private/PyGenUtil.h
@@ -1029,7 +1029,7 @@ namespace PyGenUtil
 	PyObject* GetPropertyValue(const UStruct* InStruct, const void* InStructData, const FGeneratedWrappedProperty& InPropDef, const char *InAttributeName, PyObject* InOwnerPyObject, const TCHAR* InErrorCtxt);
 
 	/** Set the current value of the given property from the given struct */
-	int SetPropertyValue(const UStruct* InStruct, void* InStructData, PyObject* InValue, const FGeneratedWrappedProperty& InPropDef, const char *InAttributeName, const FPropertyAccessChangeNotify* InChangeNotify, const uint64 InReadOnlyFlags, const bool InOwnerIsTemplate, const TCHAR* InErrorCtxt);
+	int SetPropertyValue(const UStruct* InStruct, void* InStructData, PyObject* InValue, const FGeneratedWrappedProperty& InPropDef, const char *InAttributeName, const FPropertyAccessChangeNotify* InChangeNotify, const uint64 InReadOnlyFlags, const bool InOwnerIsTemplate, const TCHAR* InErrorCtxt, const TConstArrayView<void*>& InArchetypeInstContainers);
 
 	/** Build a Python doc string for the given function and arguments list */
 	FString BuildFunctionDocString(const UFunction* InFunc, const FString& InFuncPythonName, const TArray<FGeneratedWrappedMethodParameter>& InInputParams, const TArray<FGeneratedWrappedMethodParameter>& InOutputParams, const bool* InStaticOverride = nullptr);
diff --git a/Engine/Plugins/Experimental/PythonScriptPlugin/Source/PythonScriptPlugin/Private/PyUtil.cpp b/Engine/Plugins/Experimental/PythonScriptPlugin/Source/PythonScriptPlugin/Private/PyUtil.cpp
--- a/Engine/Plugins/Experimental/PythonScriptPlugin/Source/PythonScriptPlugin/Private/PyUtil.cpp
+++ b/Engine/Plugins/Experimental/PythonScriptPlugin/Source/PythonScriptPlugin/Private/PyUtil.cpp
@@ -825,7 +825,7 @@ PyObject* GetPropertyValue(const UStruct* InStruct, const void* InStructData, co
 	Py_RETURN_NONE;
 }
 
-int SetPropertyValue(const UStruct* InStruct, void* InStructData, PyObject* InValue, const FProperty* InProp, const char *InAttributeName, const FPropertyAccessChangeNotify* InChangeNotify, const uint64 InReadOnlyFlags, const bool InOwnerIsTemplate, const TCHAR* InErrorCtxt)
+int SetPropertyValue(const UStruct* InStruct, void* InStructData, PyObject* InValue, const FProperty* InProp, const char *InAttributeName, const FPropertyAccessChangeNotify* InChangeNotify, const uint64 InReadOnlyFlags, const bool InOwnerIsTemplate, const TCHAR* InErrorCtxt, const TConstArrayView<void*>& InArchetypeInstStructData)
 {
 	if (!InValue)
 	{
@@ -866,7 +866,7 @@ int SetPropertyValue(const UStruct* InStruct, void* InStructData, PyObject* InVa
 			return -1;
 		}
 
-		if (!PyConversion::NativizeProperty_InContainer(InValue, InProp, InStructData, 0, InChangeNotify))
+		if (!PyConversion::NativizeProperty_InContainer(InValue, InProp, InStructData, 0, InArchetypeInstStructData, InChangeNotify))
 		{
 			SetPythonError(PyExc_TypeError, InErrorCtxt, *FString::Printf(TEXT("Failed to convert type '%s' to property '%s' (%s) for attribute '%s' on '%s'"), *GetFriendlyTypename(InValue), *InProp->GetName(), *InProp->GetClass()->GetName(), UTF8_TO_TCHAR(InAttributeName), *InStruct->GetName()));
 			return -1;
diff --git a/Engine/Plugins/Experimental/PythonScriptPlugin/Source/PythonScriptPlugin/Private/PyUtil.h b/Engine/Plugins/Experimental/PythonScriptPlugin/Source/PythonScriptPlugin/Private/PyUtil.h
--- a/Engine/Plugins/Experimental/PythonScriptPlugin/Source/PythonScriptPlugin/Private/PyUtil.h
+++ b/Engine/Plugins/Experimental/PythonScriptPlugin/Source/PythonScriptPlugin/Private/PyUtil.h
@@ -438,7 +438,7 @@ namespace PyUtil
 	PyObject* GetPropertyValue(const UStruct* InStruct, const void* InStructData, const FProperty* InProp, const char *InAttributeName, PyObject* InOwnerPyObject, const TCHAR* InErrorCtxt);
 
 	/** Set the current value of the given property from the given struct */
-	int SetPropertyValue(const UStruct* InStruct, void* InStructData, PyObject* InValue, const FProperty* InProp, const char *InAttributeName, const FPropertyAccessChangeNotify* InChangeNotify, const uint64 InReadOnlyFlags, const bool InOwnerIsTemplate, const TCHAR* InErrorCtxt);
+	int SetPropertyValue(const UStruct* InStruct, void* InStructData, PyObject* InValue, const FProperty* InProp, const char *InAttributeName, const FPropertyAccessChangeNotify* InChangeNotify, const uint64 InReadOnlyFlags, const bool InOwnerIsTemplate, const TCHAR* InErrorCtxt, const TConstArrayView<void*>& InArchetypeInstStructData);
 
 	/**
 	 * Check to see if the given object implements a length function.
diff --git a/Engine/Plugins/Experimental/PythonScriptPlugin/Source/PythonScriptPlugin/Private/PyWrapperObject.cpp b/Engine/Plugins/Experimental/PythonScriptPlugin/Source/PythonScriptPlugin/Private/PyWrapperObject.cpp
--- a/Engine/Plugins/Experimental/PythonScriptPlugin/Source/PythonScriptPlugin/Private/PyWrapperObject.cpp
+++ b/Engine/Plugins/Experimental/PythonScriptPlugin/Source/PythonScriptPlugin/Private/PyWrapperObject.cpp
@@ -137,7 +137,17 @@ int FPyWrapperObject::SetPropertyValue(FPyWrapperObject* InSelf, PyObject* InVal
 	}
 
 	const TUniquePtr<FPropertyAccessChangeNotify> ChangeNotify = FPyWrapperOwnerContext((PyObject*)InSelf, InPropDef.Prop).BuildChangeNotify(InNotifyMode);
-	return PyGenUtil::SetPropertyValue(InSelf->ObjectInstance->GetClass(), InSelf->ObjectInstance, InValue, InPropDef, InPythonAttrName, ChangeNotify.Get(), InReadOnlyFlags, PropertyAccessUtil::IsObjectTemplate(InSelf->ObjectInstance), *PyUtil::GetErrorContext(InSelf));
+
+	// If the object is a template, gather instances (including subclass CDOs) inheriting their property value from this template.
+	// Those instances are passed into PyGenUtil::SetPropertyValue so that they will receive the value change as well.
+	TArray<void*> ArchetypeInsts;
+	const bool bIsObjectTemplate = PropertyAccessUtil::IsObjectTemplate(InSelf->ObjectInstance);
+	if (bIsObjectTemplate)
+	{
+		PropertyAccessUtil::GetArchetypeInstancesInheritingPropertyValue_AsContainerData(InPropDef.Prop, InSelf->ObjectInstance, ArchetypeInsts);
+	}
+
+	return PyGenUtil::SetPropertyValue(InSelf->ObjectInstance->GetClass(), InSelf->ObjectInstance, InValue, InPropDef, InPythonAttrName, ChangeNotify.Get(), InReadOnlyFlags, bIsObjectTemplate, *PyUtil::GetErrorContext(InSelf), ArchetypeInsts);
 }
 
 PyObject* FPyWrapperObject::CallGetterFunction(FPyWrapperObject* InSelf, const PyGenUtil::FGeneratedWrappedFunction& InFuncDef)
@@ -878,13 +888,22 @@ PyTypeObject InitializePyWrapperObjectType()
 				return nullptr;
 			}
 
+			// If the object is a template, gather instances (including subclass CDOs) inheriting their property value from this template.
+			// Those instances are passed into PyGenUtil::SetPropertyValue so that they will receive the value change as well.
+			TArray<void*> ArchetypeInsts;
+			const bool bIsObjectTemplate = PropertyAccessUtil::IsObjectTemplate(InSelf->ObjectInstance);
+			if (bIsObjectTemplate)
+			{
+				PropertyAccessUtil::GetArchetypeInstancesInheritingPropertyValue_AsContainerData(WrappedPropDef.Prop, InSelf->ObjectInstance, ArchetypeInsts);
+			}
+
 			// If the owner class is set then this property is from 'self', otherwise it's from the sparse class data
 			const bool bPropertyIsOwnedBySelf = WrappedPropDef.Prop->GetOwnerClass() != nullptr;
 			UClass* Class = InSelf->ObjectInstance->GetClass();
 			if (bPropertyIsOwnedBySelf)
 			{
 				const TUniquePtr<FPropertyAccessChangeNotify> ChangeNotify = FPyWrapperOwnerContext((PyObject*)InSelf, WrappedPropDef.Prop).BuildChangeNotify(NotifyMode);
-				const int Result = PyGenUtil::SetPropertyValue(Class, InSelf->ObjectInstance, PyValueObj, WrappedPropDef, TCHAR_TO_UTF8(*Name.ToString()), ChangeNotify.Get(), PropertyAccessUtil::EditorReadOnlyFlags, PropertyAccessUtil::IsObjectTemplate(InSelf->ObjectInstance), *PyUtil::GetErrorContext(InSelf));
+				const int Result = PyGenUtil::SetPropertyValue(Class, InSelf->ObjectInstance, PyValueObj, WrappedPropDef, TCHAR_TO_UTF8(*Name.ToString()), ChangeNotify.Get(), PropertyAccessUtil::EditorReadOnlyFlags, bIsObjectTemplate, *PyUtil::GetErrorContext(InSelf), ArchetypeInsts);
 				if (Result != 0)
 				{
 					return nullptr;
@@ -897,7 +916,7 @@ PyTypeObject InitializePyWrapperObjectType()
 				if (SparseDataStruct && SparseData)
 				{
 					const TUniquePtr<FPropertyAccessChangeNotify> ChangeNotify = FPyWrapperOwnerContext((PyObject*)InSelf, WrappedPropDef.Prop).BuildChangeNotify(NotifyMode);
-					const int Result = PyGenUtil::SetPropertyValue(SparseDataStruct, SparseData, PyValueObj, WrappedPropDef, TCHAR_TO_UTF8(*Name.ToString()), ChangeNotify.Get(), PropertyAccessUtil::EditorReadOnlyFlags, PropertyAccessUtil::IsObjectTemplate(InSelf->ObjectInstance), *PyUtil::GetErrorContext(InSelf));
+					const int Result = PyGenUtil::SetPropertyValue(SparseDataStruct, SparseData, PyValueObj, WrappedPropDef, TCHAR_TO_UTF8(*Name.ToString()), ChangeNotify.Get(), PropertyAccessUtil::EditorReadOnlyFlags, bIsObjectTemplate, *PyUtil::GetErrorContext(InSelf), ArchetypeInsts);
 					if (Result != 0)
 					{
 						return nullptr;
@@ -980,17 +999,27 @@ PyTypeObject InitializePyWrapperObjectType()
 
 			// Try and set the value of each property
 			UClass* Class = InSelf->ObjectInstance->GetClass();
+			const bool bIsObjectTemplate = PropertyAccessUtil::IsObjectTemplate(InSelf->ObjectInstance);
 			for (const FPropertyInfoPair& PropertyInfo : PropertyInfos)
 			{
 				const FName Name = PropertyInfo.Get<0>();
 				const PyGenUtil::FGeneratedWrappedProperty& WrappedPropDef = PropertyInfo.Get<1>();
 				PyObject* PyValueObj = PropertyInfo.Get<2>().GetPtr();
 
+				// If the object is a template, gather instances (including subclass CDOs) inheriting their property value from this template.
+				// Those instances are passed into PyGenUtil::SetPropertyValue so that they will receive the value change as well.
+				// Each property can have a different set of archetype instances that are inheriting the value.
+				TArray<void*> ArchetypeInsts;
+				if (bIsObjectTemplate)
+				{
+					PropertyAccessUtil::GetArchetypeInstancesInheritingPropertyValue_AsContainerData(WrappedPropDef.Prop, InSelf->ObjectInstance, ArchetypeInsts);
+				}
+
 				// If the owner class is set then this property is from 'self', otherwise it's from the sparse class data
 				const bool bPropertyIsOwnedBySelf = WrappedPropDef.Prop->GetOwnerClass() != nullptr;
 				if (bPropertyIsOwnedBySelf)
 				{
-					const int Result = PyGenUtil::SetPropertyValue(Class, InSelf->ObjectInstance, PyValueObj, WrappedPropDef, TCHAR_TO_UTF8(*Name.ToString()), nullptr, PropertyAccessUtil::EditorReadOnlyFlags, PropertyAccessUtil::IsObjectTemplate(InSelf->ObjectInstance), *PyUtil::GetErrorContext(InSelf));
+					const int Result = PyGenUtil::SetPropertyValue(Class, InSelf->ObjectInstance, PyValueObj, WrappedPropDef, TCHAR_TO_UTF8(*Name.ToString()), nullptr, PropertyAccessUtil::EditorReadOnlyFlags, bIsObjectTemplate, *PyUtil::GetErrorContext(InSelf), ArchetypeInsts);
 					if (Result != 0)
 					{
 						return nullptr;
@@ -1002,7 +1031,7 @@ PyTypeObject InitializePyWrapperObjectType()
 					void* SparseData = Class->GetOrCreateSparseClassData();
 					if (SparseDataStruct && SparseData)
 					{
-						const int Result = PyGenUtil::SetPropertyValue(SparseDataStruct, SparseData, PyValueObj, WrappedPropDef, TCHAR_TO_UTF8(*Name.ToString()), nullptr, PropertyAccessUtil::EditorReadOnlyFlags, PropertyAccessUtil::IsObjectTemplate(InSelf->ObjectInstance), *PyUtil::GetErrorContext(InSelf));
+						const int Result = PyGenUtil::SetPropertyValue(SparseDataStruct, SparseData, PyValueObj, WrappedPropDef, TCHAR_TO_UTF8(*Name.ToString()), nullptr, PropertyAccessUtil::EditorReadOnlyFlags, bIsObjectTemplate, *PyUtil::GetErrorContext(InSelf), ArchetypeInsts);
 						if (Result != 0)
 						{
 							return nullptr;
diff --git a/Engine/Plugins/Experimental/PythonScriptPlugin/Source/PythonScriptPlugin/Private/PyWrapperStruct.cpp b/Engine/Plugins/Experimental/PythonScriptPlugin/Source/PythonScriptPlugin/Private/PyWrapperStruct.cpp
--- a/Engine/Plugins/Experimental/PythonScriptPlugin/Source/PythonScriptPlugin/Private/PyWrapperStruct.cpp
+++ b/Engine/Plugins/Experimental/PythonScriptPlugin/Source/PythonScriptPlugin/Private/PyWrapperStruct.cpp
@@ -281,7 +281,8 @@ FPyWrapperStruct* FPyWrapperStruct::CastPyObject(PyObject* InPyObject, PyTypeObj
 					break;
 				}
 
-				const int Result = PyUtil::SetPropertyValue(NewStruct->ScriptStruct, NewStruct->StructInstance, SequenceItem, InitParam.ParamProp, InitParam.ParamName.GetData(), nullptr, 0, false, *PyUtil::GetErrorContext(NewStruct.Get()));
+				const TArray<void*> NoArchetypeInsts;
+				const int Result = PyUtil::SetPropertyValue(NewStruct->ScriptStruct, NewStruct->StructInstance, SequenceItem, InitParam.ParamProp, InitParam.ParamName.GetData(), nullptr, 0, false, *PyUtil::GetErrorContext(NewStruct.Get()), NoArchetypeInsts);
 				if (Result != 0)
 				{
 					return nullptr;
@@ -321,7 +322,8 @@ FPyWrapperStruct* FPyWrapperStruct::CastPyObject(PyObject* InPyObject, PyTypeObj
 			PyObject* MappingItem = PyMapping_GetItemString(InPyObject, (char*)InitParam.ParamName.GetData());
 			if (MappingItem)
 			{
-				const int Result = PyUtil::SetPropertyValue(NewStruct->ScriptStruct, NewStruct->StructInstance, MappingItem, InitParam.ParamProp, InitParam.ParamName.GetData(), nullptr, 0, false, *PyUtil::GetErrorContext(NewStruct.Get()));
+				TArray<void*> NoArchetypeInsts;
+				const int Result = PyUtil::SetPropertyValue(NewStruct->ScriptStruct, NewStruct->StructInstance, MappingItem, InitParam.ParamProp, InitParam.ParamName.GetData(), nullptr, 0, false, *PyUtil::GetErrorContext(NewStruct.Get()), NoArchetypeInsts);
 				if (Result != 0)
 				{
 					return nullptr;
@@ -452,7 +454,8 @@ int FPyWrapperStruct::SetPropertyValue(FPyWrapperStruct* InSelf, PyObject* InVal
 	}
 
 	const TUniquePtr<FPropertyAccessChangeNotify> ChangeNotify = FPyWrapperOwnerContext((PyObject*)InSelf, InPropDef.Prop).BuildChangeNotify(InNotifyMode);
-	return PyGenUtil::SetPropertyValue(InSelf->ScriptStruct, InSelf->StructInstance, InValue, InPropDef, InPythonAttrName, ChangeNotify.Get(), InReadOnlyFlags, OwnerIsTemplate, *PyUtil::GetErrorContext(InSelf));
+	const TArray<void*> NoArchetypeInsts;
+	return PyGenUtil::SetPropertyValue(InSelf->ScriptStruct, InSelf->StructInstance, InValue, InPropDef, InPythonAttrName, ChangeNotify.Get(), InReadOnlyFlags, OwnerIsTemplate, *PyUtil::GetErrorContext(InSelf), NoArchetypeInsts);
 }
 
 int FPyWrapperStruct::CallMakeFunction_Impl(FPyWrapperStruct* InSelf, PyObject* InArgs, PyObject* InKwds, const PyGenUtil::FGeneratedWrappedFunction& InFuncDef)
@@ -635,7 +638,8 @@ PyObject* FPyWrapperStruct::CallOperatorFunction_Impl(FPyWrapperStruct* InSelf,
 		PyGenUtil::ApplyParamDefaults(FuncParams.GetMemory(), InOpFunc.AdditionalParams);
 		if (InOpFunc.OtherParam.ParamProp)
 		{
-			const FPyConversionResult RHSResult = PyConversion::NativizeProperty_InContainer(InRHS, InOpFunc.OtherParam.ParamProp, FuncParams.GetMemory(), 0, nullptr, PyConversion::ESetErrorState::No);
+			TArray<void*> NoArchetypeInsts;
+			const FPyConversionResult RHSResult = PyConversion::NativizeProperty_InContainer(InRHS, InOpFunc.OtherParam.ParamProp, FuncParams.GetMemory(), 0, NoArchetypeInsts, nullptr, PyConversion::ESetErrorState::No);
 			SetOptionalPyConversionResult(RHSResult, OutRHSConversionResult);
 
 			if (!RHSResult)
diff --git a/Engine/Source/Editor/DetailCustomizations/Private/OverrideResetToDefault.h b/Engine/Source/Editor/DetailCustomizations/Private/OverrideResetToDefault.h
--- a/Engine/Source/Editor/DetailCustomizations/Private/OverrideResetToDefault.h
+++ b/Engine/Source/Editor/DetailCustomizations/Private/OverrideResetToDefault.h
@@ -11,6 +11,11 @@
 
 /**
  * Class template aiming to override ResetToDefault methods by comparing properties with an internal static DefaultObject.
+ * 
+ * This behavior was originally written for settings UI and does not properly support inheritance (resetting a parent class's
+ * struct member to default and have it propagate to loaded children). If a struct wants to override ResetToDefault behavior
+ * using this class and must support the outer-object being a parent class, OnResetToDefault should be revisited to properly
+ * reset archetype instances as well.
  */
 template<typename UStructType>
 class TOverrideResetToDefaultWithStaticUStruct 
@@ -66,7 +71,11 @@ void TOverrideResetToDefaultWithStaticUStruct<UStructType>::OnResetToDefault(TSh
 
 	if ((DefaultValuePtr != nullptr) && (ValuePtr != nullptr))
 	{
-		PropertyAccessUtil::SetPropertyValue_DirectSingle(Property, DefaultValuePtr, Property, ValuePtr, 0, Property->HasAnyFlags(RF_ArchetypeObject), []() { return nullptr; });
+		// SetPropertyValue_DirectSingle when called on an archetype can propagate to instances. 
+		// See class comments: this ResetToDefault override was not written for contexts involving 
+		// inheritance, so pass in empty archetype instances list.
+		TArray<void*> NoArchetypeInsts;
+		PropertyAccessUtil::SetPropertyValue_DirectSingle(Property, DefaultValuePtr, Property, ValuePtr, NoArchetypeInsts, 0, Property->HasAnyFlags(RF_ArchetypeObject), []() { return nullptr; });
 	}
 }
 
diff --git a/Engine/Source/Runtime/CoreUObject/Private/UObject/PropertyAccessUtil.cpp b/Engine/Source/Runtime/CoreUObject/Private/UObject/PropertyAccessUtil.cpp
--- a/Engine/Source/Runtime/CoreUObject/Private/UObject/PropertyAccessUtil.cpp
+++ b/Engine/Source/Runtime/CoreUObject/Private/UObject/PropertyAccessUtil.cpp
@@ -424,31 +424,204 @@ EPropertyAccessResultFlags CanGetPropertyValue(const FProperty* InProp)
 	return EPropertyAccessResultFlags::Success;
 }
 
+bool GetArchetypeInstancesInheritingPropertyValue(const FProperty* InObjectProp, UObject* InObject, TArray<UObject*>& OutArchetypeInstances)
+{
+	// Clear result container
+	OutArchetypeInstances.Empty();
+
+	if (!InObjectProp || !InObject)
+	{
+		// Invalid parameters
+		return false;
+	}
+
+	if (!InObject->HasAnyFlags(RF_ClassDefaultObject | RF_ArchetypeObject))
+	{
+		// InObject isn't an archetype
+		return false;
+	}
+
+	// Get all archetype instances, which includes instances through multiple levels of inheritance
+	TArray<UObject*> Instances;
+	InObject->GetArchetypeInstances(Instances);
+
+	//
+	const bool bIsSparseProp = !InObject->GetClass()->HasProperty(InObjectProp);
+
+	// Transform object address to the property's reference value address
+	const void* ObjectContainerData = bIsSparseProp ? InObject->GetClass()->GetOrCreateSparseClassData() : InObject;
+	const void* ObjectReferenceValue = InObjectProp->ContainerPtrToValuePtr<void>(ObjectContainerData);
+
+	// Find archetype instances that inherited the property value, which is true when:
+	// - It has the same value as the archetype
+	// - And this is true for every archetype between InObject (base archetype) and Inst
+	// Since archetype instances may be visited multiple times - directly or as archetype 
+	// of another instance - we'll cache their property comparison result.
+	TSet<UObject*> SameValueInstances;
+	TSet<UObject*> DifferentValueInstances;
+	for (UObject* CandidateInstance : Instances)
+	{
+		// If we've already determined for this instance that it has a different value while processing it 
+		// as the archetype for an earlier iterated instance, ignore it
+		if (DifferentValueInstances.Contains(CandidateInstance))
+		{
+			continue;
+		}
+
+		// Visit archetype upwards until InObject is reached or abort if any archetype with a wrong value is encountered
+		UObject* InstIter = CandidateInstance;
+		bool bInheritsArchetypeValue = true;
+		do
+		{
+			if (SameValueInstances.Contains(InstIter))
+			{
+				// Instance has been evaluated during another processing of another instance,
+				// and property value was the same. Continue traveling up.
+			}
+			else if (DifferentValueInstances.Contains(InstIter))
+			{
+				// Instance has been evaluated during another processing of another instance,
+				// and property value was different. This candidate is rejected.
+				bInheritsArchetypeValue = false;
+				continue;
+			}
+			else
+			{
+				// Instance hasn't been evaluated, compare properties now.
+
+				// Get the container of the property: either the object itself or if the property is stored externally
+				// in the class's sparse data, get the sparse data container.
+				void* ContainerData = bIsSparseProp ? InstIter->GetClass()->GetOrCreateSparseClassData() : InstIter;
+				if (ContainerData != nullptr)
+				{
+					// Get the value pointer inside the container
+					void* InstanceValue = InObjectProp->ContainerPtrToValuePtr<void>(ContainerData);
+
+					if (PropertyAccessUtil::IsCompletePropertyIdentical(InObjectProp, ObjectReferenceValue, InObjectProp, InstanceValue))
+					{
+						SameValueInstances.Add(InstIter);
+					}
+					else
+					{
+						DifferentValueInstances.Add(InstIter);
+						bInheritsArchetypeValue = false;
+						continue;
+					}
+				}
+				else
+				{
+					// Could not access sparse class data. At moment of writing this, a child blueprint class with sparse data
+					// can have a null SparseDataClass in editor until it's compiled manually at least once. Since the data is
+					// yet to be initialized, ignore the candidate now. Serialization will take care of inheritance later.
+					bInheritsArchetypeValue = false;
+					continue;
+				}
+			}
+
+			// Process archetype of this instance, it may be the base archetype but there may be
+			// multiple levels of archetypes in-between.
+			InstIter = InstIter->GetArchetype();
+			ensureMsgf(InstIter, TEXT("Unexpected: while traveling up archetypes did not encounter the base archetype"));
+		} while (InstIter && InstIter != InObject && bInheritsArchetypeValue);
+
+		if (bInheritsArchetypeValue && InstIter == InObject)
+		{
+			OutArchetypeInstances.Add(CandidateInstance);
+		}
+	}
+
+	// Successfully searched for archetype instances
+	return true;
+}
+
+bool GetArchetypeInstancesInheritingPropertyValue_AsContainerData(const FProperty* InObjectProp, UObject* InObject, TArray<void*>& OutArchetypeInstContainers)
+{
+	OutArchetypeInstContainers.Empty();
+
+	// First get all the archetype UObjects that inherit InObject's value
+	TArray<UObject*> CandidateInstances;
+	GetArchetypeInstancesInheritingPropertyValue(InObjectProp, InObject, CandidateInstances);
+
+	// If the property is part of the UClass's SparseClassDataStruct, the value is not stored in the UObject's allocated
+	// memory but externally in the sparse data struct. Check which is the case and then output the mem address of the
+	// relevant container (UObject, or the sparse data struct).
+	const bool bIsSparseProp = !InObject->GetClass()->HasProperty(InObjectProp);
+	const void* ObjectContainerData = bIsSparseProp ? InObject->GetClass()->GetOrCreateSparseClassData() : InObject;
+	for (UObject* CandidateInstance : CandidateInstances)
+	{
+		if (bIsSparseProp)
+		{
+			// If the property is stored in sparse data struct, resolve address to that data.
+			void* CandidateContainerData = CandidateInstance->GetClass()->GetOrCreateSparseClassData();
+
+			// Objects of the same class share the sparse data struct. If this candidate shares the data struct 
+			// with the reference object, ignore it since the reference object will already be updated.
+			if (CandidateContainerData != ObjectContainerData)
+			{
+				// Instances of the same class share the sparse data struct, so make sure to add
+				// mem addresses uniquely.
+				OutArchetypeInstContainers.AddUnique(CandidateContainerData);
+			}
+		}
+		else
+		{
+			// If not a sparse property, the data is stored in the instance itself so add the instance
+			// pointer as container.
+			OutArchetypeInstContainers.Add(CandidateInstance);
+		}
+	}
+
+	return false;
+}
+
 EPropertyAccessResultFlags SetPropertyValue_Object(const FProperty* InObjectProp, UObject* InObject, const FProperty* InSrcProp, const void* InSrcValue, const int32 InArrayIndex, const uint64 InReadOnlyFlags, const EPropertyAccessChangeNotifyMode InNotifyMode)
 {
 	check(InObject->IsA(InObjectProp->GetOwnerClass()));
-	return SetPropertyValue_InContainer(InObjectProp, InObject, InSrcProp, InSrcValue, InArrayIndex, InReadOnlyFlags, IsObjectTemplate(InObject), [InObjectProp, InObject, InNotifyMode]()
+
+	// If the object is an archetype, we'll propagate the value change to instances that currently inheriting the object's value.
+	// Collect all archetype instances that satisfy that condition.
+	TArray<void*> PropInheritingInstContainers;
+	if (InObject->HasAnyFlags(RF_ClassDefaultObject | RF_ArchetypeObject))
+	{
+		GetArchetypeInstancesInheritingPropertyValue_AsContainerData(InObjectProp, InObject, PropInheritingInstContainers);
+	}
+
+	return SetPropertyValue_InContainer(InObjectProp, InObject, PropInheritingInstContainers, InSrcProp, InSrcValue, InArrayIndex, InReadOnlyFlags, IsObjectTemplate(InObject), [InObjectProp, InObject, InNotifyMode]()
 	{
-		return BuildBasicChangeNotify(InObjectProp, InObject, InNotifyMode);
+		return BuildBasicChangeNotify(InObjectProp, InObject, InNotifyMode, /*ChangeType=*/EPropertyChangeType::ValueSet);
 	});
 }
 
-EPropertyAccessResultFlags SetPropertyValue_InContainer(const FProperty* InContainerProp, void* InContainerData, const FProperty* InSrcProp, const void* InSrcValue, const int32 InArrayIndex, const uint64 InReadOnlyFlags, const bool InOwnerIsTemplate, const FPropertyAccessBuildChangeNotifyFunc& InBuildChangeNotifyFunc)
+EPropertyAccessResultFlags SetPropertyValue_InContainer(const FProperty* InContainerProp, void* InContainerData, const TArray<void*>& InArchetypeInstContainerData, const FProperty* InSrcProp, const void* InSrcValue, const int32 InArrayIndex, const uint64 InReadOnlyFlags, const bool InOwnerIsTemplate, const FPropertyAccessBuildChangeNotifyFunc& InBuildChangeNotifyFunc)
 {
+	// If the object we're setting the property value of is an archetype, InArchetypeInstContainerData will have been
+	// populated with pointers to archetype instances that also need their property value set. Transform the container 
+	// pointers of archetype instances to their value pointers.
+	TArray<void*> ArchetypeInstDestValues;
+	if (InOwnerIsTemplate)
+	{
+		ArchetypeInstDestValues.Reserve(InArchetypeInstContainerData.Num());
+		for (void* InstContainerData : InArchetypeInstContainerData)
+		{
+			void* InstValuePtr = InContainerProp->ContainerPtrToValuePtr<void>(InstContainerData);
+			ArchetypeInstDestValues.Add(InstValuePtr);
+		}
+	}
+
 	if (InArrayIndex == INDEX_NONE || InContainerProp->ArrayDim == 1)
 	{
 		void* DestValue = InContainerProp->ContainerPtrToValuePtr<void>(InContainerData);
-		return SetPropertyValue_DirectComplete(InSrcProp, InSrcValue, InContainerProp, DestValue, InReadOnlyFlags, InOwnerIsTemplate, InBuildChangeNotifyFunc);
+		return SetPropertyValue_DirectComplete(InSrcProp, InSrcValue, InContainerProp, DestValue, ArchetypeInstDestValues, InReadOnlyFlags, InOwnerIsTemplate, InBuildChangeNotifyFunc);
 	}
 	else
 	{
 		check(InArrayIndex < InContainerProp->ArrayDim);
 		void* DestValue = InContainerProp->ContainerPtrToValuePtr<void>(InContainerData, InArrayIndex);
-		return SetPropertyValue_DirectSingle(InSrcProp, InSrcValue, InContainerProp, DestValue, InReadOnlyFlags, InOwnerIsTemplate, InBuildChangeNotifyFunc);
+		return SetPropertyValue_DirectSingle(InSrcProp, InSrcValue, InContainerProp, DestValue, ArchetypeInstDestValues, InReadOnlyFlags, InOwnerIsTemplate, InBuildChangeNotifyFunc);
 	}
 }
 
-EPropertyAccessResultFlags SetPropertyValue_DirectSingle(const FProperty* InSrcProp, const void* InSrcValue, const FProperty* InDestProp, void* InDestValue, const uint64 InReadOnlyFlags, const bool InOwnerIsTemplate, const FPropertyAccessBuildChangeNotifyFunc& InBuildChangeNotifyFunc)
+EPropertyAccessResultFlags SetPropertyValue_DirectSingle(const FProperty* InSrcProp, const void* InSrcValue, const FProperty* InDestProp, void* InDestValue, const TArray<void*>& InArchetypeInstDestValues, const uint64 InReadOnlyFlags, const bool InOwnerIsTemplate, const FPropertyAccessBuildChangeNotifyFunc& InBuildChangeNotifyFunc)
 {
 	EPropertyAccessResultFlags Result = CanSetPropertyValue(InDestProp, InReadOnlyFlags, InOwnerIsTemplate);
 	if (Result != EPropertyAccessResultFlags::Success)
@@ -456,13 +629,27 @@ EPropertyAccessResultFlags SetPropertyValue_DirectSingle(const FProperty* InSrcP
 		return Result;
 	}
 
-	return SetPropertyValue([InSrcProp, InSrcValue, InDestProp, InDestValue](const FPropertyAccessChangeNotify* InChangeNotify)
+	return SetPropertyValue([InSrcProp, InSrcValue, InDestProp, InDestValue, InArchetypeInstDestValues, InOwnerIsTemplate](const FPropertyAccessChangeNotify* InChangeNotify)
 	{
 		bool bResult = true;
 		const bool bIdenticalValue = IsSinglePropertyIdentical(InSrcProp, InSrcValue, InDestProp, InDestValue);
 		EmitPreChangeNotify(InChangeNotify, bIdenticalValue);
 		if (!bIdenticalValue)
 		{
+			// Special case: modifying an archetype, process instances
+			if (InOwnerIsTemplate)
+			{
+				// Fill archetype instances who currently have the same value as the archetype
+				for (void* ArchetypeInstDestValue : InArchetypeInstDestValues)
+				{
+					// Only expecting to get archetype instances passed in whose current value is the same as the pre-change archetype value
+					ensureMsgf(IsSinglePropertyIdentical(InDestProp, InDestValue, InDestProp, ArchetypeInstDestValue), TEXT("SetPropertyValue modifying archetype instance which had non-inherited value"));
+
+					// Write the new property to the archetype instance
+					CopySinglePropertyValue(InSrcProp, InSrcValue, InDestProp, ArchetypeInstDestValue);
+				}
+			}
+			
 			bResult = CopySinglePropertyValue(InSrcProp, InSrcValue, InDestProp, InDestValue);
 		}
 		EmitPostChangeNotify(InChangeNotify, bIdenticalValue);
@@ -471,7 +658,7 @@ EPropertyAccessResultFlags SetPropertyValue_DirectSingle(const FProperty* InSrcP
 	}, InBuildChangeNotifyFunc);
 }
 
-EPropertyAccessResultFlags SetPropertyValue_DirectComplete(const FProperty* InSrcProp, const void* InSrcValue, const FProperty* InDestProp, void* InDestValue, const uint64 InReadOnlyFlags, const bool InOwnerIsTemplate, const FPropertyAccessBuildChangeNotifyFunc& InBuildChangeNotifyFunc)
+EPropertyAccessResultFlags SetPropertyValue_DirectComplete(const FProperty* InSrcProp, const void* InSrcValue, const FProperty* InDestProp, void* InDestValue, const TArray<void*>& InArchetypeInstDestValues, const uint64 InReadOnlyFlags, const bool InOwnerIsTemplate, const FPropertyAccessBuildChangeNotifyFunc& InBuildChangeNotifyFunc)
 {
 	EPropertyAccessResultFlags Result = CanSetPropertyValue(InDestProp, InReadOnlyFlags, InOwnerIsTemplate);
 	if (Result != EPropertyAccessResultFlags::Success)
@@ -479,13 +666,35 @@ EPropertyAccessResultFlags SetPropertyValue_DirectComplete(const FProperty* InSr
 		return Result;
 	}
 
-	return SetPropertyValue([InSrcProp, InSrcValue, InDestProp, InDestValue](const FPropertyAccessChangeNotify* InChangeNotify)
+	return SetPropertyValue([InSrcProp, InSrcValue, InDestProp, InDestValue, InArchetypeInstDestValues, InOwnerIsTemplate](const FPropertyAccessChangeNotify* InChangeNotify)
 	{
 		bool bResult = true;
 		const bool bIdenticalValue = IsCompletePropertyIdentical(InSrcProp, InSrcValue, InDestProp, InDestValue);
 		EmitPreChangeNotify(InChangeNotify, bIdenticalValue);
 		if (!bIdenticalValue)
 		{
+			// Special case: modifying an archetype, process instances
+			if (InOwnerIsTemplate)
+			{
+				// Fill archetype instances who currently have the same value as the archetype
+				for (void* ArchetypeInstDestValue : InArchetypeInstDestValues)
+				{
+#if DO_ENSURE
+					// Only expecting to get archetype instances passed in whose current value is the same as the pre-change archetype value
+					if (!IsCompletePropertyIdentical(InDestProp, InDestValue, InDestProp, ArchetypeInstDestValue))
+					{
+						FString A, B;
+						InDestProp->ExportTextItem_Direct(A, InDestValue, nullptr, nullptr, 0);
+						InDestProp->ExportTextItem_Direct(B, ArchetypeInstDestValue, nullptr, nullptr, 0);
+						ensureMsgf(IsCompletePropertyIdentical(InDestProp, InDestValue, InDestProp, ArchetypeInstDestValue), TEXT("SetPropertyValue modifying archetype instance which had non-inherited value: %s (template) vs. %s (instance)"), *A, *B);
+					}
+#endif
+
+					// Write the new property to the archetype instance
+					CopyCompletePropertyValue(InSrcProp, InSrcValue, InDestProp, ArchetypeInstDestValue);
+				}
+			}
+
 			bResult = CopyCompletePropertyValue(InSrcProp, InSrcValue, InDestProp, InDestValue);
 		}
 		EmitPostChangeNotify(InChangeNotify, bIdenticalValue);
@@ -568,7 +777,7 @@ void EmitPostChangeNotify(const FPropertyAccessChangeNotify* InChangeNotify, con
 #endif
 }
 
-TUniquePtr<FPropertyAccessChangeNotify> BuildBasicChangeNotify(const FProperty* InProp, const UObject* InObject, const EPropertyAccessChangeNotifyMode InNotifyMode)
+TUniquePtr<FPropertyAccessChangeNotify> BuildBasicChangeNotify(const FProperty* InProp, const UObject* InObject, const EPropertyAccessChangeNotifyMode InNotifyMode, const EPropertyChangeType::Type ChangeType)
 {
 	const UScriptStruct* SparseStruct = InObject->GetClass()->GetSparseClassDataStruct();
 	const bool bIsValidSparseProp = SparseStruct &&
@@ -578,6 +787,7 @@ TUniquePtr<FPropertyAccessChangeNotify> BuildBasicChangeNotify(const FProperty*
 	if (InNotifyMode != EPropertyAccessChangeNotifyMode::Never)
 	{
 		TUniquePtr<FPropertyAccessChangeNotify> ChangeNotify = MakeUnique<FPropertyAccessChangeNotify>();
+		ChangeNotify->ChangeType = ChangeType;
 		ChangeNotify->ChangedObject = const_cast<UObject*>(InObject);
 		ChangeNotify->ChangedPropertyChain.AddHead(const_cast<FProperty*>(InProp));
 		ChangeNotify->ChangedPropertyChain.SetActivePropertyNode(const_cast<FProperty*>(InProp));
diff --git a/Engine/Source/Runtime/CoreUObject/Public/UObject/PropertyAccessUtil.h b/Engine/Source/Runtime/CoreUObject/Public/UObject/PropertyAccessUtil.h
--- a/Engine/Source/Runtime/CoreUObject/Public/UObject/PropertyAccessUtil.h
+++ b/Engine/Source/Runtime/CoreUObject/Public/UObject/PropertyAccessUtil.h
@@ -168,7 +168,29 @@ namespace PropertyAccessUtil
 	COREUOBJECT_API EPropertyAccessResultFlags CanGetPropertyValue(const FProperty* InProp);
 
 	/**
-	 * High-level function for setting the value of a property on an object.
+	 * Finds all instances of an archetype object which are inheriting the given property's value from the archetype, then returns those instance UObjects.
+	 *
+	 * @param InObjectProp Property to compare the value of.
+	 * @param InObject Archetype object containing the property.
+	 * @param OutArchetypeInstances Array to store results in.
+	 *
+	 * @return True if archetype instances were searched and regardless of whether any instances were returned.
+	 */
+	COREUOBJECT_API bool GetArchetypeInstancesInheritingPropertyValue(const FProperty* InObjectProp, UObject* InObject, TArray<UObject*>& OutArchetypeInstances);
+
+	/**
+	 * Finds all instances of an archetype object which are inheriting the given property's value from the archetype, then returns the containers for the given property which is either the UObject's address or the sparse data struct.
+	 *
+	 * @param InObjectProp Property to compare the value of.
+	 * @param InObject Archetype object containing the property.
+	 * @param OutArchetypeInstContainers Array to store results in. These are pointers to the container containing the given property value. If the property is sparse class data, these will point to the sparse data struct.
+	 *
+	 * @return True if archetype instances were searched and regardless of whether any instances were returned.
+	 */
+	COREUOBJECT_API bool GetArchetypeInstancesInheritingPropertyValue_AsContainerData(const FProperty* InObjectProp, UObject* InObject, TArray<void*>& OutArchetypeInstContainers);
+
+	/**
+	 * High-level function for setting the value of a property on an object. If the object is an archetype, value will be propagated to instances that haven't modified the value.
 	 * @note This function calls CanSetPropertyValue internally, and will emit property change notifications for the object.
 	 *
 	 * @param InObjectProp Property to set the value of.
@@ -189,6 +211,7 @@ namespace PropertyAccessUtil
 	 *
 	 * @param InContainerProp Property to set the value of.
 	 * @param InContainerData The instance data containing the property.
+	 * @param InArchetypeInstContainerData When the object is an archetype: container addresses of instances of the archetype to evaluate. Those with currently the same value as archetype will also have their property updates.
 	 * @param InSrcProp Property of the value to set (must be compatible with the dest property).
 	 * @param InSrcValue The value to set on the property.
 	 * @param InArrayIndex For fixed-size array properties denotes which index of the array to set, or INDEX_NONE to set the entire property.
@@ -198,7 +221,7 @@ namespace PropertyAccessUtil
 	 *
 	 * @return Flags describing whether the set was successful.
 	 */
-	COREUOBJECT_API EPropertyAccessResultFlags SetPropertyValue_InContainer(const FProperty* InContainerProp, void* InContainerData, const FProperty* InSrcProp, const void* InSrcValue, const int32 InArrayIndex, const uint64 InReadOnlyFlags, const bool InOwnerIsTemplate, const FPropertyAccessBuildChangeNotifyFunc& InBuildChangeNotifyFunc);
+	COREUOBJECT_API EPropertyAccessResultFlags SetPropertyValue_InContainer(const FProperty* InContainerProp, void* InContainerData, const TArray<void*>& InArchetypeInstContainerData, const FProperty* InSrcProp, const void* InSrcValue, const int32 InArrayIndex, const uint64 InReadOnlyFlags, const bool InOwnerIsTemplate, const FPropertyAccessBuildChangeNotifyFunc& InBuildChangeNotifyFunc);
 	
 	/**
 	 * High-level function for setting the single-element value of a property in memory.
@@ -208,13 +231,14 @@ namespace PropertyAccessUtil
 	 * @param InSrcValue The value to set on the property.
 	 * @param InDestProp Property to get the value from (must be compatible with the source property).
 	 * @param InDestValue Instance to fill with the property value (must be a valid and constructed block of memory that is compatible with the property).
+	 * @param InArchetypeInstDestValues When the object is an archetype: direct value addresses of instances of the archetype to evaluate. Those with currently the same value as archetype will also have their property updates.
 	 * @param InReadOnlyFlags Flags controlling which properties are considered read-only.
 	 * @param InOwnerIsTemplate True if the owner object is considered a template (see IsObjectTemplate).
 	 * @param InBuildChangeNotifyFunc Logic for building the information needed to emit property change notifications when setting a property value (can return nullptr if no notifications are needed or possible).
 	 *
 	 * @return Flags describing whether the set was successful.
 	 */
-	COREUOBJECT_API EPropertyAccessResultFlags SetPropertyValue_DirectSingle(const FProperty* InSrcProp, const void* InSrcValue, const FProperty* InDestProp, void* InDestValue, const uint64 InReadOnlyFlags, const bool InOwnerIsTemplate, const FPropertyAccessBuildChangeNotifyFunc& InBuildChangeNotifyFunc);
+	COREUOBJECT_API EPropertyAccessResultFlags SetPropertyValue_DirectSingle(const FProperty* InSrcProp, const void* InSrcValue, const FProperty* InDestProp, void* InDestValue, const TArray<void*>& InArchetypeInstDestValues, const uint64 InReadOnlyFlags, const bool InOwnerIsTemplate, const FPropertyAccessBuildChangeNotifyFunc& InBuildChangeNotifyFunc);
 
 	/**
 	 * High-level function for setting the multi-element value of a property in memory.
@@ -224,13 +248,14 @@ namespace PropertyAccessUtil
 	 * @param InSrcValue The value to set on the property.
 	 * @param InDestProp Property to get the value from (must be compatible with the source property).
 	 * @param InDestValue Instance to fill with the property value (must be a valid and constructed block of memory that is compatible with the property).
+	 * @param InArchetypeInstDestValues When the object is an archetype: direct value addresses of instances of the archetype to evaluate. Those with currently the same value as archetype will also have their property updates.
 	 * @param InReadOnlyFlags Flags controlling which properties are considered read-only.
 	 * @param InOwnerIsTemplate True if the owner object is considered a template (see IsObjectTemplate).
 	 * @param InBuildChangeNotifyFunc Logic for building the information needed to emit property change notifications when setting a property value (can return nullptr if no notifications are needed or possible).
 	 *
 	 * @return Flags describing whether the set was successful.
 	 */
-	COREUOBJECT_API EPropertyAccessResultFlags SetPropertyValue_DirectComplete(const FProperty* InSrcProp, const void* InSrcValue, const FProperty* InDestProp, void* InDestValue, const uint64 InReadOnlyFlags, const bool InOwnerIsTemplate, const FPropertyAccessBuildChangeNotifyFunc& InBuildChangeNotifyFunc);
+	COREUOBJECT_API EPropertyAccessResultFlags SetPropertyValue_DirectComplete(const FProperty* InSrcProp, const void* InSrcValue, const FProperty* InDestProp, void* InDestValue, const TArray<void*>& InArchetypeInstDestValues, const uint64 InReadOnlyFlags, const bool InOwnerIsTemplate, const FPropertyAccessBuildChangeNotifyFunc& InBuildChangeNotifyFunc);
 
 	/**
 	 * Low-level function for setting the value of a property.
@@ -278,7 +303,7 @@ namespace PropertyAccessUtil
 	 *
 	 * @return The information needed to emit property change notifications.
 	 */
-	COREUOBJECT_API TUniquePtr<FPropertyAccessChangeNotify> BuildBasicChangeNotify(const FProperty* InProp, const UObject* InObject, const EPropertyAccessChangeNotifyMode InNotifyMode);
+	COREUOBJECT_API TUniquePtr<FPropertyAccessChangeNotify> BuildBasicChangeNotify(const FProperty* InProp, const UObject* InObject, const EPropertyAccessChangeNotifyMode InNotifyMode, const EPropertyChangeType::Type ChangeType);
 
 	/**
 	 * Low-level function for checking whether the given object instance is considered a template for property access.
diff --git a/Engine/Source/Runtime/Engine/Private/KismetSystemLibrary.cpp b/Engine/Source/Runtime/Engine/Private/KismetSystemLibrary.cpp
--- a/Engine/Source/Runtime/Engine/Private/KismetSystemLibrary.cpp
+++ b/Engine/Source/Runtime/Engine/Private/KismetSystemLibrary.cpp
@@ -3280,17 +3280,28 @@ bool UKismetSystemLibrary::Generic_SetEditorProperty(UObject* Object, const FNam
 			{
 				void* SparseDest = Object->GetClass()->GetOrCreateSparseClassData();
 
+				// If the object is a template, instances of it may need to be updated as well depending on whether they are inheriting their
+				// current value. Detect those instances which haven't overwritten the value. In the case of sparse data, this applies to
+				// subclass CDOs.
+				TArray<void*> ArchetypeInstances;
+				const bool bObjectIsTemplate = PropertyAccessUtil::IsObjectTemplate(Object);
+				if (bObjectIsTemplate)
+				{
+					PropertyAccessUtil::GetArchetypeInstancesInheritingPropertyValue_AsContainerData(SparseProp, Object, ArchetypeInstances);
+				}
+
 				SparseDataAccessResult = PropertyAccessUtil::SetPropertyValue_InContainer(
 					SparseProp,
 					SparseDest,
+					ArchetypeInstances,
 					ValueProp,
 					ValuePtr,
 					INDEX_NONE,
 					PropertyAccessUtil::EditorReadOnlyFlags,
-					PropertyAccessUtil::IsObjectTemplate(Object),
+					bObjectIsTemplate,
 					[SparseProp, Object, ChangeNotifyMode]()
 					{
-						return PropertyAccessUtil::BuildBasicChangeNotify(SparseProp, Object, ChangeNotifyMode);
+						return PropertyAccessUtil::BuildBasicChangeNotify(SparseProp, Object, ChangeNotifyMode, EPropertyChangeType::ValueSet);
 					});
 				if (*SparseDataAccessResult == EPropertyAccessResultFlags::Success)
 				{
