Core: Fixed missing UE_IMPLEMENT_STRUCT for FFrameRate, FFrameTime, FQualifiedFrameTime, FTimecode.

Unfortunately due to assets/classes already using these datatypes (for the last six years!) the custom serializer used by the native code can't be used by the USTRUCT code. This is okay as long as the serialization doesn't do anything that can't be supported by the generic UPROPERTY serialization which should be fine for POD like this.

Things specifically tested:
Loaded assets that used the default values for FrameRates to ensure the default had not changed.
Loaded assets that used non-default values for FFrameNumber to ensure the values still loaded.
Played back Level Sequences in the Editor
Rendered jobs using Movie Render Queue using default and non-standard Frame Rate/Range overrides.

diff --git a/Engine/Source/Runtime/CoreUObject/Private/UObject/Property.cpp b/Engine/Source/Runtime/CoreUObject/Private/UObject/Property.cpp
--- a/Engine/Source/Runtime/CoreUObject/Private/UObject/Property.cpp
+++ b/Engine/Source/Runtime/CoreUObject/Private/UObject/Property.cpp
@@ -11,6 +11,11 @@
 #include "Misc/AsciiSet.h"
 #include "Misc/Guid.h"
 #include "Misc/StringBuilder.h"
+#include "Misc/FrameRate.h"
+#include "Misc/FrameTime.h"
+#include "Misc/QualifiedFrameTime.h"
+#include "Misc/FrameNumber.h"
+#include "Misc/Timecode.h"
 #include "Serialization/TestUndeclaredScriptStructObjectReferences.h"
 #include "Templates/Casts.h"
 #include "UObject/Class.h"
@@ -525,6 +530,7 @@ struct TStructOpsTypeTraits<FFrameNumber> : public TStructOpsTypeTraitsBase2<FFr
 {
 	enum
 	{
+		WithZeroConstructor = true,
 		WithSerializer = true,
 		WithIdenticalViaEquality = true
 	};
@@ -532,6 +538,51 @@ struct TStructOpsTypeTraits<FFrameNumber> : public TStructOpsTypeTraitsBase2<FFr
 };
 UE_IMPLEMENT_STRUCT("/Script/CoreUObject", FrameNumber);
 
+template<>
+struct TStructOpsTypeTraits<FFrameRate> : public TStructOpsTypeTraitsBase2<FFrameRate>
+{
+	enum
+	{
+		// The native function has a custom serializer but assets have already been created with the generic UPROPERTY serializer,
+		// so we can't switch them to use a custom serializer without breaking assets (creates mismatched sizes in data).
+		// WithSerializer = true, 
+		WithIdenticalViaEquality = true
+	};
+	static constexpr EPropertyObjectReferenceType WithSerializerObjectReferences = EPropertyObjectReferenceType::None;
+};
+UE_IMPLEMENT_STRUCT("/Script/CoreUObject", FrameRate);
+
+template<>
+struct TStructOpsTypeTraits<FFrameTime> : public TStructOpsTypeTraitsBase2<FFrameTime>
+{
+	enum
+	{
+		WithZeroConstructor = true,
+		// The native function has a custom serializer but assets have already been created with the generic UPROPERTY serializer,
+		// so we can't switch them to use a custom serializer without breaking assets (creates mismatched sizes in data).
+		// WithSerializer = true, 
+		WithIdenticalViaEquality = true
+	};
+	static constexpr EPropertyObjectReferenceType WithSerializerObjectReferences = EPropertyObjectReferenceType::None;
+};
+UE_IMPLEMENT_STRUCT("/Script/CoreUObject", FrameTime);
+
+// Qualified Frame Times can't be zero-initialized because they contain FrameRates,
+// and have no equality operator.
+UE_IMPLEMENT_STRUCT("/Script/CoreUObject", QualifiedFrameTime);
+
+template<>
+struct TStructOpsTypeTraits<FTimecode> : public TStructOpsTypeTraitsBase2<FTimecode>
+{
+	enum
+	{
+		WithZeroConstructor = true,
+		WithIdenticalViaEquality = true
+	};
+	static constexpr EPropertyObjectReferenceType WithSerializerObjectReferences = EPropertyObjectReferenceType::None;
+};
+UE_IMPLEMENT_STRUCT("/Script/CoreUObject", Timecode);
+
 template<>
 struct TStructOpsTypeTraits<FSoftObjectPath> : public TStructOpsTypeTraitsBase2<FSoftObjectPath>
 {
diff --git a/Engine/Source/Runtime/CoreUObject/Public/UObject/NoExportTypes.h b/Engine/Source/Runtime/CoreUObject/Public/UObject/NoExportTypes.h
--- a/Engine/Source/Runtime/CoreUObject/Public/UObject/NoExportTypes.h
+++ b/Engine/Source/Runtime/CoreUObject/Public/UObject/NoExportTypes.h
@@ -21,6 +21,11 @@
 #include "Misc/Guid.h"
 #include "Misc/DateTime.h"
 #include "Misc/Timespan.h"
+#include "Misc/FrameRate.h"
+#include "Misc/FrameTime.h"
+#include "Misc/QualifiedFrameTime.h"
+#include "Misc/FrameNumber.h"
+#include "Misc/Timecode.h"
 
 #include "UObject/TopLevelAssetPath.h"
 #include "UObject/SoftObjectPath.h"
