Added creation of transient enums for IDOs to support unknown enumerators

diff --git a/Engine/Source/Runtime/CoreUObject/Internal/UObject/InstanceDataObjectUtils.h b/Engine/Source/Runtime/CoreUObject/Internal/UObject/InstanceDataObjectUtils.h
--- a/Engine/Source/Runtime/CoreUObject/Internal/UObject/InstanceDataObjectUtils.h
+++ b/Engine/Source/Runtime/CoreUObject/Internal/UObject/InstanceDataObjectUtils.h
@@ -15,6 +15,7 @@ class UObject;
 class UStruct;
 
 namespace UE { class FPropertyPathNameTree; }
+namespace UE { class FUnknownEnumNames; }
 
 namespace UE
 {
@@ -29,7 +30,7 @@ bool StructContainsLooseProperties(const UStruct* Struct);
 bool IsClassOfInstanceDataObjectClass(UStruct* Class);
 
 /** Generate a UClass that contains the union of the properties of PropertyTree and OwnerClass. */
-UClass* CreateInstanceDataObjectClass(const FPropertyPathNameTree* PropertyTree, UClass* OwnerClass, UObject* Outer);
+UClass* CreateInstanceDataObjectClass(const FPropertyPathNameTree* PropertyTree, const FUnknownEnumNames* EnumNames, UClass* OwnerClass, UObject* Outer);
 
 /** Notify that a property in an struct was set when the struct was deserialized. */
 void MarkPropertyValueSerialized(const UStruct* Struct, void* StructData, const FProperty* Property, int32 ArrayIndex = 0);
diff --git a/Engine/Source/Runtime/CoreUObject/Internal/UObject/PropertyBagRepository.h b/Engine/Source/Runtime/CoreUObject/Internal/UObject/PropertyBagRepository.h
--- a/Engine/Source/Runtime/CoreUObject/Internal/UObject/PropertyBagRepository.h
+++ b/Engine/Source/Runtime/CoreUObject/Internal/UObject/PropertyBagRepository.h
@@ -6,6 +6,7 @@
 #include "Containers/Map.h"
 #include "HAL/CriticalSection.h"
 #include "UObject/GCObject.h"
+#include "UObject/PropertyTypeName.h"
 #include "Templates/FunctionFwd.h"
 
 class UObject;
@@ -17,7 +18,6 @@ namespace UE
 #if WITH_EDITORONLY_DATA
 
 class FPropertyPathNameTree;
-class FPropertyTypeName;
 class FUnknownEnumNames;
 
 // Singleton class tracking property bag association with objects
@@ -88,7 +88,12 @@ public:
 	 * @param OutNames		Array to assign the unknown names to. Empty on return if no names are found.
 	 * @param bOutHasFlags	Assigned to true if the enum is known to have flags, otherwise false.
 	 */
-	void FindUnknownEnumNames(const UObject* Owner, FPropertyTypeName EnumTypeName, TArray<FName>& OutNames, bool& bOutHasFlags);
+	void FindUnknownEnumNames(const UObject* Owner, FPropertyTypeName EnumTypeName, TArray<FName>& OutNames, bool& bOutHasFlags) const;
+
+	/**
+	 * Finds tracked unknown enum names associated with the object, otherwise null.
+	 */
+	const FUnknownEnumNames* FindUnknownEnumNames(const UObject* Owner) const;
 
 	/**
 	 * Resets tracked unknown enum names associated with the object.
@@ -199,6 +204,22 @@ private:
 	void CreateInstanceDataObjectUnsafe(UObject* Owner, FPropertyBagAssociationData& BagData, FArchive* Archive = nullptr);
 };
 
+class FUnknownEnumNames
+{
+public:
+	void Add(const UEnum* Enum, FPropertyTypeName EnumTypeName, FName EnumValueName);
+	void Find(FPropertyTypeName EnumTypeName, TArray<FName>& OutNames, bool& bOutHasFlags) const;
+
+private:
+	struct FInfo
+	{
+		TSet<FName> Names;
+		bool bHasFlags = false;
+	};
+
+	TMap<FPropertyTypeName, FInfo> Enums;
+};
+
 #endif // WITH_EDITORONLY_DATA
 
 // construct this context in the same scope as an object is being serialized to support InstanceDataObjects.
diff --git a/Engine/Source/Runtime/CoreUObject/Private/UObject/InstanceDataObjectUtils.cpp b/Engine/Source/Runtime/CoreUObject/Private/UObject/InstanceDataObjectUtils.cpp
--- a/Engine/Source/Runtime/CoreUObject/Private/UObject/InstanceDataObjectUtils.cpp
+++ b/Engine/Source/Runtime/CoreUObject/Private/UObject/InstanceDataObjectUtils.cpp
@@ -197,12 +197,94 @@ namespace UE
 		return Struct->GetBoolMetaData(NAME_ContainsLoosePropertiesMetadata);
 	}
 
-	static UStruct* CreateInstanceDataObjectStructRec(const UClass* StructClass, UStruct* TemplateStruct, UObject* Outer, const FPropertyPathNameTree* PropertyTree);
+	static UStruct* CreateInstanceDataObjectStructRec(const UClass* StructClass, UStruct* TemplateStruct, UObject* Outer, const FPropertyPathNameTree* PropertyTree, const FUnknownEnumNames* EnumNames);
 
 	template <typename StructType>
-	StructType* CreateInstanceDataObjectStructRec(UStruct* TemplateStruct, UObject* Outer, const FPropertyPathNameTree* PropertyTree)
+	StructType* CreateInstanceDataObjectStructRec(UStruct* TemplateStruct, UObject* Outer, const FPropertyPathNameTree* PropertyTree, const FUnknownEnumNames* EnumNames)
 	{
-		return CastChecked<StructType>(CreateInstanceDataObjectStructRec(StructType::StaticClass(), TemplateStruct, Outer, PropertyTree));
+		return CastChecked<StructType>(CreateInstanceDataObjectStructRec(StructType::StaticClass(), TemplateStruct, Outer, PropertyTree, EnumNames));
+	}
+
+	UEnum* FindOrCreateInstanceDataObjectEnum(UEnum* TemplateEnum, UObject* Outer, const FProperty* Property, const FUnknownEnumNames* EnumNames)
+	{
+		if (!TemplateEnum || !EnumNames)
+		{
+			return TemplateEnum;
+		}
+
+		TArray<FName> UnknownNames;
+		bool bHasFlags = false;
+
+		// Use the original type name because the template may be a fallback enum or an IDO.
+		FPropertyTypeName EnumTypeName = FindOriginalType(Property);
+		if (EnumTypeName.IsEmpty())
+		{
+			FPropertyTypeNameBuilder Builder;
+			Builder.AddPath(TemplateEnum);
+			EnumTypeName = Builder.Build();
+		}
+
+		EnumNames->Find(EnumTypeName, UnknownNames, bHasFlags);
+		if (UnknownNames.IsEmpty())
+		{
+			return TemplateEnum;
+		}
+
+		int64 MaxEnumValue = -1;
+		int64 CombinedEnumValues = 0;
+		TArray<TPair<FName, int64>> EnumValueNames;
+		TStringBuilder<128> EnumName(InPlace, EnumTypeName.GetName());
+
+		const auto MakeFullEnumName = [&EnumName, Form = TemplateEnum->GetCppForm()](FName Name) -> FName
+		{
+			if (Form == UEnum::ECppForm::Regular)
+			{
+				return Name;
+			}
+			return FName(WriteToString<128>(EnumName, TEXTVIEW("::"), Name));
+		};
+
+		const auto MakeNextEnumValue = [&MaxEnumValue, &CombinedEnumValues, bHasFlags]() -> int64
+		{
+			if (!bHasFlags)
+			{
+				return ++MaxEnumValue;
+			}
+			const int64 NextEnumValue = ~CombinedEnumValues & (CombinedEnumValues + 1);
+			CombinedEnumValues |= NextEnumValue;
+			return NextEnumValue;
+		};
+
+		// Copy existing values except for MAX.
+		const bool bContainsExistingMax = TemplateEnum->ContainsExistingMax();
+		for (int32 Index = 0, Count = TemplateEnum->NumEnums() - (bContainsExistingMax ? 1 : 0); Index < Count; ++Index)
+		{
+			FName EnumValueName = TemplateEnum->GetNameByIndex(Index);
+			int64 EnumValue = TemplateEnum->GetValueByIndex(Index);
+			EnumValueNames.Emplace(EnumValueName, EnumValue);
+			MaxEnumValue = FMath::Max(MaxEnumValue, EnumValue);
+			CombinedEnumValues |= EnumValue;
+		}
+
+		// Copy unknown names and assign values sequentially.
+		for (FName UnknownName : UnknownNames)
+		{
+			EnumValueNames.Emplace(MakeFullEnumName(UnknownName), MakeNextEnumValue());
+		}
+
+		// Copy or create MAX with a new value.
+		const FName MaxEnumName = bContainsExistingMax ? TemplateEnum->GetNameByIndex(TemplateEnum->NumEnums() - 1) : MakeFullEnumName("MAX");
+		EnumValueNames.Emplace(MaxEnumName, bHasFlags ? CombinedEnumValues : MaxEnumValue);
+
+		// Construct a transient type that impersonates the original type.
+		const FName InstanceDataObjectName(WriteToString<128>(EnumName, TEXTVIEW("_InstanceDataObject")));
+		UEnum* Enum = NewObject<UEnum>(Outer, MakeUniqueObjectName(nullptr, UEnum::StaticClass(), InstanceDataObjectName));
+		Enum->SetEnums(EnumValueNames, TemplateEnum->GetCppForm(), bHasFlags ? EEnumFlags::Flags : EEnumFlags::None, /*bAddMaxKeyIfMissing*/ false);
+		Enum->SetMetaData(*WriteToString<32>(NAME_OriginalType), *WriteToString<128>(EnumTypeName));
+
+		// TODO: Detect out-of-bounds values and increase the size of the underlying type accordingly.
+
+		return Enum;
 	}
 
 	static FString UnmanglePropertyName(const FName MaybeMangledName, bool& bOutNameWasMangled)
@@ -222,7 +304,7 @@ namespace UE
 	}
 
 	// recursively re-instances all structs contained by this property to include loose properties
-	static void ConvertToInstanceDataObjectProperty(FProperty* Property, FPropertyTypeName PropertyType, UObject* Outer, const FPropertyPathNameTree* PropertyTree)
+	static void ConvertToInstanceDataObjectProperty(FProperty* Property, FPropertyTypeName PropertyType, UObject* Outer, const FPropertyPathNameTree* PropertyTree, const FUnknownEnumNames* EnumNames)
 	{
 		if (!Property->HasMetaData(NAME_DisplayName))
 		{
@@ -260,7 +342,7 @@ namespace UE
 					OriginalName = WriteToString<256>(OriginalNameBuilder.Build()).ToView();
 				}
 
-				UInstanceDataObjectStruct* Struct = CreateInstanceDataObjectStructRec<UInstanceDataObjectStruct>(AsStructProperty->Struct, Outer, PropertyTree);
+				UInstanceDataObjectStruct* Struct = CreateInstanceDataObjectStructRec<UInstanceDataObjectStruct>(AsStructProperty->Struct, Outer, PropertyTree, EnumNames);
 				if (const FName StructGuidName = PropertyType.GetParameterName(1); !StructGuidName.IsNone())
 				{
 					FGuid::Parse(StructGuidName.ToString(), Struct->Guid);
@@ -274,14 +356,22 @@ namespace UE
 				TrySetContainsLooseProperties(AsStructProperty, AsStructProperty->Struct);
 			}
 		}
+		else if (FByteProperty* AsByteProperty = CastField<FByteProperty>(Property))
+		{
+			AsByteProperty->Enum = FindOrCreateInstanceDataObjectEnum(AsByteProperty->Enum, Outer, Property, EnumNames);
+		}
+		else if (FEnumProperty* AsEnumProperty = CastField<FEnumProperty>(Property))
+		{
+			AsEnumProperty->SetEnumForImpersonation(FindOrCreateInstanceDataObjectEnum(AsEnumProperty->GetEnum(), Outer, Property, EnumNames));
+		}
 		else if (FArrayProperty* AsArrayProperty = CastField<FArrayProperty>(Property))
 		{
-			ConvertToInstanceDataObjectProperty(AsArrayProperty->Inner, PropertyType.GetParameter(0), Outer, PropertyTree);
+			ConvertToInstanceDataObjectProperty(AsArrayProperty->Inner, PropertyType.GetParameter(0), Outer, PropertyTree, EnumNames);
 			TrySetContainsLooseProperties(AsArrayProperty, AsArrayProperty->Inner);
 		}
 		else if (FSetProperty* AsSetProperty = CastField<FSetProperty>(Property))
 		{
-			ConvertToInstanceDataObjectProperty(AsSetProperty->ElementProp, PropertyType.GetParameter(0), Outer, PropertyTree);
+			ConvertToInstanceDataObjectProperty(AsSetProperty->ElementProp, PropertyType.GetParameter(0), Outer, PropertyTree, EnumNames);
 			TrySetContainsLooseProperties(AsSetProperty, AsSetProperty->ElementProp);
 		}
 		else if (FMapProperty* AsMapProperty = CastField<FMapProperty>(Property))
@@ -299,14 +389,14 @@ namespace UE
 				Path.Pop();
 			}
 
-			ConvertToInstanceDataObjectProperty(AsMapProperty->KeyProp, PropertyType.GetParameter(0), Outer, KeyTree);
+			ConvertToInstanceDataObjectProperty(AsMapProperty->KeyProp, PropertyType.GetParameter(0), Outer, KeyTree, EnumNames);
 			TrySetContainsLooseProperties(AsMapProperty, AsMapProperty->KeyProp);
-			ConvertToInstanceDataObjectProperty(AsMapProperty->ValueProp, PropertyType.GetParameter(1), Outer, ValueTree);
+			ConvertToInstanceDataObjectProperty(AsMapProperty->ValueProp, PropertyType.GetParameter(1), Outer, ValueTree, EnumNames);
 			TrySetContainsLooseProperties(AsMapProperty, AsMapProperty->ValueProp);
 		}
 		else if (FOptionalProperty* AsOptionalProperty = CastField<FOptionalProperty>(Property))
 		{
-			ConvertToInstanceDataObjectProperty(AsOptionalProperty->GetValueProperty(), PropertyType.GetParameter(0), Outer, PropertyTree);
+			ConvertToInstanceDataObjectProperty(AsOptionalProperty->GetValueProperty(), PropertyType.GetParameter(0), Outer, PropertyTree, EnumNames);
 			TrySetContainsLooseProperties(AsOptionalProperty, AsOptionalProperty->GetValueProperty());
 		}
 
@@ -363,7 +453,7 @@ namespace UE
 	}
 
 	// constructs an InstanceDataObject struct by merging the properties in 
-	static UStruct* CreateInstanceDataObjectStructRec(const UClass* StructClass, UStruct* TemplateStruct, UObject* Outer, const FPropertyPathNameTree* PropertyTree)
+	static UStruct* CreateInstanceDataObjectStructRec(const UClass* StructClass, UStruct* TemplateStruct, UObject* Outer, const FPropertyPathNameTree* PropertyTree, const FUnknownEnumNames* EnumNames)
 	{
 		TSet<FPropertyPathName> SuperPropertyPathsFromTree;
 
@@ -409,7 +499,7 @@ namespace UE
 					}
 				}
 
-				ConvertToInstanceDataObjectProperty(SuperProperty, Type, Super, SubTree);
+				ConvertToInstanceDataObjectProperty(SuperProperty, Type, Super, SubTree, EnumNames);
 			}
 
 			// AddCppProperty expects reverse property order for StaticLink to work correctly
@@ -460,7 +550,7 @@ namespace UE
 							// skip loose types that have been explicitly excluded from IDOs
 							continue;
 						}
-						ConvertToInstanceDataObjectProperty(Property, Type, Result, It.GetNode().GetSubTree());
+						ConvertToInstanceDataObjectProperty(Property, Type, Result, It.GetNode().GetSubTree(), EnumNames);
 						MarkPropertyAsLoose(Property);	// note: make sure not to mark until AFTER conversion, as this can mutate property flags on nested struct fields
 						LooseInstanceDataObjectProperties.Add(Property);
 						continue;
@@ -542,9 +632,9 @@ namespace UE
 			CLASS_EditInlineNew | CLASS_CollapseCategories | CLASS_Const | CLASS_CompiledFromBlueprint | CLASS_HasInstancedReference);
 	}
 
-	UClass* CreateInstanceDataObjectClass(const FPropertyPathNameTree* PropertyTree, UClass* OwnerClass, UObject* Outer)
+	UClass* CreateInstanceDataObjectClass(const FPropertyPathNameTree* PropertyTree, const FUnknownEnumNames* EnumNames, UClass* OwnerClass, UObject* Outer)
 	{
-		UClass* Result = CreateInstanceDataObjectStructRec<UInstanceDataObjectClass>(OwnerClass, Outer, PropertyTree);
+		UClass* Result = CreateInstanceDataObjectStructRec<UInstanceDataObjectClass>(OwnerClass, Outer, PropertyTree, EnumNames);
 		if (const FString& DisplayName = OwnerClass->GetMetaData(NAME_DisplayName); !DisplayName.IsEmpty())
 		{
 			Result->SetMetaData(NAME_DisplayName, *DisplayName);
diff --git a/Engine/Source/Runtime/CoreUObject/Private/UObject/PropertyBagRepository.cpp b/Engine/Source/Runtime/CoreUObject/Private/UObject/PropertyBagRepository.cpp
--- a/Engine/Source/Runtime/CoreUObject/Private/UObject/PropertyBagRepository.cpp
+++ b/Engine/Source/Runtime/CoreUObject/Private/UObject/PropertyBagRepository.cpp
@@ -128,18 +128,6 @@ public:
 	}
 };
 
-class FUnknownEnumNames
-{
-public:
-	struct FInfo
-	{
-		TSet<FName> Names;
-		bool bHasFlags = false;
-	};
-
-	TMap<FPropertyTypeName, FInfo> Enums;
-};
-
 void FPropertyBagRepository::FPropertyBagAssociationData::Destroy()
 {
 	delete Tree;
@@ -558,6 +546,13 @@ void FPropertyBagRepository::AddUnknownEnumName(const UObject* Owner, const UEnu
 		BagData.EnumNames = new FUnknownEnumNames;
 	}
 
+	BagData.EnumNames->Add(Enum, EnumTypeName, EnumValueName);
+}
+
+void FUnknownEnumNames::Add(const UEnum* Enum, FPropertyTypeName EnumTypeName, FName EnumValueName)
+{
+	check(Enum || !EnumTypeName.IsEmpty());
+
 	if (EnumTypeName.IsEmpty())
 	{
 		FPropertyTypeNameBuilder Builder;
@@ -565,7 +560,7 @@ void FPropertyBagRepository::AddUnknownEnumName(const UObject* Owner, const UEnu
 		EnumTypeName = Builder.Build();
 	}
 
-	FUnknownEnumNames::FInfo& Info = BagData.EnumNames->Enums.FindOrAdd(EnumTypeName);
+	FUnknownEnumNames::FInfo& Info = Enums.FindOrAdd(EnumTypeName);
 
 	TStringBuilder<128> EnumValueString(InPlace, EnumValueName);
 	if (String::FindFirstChar(EnumValueString, TEXT('|')) == INDEX_NONE)
@@ -598,22 +593,34 @@ void FPropertyBagRepository::AddUnknownEnumName(const UObject* Owner, const UEnu
 	}
 }
 
-void FPropertyBagRepository::FindUnknownEnumNames(const UObject* Owner, FPropertyTypeName EnumTypeName, TArray<FName>& OutNames, bool& bOutHasFlags)
+void FPropertyBagRepository::FindUnknownEnumNames(const UObject* Owner, FPropertyTypeName EnumTypeName, TArray<FName>& OutNames, bool& bOutHasFlags) const
 {
-	check(Owner);
 	checkf(!EnumTypeName.IsEmpty(), TEXT("FindUnknownEnumNames requires an enum type name. Owner: %s"), *Owner->GetPathName());
 
 	OutNames.Empty();
 	bOutHasFlags = false;
 
-	FPropertyBagRepositoryLock LockRepo(this);
-	FPropertyBagAssociationData* BagData = AssociatedData.Find(Owner);
-	if (!BagData || !BagData->EnumNames)
+	if (const FUnknownEnumNames* EnumNames = FindUnknownEnumNames(Owner))
 	{
-		return;
+		EnumNames->Find(EnumTypeName, OutNames, bOutHasFlags);
 	}
+}
+
+const FUnknownEnumNames* FPropertyBagRepository::FindUnknownEnumNames(const UObject* Owner) const
+{
+	check(Owner);
+
+	FPropertyBagRepositoryLock LockRepo(this);
+	const FPropertyBagAssociationData* BagData = AssociatedData.Find(Owner);
+	return BagData ? BagData->EnumNames : nullptr;
+}
+
+void FUnknownEnumNames::Find(FPropertyTypeName EnumTypeName, TArray<FName>& OutNames, bool& bOutHasFlags) const
+{
+	OutNames.Empty();
+	bOutHasFlags = false;
 
-	if (const FUnknownEnumNames::FInfo* Info = BagData->EnumNames->Enums.Find(EnumTypeName))
+	if (const FUnknownEnumNames::FInfo* Info = Enums.Find(EnumTypeName))
 	{
 		OutNames = Info->Names.Array();
 		bOutHasFlags = Info->bHasFlags;
@@ -828,9 +835,10 @@ void FPropertyBagRepository::CreateInstanceDataObjectUnsafe(UObject* Owner, FPro
 {
 	check(!BagData.InstanceDataObject);	// No repeated calls
 	const FPropertyPathNameTree* PropertyTree = BagData.Tree;
+	const FUnknownEnumNames* EnumNames = BagData.EnumNames;
 	// construct InstanceDataObject class
 	// TODO: should we put the InstanceDataObject or it's class in a package?
-	const UClass* InstanceDataObjectClass = CreateInstanceDataObjectClass(PropertyTree, Owner->GetClass(), GetTransientPackage());
+	const UClass* InstanceDataObjectClass = CreateInstanceDataObjectClass(PropertyTree, EnumNames, Owner->GetClass(), GetTransientPackage());
 
 	BagData.bNeedsFixup = StructContainsLooseProperties(InstanceDataObjectClass);
 
diff --git a/Engine/Source/Runtime/CoreUObject/Private/UObject/PropertyHelper.cpp b/Engine/Source/Runtime/CoreUObject/Private/UObject/PropertyHelper.cpp
--- a/Engine/Source/Runtime/CoreUObject/Private/UObject/PropertyHelper.cpp
+++ b/Engine/Source/Runtime/CoreUObject/Private/UObject/PropertyHelper.cpp
@@ -308,7 +308,7 @@ const FString* FindOriginalTypeName(const FProperty* Property)
 	{
 		return nullptr;
 	}
-	
+
 	// Prioritize metadata on the property over metadata on the type.
 	if (const FString* OriginalType = Property->FindMetaData(NAME_OriginalType))
 	{
@@ -468,16 +468,13 @@ static bool FindRedirectForProperty(FPropertyTypeName OldType, FPropertyTypeName
 		FName FieldName = Field->GetFName();
 
 	#if WITH_EDITORONLY_DATA
-		// Compare against the original type of an impersonated struct.
-		if constexpr (std::is_same_v<PropertyType, FStructProperty>)
+		// Compare against the original type of an impersonated type.
+		if (const FPropertyTypeName OriginalType = FindOriginalType(Property); !OriginalType.IsEmpty())
 		{
-			if (const FPropertyTypeName OriginalType = FindOriginalType(Property); !OriginalType.IsEmpty())
+			FieldName = OriginalType.GetName();
+			if (FieldName == OldNameRedirect.ObjectName)
 			{
-				FieldName = OriginalType.GetName();
-				if (FieldName == OldNameRedirect.ObjectName)
-				{
-					return false;
-				}
+				return false;
 			}
 		}
 	#endif // WITH_EDITORONLY_DATA
diff --git a/Engine/Source/Runtime/CoreUObject/Public/UObject/EnumProperty.h b/Engine/Source/Runtime/CoreUObject/Public/UObject/EnumProperty.h
--- a/Engine/Source/Runtime/CoreUObject/Public/UObject/EnumProperty.h
+++ b/Engine/Source/Runtime/CoreUObject/Public/UObject/EnumProperty.h
@@ -91,6 +91,15 @@ public:
 		Enum = InEnum;
 	}
 
+	/**
+	 * Set the UEnum of this property.
+	 * @note Internal! Used for creation of enums that use impersonation.
+	 */
+	UE_INTERNAL FORCEINLINE void SetEnumForImpersonation(UEnum* InEnum)
+	{
+		Enum = InEnum;
+	}
+
 	/**
 	 * Returns a pointer to the UEnum of this property.
 	 */
diff --git a/Engine/Source/Runtime/CoreUObject/Tests/UObject/InstanceDataObjectUtilsTest.cpp b/Engine/Source/Runtime/CoreUObject/Tests/UObject/InstanceDataObjectUtilsTest.cpp
--- a/Engine/Source/Runtime/CoreUObject/Tests/UObject/InstanceDataObjectUtilsTest.cpp
+++ b/Engine/Source/Runtime/CoreUObject/Tests/UObject/InstanceDataObjectUtilsTest.cpp
@@ -26,7 +26,17 @@ namespace UE
 TEST_CASE_NAMED(FInstanceDataObjectUtilsTest, "CoreUObject::Serialization::InstanceDataObjectUtils", "[CoreUObject][EngineFilter]")
 {
 	UTestInstanceDataObjectClass* BaseObject = NewObject<UTestInstanceDataObjectClass>();
-	UClass* TestClass = CreateInstanceDataObjectClass(nullptr, BaseObject->GetClass(), BaseObject->GetOuter());
+
+	FPropertyBagRepository& Repo = FPropertyBagRepository::Get();
+	Repo.AddUnknownEnumName(BaseObject, StaticEnum<ETestInstanceDataObjectBird>(), {}, "TIDOB_Pigeon");
+	Repo.AddUnknownEnumName(BaseObject, StaticEnum<ETestInstanceDataObjectGrain::Type>(), {}, "Rye");
+	Repo.AddUnknownEnumName(BaseObject, StaticEnum<ETestInstanceDataObjectFruit>(), {}, "Cherry");
+	Repo.AddUnknownEnumName(BaseObject, StaticEnum<ETestInstanceDataObjectDirection>(), {}, "Up");
+	Repo.AddUnknownEnumName(BaseObject, StaticEnum<ETestInstanceDataObjectFullFlags>(), {}, "Flag3");
+	Repo.AddUnknownEnumName(BaseObject, StaticEnum<ETestInstanceDataObjectFullFlags>(), {}, "Flag8");
+	Repo.AddUnknownEnumName(BaseObject, StaticEnum<ETestInstanceDataObjectFullFlags>(), {}, "Flag9");
+
+	UClass* TestClass = CreateInstanceDataObjectClass(nullptr, Repo.FindUnknownEnumNames(BaseObject), BaseObject->GetClass(), BaseObject->GetOuter());
 
 	FIntProperty* Int32Property = FindFProperty<FIntProperty>(TestClass, TEXT("Int32"));
 	FStructProperty* StructProperty = FindFProperty<FStructProperty>(TestClass, TEXT("Struct"));
@@ -43,6 +53,26 @@ TEST_CASE_NAMED(FInstanceDataObjectUtilsTest, "CoreUObject::Serialization::Insta
 	REQUIRE(CProperty);
 	REQUIRE(DProperty);
 
+	FByteProperty* BirdProperty = FindFProperty<FByteProperty>(StructProperty->Struct, TEXT("Bird"));
+	FByteProperty* GrainProperty = FindFProperty<FByteProperty>(StructProperty->Struct, TEXT("Grain"));
+	FEnumProperty* FruitProperty = FindFProperty<FEnumProperty>(StructProperty->Struct, TEXT("Fruit"));
+	FEnumProperty* DirectionProperty = FindFProperty<FEnumProperty>(StructProperty->Struct, TEXT("Direction"));
+	FEnumProperty* FullFlagsProperty = FindFProperty<FEnumProperty>(StructProperty->Struct, TEXT("FullFlags"));
+	REQUIRE(BirdProperty);
+	REQUIRE(GrainProperty);
+	REQUIRE(FruitProperty);
+	REQUIRE(DirectionProperty);
+	REQUIRE(FullFlagsProperty);
+
+	CHECK(BirdProperty->Enum->GetIndexByName("TIDOB_Pigeon") != INDEX_NONE);
+	CHECK(GrainProperty->Enum->GetIndexByName("ETestInstanceDataObjectGrain::Rye") != INDEX_NONE);
+	CHECK(FruitProperty->GetEnum()->GetIndexByName("ETestInstanceDataObjectFruit::Cherry") != INDEX_NONE);
+	CHECK(DirectionProperty->GetEnum()->GetIndexByName("ETestInstanceDataObjectDirection::Up") != INDEX_NONE);
+	CHECK(FullFlagsProperty->GetEnum()->GetIndexByName("ETestInstanceDataObjectFullFlags::Flag3") != INDEX_NONE);
+	CHECK(FullFlagsProperty->GetEnum()->GetIndexByName("ETestInstanceDataObjectFullFlags::Flag8") != INDEX_NONE);
+	CHECK(FullFlagsProperty->GetEnum()->GetIndexByName("ETestInstanceDataObjectFullFlags::Flag9") != INDEX_NONE);
+	CHECK(FullFlagsProperty->GetEnum()->GetMaxEnumValue() == 0b11'1111'1111);
+
 	FName TestObjectName = MakeUniqueObjectName(nullptr, TestClass, FName(WriteToString<128>(TestClass->GetFName(), TEXT("_Instance"))));
 	UObject* Owner = NewObject<UObject>(GetTransientPackage(), TestClass, TestObjectName);
 	void* StructData = StructProperty->ContainerPtrToValuePtr<void>(Owner);
@@ -92,7 +122,7 @@ TEST_CASE_NAMED(FInstanceDataObjectUtilsTest, "CoreUObject::Serialization::Insta
 TEST_CASE_NAMED(FTrackInitializedPropertiesTest, "CoreUObject::Serialization::TrackInitializedProperties", "[CoreUObject][EngineFilter]")
 {
 	UTestInstanceDataObjectClass* BaseObject = NewObject<UTestInstanceDataObjectClass>();
-	UClass* TestClass = CreateInstanceDataObjectClass(nullptr, BaseObject->GetClass(), BaseObject->GetOuter());
+	UClass* TestClass = CreateInstanceDataObjectClass(nullptr, nullptr, BaseObject->GetClass(), BaseObject->GetOuter());
 
 	FStructProperty* StructProperty = FindFProperty<FStructProperty>(TestClass, TEXT("Struct"));
 	REQUIRE(StructProperty);
diff --git a/Engine/Source/Runtime/CoreUObject/Tests/UObject/InstanceDataObjectUtilsTest.h b/Engine/Source/Runtime/CoreUObject/Tests/UObject/InstanceDataObjectUtilsTest.h
--- a/Engine/Source/Runtime/CoreUObject/Tests/UObject/InstanceDataObjectUtilsTest.h
+++ b/Engine/Source/Runtime/CoreUObject/Tests/UObject/InstanceDataObjectUtilsTest.h
@@ -121,6 +121,22 @@ enum class ETestInstanceDataObjectDirectionAlternate : uint16
 
 ENUM_CLASS_FLAGS(ETestInstanceDataObjectDirectionAlternate);
 
+UENUM(Flags)
+enum class ETestInstanceDataObjectFullFlags : uint8
+{
+	None = 0,
+	Flag0 = 1 << 0,
+	Flag1 = 1 << 1,
+	Flag2 = 1 << 2,
+	// Flag3 skipped for testing insertion of an unknown flag in the middle.
+	Flag4 = 1 << 4,
+	Flag5 = 1 << 5,
+	Flag6 = 1 << 6,
+	Flag7 = 1 << 7,
+};
+
+ENUM_CLASS_FLAGS(ETestInstanceDataObjectFullFlags);
+
 USTRUCT()
 struct FTestInstanceDataObjectPoint
 {
@@ -186,6 +202,9 @@ public:
 	UPROPERTY()
 	ETestInstanceDataObjectDirection Direction = ETestInstanceDataObjectDirection::None;
 
+	UPROPERTY()
+	ETestInstanceDataObjectFullFlags FullFlags = ETestInstanceDataObjectFullFlags::None;
+
 	UPROPERTY()
 	TEnumAsByte<ETestInstanceDataObjectGrain::Type> GrainFromEnumClass = ETestInstanceDataObjectGrain::None;
 
