Fixed regression in ArePropertiesCompatible

This was caused by CL# 29761652 which accidentally removed the property class comparison when removing a special case for object property conversion

This change restores the property class comparison, and also undoes the change from CL# 31454390 which was a workaround for the previously missing class comparison

#rnx

diff --git a/Engine/Source/Runtime/CoreUObject/Private/UObject/PropertyAccessUtil.cpp b/Engine/Source/Runtime/CoreUObject/Private/UObject/PropertyAccessUtil.cpp
--- a/Engine/Source/Runtime/CoreUObject/Private/UObject/PropertyAccessUtil.cpp
+++ b/Engine/Source/Runtime/CoreUObject/Private/UObject/PropertyAccessUtil.cpp
@@ -171,22 +171,27 @@ bool ArePropertiesCompatible(const FProperty* InSrcProp, const FProperty* InDest
 		return true;
 	}
 
+	// Compare the classes as these must be an *exact* match as the access is low-level and without property coercion
+	if (InSrcProp->GetClass() != InDestProp->GetClass())
+	{
+		return false;
+	}
+
 	// Containers also need to check their inner types
 	if (const FArrayProperty* SrcArrayProp = CastField<FArrayProperty>(InSrcProp))
 	{
-		const FArrayProperty* DestArrayProp = CastField<FArrayProperty>(InDestProp);
-		return DestArrayProp && ArePropertiesCompatible(SrcArrayProp->Inner, DestArrayProp->Inner);
+		const FArrayProperty* DestArrayProp = CastFieldChecked<FArrayProperty>(InDestProp);
+		return ArePropertiesCompatible(SrcArrayProp->Inner, DestArrayProp->Inner);
 	}
 	if (const FSetProperty* SrcSetProp = CastField<FSetProperty>(InSrcProp))
 	{
-		const FSetProperty* DestSetProp = CastField<FSetProperty>(InDestProp);
-		return DestSetProp && ArePropertiesCompatible(SrcSetProp->ElementProp, DestSetProp->ElementProp);
+		const FSetProperty* DestSetProp = CastFieldChecked<FSetProperty>(InDestProp);
+		return ArePropertiesCompatible(SrcSetProp->ElementProp, DestSetProp->ElementProp);
 	}
 	if (const FMapProperty* SrcMapProp = CastField<FMapProperty>(InSrcProp))
 	{
-		const FMapProperty* DestMapProp = CastField<FMapProperty>(InDestProp);
-		return DestMapProp
-			&& ArePropertiesCompatible(SrcMapProp->KeyProp, DestMapProp->KeyProp)
+		const FMapProperty* DestMapProp = CastFieldChecked<FMapProperty>(InDestProp);
+		return ArePropertiesCompatible(SrcMapProp->KeyProp, DestMapProp->KeyProp)
 			&& ArePropertiesCompatible(SrcMapProp->ValueProp, DestMapProp->ValueProp);
 	}
 
