Fixed conversions in enum serialization

diff --git a/Engine/Source/Runtime/CoreUObject/Private/UObject/EnumProperty.cpp b/Engine/Source/Runtime/CoreUObject/Private/UObject/EnumProperty.cpp
--- a/Engine/Source/Runtime/CoreUObject/Private/UObject/EnumProperty.cpp
+++ b/Engine/Source/Runtime/CoreUObject/Private/UObject/EnumProperty.cpp
@@ -500,8 +500,7 @@ EConvertFromTypeResult FEnumProperty::ConvertFromType(const FPropertyTag& Tag, F
 	if (*TagType == NAME_EnumProperty)
 	{
 	#if WITH_EDITORONLY_DATA
-		FUObjectSerializeContext* SerializeContext = FUObjectThreadContext::Get().GetSerializeContext();
-		if (UNLIKELY(SerializeContext->bTrackUnknownEnumNames && !CanSerializeFromTypeName(Tag.GetType())))
+		if (UNLIKELY(FUObjectThreadContext::Get().GetSerializeContext()->bTrackUnknownProperties && !CanSerializeFromTypeName(Tag.GetType())))
 		{
 			FName EnumValueName;
 			int64 EnumValue = 0;
@@ -526,8 +525,6 @@ EConvertFromTypeResult FEnumProperty::ConvertFromType(const FPropertyTag& Tag, F
 	default:
 		return EConvertFromTypeResult::UseSerializeItem;
 	case NAME_ByteProperty:
-	{
-		uint8 PreviousValue = 0;
 		if (Tag.GetType().GetParameterCount() == 0)
 		{
 			// A nested property would lose its enum name on previous versions. Handle this case for backward compatibility reasons.
@@ -544,23 +541,23 @@ EConvertFromTypeResult FEnumProperty::ConvertFromType(const FPropertyTag& Tag, F
 				InnerPropertyTag.Name = Tag.Name;
 				InnerPropertyTag.ArrayIndex = 0;
 
-				PreviousValue = (uint8)FNumericProperty::ReadEnumAsInt64(Slot, DefaultsStruct, InnerPropertyTag);
+				int64 PreviousValue = FNumericProperty::ReadEnumAsInt64(Slot, DefaultsStruct, InnerPropertyTag);
+				UnderlyingProp->SetIntPropertyValue(ContainerPtrToValuePtr<void>(Data, Tag.ArrayIndex), PreviousValue);
 			}
 			else
 			{
-				// a byte property gained an enum
-				Slot << PreviousValue;
+				// A byte property gained an enum.
+				UE::CoreUObject::Private::ConvertIntToEnumProperty<uint8>(Slot, this, UnderlyingProp, Enum, ContainerPtrToValuePtr<void>(Data, Tag.ArrayIndex));
 			}
 		}
 		else
 		{
-			PreviousValue = (uint8)FNumericProperty::ReadEnumAsInt64(Slot, DefaultsStruct, Tag);
+			FName EnumValueName;
+			int64 EnumValue = 0;
+			TryLoadEnumValueByName(Slot, Slot.GetUnderlyingArchive(), Enum, EnumValueName, EnumValue);
+			UnderlyingProp->SetIntPropertyValue(ContainerPtrToValuePtr<void>(Data, Tag.ArrayIndex), EnumValue);
 		}
-
-		// now copy the value into the object's address space
-		UnderlyingProp->SetIntPropertyValue(ContainerPtrToValuePtr<void>(Data, Tag.ArrayIndex), (uint64)PreviousValue);
 		return EConvertFromTypeResult::Converted;
-	}
 	case NAME_Int8Property:
 		UE::CoreUObject::Private::ConvertIntToEnumProperty<int8>(Slot, this, UnderlyingProp, Enum, ContainerPtrToValuePtr<void>(Data, Tag.ArrayIndex));
 		return EConvertFromTypeResult::Converted;
diff --git a/Engine/Source/Runtime/CoreUObject/Private/UObject/PropertyByte.cpp b/Engine/Source/Runtime/CoreUObject/Private/UObject/PropertyByte.cpp
--- a/Engine/Source/Runtime/CoreUObject/Private/UObject/PropertyByte.cpp
+++ b/Engine/Source/Runtime/CoreUObject/Private/UObject/PropertyByte.cpp
@@ -240,7 +240,6 @@ EConvertFromTypeResult FByteProperty::ConvertFromType(const FPropertyTag& Tag, F
 	default:
 		return EConvertFromTypeResult::UseSerializeItem;
 	case NAME_ByteProperty:
-	{
 		if ((Tag.GetType().GetParameterCount() == 0) != (Enum == nullptr))
 		{
 			// A byte property gained or lost an enum.
@@ -265,16 +264,32 @@ EConvertFromTypeResult FByteProperty::ConvertFromType(const FPropertyTag& Tag, F
 			SetPropertyValue_InContainer(Data, PreviousValue, Tag.ArrayIndex);
 			return EConvertFromTypeResult::Converted;
 		}
+	#if WITH_EDITORONLY_DATA
+		if (UNLIKELY(Enum && FUObjectThreadContext::Get().GetSerializeContext()->bTrackUnknownProperties && !CanSerializeFromTypeName(Tag.GetType())))
+		{
+			FName EnumValueName;
+			int64 EnumValue = 0;
+			TryLoadEnumValueByName(Slot, Slot.GetUnderlyingArchive(), Enum, EnumValueName, EnumValue);
+			SetPropertyValue_InContainer(Data, IntCastChecked<uint8>(EnumValue), Tag.ArrayIndex);
+			return EConvertFromTypeResult::Converted;
+		}
+	#endif
 		return EConvertFromTypeResult::UseSerializeItem;
-	}
 	case NAME_EnumProperty:
-	{
-		// Attempt to find the enum from the tag and find the byte value from the enum.
-		uint8 PreviousValue = (uint8)ReadEnumAsInt64(Slot, DefaultsStruct, Tag);
-
-		SetPropertyValue_InContainer(Data, PreviousValue, Tag.ArrayIndex);
+		if (Enum)
+		{
+			FName EnumValueName;
+			int64 EnumValue = 0;
+			TryLoadEnumValueByName(Slot, Slot.GetUnderlyingArchive(), Enum, EnumValueName, EnumValue);
+			SetPropertyValue_InContainer(Data, IntCastChecked<uint8>(EnumValue), Tag.ArrayIndex);
+		}
+		else
+		{
+			// Attempt to find the enum from the tag and find the byte value from the enum.
+			uint8 PreviousValue = (uint8)ReadEnumAsInt64(Slot, DefaultsStruct, Tag);
+			SetPropertyValue_InContainer(Data, PreviousValue, Tag.ArrayIndex);
+		}
 		return EConvertFromTypeResult::Converted;
-	}
 	case NAME_Int8Property:
 		if (Enum)
 		{
diff --git a/Engine/Source/Runtime/CoreUObject/Tests/UObject/InstanceDataObjectUtilsTest.cpp b/Engine/Source/Runtime/CoreUObject/Tests/UObject/InstanceDataObjectUtilsTest.cpp
--- a/Engine/Source/Runtime/CoreUObject/Tests/UObject/InstanceDataObjectUtilsTest.cpp
+++ b/Engine/Source/Runtime/CoreUObject/Tests/UObject/InstanceDataObjectUtilsTest.cpp
@@ -313,9 +313,8 @@ TEST_CASE_NAMED(FTrackUnknownPropertiesTest, "CoreUObject::Serialization::TrackU
 	CHECK(StructData.Direction == (ETestInstanceDataObjectDirection::North | ETestInstanceDataObjectDirection::West));
 	CHECK(StructData.GrainTypeChange == ETestInstanceDataObjectGrain::Corn);
 	CHECK(StructData.FruitTypeChange == ETestInstanceDataObjectFruit::Orange);
-	// TODO: ConvertFromType incorrectly assigns numeric values based on the saved enum type.
-	//CHECK(StructData.GrainTypeAndPropertyChange == ETestInstanceDataObjectGrain::Corn);
-	//CHECK(StructData.FruitTypeAndPropertyChange == ETestInstanceDataObjectFruit::Orange);
+	CHECK(StructData.GrainTypeAndPropertyChange == ETestInstanceDataObjectGrain::Corn);
+	CHECK(StructData.FruitTypeAndPropertyChange == ETestInstanceDataObjectFruit::Orange);
 	CHECK(StructData.Point.X == 0);
 	CHECK(StructData.Point.Y == 0);
 	CHECK(StructData.Point.Z == 0);
@@ -360,15 +359,14 @@ TEST_CASE_NAMED(FTrackUnknownPropertiesTest, "CoreUObject::Serialization::TrackU
 		#endif
 		}
 	}
-	// TODO: Enum type name changes are not recorded as a conversion.
-	//{
-	//	FSerializedPropertyPathScope Path(SerializeContext, {"GrainTypeChange", SavePropertyTypeName(FindFProperty<FProperty>(StaticStruct<FTestInstanceDataObjectStructAlternate>(), TEXT("GrainTypeChange")))});
-	//	CHECK(Tree->Find(SerializeContext->SerializedPropertyPath));
-	//}
-	//{
-	//	FSerializedPropertyPathScope Path(SerializeContext, {"FruitTypeChange", SavePropertyTypeName(FindFProperty<FProperty>(StaticStruct<FTestInstanceDataObjectStructAlternate>(), TEXT("FruitTypeChange")))});
-	//	CHECK(Tree->Find(SerializeContext->SerializedPropertyPath));
-	//}
+	{
+		FSerializedPropertyPathScope Path(SerializeContext, {"GrainTypeChange", SavePropertyTypeName(FindFProperty<FProperty>(StaticStruct<FTestInstanceDataObjectStructAlternate>(), TEXT("GrainTypeChange")))});
+		CHECK(Tree->Find(SerializeContext->SerializedPropertyPath));
+	}
+	{
+		FSerializedPropertyPathScope Path(SerializeContext, {"FruitTypeChange", SavePropertyTypeName(FindFProperty<FProperty>(StaticStruct<FTestInstanceDataObjectStructAlternate>(), TEXT("FruitTypeChange")))});
+		CHECK(Tree->Find(SerializeContext->SerializedPropertyPath));
+	}
 	{
 		FSerializedPropertyPathScope Path(SerializeContext, {"GrainTypeAndPropertyChange", SavePropertyTypeName(FindFProperty<FProperty>(StaticStruct<FTestInstanceDataObjectStructAlternate>(), TEXT("GrainTypeAndPropertyChange")))});
 		CHECK(Tree->Find(SerializeContext->SerializedPropertyPath));
@@ -399,9 +397,8 @@ TEST_CASE_NAMED(FTrackUnknownPropertiesTest, "CoreUObject::Serialization::TrackU
 	CHECK(StructData.Direction == (ETestInstanceDataObjectDirection::North | ETestInstanceDataObjectDirection::West));
 	CHECK(StructData.GrainTypeChange == ETestInstanceDataObjectGrain::Corn);
 	CHECK(StructData.FruitTypeChange == ETestInstanceDataObjectFruit::Orange);
-	// TODO: ConvertFromType incorrectly assigns numeric values based on the saved enum type.
-	//CHECK(StructData.GrainTypeAndPropertyChange == ETestInstanceDataObjectGrain::Corn);
-	//CHECK(StructData.FruitTypeAndPropertyChange == ETestInstanceDataObjectFruit::Orange);
+	CHECK(StructData.GrainTypeAndPropertyChange == ETestInstanceDataObjectGrain::Corn);
+	CHECK(StructData.FruitTypeAndPropertyChange == ETestInstanceDataObjectFruit::Orange);
 	CHECK(StructData.Point.X == 0);
 	CHECK(StructData.Point.Y == 0);
 	CHECK(StructData.Point.Z == 0);
@@ -493,9 +490,8 @@ TEST_CASE_NAMED(FTrackUnknownEnumNamesTest, "CoreUObject::Serialization::TrackUn
 	CHECK(StructData.Grain == (ETestInstanceDataObjectGrain::Type)((uint8)ETestInstanceDataObjectGrain::Wheat + 1));
 	CHECK(StructData.Fruit == (ETestInstanceDataObjectFruit)((uint8)ETestInstanceDataObjectFruit::Orange + 1));
 	CHECK(StructData.Direction == (ETestInstanceDataObjectDirection)MAX_uint16);
-	// TODO: ConvertFromType incorrectly assigns numeric values based on the saved enum type.
-	//CHECK(StructData.GrainFromEnumClass == ETestInstanceDataObjectGrain::Corn);
-	//CHECK(StructData.FruitFromNamespace == ETestInstanceDataObjectFruit::Orange);
+	CHECK(StructData.GrainFromEnumClass == ETestInstanceDataObjectGrain::Corn);
+	CHECK(StructData.FruitFromNamespace == ETestInstanceDataObjectFruit::Orange);
 	CHECK(StructData.GrainTypeChange == ETestInstanceDataObjectGrain::Corn);
 	CHECK(StructData.FruitTypeChange == ETestInstanceDataObjectFruit::Orange);
 
@@ -560,9 +556,8 @@ TEST_CASE_NAMED(FTrackUnknownEnumNamesTest, "CoreUObject::Serialization::TrackUn
 	CHECK(StructData.Grain == (ETestInstanceDataObjectGrain::Type)((uint8)ETestInstanceDataObjectGrain::Wheat + 1));
 	CHECK(StructData.Fruit == (ETestInstanceDataObjectFruit)((uint8)ETestInstanceDataObjectFruit::Orange + 1));
 	CHECK(StructData.Direction == (ETestInstanceDataObjectDirection)MAX_uint16);
-	// TODO: ConvertFromType incorrectly assigns numeric values based on the saved enum type.
-	//CHECK(StructData.GrainFromEnumClass == ETestInstanceDataObjectGrain::Corn);
-	//CHECK(StructData.FruitFromNamespace == ETestInstanceDataObjectFruit::Orange);
+	CHECK(StructData.GrainFromEnumClass == ETestInstanceDataObjectGrain::Corn);
+	CHECK(StructData.FruitFromNamespace == ETestInstanceDataObjectFruit::Orange);
 	CHECK(StructData.GrainTypeChange == ETestInstanceDataObjectGrain::Corn);
 	CHECK(StructData.FruitTypeChange == ETestInstanceDataObjectFruit::Orange);
 
