SavePackage: Fix typos where | was used to take the union of bools rather than ||. | is sometimes used in performance sensitive functions that need to avoid a branch, but these functions aren't called frequently enough to warrant that.
#rnx

diff --git a/Engine/Source/Runtime/CoreUObject/Private/UObject/SavePackage/SavePackageUtilities.cpp b/Engine/Source/Runtime/CoreUObject/Private/UObject/SavePackage/SavePackageUtilities.cpp
--- a/Engine/Source/Runtime/CoreUObject/Private/UObject/SavePackage/SavePackageUtilities.cpp
+++ b/Engine/Source/Runtime/CoreUObject/Private/UObject/SavePackage/SavePackageUtilities.cpp
@@ -1326,7 +1326,7 @@ FName FEDLCookChecker::FEDLNodeData::GetPackageName(const FEDLCookChecker& Owner
 void FEDLCookChecker::FEDLNodeData::Merge(FEDLCookChecker::FEDLNodeData&& Other)
 {
 	check(ObjectEvent == Other.ObjectEvent);
-	bIsExport = bIsExport | Other.bIsExport;
+	bIsExport = bIsExport || Other.bIsExport;
 
 	ImportingPackagesSorted.Append(Other.ImportingPackagesSorted);
 	Algo::Sort(ImportingPackagesSorted, FNameFastLess());
@@ -2291,7 +2291,7 @@ void WritePackageData(FStructuredArchiveRecord& ParentRecord, FArchiveCookContex
 	FLinkerSave* Linker, const TSet<TObjectPtr<UObject>>& ImportsUsedInGame, const TSet<FName>& SoftPackagesUsedInGame,
 	TArray<FAssetData>* OutAssetDatas, bool bProceduralSave)
 {
-	bProceduralSave = bProceduralSave | (CookContext != nullptr);
+	bProceduralSave = bProceduralSave || (CookContext != nullptr);
 	IAssetRegistryInterface* AssetRegistry = IAssetRegistryInterface::GetPtr();
 
 	// To avoid large patch sizes, we have frozen cooked package format at the format before VER_UE4_ASSETREGISTRY_DEPENDENCYFLAGS
