Converted OverridableManager API to use FPropertyVistorPath and forwarded existing API to this new one.

diff --git a/Engine/Source/Runtime/CoreUObject/Private/UObject/OverridableManager.cpp b/Engine/Source/Runtime/CoreUObject/Private/UObject/OverridableManager.cpp
--- a/Engine/Source/Runtime/CoreUObject/Private/UObject/OverridableManager.cpp
+++ b/Engine/Source/Runtime/CoreUObject/Private/UObject/OverridableManager.cpp
@@ -161,7 +161,7 @@ void FOverridableManager::OverrideObject(UObject& Object)
 	if (FOverriddenPropertySet* ThisObjectOverriddenProperties = OverriddenObjectAnnotations.Find(Object))
 	{
 		// Passing no property node means we are overriding the object itself
-		ThisObjectOverriddenProperties->OverrideProperty(FPropertyChangedEvent(nullptr), /*PropertyNode*/nullptr, &Object);
+		ThisObjectOverriddenProperties->OverrideProperty(FPropertyVisitorPath::InvalidIterator(), &Object);
 	}
 #endif // WITH_EDITORONLY_DATA
 }
@@ -189,42 +189,42 @@ void FOverridableManager::PropagateOverrideToInstancedSubObjects(UObject& Object
 #endif // WITH_EDITORONLY_DATA
 }
 
-void FOverridableManager::OverrideProperty(UObject& Object, const FPropertyChangedEvent& PropertyEvent, const FEditPropertyChain& PropertyChain)
+void FOverridableManager::OverrideProperty(UObject& Object, const FPropertyVisitorPath& PropertyPath)
 {
 #if WITH_EDITORONLY_DATA
 	if (FOverriddenPropertySet* ThisObjectOverriddenProperties = OverriddenObjectAnnotations.Find(Object))
 	{
-		ThisObjectOverriddenProperties->OverrideProperty(PropertyEvent, PropertyChain.GetActiveMemberNode() ? PropertyChain.GetActiveMemberNode() : PropertyChain.GetHead(), &Object);
+		ThisObjectOverriddenProperties->OverrideProperty(PropertyPath.GetRootIterator(), &Object);
 	}
 #endif // WITH_EDITORONLY_DATA
 }
 
-bool FOverridableManager::ClearOverriddenProperty(UObject& Object, const FPropertyChangedEvent& PropertyEvent, const FEditPropertyChain::TDoubleLinkedListNode* PropertyNode)
+bool FOverridableManager::ClearOverriddenProperty(UObject& Object, FPropertyVisitorPath::Iterator PropertyIterator)
 {
 #if WITH_EDITORONLY_DATA
 	if (FOverriddenPropertySet* ThisObjectOverriddenProperties = OverriddenObjectAnnotations.Find(Object))
 	{
-		return ThisObjectOverriddenProperties->ClearOverriddenProperty(PropertyEvent, PropertyNode);
+		return ThisObjectOverriddenProperties->ClearOverriddenProperty(PropertyIterator);
 	}
 #endif // WITH_EDITORONLY_DATA
 	return false;
 }
 
-void FOverridableManager::PreOverrideProperty(UObject& Object, const FEditPropertyChain& PropertyChain)
+void FOverridableManager::PreOverrideProperty(UObject& Object, const FPropertyVisitorPath& PropertyPath)
 {
 #if WITH_EDITORONLY_DATA
-	NotifyPropertyChange(EPropertyNotificationType::PreEdit, Object, FPropertyChangedEvent(nullptr), PropertyChain.GetActiveMemberNode() ? PropertyChain.GetActiveMemberNode() : PropertyChain.GetHead());
+	NotifyPropertyChange(EPropertyNotificationType::PreEdit, Object, PropertyPath.GetRootIterator(), EPropertyChangeType::Unspecified);
 #endif // WITH_EDITORONLY_DATA
 }
 
-void FOverridableManager::PostOverrideProperty(UObject& Object, const FPropertyChangedEvent& PropertyEvent, const FEditPropertyChain& PropertyChain)
+void FOverridableManager::PostOverrideProperty(UObject& Object, const FPropertyVisitorPath& PropertyPath, const EPropertyChangeType::Type ChangeType)
 {
 #if WITH_EDITORONLY_DATA
-	NotifyPropertyChange(EPropertyNotificationType::PostEdit, Object, PropertyEvent, PropertyChain.GetActiveMemberNode() ? PropertyChain.GetActiveMemberNode() : PropertyChain.GetHead());
+	NotifyPropertyChange(EPropertyNotificationType::PostEdit, Object, PropertyPath.GetRootIterator(), ChangeType);
 #endif // WITH_EDITORONLY_DATA
 }
 
-void FOverridableManager::NotifyPropertyChange(const EPropertyNotificationType Notification, UObject& Object, const FPropertyChangedEvent& PropertyEvent, const FEditPropertyChain::TDoubleLinkedListNode* PropertyNode)
+void FOverridableManager::NotifyPropertyChange(const EPropertyNotificationType Notification, UObject& Object, FPropertyVisitorPath::Iterator PropertyIterator, const EPropertyChangeType::Type ChangeType)
 {
 #if WITH_EDITORONLY_DATA
 	if (UE::IsClassOfInstanceDataObjectClass(Object.GetClass()))
@@ -234,17 +234,17 @@ void FOverridableManager::NotifyPropertyChange(const EPropertyNotificationType N
 	}
 	if (FOverriddenPropertySet* ThisObjectOverriddenProperties = OverriddenObjectAnnotations.Find(Object))
 	{
-		ThisObjectOverriddenProperties->NotifyPropertyChange(Notification, PropertyEvent, PropertyNode, &Object);
+		ThisObjectOverriddenProperties->NotifyPropertyChange(Notification, PropertyIterator, ChangeType, &Object);
 	}
 #endif // WITH_EDITORONLY_DATA
 }
 
-EOverriddenPropertyOperation FOverridableManager::GetOverriddenPropertyOperation(UObject& Object, const FPropertyChangedEvent& PropertyEvent, const FEditPropertyChain::TDoubleLinkedListNode* PropertyNode, bool* bOutInheritedOperation)
+EOverriddenPropertyOperation FOverridableManager::GetOverriddenPropertyOperation(UObject& Object, FPropertyVisitorPath::Iterator PropertyIterator, bool* bOutInheritedOperation)
 {
 #if WITH_EDITORONLY_DATA
 	if (const FOverriddenPropertySet* ThisObjectOverriddenProperties = OverriddenObjectAnnotations.Find(Object))
 	{
-		return ThisObjectOverriddenProperties->GetOverriddenPropertyOperation(PropertyEvent, PropertyNode, bOutInheritedOperation);
+		return ThisObjectOverriddenProperties->GetOverriddenPropertyOperation(PropertyIterator, bOutInheritedOperation);
 	}
 #endif // WITH_EDITORONLY_DATA
 	return EOverriddenPropertyOperation::None;
diff --git a/Engine/Source/Runtime/CoreUObject/Private/UObject/OverriddenPropertySet.cpp b/Engine/Source/Runtime/CoreUObject/Private/UObject/OverriddenPropertySet.cpp
--- a/Engine/Source/Runtime/CoreUObject/Private/UObject/OverriddenPropertySet.cpp
+++ b/Engine/Source/Runtime/CoreUObject/Private/UObject/OverriddenPropertySet.cpp
@@ -6,7 +6,6 @@
 #include "Serialization/ArchiveSerializedPropertyChain.h"
 #include "UObject/PropertyOptional.h"
 #include "Misc/ScopeExit.h"
-#include "Serialization/StructuredArchiveNameHelpers.h"
 #include "UObject/UObjectThreadContext.h"
 
 /*
@@ -226,19 +225,18 @@ FOverriddenPropertyNode& FOverriddenPropertySet::FindOrAddNode(FOverriddenProper
 	return OverriddenPropertyNodes.Get(NewID);
 }
 
-EOverriddenPropertyOperation FOverriddenPropertySet::GetOverriddenPropertyOperation(const FOverriddenPropertyNode& ParentPropertyNode, const FPropertyChangedEvent& PropertyEvent, const FEditPropertyChain::TDoubleLinkedListNode* PropertyNode, bool* bOutInheritedOperation, const void* Data) const
+EOverriddenPropertyOperation FOverriddenPropertySet::GetOverriddenPropertyOperation(const FOverriddenPropertyNode& ParentPropertyNode, FPropertyVisitorPath::Iterator PropertyIterator, bool* bOutInheritedOperation, const void* Data) const
 {
 	FOverridableManager& OverridableManager = FOverridableManager::Get();
 
 	const void* SubValuePtr = Data;
-	const FEditPropertyChain::TDoubleLinkedListNode* PropertyIterator = PropertyNode;
 	const FOverriddenPropertyNode* OverriddenPropertyNode = &ParentPropertyNode;
 	int32 ArrayIndex = INDEX_NONE;
 	while (PropertyIterator && (!OverriddenPropertyNode || OverriddenPropertyNode->Operation != EOverriddenPropertyOperation::Replace))
 	{
 		ArrayIndex = INDEX_NONE;
 
-		const FProperty* CurrentProperty = PropertyIterator->GetValue();
+		const FProperty* CurrentProperty = PropertyIterator->Property;
 		SubValuePtr = CurrentProperty->ContainerPtrToValuePtr<void>(SubValuePtr, 0); //@todo support static arrays
 
 		const FOverriddenPropertyNode* CurrentOverriddenPropertyNode = nullptr;
@@ -251,24 +249,25 @@ EOverriddenPropertyOperation FOverriddenPropertySet::GetOverriddenPropertyOperat
 			}
 		}
 
-
-		// Special handling for for instanced subobjects 
+		FPropertyVisitorPath::Iterator NextPropertyIterator = PropertyIterator+1;
+		// Special handling for instanced subobjects 
 		if (const FObjectPropertyBase* ObjectProperty = CastField<FObjectPropertyBase>(CurrentProperty))
 		{
-			if (PropertyIterator->GetNextNode())
+			if (NextPropertyIterator)
 			{
 				// Forward any sub queries to the subobject
 				if (UObject* SubObject = ObjectProperty->GetObjectPropertyValue(SubValuePtr))
 				{
 					// This should not be needed in the property grid, as it should already been called on the subobject.
-					return OverridableManager.GetOverriddenPropertyOperation(*SubObject, PropertyEvent, PropertyIterator->GetNextNode(), bOutInheritedOperation);
+					return OverridableManager.GetOverriddenPropertyOperation(*SubObject, NextPropertyIterator, bOutInheritedOperation);
 				}
 			}
 		}
 		// Special handling for array of instanced subobjects 
 		else if (const FArrayProperty* ArrayProperty = CastField<FArrayProperty>(CurrentProperty))
 		{
-			ArrayIndex = PropertyEvent.GetArrayIndex(CurrentProperty->GetName());
+			ArrayIndex = PropertyIterator->Index;
+			checkf(ArrayIndex == INDEX_NONE || PropertyIterator->PropertyInfo == EPropertyVisitorInfoType::ContainerIndex, TEXT("Expecting a container index"));
 
 			// Only special case is instanced subobjects, otherwise we fallback to full array override
 			checkf(ArrayProperty->Inner, TEXT("Expecting an inner type for Arrays"));
@@ -279,10 +278,10 @@ EOverriddenPropertyOperation FOverriddenPropertySet::GetOverriddenPropertyOperat
 				{
 					if (UObject* SubObject = InnerObjectProperty->GetObjectPropertyValue(ArrayHelper.GetElementPtr(ArrayIndex)))
 					{
-						if (PropertyIterator->GetNextNode())
+						if (NextPropertyIterator)
 						{
 							// Forward any sub queries to the subobject
-							return OverridableManager.GetOverriddenPropertyOperation(*SubObject, PropertyEvent, PropertyIterator->GetNextNode(), bOutInheritedOperation);
+							return OverridableManager.GetOverriddenPropertyOperation(*SubObject, NextPropertyIterator, bOutInheritedOperation);
 						}
 						else if(CurrentOverriddenPropertyNode)
 						{
@@ -306,7 +305,8 @@ EOverriddenPropertyOperation FOverriddenPropertySet::GetOverriddenPropertyOperat
 		// Special handling for maps and values of instance subobjects
 		else if (const FMapProperty* MapProperty = CastField<FMapProperty>(CurrentProperty))
 		{
-			ArrayIndex = PropertyEvent.GetArrayIndex(CurrentProperty->GetName());
+			ArrayIndex = PropertyIterator->Index;
+			checkf(ArrayIndex == INDEX_NONE || PropertyIterator->PropertyInfo == EPropertyVisitorInfoType::ContainerIndex, TEXT("Expecting a container index"));
 
 			checkf(MapProperty->ValueProp, TEXT("Expecting a value type for Maps"));
 			FScriptMapHelper MapHelper(MapProperty, SubValuePtr);
@@ -314,14 +314,14 @@ EOverriddenPropertyOperation FOverriddenPropertySet::GetOverriddenPropertyOperat
 			const int32 InternalMapIndex = ArrayIndex != INDEX_NONE ? MapHelper.FindInternalIndex(ArrayIndex) : INDEX_NONE;
 			if(MapHelper.IsValidIndex(InternalMapIndex))
 			{
-				if (PropertyIterator->GetNextNode())
+				if (NextPropertyIterator)
 				{
 					// Forward any sub queries to the subobject
 					if (const FObjectPropertyBase* ValueInstancedObjectProperty = MapProperty->ValueProp->HasAnyPropertyFlags(CPF_PersistentInstance) ? CastField<FObjectPropertyBase>(MapProperty->ValueProp) : nullptr)
 					{
 						if (UObject* ValueSubObject = ValueInstancedObjectProperty->GetObjectPropertyValue(MapHelper.GetValuePtr(InternalMapIndex)))
 						{
-							return OverridableManager.GetOverriddenPropertyOperation(*ValueSubObject, PropertyEvent, PropertyIterator->GetNextNode(), bOutInheritedOperation);
+							return OverridableManager.GetOverriddenPropertyOperation(*ValueSubObject, NextPropertyIterator, bOutInheritedOperation);
 						}
 					}
 				}
@@ -346,23 +346,22 @@ EOverriddenPropertyOperation FOverriddenPropertySet::GetOverriddenPropertyOperat
 		}
 
 		OverriddenPropertyNode = CurrentOverriddenPropertyNode;
-		PropertyIterator = PropertyIterator->GetNextNode();
+		++PropertyIterator;
 	}
 
 	if (bOutInheritedOperation)
 	{
-		*bOutInheritedOperation = PropertyIterator != nullptr || ArrayIndex != INDEX_NONE;
+		*bOutInheritedOperation = PropertyIterator || ArrayIndex != INDEX_NONE;
 	}
 	return OverriddenPropertyNode ? OverriddenPropertyNode->Operation : EOverriddenPropertyOperation::None;
 }
 
-bool FOverriddenPropertySet::ClearOverriddenProperty(FOverriddenPropertyNode& ParentPropertyNode, const FPropertyChangedEvent& PropertyEvent, const FEditPropertyChain::TDoubleLinkedListNode* PropertyNode, const void* Data)
+bool FOverriddenPropertySet::ClearOverriddenProperty(FOverriddenPropertyNode& ParentPropertyNode, FPropertyVisitorPath::Iterator PropertyIterator, const void* Data)
 {
 	FOverridableManager& OverridableManager = FOverridableManager::Get();
 
 	bool bClearedOverrides = false;
 	const void* SubValuePtr = Data;
-	const FEditPropertyChain::TDoubleLinkedListNode* PropertyIterator = PropertyNode;
 	FOverriddenPropertyNode* OverriddenPropertyNode = &ParentPropertyNode;
 	int32 ArrayIndex = INDEX_NONE;
 	TArray<FOverriddenPropertyNodeID> PropertyNodeIDPath;
@@ -371,7 +370,7 @@ bool FOverriddenPropertySet::ClearOverriddenProperty(FOverriddenPropertyNode& Pa
 	{
 		ArrayIndex = INDEX_NONE;
 
-		const FProperty* CurrentProperty = PropertyIterator->GetValue();
+		const FProperty* CurrentProperty = PropertyIterator->Property;
 		SubValuePtr = CurrentProperty->ContainerPtrToValuePtr<void>(SubValuePtr, 0); //@todo support static arrays
 
 		FOverriddenPropertyNode* CurrentOverriddenPropertyNode = nullptr;
@@ -385,14 +384,15 @@ bool FOverriddenPropertySet::ClearOverriddenProperty(FOverriddenPropertyNode& Pa
 			}
 		}
 
-		// Special handling for for instanced subobjects 
+		// Special handling for instanced subobjects 
+		FPropertyVisitorPath::Iterator NextPropertyIterator = PropertyIterator+1;
 		if (const FObjectPropertyBase* ObjectProperty = CastField<FObjectPropertyBase>(CurrentProperty))
 		{
 			if (UObject* SubObject = ObjectProperty->GetObjectPropertyValue(SubValuePtr))
 			{
-				if (PropertyIterator->GetNextNode())
+				if (NextPropertyIterator)
 				{
-					return OverridableManager.ClearOverriddenProperty(*SubObject, PropertyEvent, PropertyIterator->GetNextNode());
+					return OverridableManager.ClearOverriddenProperty(*SubObject, NextPropertyIterator);
 				}
 				else
 				{
@@ -404,7 +404,8 @@ bool FOverriddenPropertySet::ClearOverriddenProperty(FOverriddenPropertyNode& Pa
 		// Special handling for array of instanced subobjects 
 		else if (const FArrayProperty* ArrayProperty = CastField<FArrayProperty>(CurrentProperty))
 		{
-			ArrayIndex = PropertyEvent.GetArrayIndex(CurrentProperty->GetName());
+			ArrayIndex = PropertyIterator->Index;
+			checkf(ArrayIndex == INDEX_NONE || PropertyIterator->PropertyInfo == EPropertyVisitorInfoType::ContainerIndex, TEXT("Expecting a container index"));
 
 			// Only special case is instanced subobjects, otherwise we fallback to full array override
 			if (FObjectPropertyBase* InnerObjectProperty = CastField<FObjectPropertyBase>(ArrayProperty->Inner))
@@ -430,9 +431,9 @@ bool FOverriddenPropertySet::ClearOverriddenProperty(FOverriddenPropertyNode& Pa
 					{
 						if (UObject* SubObject = InnerObjectProperty->GetObjectPropertyValue(ArrayHelper.GetElementPtr(ArrayIndex)))
 						{
-							if (PropertyIterator->GetNextNode())
+							if (NextPropertyIterator)
 							{
-								return OverridableManager.ClearOverriddenProperty(*SubObject, PropertyEvent, PropertyIterator->GetNextNode());
+								return OverridableManager.ClearOverriddenProperty(*SubObject, NextPropertyIterator);
 							}
 							else if(CurrentOverriddenPropertyNode)
 							{
@@ -453,7 +454,8 @@ bool FOverriddenPropertySet::ClearOverriddenProperty(FOverriddenPropertyNode& Pa
 		// Special handling for maps and values of instance subobjects 
 		else if (const FMapProperty* MapProperty = CastField<FMapProperty>(CurrentProperty))
 		{
-			ArrayIndex = PropertyEvent.GetArrayIndex(CurrentProperty->GetName());
+			ArrayIndex = PropertyIterator->Index;
+			checkf(ArrayIndex == INDEX_NONE || PropertyIterator->PropertyInfo == EPropertyVisitorInfoType::ContainerIndex, TEXT("Expecting a container index"));
 
 			FScriptMapHelper MapHelper(MapProperty, SubValuePtr);
 
@@ -462,14 +464,14 @@ bool FOverriddenPropertySet::ClearOverriddenProperty(FOverriddenPropertyNode& Pa
 
 			// If there is a next node, it is probably because the map value is holding a instanced subobject and the user is changing value on it.
 			// So forward the call to the instanced subobject
-			if (PropertyIterator->GetNextNode())
+			if (NextPropertyIterator)
 			{
 				if(MapHelper.IsValidIndex(InternalMapIndex))
 				{
 					checkf(MapProperty->ValueProp, TEXT("Expecting a value type for Maps"));
 					if (UObject* ValueSubObject = ValueInstancedObjectProperty ? ValueInstancedObjectProperty->GetObjectPropertyValue(MapHelper.GetValuePtr(InternalMapIndex)) : nullptr)
 					{
-						return OverridableManager.ClearOverriddenProperty(*ValueSubObject, PropertyEvent, PropertyIterator->GetNextNode());
+						return OverridableManager.ClearOverriddenProperty(*ValueSubObject, NextPropertyIterator);
 					}
 				}
 			}
@@ -512,7 +514,7 @@ bool FOverriddenPropertySet::ClearOverriddenProperty(FOverriddenPropertyNode& Pa
 		}
 
 		OverriddenPropertyNode = CurrentOverriddenPropertyNode;
-		PropertyIterator = PropertyIterator->GetNextNode();
+		++PropertyIterator;
 	}
 
 	auto CleanupClearedNodes = [this, &PropertyNodeIDPath]()
@@ -541,7 +543,7 @@ bool FOverriddenPropertySet::ClearOverriddenProperty(FOverriddenPropertyNode& Pa
 		}
 	};
 
-	if (PropertyIterator != nullptr || OverriddenPropertyNode == nullptr)
+	if (PropertyIterator || OverriddenPropertyNode == nullptr)
 	{
 		if (bClearedOverrides)
 		{
@@ -560,12 +562,12 @@ bool FOverriddenPropertySet::ClearOverriddenProperty(FOverriddenPropertyNode& Pa
 	return true;
 }
 
-void FOverriddenPropertySet::NotifyPropertyChange(FOverriddenPropertyNode* ParentPropertyNode, const EPropertyNotificationType Notification, const FPropertyChangedEvent& PropertyEvent, const FEditPropertyChain::TDoubleLinkedListNode* PropertyNode, const void* Data, bool& bNeedsCleanup)
+void FOverriddenPropertySet::NotifyPropertyChange(FOverriddenPropertyNode* ParentPropertyNode, const EPropertyNotificationType Notification, FPropertyVisitorPath::Iterator PropertyIterator, const EPropertyChangeType::Type ChangeType, const void* Data, bool& bNeedsCleanup)
 {
 	checkf(IsValid(Owner), TEXT("Expecting a valid overridable owner"));
 
 	FOverridableManager& OverridableManager = FOverridableManager::Get();
-	if (!PropertyNode)
+	if (!PropertyIterator)
 	{
 		if (ParentPropertyNode && Notification == EPropertyNotificationType::PostEdit)
 		{
@@ -586,7 +588,7 @@ void FOverriddenPropertySet::NotifyPropertyChange(FOverriddenPropertyNode* Paren
 		return;
 	}
 
-	const FProperty* Property = PropertyNode->GetValue();
+	const FProperty* Property = PropertyIterator->Property;
 	checkf(Property, TEXT("Expecting a valid property"));
 
 	const void* SubValuePtr = Property->ContainerPtrToValuePtr<void>(Data, 0); //@todo support static arrays
@@ -614,7 +616,7 @@ void FOverriddenPropertySet::NotifyPropertyChange(FOverriddenPropertyNode* Paren
 		}
 	};
 
-
+	FPropertyVisitorPath::Iterator NextPropertyIterator = PropertyIterator+1;
 	if (const FArrayProperty* ArrayProperty = CastField<FArrayProperty>(Property))
 	{
 		// Only special case is instanced subobjects, otherwise we fallback to full array override
@@ -623,8 +625,10 @@ void FOverriddenPropertySet::NotifyPropertyChange(FOverriddenPropertyNode* Paren
 			if (InnerObjectProperty->HasAnyPropertyFlags(CPF_PersistentInstance))
 			{
 				FScriptArrayHelper ArrayHelper(ArrayProperty, SubValuePtr);
-				int32 ArrayIndex = PropertyEvent.GetArrayIndex(Property->GetName());
-				if (!PropertyNode->GetNextNode())
+				int32 ArrayIndex = PropertyIterator->Index;
+				checkf(ArrayIndex == INDEX_NONE || PropertyIterator->PropertyInfo == EPropertyVisitorInfoType::ContainerIndex, TEXT("Expecting a container index"));
+
+				if (!NextPropertyIterator)
 				{
 					checkf(ArrayProperty->Inner, TEXT("Expecting an inner type for Arrays"));
 
@@ -725,7 +729,7 @@ void FOverriddenPropertySet::NotifyPropertyChange(FOverriddenPropertyNode* Paren
 					// Only arrays flagged overridable logic can record deltas, for now just override entire array
 					if (!ArrayProperty->HasAnyPropertyFlags(CPF_ExperimentalOverridableLogic))
 					{
-						if(PropertyEvent.ChangeType == EPropertyChangeType::Unspecified && ArrayIndex == INDEX_NONE)
+						if(ChangeType == EPropertyChangeType::Unspecified && ArrayIndex == INDEX_NONE)
 						{
 							// Overriding all entry in the array + override instanced sub obejects
 							ArrayReplace();
@@ -738,7 +742,7 @@ void FOverriddenPropertySet::NotifyPropertyChange(FOverriddenPropertyNode* Paren
 						return;
 					}
 
-					switch(PropertyEvent.ChangeType)
+					switch(ChangeType)
 					{
 					case EPropertyChangeType::ValueSet:
 						checkf(ArrayIndex != INDEX_NONE, TEXT("ValueSet change type should have associated indexes"));
@@ -796,7 +800,7 @@ void FOverriddenPropertySet::NotifyPropertyChange(FOverriddenPropertyNode* Paren
 					if (UObject* SubObject = InnerObjectProperty->GetObjectPropertyValue(ArrayHelper.GetElementPtr(ArrayIndex)))
 					{
 						// This should not be needed in the property grid, as it should already been called on the subobject itself.
-						OverridableManager.NotifyPropertyChange(Notification, *SubObject, PropertyEvent, PropertyNode->GetNextNode());
+						OverridableManager.NotifyPropertyChange(Notification, *SubObject, NextPropertyIterator, ChangeType);
 						return;
 					}
 				}
@@ -821,9 +825,11 @@ void FOverriddenPropertySet::NotifyPropertyChange(FOverriddenPropertyNode* Paren
 		FObjectPropertyBase* ValueInstancedObjectProperty = MapProperty->ValueProp->HasAnyPropertyFlags(CPF_PersistentInstance) ? CastField<FObjectPropertyBase>(MapProperty->ValueProp) : nullptr;
 
 		FScriptMapHelper MapHelper(MapProperty, SubValuePtr);
-		int32 LogicalMapIndex = PropertyEvent.GetArrayIndex(Property->GetName());
+		int32 LogicalMapIndex = PropertyIterator->Index;
+		checkf(LogicalMapIndex == INDEX_NONE || PropertyIterator->PropertyInfo == EPropertyVisitorInfoType::ContainerIndex, TEXT("Expecting a container index type"));
+
 		int32 InternalMapIndex = LogicalMapIndex != INDEX_NONE ? MapHelper.FindInternalIndex(LogicalMapIndex) : INDEX_NONE;
-		if (!PropertyNode->GetNextNode())
+		if (!NextPropertyIterator)
 		{
 			static const FProperty* SavedProp = nullptr;
 			static uint8* SavedPreEditMap = nullptr;
@@ -938,7 +944,7 @@ void FOverriddenPropertySet::NotifyPropertyChange(FOverriddenPropertyNode* Paren
 			// Only maps flagged overridable logic can be handled here
 			if (!MapProperty->HasAnyPropertyFlags(CPF_ExperimentalOverridableLogic))
 			{
-				if (PropertyEvent.ChangeType == EPropertyChangeType::Unspecified && InternalMapIndex == INDEX_NONE)
+				if (ChangeType == EPropertyChangeType::Unspecified && InternalMapIndex == INDEX_NONE)
 				{
 					// Overriding all entry in the array + override instanced sub obejects
 					MapReplace();
@@ -951,7 +957,7 @@ void FOverriddenPropertySet::NotifyPropertyChange(FOverriddenPropertyNode* Paren
 				return;
 			}
 
-			switch (PropertyEvent.ChangeType)
+			switch (ChangeType)
 			{
 			case EPropertyChangeType::ValueSet:
 				checkf(LogicalMapIndex != INDEX_NONE, TEXT("ValueSet change type should have associated indexes"));
@@ -1010,21 +1016,21 @@ void FOverriddenPropertySet::NotifyPropertyChange(FOverriddenPropertyNode* Paren
 			//if (UObject* SubObject = KeyInstancedObjectProperty ? KeyInstancedObjectProperty->GetObjectPropertyValue(MapHelper.GetValuePtr(InternalMapIndex)) : nullptr)
 			//{
 			//	// This should not be needed in the property grid, as it should already been called on the subobject.
-			//	OverridableManager.NotifyPropertyChange(Notification, *SubObject, PropertyEvent, PropertyNode->GetNextNode());
+			//	OverridableManager.NotifyPropertyChange(Notification, *SubObject, NextPropertyIterator, ChangeType);
 			//	return;
 			//}
 
 			if (UObject* SubObject = ValueInstancedObjectProperty ? ValueInstancedObjectProperty->GetObjectPropertyValue(MapHelper.GetValuePtr(InternalMapIndex)) : nullptr)
 			{
 				// This should not be needed in the property grid, as it should already been called on the subobject.
-				OverridableManager.NotifyPropertyChange(Notification, *SubObject, PropertyEvent, PropertyNode->GetNextNode());
+				OverridableManager.NotifyPropertyChange(Notification, *SubObject, NextPropertyIterator, ChangeType);
 				return;
 			}
 		}
 	}
 	else if (Property->IsA<FStructProperty>())
 	{
-		if (!PropertyNode->GetNextNode())
+		if (!NextPropertyIterator)
 		{
 			if (Notification == EPropertyNotificationType::PostEdit && SubPropertyNode)
 			{
@@ -1033,13 +1039,13 @@ void FOverriddenPropertySet::NotifyPropertyChange(FOverriddenPropertyNode* Paren
 		}
 		else
 		{
-			NotifyPropertyChange(SubPropertyNode, Notification, PropertyEvent, PropertyNode->GetNextNode(), SubValuePtr, bNeedsCleanup);
+			NotifyPropertyChange(SubPropertyNode, Notification, NextPropertyIterator, ChangeType, SubValuePtr, bNeedsCleanup);
 		}
 		return;
 	}
 	else if (const FObjectPropertyBase* ObjectProperty = CastField<FObjectPropertyBase>(Property))
 	{
-		if (!PropertyNode->GetNextNode())
+		if (!NextPropertyIterator)
 		{
 			if (Notification == EPropertyNotificationType::PostEdit && SubPropertyNode)
 			{
@@ -1049,13 +1055,13 @@ void FOverriddenPropertySet::NotifyPropertyChange(FOverriddenPropertyNode* Paren
 		else if (UObject* SubObject = ObjectProperty->GetObjectPropertyValue(SubValuePtr))
 		{
 			// This should not be needed in the property grid, as it should already been called on the subobject.
-			OverridableManager.NotifyPropertyChange(Notification, *SubObject, PropertyEvent, PropertyNode->GetNextNode());
+			OverridableManager.NotifyPropertyChange(Notification, *SubObject, NextPropertyIterator, ChangeType);
 		}
 		return;
 	}
 	else if (const FOptionalProperty* OptionalProperty = CastField<FOptionalProperty>(Property))
 	{
-		if (!PropertyNode->GetNextNode())
+		if (!NextPropertyIterator)
 		{
 			if (Notification == EPropertyNotificationType::PostEdit && SubPropertyNode)
 			{
@@ -1064,12 +1070,12 @@ void FOverriddenPropertySet::NotifyPropertyChange(FOverriddenPropertyNode* Paren
 		}
 		else if (OptionalProperty->IsSet(Data))
 		{
-			NotifyPropertyChange(SubPropertyNode, Notification, PropertyEvent, PropertyNode->GetNextNode(), OptionalProperty->GetValuePointerForRead(SubValuePtr), bNeedsCleanup);
+			NotifyPropertyChange(SubPropertyNode, Notification, NextPropertyIterator, ChangeType, OptionalProperty->GetValuePointerForRead(SubValuePtr), bNeedsCleanup);
 		}
 		return;
 	}
 
-	UE_CLOG(PropertyNode->GetNextNode(), LogOverridableObject, Warning, TEXT("Unsupported property type(%s), fallback to overriding entire property"), *Property->GetName());
+	UE_CLOG(NextPropertyIterator, LogOverridableObject, Warning, TEXT("Unsupported property type(%s), fallback to overriding entire property"), *Property->GetName());
 	if (Notification == EPropertyNotificationType::PostEdit)
 	{
 		if (SubPropertyNode)
@@ -1101,36 +1107,36 @@ void FOverriddenPropertySet::CacheArchetype()
 	}
 }
 
-EOverriddenPropertyOperation FOverriddenPropertySet::GetOverriddenPropertyOperation(const FPropertyChangedEvent& PropertyEvent, const FEditPropertyChain::TDoubleLinkedListNode* PropertyNode, bool* bOutInheritedOperation /*= nullptr*/) const
+EOverriddenPropertyOperation FOverriddenPropertySet::GetOverriddenPropertyOperation(FPropertyVisitorPath::Iterator PropertyIterator, bool* bOutInheritedOperation /*= nullptr*/) const
 {
 	if (const FOverriddenPropertyNode* RootNode = OverriddenPropertyNodes.Find(RootNodeID))
 	{
-		return GetOverriddenPropertyOperation(*RootNode, PropertyEvent, PropertyNode, bOutInheritedOperation, Owner);
+		return GetOverriddenPropertyOperation(*RootNode, PropertyIterator, bOutInheritedOperation, Owner);
 	}
 	return EOverriddenPropertyOperation::None;
 }
 
-bool FOverriddenPropertySet::ClearOverriddenProperty(const FPropertyChangedEvent& PropertyEvent, const FEditPropertyChain::TDoubleLinkedListNode* PropertyNode)
+bool FOverriddenPropertySet::ClearOverriddenProperty(FPropertyVisitorPath::Iterator PropertyIterator)
 {
 	if (FOverriddenPropertyNode* RootNode = OverriddenPropertyNodes.Find(RootNodeID))
 	{
-		return ClearOverriddenProperty(*RootNode, PropertyEvent, PropertyNode, Owner);
+		return ClearOverriddenProperty(*RootNode, PropertyIterator, Owner);
 	}
 	return true;
 }
 
-void FOverriddenPropertySet::OverrideProperty(const FPropertyChangedEvent& PropertyEvent, const FEditPropertyChain::TDoubleLinkedListNode* PropertyNode, const void* Data)
+void FOverriddenPropertySet::OverrideProperty(FPropertyVisitorPath::Iterator PropertyIterator, const void* Data)
 {
 	FOverriddenPropertyNode& RootPropertyNode = OverriddenPropertyNodes.FindOrAdd(RootNodeID);
 	bool bNeedsCleanup = false;
-	NotifyPropertyChange(&RootPropertyNode, EPropertyNotificationType::PreEdit, PropertyEvent, PropertyNode, Data, bNeedsCleanup);
-	NotifyPropertyChange(&RootPropertyNode, EPropertyNotificationType::PostEdit, PropertyEvent, PropertyNode, Data, bNeedsCleanup);
+	NotifyPropertyChange(&RootPropertyNode, EPropertyNotificationType::PreEdit, PropertyIterator, EPropertyChangeType::Unspecified, Data, bNeedsCleanup);
+	NotifyPropertyChange(&RootPropertyNode, EPropertyNotificationType::PostEdit, PropertyIterator, EPropertyChangeType::Unspecified, Data, bNeedsCleanup);
 }
 
-void FOverriddenPropertySet::NotifyPropertyChange(const EPropertyNotificationType Notification, const FPropertyChangedEvent& PropertyEvent, const FEditPropertyChain::TDoubleLinkedListNode* PropertyNode, const void* Data)
+void FOverriddenPropertySet::NotifyPropertyChange(const EPropertyNotificationType Notification, FPropertyVisitorPath::Iterator PropertyIterator, const EPropertyChangeType::Type ChangeType, const void* Data)
 {
 	bool bNeedsCleanup = false;
-	NotifyPropertyChange(&OverriddenPropertyNodes.FindOrAdd(RootNodeID), Notification, PropertyEvent, PropertyNode, Data, bNeedsCleanup);
+	NotifyPropertyChange(&OverriddenPropertyNodes.FindOrAdd(RootNodeID), Notification, PropertyIterator, ChangeType, Data, bNeedsCleanup);
 }
 
 EOverriddenPropertyOperation FOverriddenPropertySet::GetOverriddenPropertyOperation(const FArchiveSerializedPropertyChain* CurrentPropertyChain, FProperty* Property) const
diff --git a/Engine/Source/Runtime/CoreUObject/Private/UObject/PropertyVisitor.cpp b/Engine/Source/Runtime/CoreUObject/Private/UObject/PropertyVisitor.cpp
--- a/Engine/Source/Runtime/CoreUObject/Private/UObject/PropertyVisitor.cpp
+++ b/Engine/Source/Runtime/CoreUObject/Private/UObject/PropertyVisitor.cpp
@@ -6,6 +6,10 @@
 #include "UObject/UnrealType.h"
 #include "Serialization/ArchiveSerializedPropertyChain.h"
 
+//----------------------------------------------------------------------//
+// FPropertyVisitorInfo 
+//----------------------------------------------------------------------//
+
 bool FPropertyVisitorInfo::Identical(const FPropertyVisitorInfo& Other) const
 {
 	return Property == Other.Property
@@ -15,6 +19,54 @@ bool FPropertyVisitorInfo::Identical(const FPropertyVisitorInfo& Other) const
 		&& bContainsInnerProperties == Other.bContainsInnerProperties;
 }
 
+//----------------------------------------------------------------------//
+// FPropertyVisitorPath 
+//----------------------------------------------------------------------//
+
+FPropertyVisitorPath::Iterator FPropertyVisitorPath::InvalidIterator()
+{
+	static FPropertyVisitorPath EmptyPath;
+	return Iterator(EmptyPath.GetPath());
+}
+
+FPropertyVisitorPath::FPropertyVisitorPath(const FPropertyChangedEvent& PropertyEvent, const FEditPropertyChain& PropertyChain)
+{
+	const FEditPropertyChain::TDoubleLinkedListNode* PropertyIterator = PropertyChain.GetActiveMemberNode() ? PropertyChain.GetActiveMemberNode() : PropertyChain.GetHead();
+	while(PropertyIterator)
+	{
+		const FProperty* CurrentProperty = PropertyIterator->GetValue();
+		int32 ArrayIndex = PropertyEvent.GetArrayIndex(CurrentProperty->GetName());
+		if (ArrayIndex == INDEX_NONE)
+		{
+			Push(FPropertyVisitorInfo{CurrentProperty});
+		}
+		else if (CurrentProperty->ArrayDim > 1)
+		{
+			Push(FPropertyVisitorInfo{CurrentProperty, ArrayIndex, EPropertyVisitorInfoType::StaticArrayIndex});
+		}
+		else
+		{
+			// Only container is left, no easy way yet to know if we are editing a Key/value of a map
+			Push(FPropertyVisitorInfo{CurrentProperty, ArrayIndex, EPropertyVisitorInfoType::ContainerIndex});
+		}
+		PropertyIterator = PropertyIterator->GetNextNode();
+	}
+}
+
+FPropertyVisitorPath::FPropertyVisitorPath(const FArchiveSerializedPropertyChain& PropertyChain)
+{
+	TArray<class FProperty*, TInlineAllocator<8>>::TConstIterator PropertyIterator = PropertyChain.GetRootIterator();
+	while (PropertyIterator)
+	{
+		const FProperty* CurrentProperty = (*PropertyIterator);
+
+		// @todo add container index eventually if we ever have this information
+		Push(FPropertyVisitorInfo{CurrentProperty});
+
+		++PropertyIterator;
+	}
+}
+
 FString FPropertyVisitorPath::ToString(const TCHAR* Separator /*= TEXT(".")*/) const
 {
 	return PropertyVisitorHelpers::PathToString(Path, Separator);
@@ -162,6 +214,7 @@ FArchiveSerializedPropertyChain PropertyVisitorHelpers::PathToSerializedProperty
 	FArchiveSerializedPropertyChain Chain;
 	for (const FPropertyVisitorInfo& Info : Path)
 	{
+		// @todo add container index eventually if we ever have this information
 		Chain.PushProperty(const_cast<FProperty*>(Info.Property), Info.Property->IsEditorOnlyProperty());
 	}
 	return Chain;
diff --git a/Engine/Source/Runtime/CoreUObject/Public/UObject/OverridableManager.h b/Engine/Source/Runtime/CoreUObject/Public/UObject/OverridableManager.h
--- a/Engine/Source/Runtime/CoreUObject/Public/UObject/OverridableManager.h
+++ b/Engine/Source/Runtime/CoreUObject/Public/UObject/OverridableManager.h
@@ -29,7 +29,7 @@ struct FOverriddenPropertyAnnotation
 };
 
 /**
- * Global container of overriden object annotations
+ * Global container of overridden object annotations
  */
 class FOverriddenPropertyAnnotations : public FUObjectAnnotationSparse<FOverriddenPropertyAnnotation, true/*bAutoRemove*/>
 {
@@ -202,35 +202,82 @@ public:
 	 * Override a specific property of an object (Helper methods to call Pre/PostOverride)
 	 * Note: Supports object that does not have overridable serialization enabled
 	 * @param Object owning the property
-	 * @param PropertyEvent information about the type of change including any container item index
-	 * @param PropertyChain leading to the property that is about to be overridden */
-	COREUOBJECT_API void OverrideProperty(UObject& Object, const FPropertyChangedEvent& PropertyEvent, const FEditPropertyChain& PropertyChain);
+	 * @param PropertyPath leading to the property that is about to be overridden */
+	COREUOBJECT_API void OverrideProperty(UObject& Object, const FPropertyVisitorPath& PropertyPath);
 
 	/**
 	 * Clears an overridden properties specified by the property chain
 	 * @param Object owning the property to clear
-	 * @param PropertyEvent only needed to know about the container item index in any
-	 * @param PropertyChain to the property to clear from the root of the specified object
+	 * @param PropertyPath to the property to clear from the root of the specified object
 	 * @return true if the property was successfully cleared. */
-	FORCEINLINE bool ClearOverriddenProperty(UObject& Object, const FPropertyChangedEvent& PropertyEvent, const FEditPropertyChain& PropertyChain)
+	bool ClearOverriddenProperty(UObject& Object, const FPropertyVisitorPath& PropertyPath)
+	{
+		return ClearOverriddenProperty(Object, PropertyPath.GetRootIterator());
+	}
+
+	/**
+	 * To be called prior to override a property of the specified object
+	 * Note: Supports object that does not have overridable serialization enabled
+	 * @param Object owning the property
+	 * @param PropertyPath leading to the property about to be overridden */
+	COREUOBJECT_API void PreOverrideProperty(UObject& Object, const FPropertyVisitorPath& PropertyPath);
+
+	/**
+	 * To be called after the property of the specified object was overridden
+	 * Note: Supports object that does not have overridable serialization enabled
+	 * @param Object owning the property
+	 * @param PropertyPath leading to the property that was overridden
+ 	 * @param ChangeType of the current operation */
+	COREUOBJECT_API void PostOverrideProperty(UObject& Object, const FPropertyVisitorPath& PropertyPath, const EPropertyChangeType::Type ChangeType);
+
+	/**
+	 * Retrieve the overridable operation from the specified property path
+	 * @param Object owning the property
+	 * @param PropertyPath leading to the property the caller is interested in
+	 * @param bOutInheritedOperation optional parameter to know if the operation returned was inherited from a parent property
+	 * @return the current type of override operation on the property */
+	EOverriddenPropertyOperation GetOverriddenPropertyOperation(UObject& Object, const FPropertyVisitorPath& PropertyPath, bool* bOutInheritedOperation = nullptr)
 	{
-		return ClearOverriddenProperty(Object, PropertyEvent, PropertyChain.GetActiveMemberNode() ? PropertyChain.GetActiveMemberNode() : PropertyChain.GetHead());
+		return GetOverriddenPropertyOperation(Object, PropertyPath.GetRootIterator(), bOutInheritedOperation);
+	}
+
+	/**
+	 * Serializes the overridden properties of the specified object into the record
+	 * @param Object to serialize the overridden property
+	 * @param ObjectRecord the record to use for serialization */
+	COREUOBJECT_API void SerializeOverriddenProperties(UObject& Object, FStructuredArchive::FRecord ObjectRecord);
+
+	/**
+	 * Override a specific property of an object (Helper methods to call Pre/PostOverride)
+	 * Note: Supports object that does not have overridable serialization enabled
+	 * @param Object owning the property
+	 * @param PropertyEvent information about the type of change including any container item index
+	 * @param PropertyChain leading to the property that is about to be overridden */
+	void OverrideProperty(UObject& Object, const FPropertyChangedEvent& PropertyEvent, const FEditPropertyChain& PropertyChain)
+	{
+		OverrideProperty(Object, FPropertyVisitorPath(PropertyEvent, PropertyChain));
 	}
 
 	/**
 	 * Clears an overridden properties specified by the property chain
 	 * @param Object owning the property to clear
 	 * @param PropertyEvent only needed to know about the container item index in any
-	 * @param PropertyNode leading to the property to clear, null means to clear the specified object
+	 * @param PropertyChain to the property to clear from the root of the specified object
 	 * @return true if the property was successfully cleared. */
-	COREUOBJECT_API bool ClearOverriddenProperty(UObject& Object, const FPropertyChangedEvent& PropertyEvent, const FEditPropertyChain::TDoubleLinkedListNode* PropertyNode);
+	bool ClearOverriddenProperty(UObject& Object, const FPropertyChangedEvent& PropertyEvent, const FEditPropertyChain& PropertyChain)
+	{
+		return ClearOverriddenProperty(Object, FPropertyVisitorPath(PropertyEvent, PropertyChain));
+	}
 
 	/**
 	 * To be called prior to override a property of the specified object
 	 * Note: Supports object that does not have overridable serialization enabled
 	 * @param Object owning the property
-	 * @param PropertyChain leading to the property about to be overriden */
-	COREUOBJECT_API void PreOverrideProperty(UObject& Object, const FEditPropertyChain& PropertyChain);
+	 * @param PropertyChain leading to the property about to be overridden */
+	void PreOverrideProperty(UObject& Object, const FEditPropertyChain& PropertyChain)
+	{
+		PreOverrideProperty(Object, FPropertyVisitorPath(FPropertyChangedEvent(nullptr), PropertyChain));
+	}
 
 	/**
 	 * To be called after the property was overridden of the specified object
@@ -238,16 +285,10 @@ public:
 	 * @param Object owning the property
 	 * @param PropertyEvent information about the type of change including any container item index
 	 * @param PropertyChain leading to the property that was overridden */
-	COREUOBJECT_API void PostOverrideProperty(UObject& Object, const FPropertyChangedEvent& PropertyEvent, const FEditPropertyChain& PropertyChain);
-
-	/**
-	 * Notifies that a properties is about to change or has been changed for the specified object.
-	 * Note: Will ensure if the specified object does not have overridable serialization enabled
-	 * @param Notification the type of notification (pre or post
-	 * @param Object owning the property
-	 * @param PropertyEvent information about the type of change including any container item index
-	 * @param PropertyNode leading to the property that is changing, null indicate that it is the specified object has changed */
-	COREUOBJECT_API void NotifyPropertyChange(const EPropertyNotificationType Notification, UObject& Object, const FPropertyChangedEvent& PropertyEvent, const FEditPropertyChain::TDoubleLinkedListNode* PropertyNode);
+	void PostOverrideProperty(UObject& Object, const FPropertyChangedEvent& PropertyEvent, const FEditPropertyChain& PropertyChain)
+	{
+		PostOverrideProperty(Object, FPropertyVisitorPath(PropertyEvent, PropertyChain), PropertyEvent.ChangeType);
+	}
 
 	/**
 	 * Retrieve the overridable operation from the specified the edit property chain
@@ -258,23 +299,9 @@ public:
 	 * @return the current type of override operation on the property */
 	FORCEINLINE EOverriddenPropertyOperation GetOverriddenPropertyOperation(UObject& Object, const FPropertyChangedEvent& PropertyEvent, const FEditPropertyChain& PropertyChain, bool* bOutInheritedOperation = nullptr)
 	{
-		return GetOverriddenPropertyOperation(Object, PropertyEvent, PropertyChain.GetActiveMemberNode() ? PropertyChain.GetActiveMemberNode() : PropertyChain.GetHead(), bOutInheritedOperation);
+		return GetOverriddenPropertyOperation(Object, FPropertyVisitorPath(PropertyEvent, PropertyChain), bOutInheritedOperation);
 	}
 
-	/**
-	 * Retrieve the overridable operation from the specified the edit property chain
-	 * @param Object owning the property
-	 * @param PropertyEvent only needed to know about the container item index in any
-	 * @param PropertyNode leading to the property interested in, null will return the operation on the specified object 
-	 * @param bOutInheritedOperation optional parameter to know if the state returned was inherited from a parent property
-	 * @return the current type of override operation on the property */
-	COREUOBJECT_API EOverriddenPropertyOperation GetOverriddenPropertyOperation(UObject& Object, const FPropertyChangedEvent& PropertyEvent, const FEditPropertyChain::TDoubleLinkedListNode* PropertyNode, bool* bOutInheritedOperation = nullptr);
-
-	/**
-	 * Serializes the overriden properties of the specified object into the record
-	 * @param Object to serialize the overriden property
-	 * @param ObjectRecord the record to use for serialization */
-	COREUOBJECT_API void SerializeOverriddenProperties(UObject& Object, FStructuredArchive::FRecord ObjectRecord);
 
 	void HandleObjectsReInstantiated(const TMap<UObject*, UObject*>& OldToNewInstanceMap);
 
@@ -285,6 +312,11 @@ protected:
 	FOverriddenPropertyAnnotations OverriddenObjectAnnotations;
 #else
 	FUObjectAnnotationSparseBool NeedsSubobjectTemplateInstantiation;
-#endif // WITH_EDITORONLY_DATA 
+#endif // WITH_EDITORONLY_DATA
+
+	friend FOverriddenPropertySet;
+	COREUOBJECT_API void NotifyPropertyChange(const EPropertyNotificationType Notification, UObject& Object, FPropertyVisitorPath::Iterator PropertyIterator, const EPropertyChangeType::Type ChangeType);
+	COREUOBJECT_API EOverriddenPropertyOperation GetOverriddenPropertyOperation(UObject& Object, FPropertyVisitorPath::Iterator PropertyIterator, bool* bOutInheritedOperation = nullptr);
+	COREUOBJECT_API bool ClearOverriddenProperty(UObject& Object, FPropertyVisitorPath::Iterator PropertyIterator);
 };
 
diff --git a/Engine/Source/Runtime/CoreUObject/Public/UObject/OverriddenPropertySet.h b/Engine/Source/Runtime/CoreUObject/Public/UObject/OverriddenPropertySet.h
--- a/Engine/Source/Runtime/CoreUObject/Public/UObject/OverriddenPropertySet.h
+++ b/Engine/Source/Runtime/CoreUObject/Public/UObject/OverriddenPropertySet.h
@@ -3,13 +3,12 @@
 #pragma once
 
 #include "UObject/Object.h"
-#include "Serialization/ArchiveSerializedPropertyChain.h"
 
 #include "OverriddenPropertySet.generated.h"
 
 COREUOBJECT_API DECLARE_LOG_CATEGORY_EXTERN(LogOverridableObject, Warning, All);
 
-struct FPropertyChangedChainEvent;
+struct FArchiveSerializedPropertyChain;
 
 /*
  *************************************************************************************
@@ -222,33 +221,30 @@ public:
 
 	/**
 	 * Retrieve the overridable operation from the specified the edit property chain node
-	 * @param PropertyEvent only needed to know about the container item index in any
-	 * @param PropertyNode leading to the property interested in, null will return the operation of the object itself
+	 * @param PropertyIterator leading to the property interested in, invalid iterator will return the operation of the object itself
 	 * @param bOutInheritedOperation optional parameter to know if the state returned was inherited from a parent property
 	 * @return the current type of override operation on the property */
-	EOverriddenPropertyOperation GetOverriddenPropertyOperation(const FPropertyChangedEvent& PropertyEvent, const FEditPropertyChain::TDoubleLinkedListNode* PropertyNode, bool* bOutInheritedOperation = nullptr) const;
+	EOverriddenPropertyOperation GetOverriddenPropertyOperation(FPropertyVisitorPath::Iterator PropertyIterator, bool* bOutInheritedOperation = nullptr) const;
 
 	/**
 	 * Clear any properties from the serialized property chain node
-	 * @param PropertyEvent only needed to know about the container item index in any
-	 * @param PropertyNode leading to the property to clear, null will clear the overrides on the object itself
+	 * @param PropertyIterator leading to the property to clear, invalid iterator will clear the overrides on the object itself
 	 * @return if the operation was successful */
-	bool ClearOverriddenProperty(const FPropertyChangedEvent& PropertyEvent, const FEditPropertyChain::TDoubleLinkedListNode* PropertyNode);
+	bool ClearOverriddenProperty(FPropertyVisitorPath::Iterator PropertyIterator);
 
 	/**
 	 * Utility methods that call NotifyPropertyChange(Pre/PostEdit)
-	 * @param PropertyEvent information about the type of change
-	 * @param PropertyNode leading to the property that is changing, null means it is the object itself that is changing
+	 * @param PropertyIterator leading to the property that is changing, invalid iterator means it is the object itself that is changing
 	 * @param Data memory of the current property */
-	void OverrideProperty(const FPropertyChangedEvent& PropertyEvent, const FEditPropertyChain::TDoubleLinkedListNode* PropertyNode, const void* Data);
+	void OverrideProperty(FPropertyVisitorPath::Iterator PropertyIterator, const void* Data);
 
 	/**
 	 * Handling and storing modification on a property of an object
 	 * @param Notification type either pre/post property overridden
-	 * @param PropertyEvent information about the type of change
-	 * @param PropertyNode leading to the property that is changing, null means it is the object itself that is changing
+	 * @param PropertyIterator leading to the property that is changing, null means it is the object itself that is changing
+	 * @param ChangeType of the current operation
 	 * @param Data memory of the current property */
-	void NotifyPropertyChange(const EPropertyNotificationType Notification, const FPropertyChangedEvent& PropertyEvent, const FEditPropertyChain::TDoubleLinkedListNode* PropertyNode, const void* Data);
+	void NotifyPropertyChange(const EPropertyNotificationType Notification, FPropertyVisitorPath::Iterator PropertyIterator, const EPropertyChangeType::Type ChangeType, const void* Data);
 
 	/**
 	 * Retrieve the overridable operation from the specified the serialized property chain and the specified property
@@ -314,15 +310,15 @@ protected:
 
 	FOverriddenPropertyNode& FindOrAddNode(FOverriddenPropertyNode& ParentPropertyNode, FOverriddenPropertyNodeID NodeID);
 
-	EOverriddenPropertyOperation GetOverriddenPropertyOperation(const FOverriddenPropertyNode& ParentPropertyNode, const FPropertyChangedEvent& PropertyEvent, const FEditPropertyChain::TDoubleLinkedListNode* PropertyNode, bool* bOutInheritedOperation, const void* Data) const;
-	bool ClearOverriddenProperty(FOverriddenPropertyNode& ParentPropertyNode, const FPropertyChangedEvent& PropertyEvent, const FEditPropertyChain::TDoubleLinkedListNode* PropertyNode, const void* Data);
-	void NotifyPropertyChange(FOverriddenPropertyNode* ParentPropertyNode, const EPropertyNotificationType Notification, const FPropertyChangedEvent& PropertyEvent, const FEditPropertyChain::TDoubleLinkedListNode* PropertyNode, const void* Data, bool& bNeedsCleanup);
+	EOverriddenPropertyOperation GetOverriddenPropertyOperation(const FOverriddenPropertyNode& ParentPropertyNode, FPropertyVisitorPath::Iterator PropertyIterator, bool* bOutInheritedOperation, const void* Data) const;
+	bool ClearOverriddenProperty(FOverriddenPropertyNode& ParentPropertyNode, FPropertyVisitorPath::Iterator PropertyIterator, const void* Data);
+	void NotifyPropertyChange(FOverriddenPropertyNode* ParentPropertyNode, const EPropertyNotificationType Notification, FPropertyVisitorPath::Iterator PropertyIterator, const EPropertyChangeType::Type ChangeType, const void* Data, bool& bNeedsCleanup);
 
 	EOverriddenPropertyOperation GetOverriddenPropertyOperation(const FOverriddenPropertyNode& ParentPropertyNode, const FArchiveSerializedPropertyChain* CurrentPropertyChain, FProperty* Property) const;
 	FOverriddenPropertyNode* SetOverriddenPropertyOperation(EOverriddenPropertyOperation Operation, FOverriddenPropertyNode& ParentPropertyNode, const FArchiveSerializedPropertyChain* CurrentPropertyChain, FProperty* Property);
 	const FOverriddenPropertyNode* GetOverriddenPropertyNode(const FOverriddenPropertyNode& ParentPropertyNode, const FArchiveSerializedPropertyChain* CurrentPropertyChain) const;
 
-	void RemoveOverriddenSubProperties(FOverriddenPropertyNode& PropertyNode);
+	void RemoveOverriddenSubProperties(FOverriddenPropertyNode& PropertyIterator);
 
 
 private:
diff --git a/Engine/Source/Runtime/CoreUObject/Public/UObject/PropertyVisitor.h b/Engine/Source/Runtime/CoreUObject/Public/UObject/PropertyVisitor.h
--- a/Engine/Source/Runtime/CoreUObject/Public/UObject/PropertyVisitor.h
+++ b/Engine/Source/Runtime/CoreUObject/Public/UObject/PropertyVisitor.h
@@ -6,6 +6,8 @@
 #include "Containers/ArrayView.h"
 
 class FProperty;
+class FEditPropertyChain;
+struct FPropertyChangedEvent;
 struct FArchiveSerializedPropertyChain;
 
 enum class EPropertyVisitorControlFlow : uint8
@@ -106,6 +108,9 @@ public:
 		Path = InPath;
 	}
 
+	explicit COREUOBJECT_API FPropertyVisitorPath(const FPropertyChangedEvent& PropertyEvent, const FEditPropertyChain& PropertyChain);
+	explicit COREUOBJECT_API FPropertyVisitorPath(const FArchiveSerializedPropertyChain& PropertyChain);
+
 	void Push(const FPropertyVisitorInfo& Info)
 	{
 		Path.Push(Info);
@@ -169,11 +174,24 @@ public:
 	 */
 	COREUOBJECT_API void* GetPropertyDataPtr(UObject* Object) const;
 
-	TArray<FPropertyVisitorInfo>::TConstIterator GetRootIterator() const
+	/** Iterator for a property visitor path */
+	using Iterator = TArray<FPropertyVisitorInfo>::TConstIterator;
+
+	/**
+	 * Returns an iterator on the root path node, useful when calling methods that are recursive
+	 * @return an iterator on the root path node
+	 */
+	Iterator GetRootIterator() const
 	{
-		return Path.CreateConstIterator();
+		return Iterator(Path);
 	}
 
+	/**
+	 * Invalid iterator pointing to an empty path
+	 * @return an iterator to an empty path
+	 */
+	COREUOBJECT_API static Iterator InvalidIterator();
+
 	/**
 	 * Converts path to an archive serialized property chain
 	 * This method useful when APIs like in the FOverridableManager are taking in a
