Optimize UEnum (104 -> 96)
Optimize USCS_Node(72->64)
Saves 0.15 MB
#rnx

diff --git a/Engine/Source/Runtime/CoreUObject/Public/UObject/Class.h b/Engine/Source/Runtime/CoreUObject/Public/UObject/Class.h
--- a/Engine/Source/Runtime/CoreUObject/Public/UObject/Class.h
+++ b/Engine/Source/Runtime/CoreUObject/Public/UObject/Class.h
@@ -2160,7 +2160,7 @@ class UEnum : public UField
 
 public:
 	/** How this enum is declared in C++, affects the internal naming of enum values */
-	enum class ECppForm
+	enum class ECppForm : uint8
 	{
 		Regular,
 		Namespaced,
@@ -2691,11 +2691,11 @@ protected:
 	/** Enum flags. */
 	EEnumFlags EnumFlags;
 
-	/** pointer to function used to look up the enum's display name. Currently only assigned for UEnums generated for nativized blueprints */
-	FEnumDisplayNameFn EnumDisplayNameFn;
-
 	/** Package name this enum was in when its names were being added to the primary list */
 	FName EnumPackage;
+	
+	/** pointer to function used to look up the enum's display name. Currently only assigned for UEnums generated for nativized blueprints */
+	FEnumDisplayNameFn EnumDisplayNameFn;
 
 	/** lock to be taken when accessing AllEnumNames */
 	static FRWLock AllEnumNamesLock;
