Fixed the IDO CopyTaggedProperties function to set the serializing defaults flag based on the object

#rnx

diff --git a/Engine/Source/Runtime/CoreUObject/Private/UObject/InstanceDataObjectUtils.cpp b/Engine/Source/Runtime/CoreUObject/Private/UObject/InstanceDataObjectUtils.cpp
--- a/Engine/Source/Runtime/CoreUObject/Private/UObject/InstanceDataObjectUtils.cpp
+++ b/Engine/Source/Runtime/CoreUObject/Private/UObject/InstanceDataObjectUtils.cpp
@@ -602,7 +602,31 @@ namespace UE
 		Result->StaticLink(/*bRelinkExistingProperties*/true);
 		return Result;
 	}
-	
+
+	struct FSerializingDefaultsScope
+	{
+		UE_NONCOPYABLE(FSerializingDefaultsScope);
+
+		inline FSerializingDefaultsScope(FArchive& Ar, const UObject* Object)
+		{
+			if (Object->HasAnyFlags(RF_ClassDefaultObject))
+			{
+				Archive = &Ar;
+				Archive->StartSerializingDefaults();
+			}
+		}
+
+		inline ~FSerializingDefaultsScope()
+		{
+			if (Archive)
+			{
+				Archive->StopSerializingDefaults();
+			}
+		}
+
+		FArchive* Archive = nullptr;
+	};
+
 	void CopyTaggedProperties(const UObject* Source, UObject* Dest)
 	{
 		FUObjectSerializeContext* SerializeContext = FUObjectThreadContext::Get().GetSerializeContext();
@@ -612,17 +636,16 @@ namespace UE
 
 		TArray<uint8> Buffer;
 		Buffer.Reserve(Source->GetClass()->GetStructureSize());
+
 		FObjectWriter Writer(Buffer);
+		FSerializingDefaultsScope WriterDefaultsScope(Writer, Source);
 		Writer.ArNoDelta = true;
-		Writer.StartSerializingDefaults();
 		Source->GetClass()->SerializeTaggedProperties(Writer, (uint8*)Source, Source->GetClass(), nullptr);
-		Writer.StopSerializingDefaults();
 
 		FObjectReader Reader(Buffer);
+		FSerializingDefaultsScope ReaderDefaultsScope(Reader, Dest);
 		Reader.ArMergeOverrides = true;
-		Reader.StartSerializingDefaults();
 		Dest->GetClass()->SerializeTaggedProperties(Reader, (uint8*)Dest, Dest->GetClass(), nullptr);
-		Reader.StopSerializingDefaults();
 	}
 
 	static void SetClassFlags(UClass* IDOClass, const UClass* OwnerClass)
