[RegressionFix] 35336947 caused NAME_ContainsLoosePropertiesMetadata to stop being set which broke IDOs.

diff --git a/Engine/Source/Runtime/CoreUObject/Private/UObject/InstanceDataObjectUtils.cpp b/Engine/Source/Runtime/CoreUObject/Private/UObject/InstanceDataObjectUtils.cpp
--- a/Engine/Source/Runtime/CoreUObject/Private/UObject/InstanceDataObjectUtils.cpp
+++ b/Engine/Source/Runtime/CoreUObject/Private/UObject/InstanceDataObjectUtils.cpp
@@ -316,14 +316,6 @@ namespace UE
 			}
 		}
 
-		const auto TrySetContainsLooseProperties = [](FProperty* Property, const FFieldVariant& Inner)
-		{
-			if (Inner.HasMetaData(NAME_ContainsLoosePropertiesMetadata))
-			{
-				Property->SetMetaData(NAME_ContainsLoosePropertiesMetadata, TEXT("True"));
-			}
-		};
-
 		if (FStructProperty* AsStructProperty = CastField<FStructProperty>(Property))
 		{
 			if (!AsStructProperty->Struct->UseNativeSerialization())
@@ -352,8 +344,6 @@ namespace UE
 				AsStructProperty->SetMetaData(NAME_PresentAsTypeMetadata, *OriginalName);
 				AsStructProperty->Struct->SetMetaData(NAME_OriginalType, *OriginalName);
 				AsStructProperty->Struct->SetMetaData(NAME_PresentAsTypeMetadata, *OriginalName);
-
-				TrySetContainsLooseProperties(AsStructProperty, AsStructProperty->Struct);
 			}
 		}
 		else if (FByteProperty* AsByteProperty = CastField<FByteProperty>(Property))
@@ -367,12 +357,10 @@ namespace UE
 		else if (FArrayProperty* AsArrayProperty = CastField<FArrayProperty>(Property))
 		{
 			ConvertToInstanceDataObjectProperty(AsArrayProperty->Inner, PropertyType.GetParameter(0), Outer, PropertyTree, EnumNames);
-			TrySetContainsLooseProperties(AsArrayProperty, AsArrayProperty->Inner);
 		}
 		else if (FSetProperty* AsSetProperty = CastField<FSetProperty>(Property))
 		{
 			ConvertToInstanceDataObjectProperty(AsSetProperty->ElementProp, PropertyType.GetParameter(0), Outer, PropertyTree, EnumNames);
-			TrySetContainsLooseProperties(AsSetProperty, AsSetProperty->ElementProp);
 		}
 		else if (FMapProperty* AsMapProperty = CastField<FMapProperty>(Property))
 		{
@@ -390,14 +378,50 @@ namespace UE
 			}
 
 			ConvertToInstanceDataObjectProperty(AsMapProperty->KeyProp, PropertyType.GetParameter(0), Outer, KeyTree, EnumNames);
-			TrySetContainsLooseProperties(AsMapProperty, AsMapProperty->KeyProp);
 			ConvertToInstanceDataObjectProperty(AsMapProperty->ValueProp, PropertyType.GetParameter(1), Outer, ValueTree, EnumNames);
-			TrySetContainsLooseProperties(AsMapProperty, AsMapProperty->ValueProp);
 		}
 		else if (FOptionalProperty* AsOptionalProperty = CastField<FOptionalProperty>(Property))
 		{
 			ConvertToInstanceDataObjectProperty(AsOptionalProperty->GetValueProperty(), PropertyType.GetParameter(0), Outer, PropertyTree, EnumNames);
-			TrySetContainsLooseProperties(AsOptionalProperty, AsOptionalProperty->GetValueProperty());
+		}
+	}
+
+	// recursively sets NAME_ContainsLoosePropertiesMetadata on all properties that contain loose properties
+	static void TrySetContainsLoosePropertyMetadata(FProperty* Property)
+	{
+		const auto Helper = [](FProperty* Property, const FFieldVariant& Inner)
+		{
+			if (Inner.HasMetaData(NAME_ContainsLoosePropertiesMetadata))
+			{
+				Property->SetMetaData(NAME_ContainsLoosePropertiesMetadata, TEXT("True"));
+			}
+		};
+
+		if (FStructProperty* AsStructProperty = CastField<FStructProperty>(Property))
+		{
+			Helper(AsStructProperty, AsStructProperty->Struct);
+		}
+		else if (FArrayProperty* AsArrayProperty = CastField<FArrayProperty>(Property))
+		{
+			TrySetContainsLoosePropertyMetadata(AsArrayProperty->Inner);
+			Helper(AsArrayProperty, AsArrayProperty->Inner);
+		}
+		else if (FSetProperty* AsSetProperty = CastField<FSetProperty>(Property))
+		{
+			TrySetContainsLoosePropertyMetadata(AsSetProperty->ElementProp);
+			Helper(AsSetProperty, AsSetProperty->ElementProp);
+		}
+		else if (FMapProperty* AsMapProperty = CastField<FMapProperty>(Property))
+		{
+			TrySetContainsLoosePropertyMetadata(AsMapProperty->KeyProp);
+			Helper(AsMapProperty, AsMapProperty->KeyProp);
+			TrySetContainsLoosePropertyMetadata(AsMapProperty->ValueProp);
+			Helper(AsMapProperty, AsMapProperty->ValueProp);
+		}
+		else if (FOptionalProperty* AsOptionalProperty = CastField<FOptionalProperty>(Property))
+		{
+			TrySetContainsLoosePropertyMetadata(AsOptionalProperty->GetValueProperty());
+			Helper(AsOptionalProperty, AsOptionalProperty->GetValueProperty());
 		}
 
 		if (Property->GetBoolMetaData(NAME_IsLooseMetadata) || Property->GetBoolMetaData(NAME_ContainsLoosePropertiesMetadata))
@@ -500,6 +524,7 @@ namespace UE
 				}
 
 				ConvertToInstanceDataObjectProperty(SuperProperty, Type, Super, SubTree, EnumNames);
+				TrySetContainsLoosePropertyMetadata(SuperProperty);
 			}
 
 			// AddCppProperty expects reverse property order for StaticLink to work correctly
@@ -552,6 +577,7 @@ namespace UE
 						}
 						ConvertToInstanceDataObjectProperty(Property, Type, Result, It.GetNode().GetSubTree(), EnumNames);
 						MarkPropertyAsLoose(Property);	// note: make sure not to mark until AFTER conversion, as this can mutate property flags on nested struct fields
+						TrySetContainsLoosePropertyMetadata(Property);
 						LooseInstanceDataObjectProperties.Add(Property);
 						continue;
 					}
