Miscellaneous redirector support.
* Now that asset redirectors can be ingested as core redirects, FSoftObjectPath should always check for redirects rather than only in the editor.
* Improved debug info in FCoreRedirects
* Improved handling of redirects by CollectionAssetRegistryBridge

#rnx

diff --git a/Engine/Source/Editor/ContentBrowser/Private/CollectionAssetRegistryBridge.cpp b/Engine/Source/Editor/ContentBrowser/Private/CollectionAssetRegistryBridge.cpp
--- a/Engine/Source/Editor/ContentBrowser/Private/CollectionAssetRegistryBridge.cpp
+++ b/Engine/Source/Editor/ContentBrowser/Private/CollectionAssetRegistryBridge.cpp
@@ -20,6 +20,7 @@
 #include "Trace/Detail/Channel.h"
 #include "UObject/Class.h"
 #include "UObject/ConstructorHelpers.h"
+#include "UObject/CoreRedirects.h"
 #include "UObject/LinkerLoad.h"
 #include "UObject/NameTypes.h"
 #include "UObject/UObjectGlobals.h"
@@ -103,6 +104,17 @@ public:
 					}
 				}
 			}
+			else
+			{
+				// Check whether there is a GeneralizedObject redirector entry in FCoreRedirects
+				// This would come from the process of converting on-disk UObjectRedirectors into 
+				// core redirectors
+				const FCoreRedirectObjectName OldObjectName(InObjectPathString);
+				
+				FCoreRedirectObjectName NewObjectName = 
+					FCoreRedirects::GetRedirectedName(ECoreRedirectFlags::Type_Asset, OldObjectName);
+				OutNewObjectPath.SetPath(NewObjectName.ToString());
+			}
 		}
 
 		return OutNewObjectPath.IsValid() && InObjectPath != OutNewObjectPath;
diff --git a/Engine/Source/Runtime/CoreUObject/Private/UObject/CoreRedirects.cpp b/Engine/Source/Runtime/CoreUObject/Private/UObject/CoreRedirects.cpp
--- a/Engine/Source/Runtime/CoreUObject/Private/UObject/CoreRedirects.cpp
+++ b/Engine/Source/Runtime/CoreUObject/Private/UObject/CoreRedirects.cpp
@@ -2045,6 +2045,8 @@ void FCoreRedirects::AddAssetRedirects(const TMap<FSoftObjectPath, FSoftObjectPa
 	{
 		FCoreRedirectorScopeLockForWrite ScopeLock(RWLock);
 		FRedirectNameMap* ExistingMap = &RedirectTypeMap.FindOrAdd(ECoreRedirectFlags::Type_Asset);
+		int32 NumAdded = 0;
+		int32 NumSkipped = 0;
 		for (const TPair<FSoftObjectPath, FSoftObjectPath>& Pair : InRedirects)
 		{
 			// Asset redirects are, by definition, not package redirects
@@ -2052,6 +2054,7 @@ void FCoreRedirects::AddAssetRedirects(const TMap<FSoftObjectPath, FSoftObjectPa
 			{
 				UE_LOG(LogCoreRedirects, Warning, TEXT("Attempted to register asset redirector that was missing a package or object name. Redirector was from %s to %s"),
 					*Pair.Key.ToString(), *Pair.Value.ToString());
+				NumSkipped++;
 				continue;
 			}
 
@@ -2084,6 +2087,7 @@ void FCoreRedirects::AddAssetRedirects(const TMap<FSoftObjectPath, FSoftObjectPa
 						Pair.Value.ToString(), Pair.Key.ToString(), ExistingTargetName.ToString());
 
 					bShouldAddRedirect = false;
+					NumSkipped++;
 					break;
 				}
 			}
@@ -2091,9 +2095,14 @@ void FCoreRedirects::AddAssetRedirects(const TMap<FSoftObjectPath, FSoftObjectPa
 			if (bShouldAddRedirect)
 			{
 				ExistingRedirects.Add(ObjectRedirector);
+				NumAdded++;
 			}
 		}
 
+		UE_LOG(LogCoreRedirects, Display, 
+			TEXT("Object redirects provided to FCoreRedirects: %d. Redirects add: %d. Redirects skipped: %d"),
+			InRedirects.Num(), NumAdded, NumSkipped);
+
 		if (IsInDebugMode())
 		{
 			ValidateAssetRedirects();
diff --git a/Engine/Source/Runtime/CoreUObject/Private/UObject/SoftObjectPath.cpp b/Engine/Source/Runtime/CoreUObject/Private/UObject/SoftObjectPath.cpp
--- a/Engine/Source/Runtime/CoreUObject/Private/UObject/SoftObjectPath.cpp
+++ b/Engine/Source/Runtime/CoreUObject/Private/UObject/SoftObjectPath.cpp
@@ -16,6 +16,18 @@
 #include "String/Find.h"
 #include "AutoRTFM/AutoRTFM.h"
 
+namespace SoftObjectPath
+{
+static bool bResolveCoreRedirects = true;
+
+#if !WITH_EDITOR
+static FAutoConsoleVariableRef CVarEnablePathFixupOutsideEditor(TEXT("SoftObjectPath.EnablePathFixupOutsideEditor"),
+	bResolveCoreRedirects,
+	TEXT("When true (by default) we will call FixupCoreRedirects when resolving, loading, or saving soft object paths outside the editor. When false, we will revert to the legacy behavior and not do the extra fixup.")
+);
+#endif
+}
+
 // Deprecated constructor
 FSoftObjectPath::FSoftObjectPath(FName InAssetPathName, FString InSubPathString)
 {
@@ -280,12 +292,15 @@ bool FSoftObjectPath::PreSavePath(bool* bReportSoftObjectPathRedirects)
 		*this = FoundRedirection;
 		return true;
 	}
+#endif
 
-	if (FixupCoreRedirects())
+	if (SoftObjectPath::bResolveCoreRedirects)
 	{
-		return true;
+		if (FixupCoreRedirects())
+		{
+			return true;
+		}
 	}
-#endif // WITH_EDITOR
 	return false;
 }
 
@@ -614,17 +629,18 @@ UObject* FSoftObjectPath::TryLoad(FUObjectSerializeContext* InLoadContext) const
 				LoadedObject = StaticLoadObject(UObject::StaticClass(), nullptr, *PathString, nullptr, LOAD_None, nullptr, true);
 			};
 
-#if WITH_EDITOR
 			// Look at core redirects if we didn't find the object
-			if (!LoadedObject)
+			if (!LoadedObject && SoftObjectPath::bResolveCoreRedirects)
 			{
 				FSoftObjectPath FixupObjectPath = *this;
 				if (FixupObjectPath.FixupCoreRedirects())
 				{
-					LoadedObject = LoadObject<UObject>(nullptr, *FixupObjectPath.ToString());
+					UE_AUTORTFM_OPEN2
+					{
+						LoadedObject = LoadObject<UObject>(nullptr, *FixupObjectPath.ToString());
+					};
 				}
 			}
-#endif
 
 			while (UObjectRedirector* Redirector = Cast<UObjectRedirector>(LoadedObject))
 			{
@@ -648,6 +664,11 @@ int32 FSoftObjectPath::LoadAsync(FLoadSoftObjectPathAsyncDelegate InCompletionDe
 	}
 #endif
 
+	if (SoftObjectPath::bResolveCoreRedirects)
+	{
+		PathToLoad.FixupCoreRedirects();
+	}
+
 	FLoadAssetAsyncDelegate WrapperDelegate = FLoadAssetAsyncDelegate::CreateLambda(
 		[RequestedPath, PathToLoad, CompletionDelegate = MoveTemp(InCompletionDelegate)](const FTopLevelAssetPath& InAssetPath, UObject* InLoadedObject, EAsyncLoadingResult::Type InResult) mutable
 		{
@@ -714,9 +735,8 @@ UObject* FSoftObjectPath::ResolveObjectInternal(const TCHAR* PathString) const
 		}
 	}
 
-#if WITH_EDITOR
 	// Look at core redirects if we didn't find the object
-	if (!FoundObject)
+	if (!FoundObject && SoftObjectPath::bResolveCoreRedirects)
 	{
 		FSoftObjectPath FixupObjectPath = *this;
 		if (FixupObjectPath.FixupCoreRedirects())
@@ -724,7 +744,6 @@ UObject* FSoftObjectPath::ResolveObjectInternal(const TCHAR* PathString) const
 			FoundObject = FindObject<UObject>(nullptr, *FixupObjectPath.ToString());
 		}
 	}
-#endif
 
 	while (UObjectRedirector* Redirector = Cast<UObjectRedirector>(FoundObject))
 	{
