[ZenLoader]
 - Add support for loading loose cooked file at runtime
 - Make sure zenloader can be the default loader in all case from now on
 - Auto-detect separate cooked files in AsyncArchive

diff --git a/Engine/Source/Runtime/CoreUObject/Private/Serialization/AsyncLoading.cpp b/Engine/Source/Runtime/CoreUObject/Private/Serialization/AsyncLoading.cpp
--- a/Engine/Source/Runtime/CoreUObject/Private/Serialization/AsyncLoading.cpp
+++ b/Engine/Source/Runtime/CoreUObject/Private/Serialization/AsyncLoading.cpp
@@ -7963,7 +7963,7 @@ int64 FAsyncArchive::TotalSize()
 	{
 		QUICK_SCOPE_CYCLE_COUNTER(STAT_FArchiveAsync2_TotalSize);
 		SizeRequestPtr->WaitCompletion();
-		if ((GEventDrivenLoaderEnabled || bCookedForEDLInEditor) && HeaderSizeWhenReadingExportsFromSplitFile)
+		if ((GEventDrivenLoaderEnabled || bCookedForEDLInEditor || GetLoaderType() == ELoaderType::ZenLoader) && HeaderSizeWhenReadingExportsFromSplitFile)
 		{
 			FileSize = SizeRequestPtr->GetSizeResults();
 		}
@@ -7996,12 +7996,15 @@ FORCEINLINE void FAsyncArchive::SetPosAndUpdatePrecacheBuffer(int64 Pos)
 
 void FAsyncArchive::Seek(int64 InPos)
 {
-	if ((GEventDrivenLoaderEnabled || bCookedForEDLInEditor) && LoadPhase < ELoadPhase::ProcessingExports)
+	if ((GEventDrivenLoaderEnabled || bCookedForEDLInEditor || GetLoaderType() == ELoaderType::ZenLoader) && LoadPhase < ELoadPhase::ProcessingExports)
 	{
-		check(!HeaderSizeWhenReadingExportsFromSplitFile && HeaderSize && TotalSize() == HeaderSize);
-		if (InPos >= HeaderSize)
+		// Auto-detect when exports are in a separate file so we only activate that mode when needed.
+		if (!HeaderSizeWhenReadingExportsFromSplitFile && HeaderSize && TotalSize() == HeaderSize)
 		{
-			FirstExportStarting();
+			if (InPos >= HeaderSize)
+			{
+				FirstExportStarting();
+			}
 		}
 	}
 	checkf(InPos >= 0 && InPos <= TotalSizeOrMaxInt64IfNotReady(), TEXT("Bad position in FAsyncArchive::Seek. Filename:%s InPos:%lu, Size:%lu"),
@@ -8287,26 +8290,30 @@ void FAsyncArchive::FirstExportStarting()
 	LogItem(TEXT("Exports"));
 	LoadPhase = ELoadPhase::ProcessingExports;
 
-	if ((GEventDrivenLoaderEnabled && !EVENT_DRIVEN_ASYNC_LOAD_ACTIVE_AT_RUNTIME) || bCookedForEDLInEditor)
+	if ((GEventDrivenLoaderEnabled && !EVENT_DRIVEN_ASYNC_LOAD_ACTIVE_AT_RUNTIME) || bCookedForEDLInEditor || GetLoaderType() == ELoaderType::ZenLoader)
 	{
-		FlushCache();
-		if (Handle)
+		// Detect when exports are in a separate file before trying to activate that mode.
+		if (!HeaderSizeWhenReadingExportsFromSplitFile && HeaderSize && TotalSize() == HeaderSize)
 		{
-			delete Handle;
-			Handle = nullptr;
-		}
+			FlushCache();
+			if (Handle)
+			{
+				delete Handle;
+				Handle = nullptr;
+			}
 
-		HeaderSizeWhenReadingExportsFromSplitFile = HeaderSize;
+			HeaderSizeWhenReadingExportsFromSplitFile = HeaderSize;
 
-		FOpenAsyncPackageResult OpenResult = IPackageResourceManager::Get().OpenAsyncReadPackage(PackagePath, EPackageSegment::Exports);
-		Handle = OpenResult.Handle.Release();
-		check(Handle); // OpenAsyncReadPackage guarantees a non-null return value; the handle will fail to read later if the path does not exist
+			FOpenAsyncPackageResult OpenResult = IPackageResourceManager::Get().OpenAsyncReadPackage(PackagePath, EPackageSegment::Exports);
+			Handle = OpenResult.Handle.Release();
+			check(Handle); // OpenAsyncReadPackage guarantees a non-null return value; the handle will fail to read later if the path does not exist
 
-		check(!SizeRequestPtr);
-		SizeRequestPtr = Handle->SizeRequest();
-		if (SizeRequestPtr->PollCompletion())
-		{
-			TotalSize(); // complete the request
+			check(!SizeRequestPtr);
+			SizeRequestPtr = Handle->SizeRequest();
+			if (SizeRequestPtr->PollCompletion())
+			{
+				TotalSize(); // complete the request
+			}
 		}
 	}
 }
diff --git a/Engine/Source/Runtime/CoreUObject/Private/Serialization/AsyncLoading2.cpp b/Engine/Source/Runtime/CoreUObject/Private/Serialization/AsyncLoading2.cpp
--- a/Engine/Source/Runtime/CoreUObject/Private/Serialization/AsyncLoading2.cpp
+++ b/Engine/Source/Runtime/CoreUObject/Private/Serialization/AsyncLoading2.cpp
@@ -264,7 +264,7 @@ void FindAllRuntimeScriptPackages(TArray<UPackage*>& OutPackages)
 #endif
 
 #ifndef ALT2_ENABLE_LINKERLOAD_SUPPORT
-#define ALT2_ENABLE_LINKERLOAD_SUPPORT WITH_EDITOR
+#define ALT2_ENABLE_LINKERLOAD_SUPPORT 1
 #endif
 
 #ifndef ALT2_ENABLE_NEW_ARCHIVE_FOR_LINKERLOAD
@@ -468,53 +468,6 @@ FString FormatPackageId(FPackageId PackageId)
 struct FAsyncPackage2;
 class FAsyncLoadingThread2;
 
-class FSimpleArchive final
-	: public FArchive
-{
-public:
-	FSimpleArchive(const uint8* BufferPtr, uint64 BufferSize)
-	{
-#if (!DEVIRTUALIZE_FLinkerLoad_Serialize)
-		ActiveFPLB = &InlineFPLB;
-#endif
-		ActiveFPLB->OriginalFastPathLoadBuffer = BufferPtr;
-		ActiveFPLB->StartFastPathLoadBuffer = BufferPtr;
-		ActiveFPLB->EndFastPathLoadBuffer = BufferPtr + BufferSize;
-	}
-
-	int64 TotalSize() override
-	{
-		return ActiveFPLB->EndFastPathLoadBuffer - ActiveFPLB->OriginalFastPathLoadBuffer;
-	}
-
-	int64 Tell() override
-	{
-		return ActiveFPLB->StartFastPathLoadBuffer - ActiveFPLB->OriginalFastPathLoadBuffer;
-	}
-
-	void Seek(int64 Position) override
-	{
-		ActiveFPLB->StartFastPathLoadBuffer = ActiveFPLB->OriginalFastPathLoadBuffer + Position;
-		check(ActiveFPLB->StartFastPathLoadBuffer <= ActiveFPLB->EndFastPathLoadBuffer);
-	}
-
-	void Serialize(void* Data, int64 Length) override
-	{
-		if (!Length || IsError())
-		{
-			return;
-		}
-		check(ActiveFPLB->StartFastPathLoadBuffer + Length <= ActiveFPLB->EndFastPathLoadBuffer);
-		FMemory::Memcpy(Data, ActiveFPLB->StartFastPathLoadBuffer, Length);
-		ActiveFPLB->StartFastPathLoadBuffer += Length;
-	}
-private:
-#if (!DEVIRTUALIZE_FLinkerLoad_Serialize)
-	FArchive::FFastPathLoadBuffer InlineFPLB;
-	FArchive::FFastPathLoadBuffer* ActiveFPLB;
-#endif
-};
-
 struct FExportObject
 {
 	UObject* Object = nullptr;
@@ -5654,8 +5607,9 @@ void FAsyncPackage2::InitializeLinkerLoadState(const FLinkerInstancingContext* I
 
 void FAsyncPackage2::CreateLinker(const FLinkerInstancingContext* InstancingContext)
 {
-	uint32 LinkerFlags = (LOAD_Async | LOAD_NoVerify | LOAD_SkipLoadImportedPackages | Desc.LoadFlags);
+	uint32 LinkerFlags = (LOAD_Async | LOAD_NoVerify | LOAD_SkipLoadImportedPackages);
 #if WITH_EDITOR
+	LinkerFlags |= Desc.LoadFlags;
 	if ((Desc.PackageFlags & PKG_PlayInEditor) != 0 && (GIsEditor || !FApp::IsGame()))
 	{
 		LinkerFlags |= LOAD_PackageForPIE;
diff --git a/Engine/Source/Runtime/CoreUObject/Private/Serialization/AsyncPackageLoader.cpp b/Engine/Source/Runtime/CoreUObject/Private/Serialization/AsyncPackageLoader.cpp
--- a/Engine/Source/Runtime/CoreUObject/Private/Serialization/AsyncPackageLoader.cpp
+++ b/Engine/Source/Runtime/CoreUObject/Private/Serialization/AsyncPackageLoader.cpp
@@ -15,6 +15,14 @@
 #include "IO/IoDispatcherInternal.h"
 #include "HAL/IConsoleManager.h"
 
+static bool GUseOldLoaderAsFallback = false;
+static FAutoConsoleVariableRef CVarGUseOldLoaderAsFallback(
+	TEXT("s.UseOldLoaderAsFallback"),
+	GUseOldLoaderAsFallback,
+	TEXT("When active, the old loader will be used as the fallback instead of zenloader"),
+	ECVF_Default
+);
+
 #define DO_TRACK_ASYNC_LOAD_REQUESTS (!UE_BUILD_SHIPPING && !UE_BUILD_TEST)
 
 #if DO_TRACK_ASYNC_LOAD_REQUESTS
@@ -26,6 +34,7 @@
 volatile int32 GIsLoaderCreated;
 TUniquePtr<IAsyncPackageLoader> GPackageLoader;
 bool GAsyncLoadingAllowed = true;
+bool GIoDispatcherInitializedByLoader = false;
 
 FThreadSafeCounter IAsyncPackageLoader::NextPackageRequestId;
 
@@ -186,18 +195,19 @@ void SetAsyncLoadingAllowed(bool bAllowAsyncLoading)
 void InitAsyncThread()
 {
 	LLM_SCOPE(ELLMTag::AsyncLoading);
+	
+	bool bCommandLineDisabled = FParse::Param(FCommandLine::Get(), TEXT("NoZenLoader"));
+
 	if (FIoDispatcher::IsInitialized())
 	{
 		bool bSettingsEnabled = false;
 		bool bCommandLineEnabled = false;
-		bool bCommandLineDisabled = false;
 		bool bHasUseIoStoreParamInEditor = false;
 #if WITH_EDITOR
 		// enable zenloader by default in editor case
 		bSettingsEnabled = true;
 
 		bCommandLineEnabled = FParse::Param(FCommandLine::Get(), TEXT("ZenLoader"));
-		bCommandLineDisabled = FParse::Param(FCommandLine::Get(), TEXT("NoZenLoader"));
 		check(GConfig);
 		GConfig->GetBool(TEXT("/Script/Engine.EditorStreamingSettings"), TEXT("s.ZenLoaderEnabled"), bSettingsEnabled, GEngineIni);
 		bHasUseIoStoreParamInEditor = UE_FORCE_USE_IOSTORE || FParse::Param(FCommandLine::Get(), TEXT("UseIoStore"));
@@ -222,7 +232,20 @@ void InitAsyncThread()
 	}
 	if (!GPackageLoader.IsValid())
 	{
-		GPackageLoader = MakeUnique<FAsyncLoadingThread>(/** ThreadIndex = */ 0);
+		if (GUseOldLoaderAsFallback || bCommandLineDisabled)
+		{
+			GPackageLoader = MakeUnique<FAsyncLoadingThread>(/** ThreadIndex = */ 0);
+		}
+		else
+		{
+			if (!FIoDispatcher::IsInitialized())
+			{
+				FIoDispatcher::Initialize();
+				GIoDispatcherInitializedByLoader = true;
+			}
+
+			GPackageLoader.Reset(MakeAsyncPackageLoader2(FIoDispatcher::Get()));
+		}
 	}
 
 	FPlatformAtomics::InterlockedIncrement(&GIsLoaderCreated);
@@ -241,6 +264,11 @@ void ShutdownAsyncThread()
 		GPackageLoader->ShutdownLoading();
 		GPackageLoader.Reset(nullptr);
 	}
+
+	if (GIoDispatcherInitializedByLoader)
+	{
+		FIoDispatcher::Shutdown();
+	}
 }
 
 bool IsInAsyncLoadingThreadCoreUObjectInternal()
diff --git a/Engine/Source/Runtime/CoreUObject/Tests/Serialization/BulkDataTests.cpp b/Engine/Source/Runtime/CoreUObject/Tests/Serialization/BulkDataTests.cpp
--- a/Engine/Source/Runtime/CoreUObject/Tests/Serialization/BulkDataTests.cpp
+++ b/Engine/Source/Runtime/CoreUObject/Tests/Serialization/BulkDataTests.cpp
@@ -54,16 +54,24 @@ bool TestBulkDataFlags(FBulkData& BulkData, uint32 Flags, T&& IsSet)
 
 class FIoDispatcherTestScope
 {
+	bool bIoDispatcherInitialized = false;
 public:
 	FIoDispatcherTestScope()
 	{
-		FIoDispatcher::Initialize();
+		if (!FIoDispatcher::IsInitialized())
+		{
+			FIoDispatcher::Initialize();
+			bIoDispatcherInitialized = true;
+		}
 		FIoDispatcher::InitializePostSettings();
 	}
 	
 	~FIoDispatcherTestScope()
 	{
-		FIoDispatcher::Shutdown();
+		if (bIoDispatcherInitialized)
+		{
+			FIoDispatcher::Shutdown();
+		}
 	}
 };
 
