Cooker: SkipOnlyEditorOnly: Fix BlueprintGeneratedClass's adapter that declares all imports from UBlueprint to exclude imports from editor-only components on the class default object, or from any other editor-only non-blueprint-node objects within the package. We now only allow editor-only objects if they are subobjects of the UBlueprint.
#rnx

diff --git a/Engine/Source/Runtime/CoreUObject/Internal/UObject/ImportExportCollector.h b/Engine/Source/Runtime/CoreUObject/Internal/UObject/ImportExportCollector.h
--- a/Engine/Source/Runtime/CoreUObject/Internal/UObject/ImportExportCollector.h
+++ b/Engine/Source/Runtime/CoreUObject/Internal/UObject/ImportExportCollector.h
@@ -12,6 +12,7 @@
 #include "UObject/Package.h"
 #include "UObject/NameTypes.h"
 #include "UObject/SoftObjectPath.h"
+#include "UObject/UObjectGlobals.h"
 
 /** An Archive that records all of the imported packages from a tree of exports. */
 class FImportExportCollector : public FArchiveUObject
@@ -32,9 +33,15 @@ public:
 	COREUOBJECT_API void SerializeObjectAndReferencedExports(UObject* RootObject);
 	/** Restore the collector to empty. */
 	COREUOBJECT_API void Reset();
-	const TSet<UObject*>& GetExports() const;
+	TArray<UObject*> GetExports() const;
 	const TMap<FSoftObjectPath, ESoftObjectPathCollectType>& GetImports() const;
 	const TMap<FName, ESoftObjectPathCollectType>& GetImportedPackages() const;
+	/**
+	 * By default, when this->IsFilterEditorOnly, exports that are EditorOnly exports are not serialized
+	 * to look for imports or other exports, and EditorOnly imports are not recorded. But the caller can
+	 * override this and allow some of them to be serialized.
+	 */
+	void SetCallbackIsEditorOnlyObjectAllowed(TFunction<bool(const UObject*)> InCallback);
 
 	COREUOBJECT_API virtual FArchive& operator<<(UObject*& Obj) override;
 	COREUOBJECT_API virtual FArchive& operator<<(FSoftObjectPath& Value) override;
@@ -42,11 +49,21 @@ public:
 private:
 	void AddImport(const FSoftObjectPath& Path, ESoftObjectPathCollectType CollectType);
 	ESoftObjectPathCollectType Union(ESoftObjectPathCollectType A, ESoftObjectPathCollectType B);
+	bool CachedIsEditorOnlyObject(const UObject* Object);
 
-	TSet<UObject*> Exports;
+	enum class EVisitResult
+	{
+		Uninitialized,
+		Excluded,
+		Import,
+		Export,
+	};
+	TMap<UObject*, EVisitResult> Visited;
 	TRingBuffer<UObject*> ExportsExploreQueue;
 	TMap<FSoftObjectPath, ESoftObjectPathCollectType> Imports;
 	TMap<FName, ESoftObjectPathCollectType> ImportedPackages;
+	TMap<const UObject*, UE::SavePackageUtilities::EEditorOnlyObjectResult> EditorOnlyObjectCache;
+	TFunction<bool(const UObject*)> CallbackIsEditorOnlyObjectAllowed;
 	UPackage* RootPackage;
 	FName RootPackageName;
 };
@@ -55,8 +72,16 @@ private:
 // Inline implementations
 ///////////////////////////////////////////////////////
 
-inline const TSet<UObject*>& FImportExportCollector::GetExports() const
+inline TArray<UObject*> FImportExportCollector::GetExports() const
 {
+	TArray<UObject*> Exports;
+	for (const TPair<UObject*, EVisitResult>& Pair : Visited)
+	{
+		if (Pair.Value == EVisitResult::Export)
+		{
+			Exports.Add(Pair.Key);
+		}
+	}
 	return Exports;
 }
 
@@ -70,4 +95,9 @@ inline const TMap<FName, ESoftObjectPathCollectType>& FImportExportCollector::Ge
 	return ImportedPackages;
 }
 
+inline void FImportExportCollector::SetCallbackIsEditorOnlyObjectAllowed(TFunction<bool(const UObject*)> InCallback)
+{
+	CallbackIsEditorOnlyObjectAllowed = MoveTemp(InCallback);
+}
+
 #endif // WITH_EDITORONLY_DATA
diff --git a/Engine/Source/Runtime/CoreUObject/Private/UObject/ImportExportCollector.cpp b/Engine/Source/Runtime/CoreUObject/Private/UObject/ImportExportCollector.cpp
--- a/Engine/Source/Runtime/CoreUObject/Private/UObject/ImportExportCollector.cpp
+++ b/Engine/Source/Runtime/CoreUObject/Private/UObject/ImportExportCollector.cpp
@@ -16,13 +16,14 @@ FImportExportCollector::FImportExportCollector(UPackage* InRootPackage)
 
 void FImportExportCollector::Reset()
 {
-	Exports.Reset();
+	Visited.Reset();
 	Imports.Reset();
+	EditorOnlyObjectCache.Reset();
 }
 
 void FImportExportCollector::AddExportToIgnore(UObject* Export)
 {
-	Exports.Add(Export);
+	Visited.Add(Export, EVisitResult::Excluded);
 }
 
 void FImportExportCollector::SerializeObjectAndReferencedExports(UObject* RootObject)
@@ -41,24 +42,61 @@ FArchive& FImportExportCollector::operator<<(UObject*& Obj)
 	{
 		return *this;
 	}
-	UPackage* Package = Obj->GetPackage();
-	if (!Package)
+
+	bool bFirstVisit = false;
+	EVisitResult& VisitResult = Visited.FindOrAdd(Obj);
+	switch (VisitResult)
 	{
+	case EVisitResult::Excluded: [[fallthrough]];
+	case EVisitResult::Export:
 		return *this;
-	}
-	if (Package != RootPackage)
+	case EVisitResult::Import:
+		// We have to call AddImport every time instead of early exiting because the SoftObjectPathCollectType might
+		// need to be updated
+		break;
+	case EVisitResult::Uninitialized:
 	{
-		AddImport(FSoftObjectPath(Obj), ESoftObjectPathCollectType::AlwaysCollect);
+		bFirstVisit = true;
+		if (IsFilterEditorOnly() && CachedIsEditorOnlyObject(Obj))
+		{
+			bool bAllowedObject = ((bool)CallbackIsEditorOnlyObjectAllowed) &&
+				CallbackIsEditorOnlyObjectAllowed(Obj);
+			if (!bAllowedObject)
+			{
+				VisitResult = EVisitResult::Excluded;
+				return *this;
+			}
+		}
+
+		UPackage* Package = Obj->GetPackage();
+		if (!Package)
+		{
+			VisitResult = EVisitResult::Excluded;
+			return *this;
+		}
+
+		VisitResult = Package != RootPackage ? EVisitResult::Import : EVisitResult::Export;
+		break;
+	}
+	default:
+		checkNoEntry();
 		return *this;
 	}
 
-	bool bAlreadyExists;
-	Exports.Add(Obj, &bAlreadyExists);
-	if (bAlreadyExists)
+	switch (VisitResult)
 	{
+	case EVisitResult::Import:
+		AddImport(FSoftObjectPath(Obj), ESoftObjectPathCollectType::AlwaysCollect);
+		break;
+	case EVisitResult::Export:
+		check(bFirstVisit); // We should have early exited above if revisiting an export
+		ExportsExploreQueue.Add(Obj);
+		break;
+	default:
+		checkNoEntry();
 		return *this;
 	}
-	ExportsExploreQueue.Add(Obj);
+
 	return *this;
 }
 
@@ -98,4 +136,24 @@ ESoftObjectPathCollectType FImportExportCollector::Union(ESoftObjectPathCollectT
 	return static_cast<ESoftObjectPathCollectType>(FMath::Max(static_cast<int>(A), static_cast<int>(B)));
 }
 
-#endif // WITH_EDITORONLY_DATA
+bool FImportExportCollector::CachedIsEditorOnlyObject(const UObject* Object)
+{
+	using namespace UE::SavePackageUtilities;
+
+#if WITH_EDITOR
+	return IsEditorOnlyObject(Object, true /* bCheckRecursive */,
+		[this](const UObject* InObject)
+		{
+			return EditorOnlyObjectCache.FindOrAdd(InObject, EEditorOnlyObjectResult::Uninitialized);
+		},
+		[this](const UObject* InObject, bool bEditorOnly)
+		{
+			EditorOnlyObjectCache.FindOrAdd(InObject) =
+				bEditorOnly ? EEditorOnlyObjectResult::EditorOnly : EEditorOnlyObjectResult::NonEditorOnly;
+		});
+#else
+	return false;
+#endif
+}
+
+#endif // WITH_EDITORONLY_DATA
\ No newline at end of file
diff --git a/Engine/Source/Runtime/CoreUObject/Private/UObject/SavePackage/SaveContext.h b/Engine/Source/Runtime/CoreUObject/Private/UObject/SavePackage/SaveContext.h
--- a/Engine/Source/Runtime/CoreUObject/Private/UObject/SavePackage/SaveContext.h
+++ b/Engine/Source/Runtime/CoreUObject/Private/UObject/SavePackage/SaveContext.h
@@ -1139,6 +1139,7 @@ public:
 		}
 	}
 
+#if WITH_EDITORONLY_DATA
 	auto GetFunctorReadCachedEditorOnlyObject()
 	{
 		using namespace UE::SavePackageUtilities;
@@ -1160,6 +1161,7 @@ public:
 				Status.bEditorOnly = bEditorOnly;
 			};
 	}
+#endif
 	UE::SavePackageUtilities::EEditorOnlyObjectFlags GetEditorOnlyObjectFlags() const;
 
 public:
diff --git a/Engine/Source/Runtime/CoreUObject/Private/UObject/SavePackage/SavePackageUtilities.cpp b/Engine/Source/Runtime/CoreUObject/Private/UObject/SavePackage/SavePackageUtilities.cpp
--- a/Engine/Source/Runtime/CoreUObject/Private/UObject/SavePackage/SavePackageUtilities.cpp
+++ b/Engine/Source/Runtime/CoreUObject/Private/UObject/SavePackage/SavePackageUtilities.cpp
@@ -847,6 +847,7 @@ void FObjectSaveContextData::Set(UPackage* Package, const ITargetPlatform* InTar
 	bUpdatingLoadedPath = UE::SavePackageUtilities::IsUpdatingLoadedPath(InTargetPlatform != nullptr, TargetPath, InSaveFlags);
 }
 
+#if WITH_EDITORONLY_DATA
 bool IsEditorOnlyObject(const UObject* InObject, bool bCheckRecursive)
 {
 	using namespace UE::SavePackageUtilities;
@@ -866,6 +867,16 @@ bool IsEditorOnlyObject(const UObject* InObject, bool bCheckRecursive, bool bChe
 	return IsEditorOnlyObjectInternal(InObject, Flags);
 }
 
+bool IsEditorOnlyObject(const UObject* InObject, bool bCheckRecursive,
+	TFunctionRef<UE::SavePackageUtilities::EEditorOnlyObjectResult(const UObject*)> LookupInCache,
+	TFunctionRef<void(const UObject*, bool)> AddToCache)
+{
+	using namespace UE::SavePackageUtilities;
+	EEditorOnlyObjectFlags Flags = EEditorOnlyObjectFlags::None;
+	Flags |= bCheckRecursive ? EEditorOnlyObjectFlags::CheckRecursive : EEditorOnlyObjectFlags::None;
+	return IsEditorOnlyObjectInternal(InObject, Flags, LookupInCache, AddToCache);
+}
+
 namespace UE::SavePackageUtilities
 {
 
@@ -992,6 +1003,8 @@ bool IsEditorOnlyObjectInternal(const UObject* InObject, EEditorOnlyObjectFlags
 
 } // namespace UE::SavePackageUtilities
 
+#endif // WITH_EDITORONLY_DATA
+
 void FObjectImportSortHelper::SortImports(FLinkerSave* Linker)
 {
 	TArray<FObjectImport>& Imports = Linker->ImportMap;
diff --git a/Engine/Source/Runtime/CoreUObject/Private/UObject/SavePackage/SavePackageUtilities.h b/Engine/Source/Runtime/CoreUObject/Private/UObject/SavePackage/SavePackageUtilities.h
--- a/Engine/Source/Runtime/CoreUObject/Private/UObject/SavePackage/SavePackageUtilities.h
+++ b/Engine/Source/Runtime/CoreUObject/Private/UObject/SavePackage/SavePackageUtilities.h
@@ -16,6 +16,7 @@
 #include "UObject/UObjectMarks.h"
 #include "UObject/ObjectPtr.h"
 #include "UObject/SavePackage.h"
+#include "UObject/UObjectGlobals.h"
 
 // This file contains private utilities shared by UPackage::Save and UPackage::Save2 
 
@@ -376,21 +377,15 @@ enum class EEditorOnlyObjectFlags
 };
 ENUM_CLASS_FLAGS(EEditorOnlyObjectFlags);
 
-enum class EEditorOnlyObjectResult
-{
-	Uninitialized,
-	EditorOnly,
-	NonEditorOnly,
-};
 #if WITH_EDITORONLY_DATA
 bool CanStripEditorOnlyImportsAndExports();
-#endif
 
 /** Returns result of IsEditorOnlyObjectInternal if Engine:[Core.System]:CanStripEditorOnlyExportsAndImports (ini) is set to true */
 bool IsEditorOnlyObjectInternal(const UObject* InObject, EEditorOnlyObjectFlags Flags,
-	TFunctionRef<EEditorOnlyObjectResult(const UObject*)> LookupInCache,
+	TFunctionRef<UE::SavePackageUtilities::EEditorOnlyObjectResult(const UObject*)> LookupInCache,
 	TFunctionRef<void(const UObject*, bool)> AddToCache);
 bool IsEditorOnlyObjectInternal(const UObject* InObject, EEditorOnlyObjectFlags Flags);
+#endif
 
 }
 
diff --git a/Engine/Source/Runtime/CoreUObject/Public/UObject/UObjectGlobals.h b/Engine/Source/Runtime/CoreUObject/Public/UObject/UObjectGlobals.h
--- a/Engine/Source/Runtime/CoreUObject/Public/UObject/UObjectGlobals.h
+++ b/Engine/Source/Runtime/CoreUObject/Public/UObject/UObjectGlobals.h
@@ -3385,6 +3385,18 @@ struct FAssetMsg
 	}
 #endif // NO_LOGGING
 
+#if WITH_EDITORONLY_DATA
+namespace UE::SavePackageUtilities
+{
+enum class EEditorOnlyObjectResult
+{
+	Uninitialized,
+	EditorOnly,
+	NonEditorOnly,
+};
+}// namespace UE::SavePackageUtilities
+#endif // WITH_EDITORONLY_DATA
+
 #if WITH_EDITOR
 /** 
  * Returns if true if the object is editor-only:
@@ -3399,6 +3411,10 @@ struct FAssetMsg
 COREUOBJECT_API bool IsEditorOnlyObject(const UObject* InObject, bool bCheckRecursive = true);
 UE_DEPRECATED(5.3, "bCheckMarks argument is no longer supported because we are transitioning away from using ObjectMarks during saving");
 COREUOBJECT_API bool IsEditorOnlyObject(const UObject* InObject, bool bCheckRecursive, bool bCheckMarks);
+COREUOBJECT_API bool IsEditorOnlyObject(const UObject* InObject, bool bCheckRecursive,
+	TFunctionRef<UE::SavePackageUtilities::EEditorOnlyObjectResult(const UObject* Object)> LookupInCache,
+	TFunctionRef<void(const UObject* Object, bool bEditorOnly)> AddToCache);
+
 #endif //WITH_EDITOR
 
 class FFieldClass;
diff --git a/Engine/Source/Runtime/Engine/Private/BlueprintGeneratedClass.cpp b/Engine/Source/Runtime/Engine/Private/BlueprintGeneratedClass.cpp
--- a/Engine/Source/Runtime/Engine/Private/BlueprintGeneratedClass.cpp
+++ b/Engine/Source/Runtime/Engine/Private/BlueprintGeneratedClass.cpp
@@ -2564,7 +2564,13 @@ void UBlueprintGeneratedClass::Serialize(FArchive& Ar)
 		Collector.SetCookData(Ar.GetCookData());
 		Collector.AddExportToIgnore(this);
 		Collector.SetFilterEditorOnly(Ar.IsFilterEditorOnly());
-		Collector.SerializeObjectAndReferencedExports(ClassGeneratedBy);
+		UObject* ClassGeneratedByPtr = ClassGeneratedBy;
+		Collector.SetCallbackIsEditorOnlyObjectAllowed(
+			[ClassGeneratedByPtr](const UObject* Object)
+			{
+				return Object == ClassGeneratedByPtr || Object->IsIn(ClassGeneratedByPtr);
+			});
+		Collector.SerializeObjectAndReferencedExports(ClassGeneratedByPtr);
 		for (const TPair<FName, ESoftObjectPathCollectType>& Pair : Collector.GetImportedPackages())
 		{
 			if (Pair.Value != ESoftObjectPathCollectType::AlwaysCollect)
