- Changes DoesPackageExist to look at the IoDispatcher (if initialized) before looking at the filesystem to match path casing. Checking the OS path casing is very expensive compared to checking the IOStore, so this ends up being a good win for cooked editor builds in particular since cooked editors will make frequent Filesystem filter requests but many assets will reside in chunk data. Opening a large project goes from 15s -> 10s spent in DoesPackageExist (more specifically TryMatchCaseOnDisk)
-- Non-cooked editor builds won't touch this code path in most cases since earlier code will have checked the AssetRegistry for the package case.
-- Runtime code will not commonly request from the filesystem and thus won't be impacted by this change
- Minor typo fixes in comments

#rnx

diff --git a/Engine/Source/Runtime/CoreUObject/Private/Misc/PackageName.cpp b/Engine/Source/Runtime/CoreUObject/Private/Misc/PackageName.cpp
--- a/Engine/Source/Runtime/CoreUObject/Private/Misc/PackageName.cpp
+++ b/Engine/Source/Runtime/CoreUObject/Private/Misc/PackageName.cpp
@@ -2034,7 +2034,7 @@ FPackageName::EPackageLocationFilter FPackageName::InternalDoesPackageExistEx(co
 	return InternalDoesPackageExistEx(PackagePath, Filter, bMatchCaseOnDisk, OutPackagePath);
 }
 
-FPackageName::EPackageLocationFilter FPackageName::InternalDoesPackageExistEx(const FPackagePath & PackagePath, EPackageLocationFilter Filter, bool bMatchCaseOnDisk, FPackagePath * OutPackagePath)
+FPackageName::EPackageLocationFilter FPackageName::InternalDoesPackageExistEx(const FPackagePath& PackagePath, EPackageLocationFilter Filter, bool bMatchCaseOnDisk, FPackagePath* OutPackagePath)
 {
 	TStringBuilder<256> PackageName;
 	PackagePath.AppendPackageName(PackageName);
@@ -2051,39 +2051,21 @@ FPackageName::EPackageLocationFilter FPackageName::InternalDoesPackageExistEx(co
 	}
 
 	FText Reason;
-	if ( !FPackageName::IsValidTextForLongPackageName( PackageName, &Reason ) )
+	if (!FPackageName::IsValidTextForLongPackageName(PackageName, &Reason))
 	{
-		UE_LOG(LogPackageName, Error, TEXT( "DoesPackageExist: DoesPackageExist FAILED: '%s' is not a long packagename name. Reason: %s"), PackageName.ToString(), *Reason.ToString() );
+		UE_LOG(LogPackageName, Error, TEXT("DoesPackageExist: DoesPackageExist FAILED: '%s' is not a long packagename name. Reason: %s"), PackageName.ToString(), *Reason.ToString());
 		return EPackageLocationFilter::None;
 	}
 
 	EPackageLocationFilter Result = EPackageLocationFilter::None;
 
-	if ((uint8)Filter & (uint8)EPackageLocationFilter::FileSystem)
-	{
-		bool bFoundInFileSystem = false;
-		if (bMatchCaseOnDisk)
-		{
-			bFoundInFileSystem = IPackageResourceManager::Get().TryMatchCaseOnDisk(PackagePath, OutPackagePath);
-		}
-		else
-		{
-			bFoundInFileSystem = IPackageResourceManager::Get().DoesPackageExist(PackagePath, OutPackagePath);
-		}
+	bool bCheckIoDispatcher = (uint8)Filter & (uint8)EPackageLocationFilter::IoDispatcher;
+	bool bCheckFilesystem = (uint8)Filter & (uint8)EPackageLocationFilter::FileSystem;
+	bool bAttemptSkipTryMatchCase = bCheckFilesystem && bMatchCaseOnDisk;
 
-		if (bFoundInFileSystem)
-		{
-			Result = EPackageLocationFilter((uint8)Result | (uint8)EPackageLocationFilter::FileSystem);
-
-			// if we just want to find any existence, then we are done
-			if (Filter == EPackageLocationFilter::Any)
-			{
-				return Result;
-			}
-		}
-	}
-
-	if (((uint8)Filter & (uint8)EPackageLocationFilter::IoDispatcher))
+	// Matching path case on disk is expensive. Even if the package filter doesn't include IoDispatcher, check 
+	// to see if the IoDispatcher contains this path since if it's there we can skip matching the case on disk.
+	if (bCheckIoDispatcher || bAttemptSkipTryMatchCase)
 	{
 		bool bFoundInIoDispatcher = false;
 		if (DoesPackageExistOverrideDelegate.IsBound())
@@ -2109,7 +2091,36 @@ FPackageName::EPackageLocationFilter FPackageName::InternalDoesPackageExistEx(co
 				}
 			}
 
-			Result = EPackageLocationFilter((uint8)Result | (uint8)EPackageLocationFilter::IoDispatcher);
+			static_assert(uint8(EPackageLocationFilter::None) == 0);
+			const uint8 CheckIoDispatcherBitMask = ~uint8(bCheckIoDispatcher) + 1;
+			Result = EPackageLocationFilter(uint8(Result) | (uint8(EPackageLocationFilter::IoDispatcher) & CheckIoDispatcherBitMask));
+
+			// if we just want to find any existence, then we are done
+			if (Filter == EPackageLocationFilter::Any)
+			{
+				return Result;
+			}
+
+			// We do not support assets in both the filesystem and IoDispatcher simultaneously, so skip checking the filesystem if requested
+			bCheckFilesystem = false;
+		}
+	}
+
+	if (bCheckFilesystem)
+	{
+		bool bFoundInFileSystem = false;
+		if (bMatchCaseOnDisk)
+		{
+			bFoundInFileSystem = IPackageResourceManager::Get().TryMatchCaseOnDisk(PackagePath, OutPackagePath);
+		}
+		else
+		{
+			bFoundInFileSystem = IPackageResourceManager::Get().DoesPackageExist(PackagePath, OutPackagePath);
+		}
+
+		if (bFoundInFileSystem)
+		{
+			Result = EPackageLocationFilter((uint8)Result | (uint8)EPackageLocationFilter::FileSystem);
 
 			// if we just want to find any existence, then we are done
 			if (Filter == EPackageLocationFilter::Any)
diff --git a/Engine/Source/Runtime/CoreUObject/Public/Misc/PackageName.h b/Engine/Source/Runtime/CoreUObject/Public/Misc/PackageName.h
--- a/Engine/Source/Runtime/CoreUObject/Public/Misc/PackageName.h
+++ b/Engine/Source/Runtime/CoreUObject/Public/Misc/PackageName.h
@@ -407,16 +407,16 @@ public:
 	};
 
 	/**
-	 * Checks if the package exists in IOStore containers, on disk outsode of IOStore, both, or neither
+	 * Checks if the package exists in IOStore containers, on disk outside of IOStore, both, or neither
 	 *
 	 * @param PackagePath Package package.
 	 * @param Filter Indication of where it should look for 
-	 * @param Guid If nonnull, and the package is found on disk but does not have this PackageGuid in its FPackageFileSummary::Guid, false is returned
+	 * @param Guid If non-null, and the package is found on disk but does not have this PackageGuid in its FPackageFileSummary::Guid, false is returned
 	 * @param bMatchCaseOnDisk If true, the OutPackagePath is modified to match the capitalization of the discovered file
-	 * @param OutPackagePath If nonnull and the package exists, set to a copy of PackagePath with the HeaderExtension set to the extension that exists on disk (and if bMatchCaseOnDisk is true, capitalization changed to match). If not found, this variable is not written
+	 * @param OutPackagePath If non-null and the package exists, set to a copy of PackagePath with the HeaderExtension set to the extension that exists on disk (and if bMatchCaseOnDisk is true, capitalization changed to match). If not found, this variable is not written
 	 * @return the set of locations where the package exists (IoDispatcher or FileSystem, both or neither)
 	 **/
-	static COREUOBJECT_API EPackageLocationFilter DoesPackageExistEx(const FPackagePath& PackagePath, EPackageLocationFilter Filterconst, bool bMatchCaseOnDisk = false, FPackagePath* OutPackagePath = nullptr);
+	static COREUOBJECT_API EPackageLocationFilter DoesPackageExistEx(const FPackagePath& PackagePath, EPackageLocationFilter Filter, bool bMatchCaseOnDisk = false, FPackagePath* OutPackagePath = nullptr);
 
 	/**
 	 * Attempts to find a package given its short name on disk (very slow).
