Fix arrays subobject instantiation in non editor builds when cooked using TPS(TagPropertySerialization) and OS(OverridableSerialization).

diff --git a/Engine/Source/Runtime/CoreUObject/Private/UObject/CoreNative.cpp b/Engine/Source/Runtime/CoreUObject/Private/UObject/CoreNative.cpp
--- a/Engine/Source/Runtime/CoreUObject/Private/UObject/CoreNative.cpp
+++ b/Engine/Source/Runtime/CoreUObject/Private/UObject/CoreNative.cpp
@@ -207,8 +207,8 @@ UObject* FObjectInstancingGraph::GetInstancedSubobject( UObject* SourceSubobject
 		if ( bShouldInstance )
 		{
 			// If the CurrentValue is within the SourceRoot, lets use it to instantiate as it must have come from the merge result of the serialization
-			const bool bAreOverridesEnabled = SourceSubobject && FOverridableManager::Get().IsEnabled(*SourceSubobject);
-			if (bAreOverridesEnabled && SourceSubobject != CurrentValue && CurrentValue->IsIn(SourceRoot))
+			const bool bIsInstantiatingSubObjectForOverridableSerialization = SourceSubobject && (FOverridableManager::Get().IsEnabled(*SourceSubobject) || FOverridableManager::Get().NeedSubObjectTemplateInstantiation(*SourceSubobject));
+			if (bIsInstantiatingSubObjectForOverridableSerialization && SourceSubobject != CurrentValue && CurrentValue->IsIn(SourceRoot))
 			{
 				SourceSubobject = CurrentValue;
 			}
