Fix arrays subobject instanciation in non editor builds when saved with OverridableSerialization

diff --git a/Engine/Source/Runtime/CoreUObject/Private/UObject/CoreNative.cpp b/Engine/Source/Runtime/CoreUObject/Private/UObject/CoreNative.cpp
--- a/Engine/Source/Runtime/CoreUObject/Private/UObject/CoreNative.cpp
+++ b/Engine/Source/Runtime/CoreUObject/Private/UObject/CoreNative.cpp
@@ -207,7 +207,7 @@ UObject* FObjectInstancingGraph::GetInstancedSubobject( UObject* SourceSubobject
 		if ( bShouldInstance )
 		{
 			// If the CurrentValue is within the SourceRoot, lets use it to instantiate as it must have come from the merge result of the serialization
-			const bool bIsInstantiatingSubObjectForOverridableSerialization = SourceSubobject && (FOverridableManager::Get().IsEnabled(*SourceSubobject) || FOverridableManager::Get().NeedSubObjectTemplateInstantiation(*SourceSubobject));
+			const bool bIsInstantiatingSubObjectForOverridableSerialization = SourceSubobject && FOverridableManager::Get().NeedSubObjectTemplateInstantiation(*SourceSubobject);
 			if (bIsInstantiatingSubObjectForOverridableSerialization && SourceSubobject != CurrentValue && CurrentValue->IsIn(SourceRoot))
 			{
 				SourceSubobject = CurrentValue;
