Deprecating FProperty::ElementSize for public access and introducing accessors.  Instead of accessing the variable directly, use GetElementSize and SetElementSize in order to facilitate a future change of underlying storage type (to pack FProperty tighter and save some memory).

[RN] UStructs above 16.7mb will trip an ensure in the Editor so that they remain compatible with upcoming changes to Max ElementSize. You can disable the ensure with CVar CoreUObject.EnsureAgainstLargeProperties.

diff --git a/Engine/Plugins/Animation/ControlRig/Source/ControlRig/Public/Rigs/RigHierarchyMetadata.h b/Engine/Plugins/Animation/ControlRig/Source/ControlRig/Public/Rigs/RigHierarchyMetadata.h
--- a/Engine/Plugins/Animation/ControlRig/Source/ControlRig/Public/Rigs/RigHierarchyMetadata.h
+++ b/Engine/Plugins/Animation/ControlRig/Source/ControlRig/Public/Rigs/RigHierarchyMetadata.h
@@ -68,9 +68,9 @@ public:
 	{
 		return GetValueProperty()->ContainerPtrToValuePtr<void>(this);
 	}
-	const int32& GetValueSize() const
+	int32 GetValueSize() const
 	{
-		return GetValueProperty()->ElementSize;
+		return GetValueProperty()->GetElementSize();
 	}
 	bool SetValueData(const void* InData, int32 InSize)
 	{
diff --git a/Engine/Plugins/Animation/DeformerGraph/Source/OptimusCore/Private/OptimusDataTypeRegistry.cpp b/Engine/Plugins/Animation/DeformerGraph/Source/OptimusCore/Private/OptimusDataTypeRegistry.cpp
--- a/Engine/Plugins/Animation/DeformerGraph/Source/OptimusCore/Private/OptimusDataTypeRegistry.cpp
+++ b/Engine/Plugins/Animation/DeformerGraph/Source/OptimusCore/Private/OptimusDataTypeRegistry.cpp
@@ -591,7 +591,7 @@ bool FOptimusDataTypeRegistry::RegisterStructType(UScriptStruct* InStructType)
 	{
 		auto Property = new FStructProperty(InScope, InName, RF_Public);
 		Property->Struct = InStructType;
-		Property->ElementSize = InStructType->GetStructureSize();
+		Property->SetElementSize(InStructType->GetStructureSize());
 		if (bIsHashable)
 		{
 			Property->SetPropertyFlags(CPF_HasGetValueTypeHash);
@@ -1081,7 +1081,7 @@ bool FOptimusDataTypeRegistry::RegisterType(
 			{
 				auto Property = new FStructProperty(InScope, InName, RF_Public);
 				Property->Struct = InStructType;
-				Property->ElementSize = InStructType->GetStructureSize();
+				Property->SetElementSize(InStructType->GetStructureSize());
 				if (bIsHashable)
 				{
 					Property->SetPropertyFlags(CPF_HasGetValueTypeHash);
@@ -1241,7 +1241,7 @@ bool FOptimusDataTypeRegistry::RegisterType(
 			{
 				FStructProperty* Property = new FStructProperty(InScope, InName, RF_Public);
 				Property->Struct = InStructType;
-				Property->ElementSize = InStructType->GetStructureSize();
+				Property->SetElementSize(InStructType->GetStructureSize());
 				if (bIsHashable)
 				{
 					Property->SetPropertyFlags(CPF_HasGetValueTypeHash);
@@ -1296,7 +1296,7 @@ bool FOptimusDataTypeRegistry::RegisterType(
 			PropertyCreateFunc = [bIsHashable, InStructType](FFieldVariant InScope, FName InName) -> FProperty* {
 				FStructProperty* Property = new FStructProperty(InScope, InName, RF_Public);
 				Property->Struct = InStructType;
-				Property->ElementSize = InStructType->GetStructureSize();
+				Property->SetElementSize(InStructType->GetStructureSize());
 				if (bIsHashable)
 				{
 					Property->SetPropertyFlags(CPF_HasGetValueTypeHash);
diff --git a/Engine/Plugins/Developer/Concert/ConcertSync/ConcertSyncCore/Source/ConcertSyncCore/Private/ConcertSyncArchives.cpp b/Engine/Plugins/Developer/Concert/ConcertSync/ConcertSyncCore/Source/ConcertSyncCore/Private/ConcertSyncArchives.cpp
--- a/Engine/Plugins/Developer/Concert/ConcertSync/ConcertSyncCore/Source/ConcertSyncCore/Private/ConcertSyncArchives.cpp
+++ b/Engine/Plugins/Developer/Concert/ConcertSync/ConcertSyncCore/Source/ConcertSyncCore/Private/ConcertSyncArchives.cpp
@@ -635,7 +635,7 @@ void FConcertSyncObjectRewriter::RewriteProperty(const FProperty* InProp)
 
 	for (int32 Idx = 0; Idx < InProp->ArrayDim; ++Idx)
 	{
-		InProp->SerializeItem(FStructuredArchiveFromArchive(*this).GetSlot(), (uint8*)TmpPropData + (InProp->ElementSize * Idx));
+		InProp->SerializeItem(FStructuredArchiveFromArchive(*this).GetSlot(), (uint8*)TmpPropData + (InProp->GetElementSize() * Idx));
 	}
 
 	InProp->DestroyValue(TmpPropData);
diff --git a/Engine/Plugins/Enterprise/VariantManager/Source/VariantManager/Private/DisplayNodes/VariantManagerPropertyNode.cpp b/Engine/Plugins/Enterprise/VariantManager/Source/VariantManager/Private/DisplayNodes/VariantManagerPropertyNode.cpp
--- a/Engine/Plugins/Enterprise/VariantManager/Source/VariantManager/Private/DisplayNodes/VariantManagerPropertyNode.cpp
+++ b/Engine/Plugins/Enterprise/VariantManager/Source/VariantManager/Private/DisplayNodes/VariantManagerPropertyNode.cpp
@@ -517,7 +517,7 @@ void FVariantManagerPropertyNode::UpdateRecordedDataFromSinglePropView(TSharedPt
 		}
 		else
 		{
-			PropertyValue->SetRecordedData((uint8*)SinglePropWidgetDataPtr, PropHandle->GetProperty()->ElementSize);
+			PropertyValue->SetRecordedData((uint8*)SinglePropWidgetDataPtr, PropHandle->GetProperty()->GetElementSize());
 		}
 	}
 
diff --git a/Engine/Plugins/Enterprise/VariantManager/Source/VariantManager/Private/DisplayNodes/VariantManagerStructPropertyNode.cpp b/Engine/Plugins/Enterprise/VariantManager/Source/VariantManager/Private/DisplayNodes/VariantManagerStructPropertyNode.cpp
--- a/Engine/Plugins/Enterprise/VariantManager/Source/VariantManager/Private/DisplayNodes/VariantManagerStructPropertyNode.cpp
+++ b/Engine/Plugins/Enterprise/VariantManager/Source/VariantManager/Private/DisplayNodes/VariantManagerStructPropertyNode.cpp
@@ -316,7 +316,7 @@ TSharedPtr<SWidget> FVariantManagerStructPropertyNode::GetPropertyValueWidget()
 		FNumericProperty* Prop = *NumPropIter;
 
 		// Check if we have the same value on all UPropertyValue, for each separate child property
-		int32 Size = Prop->ElementSize;
+		int32 Size = Prop->GetElementSize();
 		int32 Offset = Prop->GetOffset_ForInternal();
 		bool bSameValue = true;
 
@@ -346,7 +346,7 @@ TSharedPtr<SWidget> FVariantManagerStructPropertyNode::GetPropertyValueWidget()
 
 		// This is not trivial because if we're a rotator property the standard is to actually show a vector
 		// widget, but the vector widget expects to read/write from doubles, while rotators still hold floats
-		const size_t RecordedElementSize = bIsRotator ? RotatorElementSize : Prop->ElementSize;
+		const size_t RecordedElementSize = bIsRotator ? RotatorElementSize : Prop->GetElementSize();
 
 		FMargin& MarginToUse = (Prop == LastProp) ? LastMargin : CommonMargin;
 
@@ -527,7 +527,7 @@ void FVariantManagerStructPropertyNode::OnSignedPropCommitted(int64 InValue, ETe
 	}
 
 	UPropertyValue* FirstPropertyValue = PropertyValues[0].Get();
-	int32 Size = Prop->ElementSize;
+	int32 Size = Prop->GetElementSize();
 	TArray<uint8> InValueBytes;
 	InValueBytes.SetNumUninitialized(Size);
 
@@ -596,7 +596,7 @@ void FVariantManagerStructPropertyNode::OnUnsignedPropCommitted(uint64 InValue,
 	}
 
 	UPropertyValue* FirstPropertyValue = PropertyValues[0].Get();
-	int32 Size = Prop->ElementSize;
+	int32 Size = Prop->GetElementSize();
 	TArray<uint8> InValueBytes;
 	InValueBytes.SetNumUninitialized(Size);
 
@@ -724,7 +724,7 @@ TOptional<int64> FVariantManagerStructPropertyNode::GetSignedValueFromPropertyVa
 		return TOptional<int64>();
 	}
 
-	int32 Size = Prop->ElementSize;
+	int32 Size = Prop->GetElementSize();
 
 	const TArray<uint8>* FirstRecordedData = nullptr;
 
@@ -780,7 +780,7 @@ TOptional<uint64> FVariantManagerStructPropertyNode::GetUnsignedValueFromPropert
 		return TOptional<uint64>();
 	}
 
-	int32 Size = Prop->ElementSize;
+	int32 Size = Prop->GetElementSize();
 
 	const TArray<uint8>* FirstRecordedData = nullptr;
 
diff --git a/Engine/Plugins/Enterprise/VariantManagerContent/Source/VariantManagerContent/Private/PropertyValue.cpp b/Engine/Plugins/Enterprise/VariantManagerContent/Source/VariantManagerContent/Private/PropertyValue.cpp
--- a/Engine/Plugins/Enterprise/VariantManagerContent/Source/VariantManagerContent/Private/PropertyValue.cpp
+++ b/Engine/Plugins/Enterprise/VariantManagerContent/Source/VariantManagerContent/Private/PropertyValue.cpp
@@ -753,7 +753,7 @@ void UPropertyValue::ApplyDataToResolvedObject()
 	else if (FEnumProperty* PropAsEnum = CastField<FEnumProperty>(LeafProperty))
 	{
 		FNumericProperty* UnderlyingProp = PropAsEnum->GetUnderlyingProperty();
-		int32 PropertySizeBytes = UnderlyingProp->ElementSize;
+		int32 PropertySizeBytes = UnderlyingProp->GetElementSize();
 
 		ValueBytes.SetNum(PropertySizeBytes);
 		FMemory::Memcpy(PropertyValuePtr, ValueBytes.GetData(), PropertySizeBytes);
@@ -1110,11 +1110,11 @@ int32 UPropertyValue::GetValueSizeInBytes() const
 	FProperty* Prop = GetProperty();
 	if (FEnumProperty* PropAsEnumProp = CastField<FEnumProperty>(Prop))
 	{
-		return PropAsEnumProp->GetUnderlyingProperty()->ElementSize;
+		return PropAsEnumProp->GetUnderlyingProperty()->GetElementSize();
 	}
 	else if (Prop)
 	{
-		return Prop->ElementSize;
+		return Prop->GetElementSize();
 	}
 
 	UE_LOG(LogVariantContent, Warning, TEXT("Returning size zero for PropertyValue '%s'"), *GetFullDisplayString());
@@ -1274,7 +1274,7 @@ const TArray<uint8>& UPropertyValue::GetDefaultValue()
 						if (FEnumProperty* PropAsEnum = CastField<FEnumProperty>(LeafProperty))
 						{
 							FNumericProperty* UnderlyingProp = PropAsEnum->GetUnderlyingProperty();
-							NumBytes = UnderlyingProp->ElementSize;
+							NumBytes = UnderlyingProp->GetElementSize();
 						}
 
 						// If we're a material property value we won't have PropertyValuePtr
diff --git a/Engine/Plugins/Experimental/AnimNext/Source/AnimNext/Private/Graph/AnimGraph/AnimNodeExposedValueHandler_AnimNextParameters.cpp b/Engine/Plugins/Experimental/AnimNext/Source/AnimNext/Private/Graph/AnimGraph/AnimNodeExposedValueHandler_AnimNextParameters.cpp
--- a/Engine/Plugins/Experimental/AnimNext/Source/AnimNext/Private/Graph/AnimGraph/AnimNodeExposedValueHandler_AnimNextParameters.cpp
+++ b/Engine/Plugins/Experimental/AnimNext/Source/AnimNext/Private/Graph/AnimGraph/AnimNodeExposedValueHandler_AnimNextParameters.cpp
@@ -56,7 +56,7 @@ namespace UE::AnimNext::Private
 			break;
 		case EPropertyAccessCopyType::Plain:
 			checkSlow(InDestProperty->PropertyFlags & CPF_IsPlainOldData);
-			FMemory::Memcpy(InDestAddr, InSrcAddr, InDestProperty->ElementSize);
+			FMemory::Memcpy(InDestAddr, InSrcAddr, InDestProperty->GetElementSize());
 			break;
 		case EPropertyAccessCopyType::Complex:
 			InDestProperty->CopyCompleteValue(InDestAddr, InSrcAddr);
diff --git a/Engine/Plugins/Experimental/InstanceDataObjectFixupTool/Source/Private/InstanceDataObjectFixupPanel.cpp b/Engine/Plugins/Experimental/InstanceDataObjectFixupTool/Source/Private/InstanceDataObjectFixupPanel.cpp
--- a/Engine/Plugins/Experimental/InstanceDataObjectFixupTool/Source/Private/InstanceDataObjectFixupPanel.cpp
+++ b/Engine/Plugins/Experimental/InstanceDataObjectFixupTool/Source/Private/InstanceDataObjectFixupPanel.cpp
@@ -660,7 +660,7 @@ FInstanceDataObjectFixupPanel::FTypeConverter::EWarning FInstanceDataObjectFixup
 {
 	// convert from source to destination in a temp buffer to see if it's possible
 	TArray<uint8, TInlineAllocator<64>> SourceToDest;
-	SourceToDest.SetNumUninitialized(DestinationProperty->ElementSize);
+	SourceToDest.SetNumUninitialized(DestinationProperty->GetElementSize());
 	DestinationProperty->InitializeValue(SourceToDest.GetData());
 	if (!TryConvert(SourceProperty, SourceData, DestinationProperty, SourceToDest.GetData()))
 	{
@@ -669,7 +669,7 @@ FInstanceDataObjectFixupPanel::FTypeConverter::EWarning FInstanceDataObjectFixup
 
 	// convert from destination to source in a temp buffer to see if it's possible
 	TArray<uint8, TInlineAllocator<64>> DestToSource;
-	DestToSource.SetNumUninitialized(SourceProperty->ElementSize);
+	DestToSource.SetNumUninitialized(SourceProperty->GetElementSize());
 	SourceProperty->InitializeValue(DestToSource.GetData());
 	if (!TryConvert(DestinationProperty, SourceToDest.GetData(), SourceProperty, DestToSource.GetData()))
 	{
@@ -824,7 +824,7 @@ void FInstanceDataObjectFixupPanel::RedirectProperty(const FPropertyPath& From,
 		if (ToRevertInfo)
 		{
 			// cache the destination value so it can be reverted later
-			const int32 Size = DestinationProperty->ArrayDim * DestinationProperty->ElementSize;
+			const int32 Size = DestinationProperty->ArrayDim * DestinationProperty->GetElementSize();
 			ToRevertInfo->OriginalValue.AddZeroed(Size);
 			uint8* Buffer = ToRevertInfo->OriginalValue.GetData() + (ToRevertInfo->OriginalValue.Num() - Size);
 			DestinationProperty->CopyCompleteValue(Buffer, Destination);
@@ -846,7 +846,7 @@ void FInstanceDataObjectFixupPanel::RedirectProperty(const FPropertyPath& From,
 		{
 			// apply FromRevertInfo to From
 			SourceProperty->CopyCompleteValue(Source, FromRevertInfoItr);
-			FromRevertInfoItr += DestinationProperty->ArrayDim * DestinationProperty->ElementSize;
+			FromRevertInfoItr += DestinationProperty->ArrayDim * DestinationProperty->GetElementSize();
 		}
 		Instance->PostEditChangeChainProperty(ChangedChainEvent);
 	}
@@ -889,7 +889,7 @@ void FInstanceDataObjectFixupPanel::RedirectProperty(const FPropertyPath& From,
 		if (ToRevertInfo)
 		{
 			// cache the destination value so it can be reverted later
-			const int32 Size = DestinationProperty->ArrayDim * DestinationProperty->ElementSize;
+			const int32 Size = DestinationProperty->ArrayDim * DestinationProperty->GetElementSize();
 			ToRevertInfo->OriginalValue.AddZeroed(Size);
 			uint8* Buffer = ToRevertInfo->OriginalValue.GetData() + (ToRevertInfo->OriginalValue.Num() - Size);
 			DestinationProperty->CopyCompleteValue(Buffer, Destination);
@@ -908,7 +908,7 @@ void FInstanceDataObjectFixupPanel::RedirectProperty(const FPropertyPath& From,
 		{
 			// apply FromRevertInfo to From
 			SourceProperty->CopyCompleteValue(Source, FromRevertInfoItr);
-			FromRevertInfoItr += DestinationProperty->ArrayDim * DestinationProperty->ElementSize;
+			FromRevertInfoItr += DestinationProperty->ArrayDim * DestinationProperty->GetElementSize();
 		}
 		FPropertyChangedChainEvent ChangedChainEvent(Chains[I], ChangeEvents[I]);
 		Instance->PostEditChangeChainProperty(ChangedChainEvent);
diff --git a/Engine/Plugins/Experimental/PythonScriptPlugin/Source/PythonScriptPlugin/Private/PyConversion.cpp b/Engine/Plugins/Experimental/PythonScriptPlugin/Source/PythonScriptPlugin/Private/PyConversion.cpp
--- a/Engine/Plugins/Experimental/PythonScriptPlugin/Source/PythonScriptPlugin/Private/PyConversion.cpp
+++ b/Engine/Plugins/Experimental/PythonScriptPlugin/Source/PythonScriptPlugin/Private/PyConversion.cpp
@@ -674,9 +674,9 @@ FPyConversionResult NativizeProperty(PyObject* PyObj, const FProperty* Prop, voi
 				{
 					for (void* ArchInstValueAddr : InArchetypeInstValueAddrs)
 					{
-						Prop->CopySingleValue(static_cast<uint8*>(ArchInstValueAddr) + (Prop->ElementSize * ArrIndex), FPyWrapperFixedArray::GetItemPtr(PyFixedArray, ArrIndex));
+						Prop->CopySingleValue(static_cast<uint8*>(ArchInstValueAddr) + (Prop->GetElementSize() * ArrIndex), FPyWrapperFixedArray::GetItemPtr(PyFixedArray, ArrIndex));
 					}
-					Prop->CopySingleValue(static_cast<uint8*>(ValueAddr) + (Prop->ElementSize * ArrIndex), FPyWrapperFixedArray::GetItemPtr(PyFixedArray, ArrIndex));
+					Prop->CopySingleValue(static_cast<uint8*>(ValueAddr) + (Prop->GetElementSize() * ArrIndex), FPyWrapperFixedArray::GetItemPtr(PyFixedArray, ArrIndex));
 				}
 			});
 			return FPyConversionResult::Success();
diff --git a/Engine/Plugins/Experimental/PythonScriptPlugin/Source/PythonScriptPlugin/Private/PyGenUtil.cpp b/Engine/Plugins/Experimental/PythonScriptPlugin/Source/PythonScriptPlugin/Private/PyGenUtil.cpp
--- a/Engine/Plugins/Experimental/PythonScriptPlugin/Source/PythonScriptPlugin/Private/PyGenUtil.cpp
+++ b/Engine/Plugins/Experimental/PythonScriptPlugin/Source/PythonScriptPlugin/Private/PyGenUtil.cpp
@@ -2209,7 +2209,7 @@ void PythonizeValueImpl(const FProperty* InProp, const void* InPropValue, const
 	}
 	for (int32 ArrIndex = 0; ArrIndex < InProp->ArrayDim; ++ArrIndex)
 	{
-		const void* PropArrValue = ((uint8*)InPropValue) + (InProp->ElementSize * ArrIndex);
+		const void* PropArrValue = ((uint8*)InPropValue) + (InProp->GetElementSize() * ArrIndex);
 		if (ArrIndex > 0)
 		{
 			OutPythonDefaultValue += TEXT(", ");
diff --git a/Engine/Plugins/Experimental/PythonScriptPlugin/Source/PythonScriptPlugin/Private/PyUtil.cpp b/Engine/Plugins/Experimental/PythonScriptPlugin/Source/PythonScriptPlugin/Private/PyUtil.cpp
--- a/Engine/Plugins/Experimental/PythonScriptPlugin/Source/PythonScriptPlugin/Private/PyUtil.cpp
+++ b/Engine/Plugins/Experimental/PythonScriptPlugin/Source/PythonScriptPlugin/Private/PyUtil.cpp
@@ -165,7 +165,7 @@ const FProperty* FPropValueOnScope::GetProp() const
 void* FPropValueOnScope::GetValue(const int32 InArrayIndex) const
 {
 	check(InArrayIndex >= 0 && InArrayIndex < Prop->ArrayDim);
-	return ((uint8*)Value) + (Prop->ElementSize * InArrayIndex);
+	return ((uint8*)Value) + (Prop->GetElementSize() * InArrayIndex);
 }
 
 FFixedArrayElementOnScope::FFixedArrayElementOnScope(const FProperty* InProp)
diff --git a/Engine/Plugins/Experimental/PythonScriptPlugin/Source/PythonScriptPlugin/Private/PyWrapperFixedArray.cpp b/Engine/Plugins/Experimental/PythonScriptPlugin/Source/PythonScriptPlugin/Private/PyWrapperFixedArray.cpp
--- a/Engine/Plugins/Experimental/PythonScriptPlugin/Source/PythonScriptPlugin/Private/PyWrapperFixedArray.cpp
+++ b/Engine/Plugins/Experimental/PythonScriptPlugin/Source/PythonScriptPlugin/Private/PyWrapperFixedArray.cpp
@@ -366,7 +366,7 @@ void* FPyWrapperFixedArray::GetItemPtr(FPyWrapperFixedArray* InSelf, Py_ssize_t
 
 	// This doesn't use ContainerPtrToValuePtr as the ArrayInstance has already been adjusted to 
 	// point to the first element in the array and we just need to adjust it for the element size
-	return static_cast<uint8*>(InSelf->ArrayInstance) + (InSelf->ArrayProp->ElementSize * InIndex);
+	return static_cast<uint8*>(InSelf->ArrayInstance) + (InSelf->ArrayProp->GetElementSize() * InIndex);
 }
 
 Py_ssize_t FPyWrapperFixedArray::Len(FPyWrapperFixedArray* InSelf)
diff --git a/Engine/Plugins/FX/Niagara/Source/Niagara/Private/NiagaraRendererComponents.cpp b/Engine/Plugins/FX/Niagara/Source/Niagara/Private/NiagaraRendererComponents.cpp
--- a/Engine/Plugins/FX/Niagara/Source/Niagara/Private/NiagaraRendererComponents.cpp
+++ b/Engine/Plugins/FX/Niagara/Source/Niagara/Private/NiagaraRendererComponents.cpp
@@ -176,7 +176,7 @@ void InvokeSetterFunction(UObject* InRuntimeObject, UFunction* Setter, const uin
 				{
 					// LWC backwards compatibility for old bindings
 					UScriptStruct* ScriptStruct = ((FStructProperty*)Property)->Struct;
-					if (DataSize == sizeof(FVector3f) && Property->ElementSize == sizeof(FVector3d) && ScriptStruct && (ScriptStruct->GetFName() == NAME_Vector || ScriptStruct->GetFName() == NAME_Vector3d))
+					if (DataSize == sizeof(FVector3f) && Property->GetElementSize() == sizeof(FVector3d) && ScriptStruct && (ScriptStruct->GetFName() == NAME_Vector || ScriptStruct->GetFName() == NAME_Vector3d))
 					{
 						FVector3f* Vec3f = (FVector3f*)InData;
 						FVector Vec3d(*Vec3f);
@@ -184,7 +184,7 @@ void InvokeSetterFunction(UObject* InRuntimeObject, UFunction* Setter, const uin
 					}
 					else
 					{
-						const bool bIsValid = ensureMsgf(DataSize == Property->ElementSize, TEXT("Property type does not match for setter function %s::%s (%ibytes != %ibytes"), *InRuntimeObject->GetName(), *Setter->GetName(), DataSize, Property->ElementSize);
+						const bool bIsValid = ensureMsgf(DataSize == Property->GetElementSize(), TEXT("Property type does not match for setter function %s::%s (%ibytes != %ibytes"), *InRuntimeObject->GetName(), *Setter->GetName(), DataSize, Property->GetElementSize());
 						if (bIsValid)
 						{
 							Property->CopyCompleteValue(Property->ContainerPtrToValuePtr<void>(Params), InData);
diff --git a/Engine/Plugins/FX/Niagara/Source/Niagara/Private/NiagaraTypes.cpp b/Engine/Plugins/FX/Niagara/Source/Niagara/Private/NiagaraTypes.cpp
--- a/Engine/Plugins/FX/Niagara/Source/Niagara/Private/NiagaraTypes.cpp
+++ b/Engine/Plugins/FX/Niagara/Source/Niagara/Private/NiagaraTypes.cpp
@@ -419,35 +419,35 @@ FNiagaraLwcStructConverter BuildSWCStructure(UScriptStruct* NewStruct, UScriptSt
 			if ((StructName == NAME_Vector2d) || (StructName == NAME_Vector2D))
 			{
 				NewStructProperty->Struct = Vector2fStruct;
-				NewStructProperty->ElementSize = Vector2fStruct->GetStructureSize();
+				NewStructProperty->SetElementSize(Vector2fStruct->GetStructureSize());
 				AlignedOffset = Align(AlignedOffset, Vector2fStruct->GetMinAlignment());
 				StructConverter.AddConversionStep(OldPropertySize, OldPropertyOffset, NewStructProperty->GetSize(), AlignedOffset, ENiagaraStructConversionType::Vector2);
 			}
 			else if ((StructName == NAME_Vector3d) || (StructName == NAME_Vector))
 			{
 				NewStructProperty->Struct = Vector3fStruct;
-				NewStructProperty->ElementSize = Vector3fStruct->GetStructureSize();
+				NewStructProperty->SetElementSize(Vector3fStruct->GetStructureSize());
 				AlignedOffset = Align(AlignedOffset, Vector3fStruct->GetMinAlignment());
 				StructConverter.AddConversionStep(OldPropertySize, OldPropertyOffset, NewStructProperty->GetSize(), AlignedOffset, ENiagaraStructConversionType::Vector3);
 			}
 			else if ((StructName == NAME_Vector4d) || (StructName == NAME_Vector4))
 			{
 				NewStructProperty->Struct = Vector4fStruct;
-				NewStructProperty->ElementSize = Vector4fStruct->GetStructureSize();
+				NewStructProperty->SetElementSize(Vector4fStruct->GetStructureSize());
 				AlignedOffset = Align(AlignedOffset, Vector4fStruct->GetMinAlignment());
 				StructConverter.AddConversionStep(OldPropertySize, OldPropertyOffset, NewStructProperty->GetSize(), AlignedOffset, ENiagaraStructConversionType::Vector4);
 			}
 			else if ((StructName == NAME_Quat4d) || (StructName == NAME_Quat))
 			{
 				NewStructProperty->Struct = Quat4fStruct;
-				NewStructProperty->ElementSize = Quat4fStruct->GetStructureSize();
+				NewStructProperty->SetElementSize(Quat4fStruct->GetStructureSize());
 				AlignedOffset = Align(AlignedOffset, Quat4fStruct->GetMinAlignment());
 				StructConverter.AddConversionStep(OldPropertySize, OldPropertyOffset, NewStructProperty->GetSize(), AlignedOffset, ENiagaraStructConversionType::Quat);
 			}
 			else
 			{
 				NewStructProperty->Struct = ChildAsStruct->Struct;
-				NewStructProperty->ElementSize = ChildAsStruct->Struct->GetStructureSize();
+				NewStructProperty->SetElementSize(ChildAsStruct->Struct->GetStructureSize());
 				AlignedOffset = Align(AlignedOffset, ChildAsStruct->Struct->GetMinAlignment());
 				StructConverter.AddConversionStep(OldPropertySize, OldPropertyOffset, NewStructProperty->GetSize(), AlignedOffset, ENiagaraStructConversionType::CopyOnly);
 			}
diff --git a/Engine/Plugins/Runtime/PropertyBindingUtils/Source/PropertyBindingUtils/Private/PropertyBindingPath.cpp b/Engine/Plugins/Runtime/PropertyBindingUtils/Source/PropertyBindingUtils/Private/PropertyBindingPath.cpp
--- a/Engine/Plugins/Runtime/PropertyBindingUtils/Source/PropertyBindingUtils/Private/PropertyBindingPath.cpp
+++ b/Engine/Plugins/Runtime/PropertyBindingUtils/Source/PropertyBindingUtils/Private/PropertyBindingPath.cpp
@@ -354,7 +354,7 @@ bool FPropertyBindingPath::ResolveIndirectionsWithValue(const FPropertyBindingDa
 				return false;
 			}
 			ArrayIndex = FMath::Max(0, Segment->GetArrayIndex());
-			Offset = Property->GetOffset_ForInternal() + Property->ElementSize * ArrayIndex;
+			Offset = Property->GetOffset_ForInternal() + Property->GetElementSize() * ArrayIndex;
 		}
 
 		FPropertyBindingPathIndirection& Indirection = OutIndirections.AddDefaulted_GetRef();
diff --git a/Engine/Plugins/Runtime/RigVM/Source/RigVM/Private/RigVMCore/RigVMMemoryDeprecated.cpp b/Engine/Plugins/Runtime/RigVM/Source/RigVM/Private/RigVMCore/RigVMMemoryDeprecated.cpp
--- a/Engine/Plugins/Runtime/RigVM/Source/RigVM/Private/RigVMCore/RigVMMemoryDeprecated.cpp
+++ b/Engine/Plugins/Runtime/RigVM/Source/RigVM/Private/RigVMCore/RigVMMemoryDeprecated.cpp
@@ -248,7 +248,7 @@ FRigVMRegisterOffset::FRigVMRegisterOffset(UScriptStruct* InScriptStruct, const
 			else
 			{
 				Offset.CPPType = *Property->GetCPPType();
-				Offset.ElementSize = Property->ElementSize;
+				Offset.ElementSize = Property->GetElementSize();
 
 				if (FArrayProperty* ArrayProperty = CastField<FArrayProperty>(Property))
 				{
@@ -256,7 +256,7 @@ FRigVMRegisterOffset::FRigVMRegisterOffset(UScriptStruct* InScriptStruct, const
 					Property = ArrayProperty->Inner;
 
 					Offset.CPPType = *ArrayProperty->Inner->GetCPPType();
-					Offset.ElementSize = ArrayProperty->Inner->ElementSize;
+					Offset.ElementSize = ArrayProperty->Inner->GetElementSize();
 				}
 
 				if (FStructProperty* StructProperty = CastField<FStructProperty>(Property))
@@ -282,7 +282,7 @@ FRigVMRegisterOffset::FRigVMRegisterOffset(UScriptStruct* InScriptStruct, const
 			}
 
 			int32 ArrayIndex = FCString::Atoi(*Left);
-			int32 SegmentIndex = -1 - InArrayProperty->Inner->ElementSize * ArrayIndex;
+			int32 SegmentIndex = -1 - InArrayProperty->Inner->GetElementSize() * ArrayIndex;
 
 			if (Offset.Segments.Num() > 0)
 			{
@@ -314,7 +314,7 @@ FRigVMRegisterOffset::FRigVMRegisterOffset(UScriptStruct* InScriptStruct, const
 			else
 			{
 				Offset.CPPType = *InArrayProperty->Inner->GetCPPType();
-				Offset.ElementSize = InArrayProperty->Inner->ElementSize;
+				Offset.ElementSize = InArrayProperty->Inner->GetElementSize();
 
 				if (FArrayProperty* ArrayProperty = CastField<FArrayProperty>(InArrayProperty->Inner))
 				{
diff --git a/Engine/Plugins/Runtime/ScriptableToolsFramework/Source/ScriptableToolsFramework/Private/ScriptableInteractiveTool.cpp b/Engine/Plugins/Runtime/ScriptableToolsFramework/Source/ScriptableToolsFramework/Private/ScriptableInteractiveTool.cpp
--- a/Engine/Plugins/Runtime/ScriptableToolsFramework/Source/ScriptableToolsFramework/Private/ScriptableInteractiveTool.cpp
+++ b/Engine/Plugins/Runtime/ScriptableToolsFramework/Source/ScriptableToolsFramework/Private/ScriptableInteractiveTool.cpp
@@ -874,7 +874,7 @@ UScriptableInteractiveToolPropertySet* UScriptableInteractiveTool::WatchProperty
 			[this, ArrayIndex, ArrayProperty, PropertySet]()->uint32 { 
 				const void* ArrayValuePtr = ArrayProperty->ContainerPtrToValuePtr<void>(PropertySet);
 				FScriptArrayHelper ArrayHelper(ArrayProperty, ArrayValuePtr);
-				int32 ElementSize = ArrayProperty->Inner->ElementSize;
+				int32 ElementSize = ArrayProperty->Inner->GetElementSize();
 				int32 NumElements = ArrayHelper.Num();
  				uint32 CRCValue = NumElements;
 				for ( int32 k = 0; k < NumElements; ++k )
diff --git a/Engine/Plugins/Runtime/StateTree/Source/StateTreeModule/Private/StateTreePropertyBindings.cpp b/Engine/Plugins/Runtime/StateTree/Source/StateTreeModule/Private/StateTreePropertyBindings.cpp
--- a/Engine/Plugins/Runtime/StateTree/Source/StateTreeModule/Private/StateTreePropertyBindings.cpp
+++ b/Engine/Plugins/Runtime/StateTree/Source/StateTreeModule/Private/StateTreePropertyBindings.cpp
@@ -566,7 +566,7 @@ bool FStateTreePropertyBindings::ResolveCopyType(const FStateTreePropertyPathInd
 		else if (TargetProperty->PropertyFlags & CPF_IsPlainOldData)
 		{
 			OutCopy.Type = EStateTreePropertyCopyType::CopyPlain;
-			OutCopy.CopySize = SourceProperty->ElementSize * SourceProperty->ArrayDim;
+			OutCopy.CopySize = SourceProperty->GetElementSize() * SourceProperty->ArrayDim;
 			return true;
 		}
 		else
@@ -1660,7 +1660,7 @@ bool FStateTreePropertyPath::ResolveIndirectionsWithValue(const FStateTreeDataVi
 				return false;
 			}
 			ArrayIndex = FMath::Max(0, Segment->GetArrayIndex());
-			Offset = Property->GetOffset_ForInternal() + Property->ElementSize * ArrayIndex;
+			Offset = Property->GetOffset_ForInternal() + Property->GetElementSize() * ArrayIndex;
 		}
 
 		FStateTreePropertyPathIndirection& Indirection = OutIndirections.AddDefaulted_GetRef();
diff --git a/Engine/Plugins/VirtualProduction/RemoteControl/Source/RemoteControl/Private/RemoteControlProtocolBinding.cpp b/Engine/Plugins/VirtualProduction/RemoteControl/Source/RemoteControl/Private/RemoteControlProtocolBinding.cpp
--- a/Engine/Plugins/VirtualProduction/RemoteControl/Source/RemoteControl/Private/RemoteControlProtocolBinding.cpp
+++ b/Engine/Plugins/VirtualProduction/RemoteControl/Source/RemoteControl/Private/RemoteControlProtocolBinding.cpp
@@ -434,7 +434,7 @@ FRemoteControlProtocolMapping::FRemoteControlProtocolMapping(FProperty* InProper
 	}
 	else if (FNumericProperty* NumericProperty = CastField<FNumericProperty>(InProperty))
 	{
-		InterpolationMappingPropertyData.AddZeroed(NumericProperty->ElementSize);
+		InterpolationMappingPropertyData.AddZeroed(NumericProperty->GetElementSize());
 	}
 	else if (FStructProperty* StructProperty = CastField<FStructProperty>(InProperty))
 	{
diff --git a/Engine/Plugins/VirtualProduction/RemoteControl/Source/RemoteControl/Public/RemoteControlProtocolBinding.h b/Engine/Plugins/VirtualProduction/RemoteControl/Source/RemoteControl/Public/RemoteControlProtocolBinding.h
--- a/Engine/Plugins/VirtualProduction/RemoteControl/Source/RemoteControl/Public/RemoteControlProtocolBinding.h
+++ b/Engine/Plugins/VirtualProduction/RemoteControl/Source/RemoteControl/Public/RemoteControlProtocolBinding.h
@@ -339,13 +339,13 @@ inline void FRemoteControlProtocolMapping::SetMappingValueAsPrimitive<FText>(FTe
 template <>
 inline bool FRemoteControlProtocolMapping::PropertySizeMatchesData<FBoolProperty>(const TArray<uint8>& InSource, const FBoolProperty* InProperty)
 {
-	return ensure(InProperty->ElementSize == InSource.Num());
+	return ensure(InProperty->GetElementSize() == InSource.Num());
 }
 
 template <>
 inline bool FRemoteControlProtocolMapping::PropertySizeMatchesData<FNumericProperty>(const TArray<uint8>& InSource, const FNumericProperty* InProperty)
 {
-	return ensure(InProperty->ElementSize >= InSource.Num());
+	return ensure(InProperty->GetElementSize() >= InSource.Num());
 }
 
 template <>
diff --git a/Engine/Plugins/VirtualProduction/RemoteControl/Source/RemoteControlCommon/Private/RCTypeUtilities.cpp b/Engine/Plugins/VirtualProduction/RemoteControl/Source/RemoteControlCommon/Private/RCTypeUtilities.cpp
--- a/Engine/Plugins/VirtualProduction/RemoteControl/Source/RemoteControlCommon/Private/RCTypeUtilities.cpp
+++ b/Engine/Plugins/VirtualProduction/RemoteControl/Source/RemoteControlCommon/Private/RCTypeUtilities.cpp
@@ -12,7 +12,7 @@ SIZE_T RemoteControlTypeUtilities::GetPropertySize(const FProperty* InProperty,
 {
 	check(InProperty);
 
-	const int32 ElementSize = InProperty->ElementSize;
+	const int32 ElementSize = InProperty->GetElementSize();
 	if (const FArrayProperty* ArrayProperty = CastField<FArrayProperty>(InProperty))
 	{
 		return FScriptArrayHelper(ArrayProperty, InData).Num() * ElementSize;
diff --git a/Engine/Plugins/VirtualProduction/RemoteControl/Source/RemoteControlCommon/Public/RCPropertyUtilities.inl b/Engine/Plugins/VirtualProduction/RemoteControl/Source/RemoteControlCommon/Public/RCPropertyUtilities.inl
--- a/Engine/Plugins/VirtualProduction/RemoteControl/Source/RemoteControlCommon/Public/RCPropertyUtilities.inl
+++ b/Engine/Plugins/VirtualProduction/RemoteControl/Source/RemoteControlCommon/Public/RCPropertyUtilities.inl
@@ -228,7 +228,7 @@ namespace RemoteControlPropertyUtilities
 		const int32 SrcSize = InSrc.Size();
 
 		// The stored data size doesn't match, so cast
-		if(OutDst.GetProperty()->ElementSize != SrcSize)
+		if(OutDst.GetProperty()->GetElementSize() != SrcSize)
 		{
 			if(const FNumericProperty* DstProperty = OutDst.GetProperty<FNumericProperty>())
 			{
diff --git a/Engine/Source/Editor/EditorConfig/Private/EditorConfig.cpp b/Engine/Source/Editor/EditorConfig/Private/EditorConfig.cpp
--- a/Engine/Source/Editor/EditorConfig/Private/EditorConfig.cpp
+++ b/Engine/Source/Editor/EditorConfig/Private/EditorConfig.cpp
@@ -368,7 +368,7 @@ void FEditorConfig::ReadValue(const TSharedPtr<FJsonValue>& JsonValue, const FPr
 
 			// temporary buffer to read elements into
 			TArray<uint8> TempBuffer;
-			TempBuffer.AddUninitialized(InnerProperty->ElementSize);
+			TempBuffer.AddUninitialized(InnerProperty->GetElementSize());
 
 			for (int32 Idx = 0; Idx < SetJsonValue->Num(); ++Idx)
 			{
@@ -400,10 +400,10 @@ void FEditorConfig::ReadValue(const TSharedPtr<FJsonValue>& JsonValue, const FPr
 
 			// temporary buffers to read elements into
 			TArray<uint8> TempKey;
-			TempKey.AddZeroed(KeyProperty->ElementSize);
+			TempKey.AddZeroed(KeyProperty->GetElementSize());
 
 			TArray<uint8> TempValue;
-			TempValue.AddZeroed(ValueProperty->ElementSize);
+			TempValue.AddZeroed(ValueProperty->GetElementSize());
 
 			for (const TPair<FString, TSharedPtr<FJsonValue>>& JsonPair : (*JsonObjectValue)->Values)
 			{
@@ -430,10 +430,10 @@ void FEditorConfig::ReadValue(const TSharedPtr<FJsonValue>& JsonValue, const FPr
 
 			// temporary buffers to read elements into
 			TArray<uint8> TempKey;
-			TempKey.AddUninitialized(KeyProperty->ElementSize);
+			TempKey.AddUninitialized(KeyProperty->GetElementSize());
 
 			TArray<uint8> TempValue;
-			TempValue.AddUninitialized(ValueProperty->ElementSize);
+			TempValue.AddUninitialized(ValueProperty->GetElementSize());
 
 			for (const TSharedPtr<FJsonValue>& JsonElement : *JsonArrayPtr)
 			{
diff --git a/Engine/Source/Editor/Kismet/Private/FindInBlueprintManager.cpp b/Engine/Source/Editor/Kismet/Private/FindInBlueprintManager.cpp
--- a/Engine/Source/Editor/Kismet/Private/FindInBlueprintManager.cpp
+++ b/Engine/Source/Editor/Kismet/Private/FindInBlueprintManager.cpp
@@ -931,7 +931,7 @@ namespace BlueprintSearchMetaDataHelpers
 							TArray< TSharedPtr<FJsonValue> > Array;
 							for (int Index = 0; Index != Property->ArrayDim; ++Index)
 							{
-								GatherSearchablesFromProperty(InWriter, Property, (char*)Value + Index * Property->ElementSize, InStruct);
+								GatherSearchablesFromProperty(InWriter, Property, (char*)Value + Index * Property->GetElementSize(), InStruct);
 							}
 						}
 					}
@@ -1055,7 +1055,7 @@ namespace BlueprintSearchMetaDataHelpers
 						TArray< TSharedPtr<FJsonValue> > Array;
 						for (int Index = 0; Index != Property->ArrayDim; ++Index)
 						{
-							CacheSubPropertySearchables(InOutCachePropertyMapping, Property, (char*)Value + Index * Property->ElementSize, InStruct);
+							CacheSubPropertySearchables(InOutCachePropertyMapping, Property, (char*)Value + Index * Property->GetElementSize(), InStruct);
 						}
 					}
 					SearchableProperties.Add(MoveTemp(SearchableProperty));
diff --git a/Engine/Source/Editor/KismetCompiler/Private/KismetCompiler.cpp b/Engine/Source/Editor/KismetCompiler/Private/KismetCompiler.cpp
--- a/Engine/Source/Editor/KismetCompiler/Private/KismetCompiler.cpp
+++ b/Engine/Source/Editor/KismetCompiler/Private/KismetCompiler.cpp
@@ -1826,7 +1826,7 @@ void FKismetCompilerContext::PrintVerboseInfoStruct(UStruct* Struct) const
 				FText::FromString(Prop->GetClass()->GetDescription()),
 				FText::FromString(Prop->GetName()),
 				Prop->GetOffset_ForDebug(),
-				Prop->ElementSize,
+				Prop->GetElementSize(),
 				Prop->ArrayDim,
 				FText::FromString(FString::Printf(TEXT("%" UINT64_x_FMT), (uint64)Prop->PropertyFlags))
 			).ToString()
diff --git a/Engine/Source/Editor/PropertyEditor/Private/ItemPropertyNode.cpp b/Engine/Source/Editor/PropertyEditor/Private/ItemPropertyNode.cpp
--- a/Engine/Source/Editor/PropertyEditor/Private/ItemPropertyNode.cpp
+++ b/Engine/Source/Editor/PropertyEditor/Private/ItemPropertyNode.cpp
@@ -218,7 +218,7 @@ void FItemPropertyNode::InitChildNodes()
 				FPropertyNodeInitParams InitParams;
 				InitParams.ParentNode = SharedThis(this);
 				InitParams.Property = MyProperty;
-				InitParams.ArrayOffset = Index*MyProperty->ElementSize;
+				InitParams.ArrayOffset = Index*MyProperty->GetElementSize();
 				InitParams.ArrayIndex = Index;
 				InitParams.bAllowChildren = true;
 				InitParams.bForceHiddenPropertyVisibility = bShouldShowHiddenProperties;
@@ -248,7 +248,7 @@ void FItemPropertyNode::InitChildNodes()
 				FPropertyNodeInitParams InitParams;
 				InitParams.ParentNode = SharedThis(this);
 				InitParams.Property = ArrayProperty->Inner;
-				InitParams.ArrayOffset = Index * ArrayProperty->Inner->ElementSize;
+				InitParams.ArrayOffset = Index * ArrayProperty->Inner->GetElementSize();
 				InitParams.ArrayIndex = Index;
 				InitParams.bAllowChildren = true;
 				InitParams.bForceHiddenPropertyVisibility = bShouldShowHiddenProperties;
diff --git a/Engine/Source/Editor/PropertyEditor/Private/PropertyChangeListener.cpp b/Engine/Source/Editor/PropertyEditor/Private/PropertyChangeListener.cpp
--- a/Engine/Source/Editor/PropertyEditor/Private/PropertyChangeListener.cpp
+++ b/Engine/Source/Editor/PropertyEditor/Private/PropertyChangeListener.cpp
@@ -77,13 +77,13 @@ public:
 			if( PropertyNodeRef.GetArrayIndex() == INDEX_NONE && Property->ArrayDim > 1 )
 			{
 				// Check static arrays
-				Data.AddZeroed( Property->ArrayDim * Property->ElementSize );
+				Data.AddZeroed( Property->ArrayDim * Property->GetElementSize() );
 				Property->CopyCompleteValue( Data.GetData(), PropertyValueAddresses.Address);
 			}
 			else
 			{
 				// Regular properties
-				Data.AddZeroed( Property->ElementSize );
+				Data.AddZeroed( Property->GetElementSize() );
 				Property->CopySingleValue( Data.GetData(), PropertyValueAddresses.Address );
 			}
 		}
diff --git a/Engine/Source/Editor/PropertyEditor/Private/PropertyEditorModule.cpp b/Engine/Source/Editor/PropertyEditor/Private/PropertyEditorModule.cpp
--- a/Engine/Source/Editor/PropertyEditor/Private/PropertyEditorModule.cpp
+++ b/Engine/Source/Editor/PropertyEditor/Private/PropertyEditorModule.cpp
@@ -424,7 +424,7 @@ FStructProperty* FPropertyEditorModule::RegisterStructProperty(const UStruct* St
 		UScriptStruct* InnerStruct = CastChecked<UScriptStruct>(const_cast<UStruct*>(StructClass));
 		StructProperty = new FStructProperty(StructOnScopePropertyOwner, *MakeUniqueObjectName(StructOnScopePropertyOwner, UField::StaticClass(), InnerStruct->GetFName()).ToString(), RF_Transient);
 		StructProperty->Struct = InnerStruct;
-		StructProperty->ElementSize = StructClass->GetStructureSize();
+		StructProperty->SetElementSize(StructClass->GetStructureSize());
 		StructOnScopePropertyOwner->AddCppProperty(StructProperty);
 
 		RegisteredStructToProxyMap.Add(StructName, StructProperty);
diff --git a/Engine/Source/Editor/PropertyEditor/Private/PropertyHandleImpl.cpp b/Engine/Source/Editor/PropertyEditor/Private/PropertyHandleImpl.cpp
--- a/Engine/Source/Editor/PropertyEditor/Private/PropertyHandleImpl.cpp
+++ b/Engine/Source/Editor/PropertyEditor/Private/PropertyHandleImpl.cpp
@@ -124,7 +124,7 @@ FPropertyAccess::Result FPropertyValueImpl::GetValueData( void*& OutAddress ) co
 			if (ValueAddress && Property)
 			{
 				const int32 Index = 0;
-				OutAddress = ValueAddress + Index * Property->ElementSize;
+				OutAddress = ValueAddress + Index * Property->GetElementSize();
 				Res = FPropertyAccess::Success;
 			}
 		}
diff --git a/Engine/Source/Editor/PropertyEditor/Private/PropertyNode.cpp b/Engine/Source/Editor/PropertyEditor/Private/PropertyNode.cpp
--- a/Engine/Source/Editor/PropertyEditor/Private/PropertyNode.cpp
+++ b/Engine/Source/Editor/PropertyEditor/Private/PropertyNode.cpp
@@ -981,7 +981,7 @@ FPropertyAccess::Result FPropertyNodeEditStack::InitializeInternal(const FProper
 		if (Property == ParentProperty) // Static array items
 		{
 			// Static array property node creates subnodes that point to individual array items
-			MemoryStack.Add(FMemoryFrame(Property, MemoryStack.Last().Memory + InNode->GetArrayIndex() * Property->ElementSize));
+			MemoryStack.Add(FMemoryFrame(Property, MemoryStack.Last().Memory + InNode->GetArrayIndex() * Property->GetElementSize()));
 		}
 		else if (const FStructProperty* StructProp = CastField<FStructProperty>(ParentProperty)) // structs
 		{
@@ -2140,7 +2140,7 @@ struct FPropertyItemComponentCollector
 			// either the associated property is not an array property, or it's the header for the property (meaning the entire array)
 			for ( int32 ArrayIndex = 0; ArrayIndex < Prop->ArrayDim; ArrayIndex++ )
 			{
-				ProcessProperty(Prop, ValueTracker.GetPropertyValueAddress() + ArrayIndex * Prop->ElementSize);
+				ProcessProperty(Prop, ValueTracker.GetPropertyValueAddress() + ArrayIndex * Prop->GetElementSize());
 			}
 		}
 		else
@@ -2218,7 +2218,7 @@ private:
 			int32 ArraySize = ArrayHelper.Num();
 			for ( int32 ArrayIndex = 0; ArrayIndex < ArraySize; ArrayIndex++ )
 			{
-				ProcessProperty(ArrayProp->Inner, ArrayValue + ArrayIndex * ArrayProp->Inner->ElementSize);
+				ProcessProperty(ArrayProp->Inner, ArrayValue + ArrayIndex * ArrayProp->Inner->GetElementSize());
 			}
 
 			bResult = true;
@@ -2243,7 +2243,7 @@ private:
 		{
 			FScriptSet* SetValuePtr = SetProp->GetPropertyValuePtr(PropertyValueAddress);
 
-			FScriptSetLayout SetLayout = SetValuePtr->GetScriptLayout(SetProp->ElementProp->ElementSize, SetProp->ElementProp->GetMinAlignment());
+			FScriptSetLayout SetLayout = SetValuePtr->GetScriptLayout(SetProp->ElementProp->GetElementSize(), SetProp->ElementProp->GetMinAlignment());
 			int32 ItemsLeft = SetValuePtr->Num();
 
 			for (int32 Index = 0; ItemsLeft > 0; ++Index)
diff --git a/Engine/Source/Editor/UnrealEd/Private/DiffUtils.cpp b/Engine/Source/Editor/UnrealEd/Private/DiffUtils.cpp
--- a/Engine/Source/Editor/UnrealEd/Private/DiffUtils.cpp
+++ b/Engine/Source/Editor/UnrealEd/Private/DiffUtils.cpp
@@ -818,7 +818,7 @@ static void IdenticalHelper(const FProperty* AProperty, const FProperty* BProper
 		// our container (e.g. the struct or class or dynamic array) has already done the initial offset calculation:
 		for (int32 I = 0; I < AProperty->ArrayDim; ++I)
 		{
-			int32 Offset = AProperty->ElementSize * I;
+			int32 Offset = AProperty->GetElementSize() * I;
 
 			DiffParameters.RootPath = FPropertySoftPath(DiffParameters.RootPath, I);
 			DiffParameters.bShouldDiffArrayElements = false;
diff --git a/Engine/Source/Editor/UnrealEd/Private/MergeUtils.cpp b/Engine/Source/Editor/UnrealEd/Private/MergeUtils.cpp
--- a/Engine/Source/Editor/UnrealEd/Private/MergeUtils.cpp
+++ b/Engine/Source/Editor/UnrealEd/Private/MergeUtils.cpp
@@ -326,7 +326,7 @@ struct FPropertyInstance
 		{
 			const void* DataA = ValProperty->ContainerPtrToValuePtr<void*>(Object);
 			const void* DataB = Other.ValProperty->ContainerPtrToValuePtr<void*>(Other.Object);
-			check(ValProperty->ElementSize == Other.ValProperty->ElementSize);
+			check(ValProperty->GetElementSize() == Other.ValProperty->GetElementSize());
 			check(ValProperty->GetClass() == Other.ValProperty->GetClass());
 			
 			return ValProperty->Identical(DataA, DataB, PPF_DeepComparison);
diff --git a/Engine/Source/Runtime/CoreUObject/Private/Internationalization/StabilizeLocalizationKeys.cpp b/Engine/Source/Runtime/CoreUObject/Private/Internationalization/StabilizeLocalizationKeys.cpp
--- a/Engine/Source/Runtime/CoreUObject/Private/Internationalization/StabilizeLocalizationKeys.cpp
+++ b/Engine/Source/Runtime/CoreUObject/Private/Internationalization/StabilizeLocalizationKeys.cpp
@@ -34,7 +34,7 @@ void StabilizeLocalizationKeys::StabilizeLocalizationKeysForProperty(FProperty*
 	{
 		for (int32 ArrIndex = 0; ArrIndex < TextProp->ArrayDim; ++ArrIndex)
 		{
-			void* PropValueData = ((uint8*)InPropData) + (TextProp->ElementSize * ArrIndex);
+			void* PropValueData = ((uint8*)InPropData) + (TextProp->GetElementSize() * ArrIndex);
 
 			FText* TextValuePtr = TextProp->GetPropertyValuePtr(PropValueData);
 			check(TextValuePtr);
@@ -58,7 +58,7 @@ void StabilizeLocalizationKeys::StabilizeLocalizationKeysForProperty(FProperty*
 	{
 		for (int32 ArrIndex = 0; ArrIndex < StructProp->ArrayDim; ++ArrIndex)
 		{
-			void* PropValueData = ((uint8*)InPropData) + (StructProp->ElementSize * ArrIndex);
+			void* PropValueData = ((uint8*)InPropData) + (StructProp->GetElementSize() * ArrIndex);
 
 			if (StructProp->ArrayDim > 1)
 			{
diff --git a/Engine/Source/Runtime/CoreUObject/Private/Serialization/PropertyLocalizationDataGathering.cpp b/Engine/Source/Runtime/CoreUObject/Private/Serialization/PropertyLocalizationDataGathering.cpp
--- a/Engine/Source/Runtime/CoreUObject/Private/Serialization/PropertyLocalizationDataGathering.cpp
+++ b/Engine/Source/Runtime/CoreUObject/Private/Serialization/PropertyLocalizationDataGathering.cpp
@@ -385,8 +385,8 @@ void FPropertyLocalizationDataGatherer::GatherLocalizationDataFromChildTextPrope
 		}
 
 		const FString& PathToElement = IsFixedSizeArray ? PathToFixedSizeArrayElement : PathToElementRoot;
-		const void* const ElementValueAddress = reinterpret_cast<const uint8*>(ValueAddress) + Property->ElementSize * i;
-		const void* const DefaultElementValueAddress = DefaultValueAddress ? (reinterpret_cast<const uint8*>(DefaultValueAddress) + Property->ElementSize * i) : nullptr;
+		const void* const ElementValueAddress = reinterpret_cast<const uint8*>(ValueAddress) + Property->GetElementSize() * i;
+		const void* const DefaultElementValueAddress = DefaultValueAddress ? (reinterpret_cast<const uint8*>(DefaultValueAddress) + Property->GetElementSize() * i) : nullptr;
 
 		EPropertyLocalizationGathererTextFlags ElementChildPropertyGatherTextFlags = FixedChildPropertyGatherTextFlags;
 		if (!EnumHasAnyFlags(ElementChildPropertyGatherTextFlags, EPropertyLocalizationGathererTextFlags::ForceIsDefaultValue))
diff --git a/Engine/Source/Runtime/CoreUObject/Private/Serialization/UnversionedPropertySerialization.cpp b/Engine/Source/Runtime/CoreUObject/Private/Serialization/UnversionedPropertySerialization.cpp
--- a/Engine/Source/Runtime/CoreUObject/Private/Serialization/UnversionedPropertySerialization.cpp
+++ b/Engine/Source/Runtime/CoreUObject/Private/Serialization/UnversionedPropertySerialization.cpp
@@ -52,7 +52,7 @@ public:
 	FUnversionedPropertySerializer(FProperty* InProperty, int32 InArrayIndex)
 		: Property(InProperty)
 #if CACHE_UNVERSIONED_PROPERTY_SCHEMA
-		, Offset(Property->GetOffset_ForInternal() + Property->ElementSize * InArrayIndex)
+		, Offset(Property->GetOffset_ForInternal() + Property->GetElementSize() * InArrayIndex)
 		, bSerializeAsInteger(CanSerializeAsInteger(Property))
 		, bIsOptional(IsOptional(Property->GetClass()->GetCastFlags()))
 		, IntType(GetIntType(Property->GetMinAlignment()))
@@ -126,8 +126,8 @@ public:
 		// Cached FastZeroIntNum is only uint8 and not sufficient for large unset optionals
 		if (FastZeroIntNum == 0)
 		{
-			checkf(bIsOptional && Property->ElementSize >= 256, TEXT("Only large unset optionals should hit this loading path"));
-			FMemory::Memzero(ValueData, Property->ElementSize);
+			checkf(bIsOptional && Property->GetElementSize() >= 256, TEXT("Only large unset optionals should hit this loading path"));
+			FMemory::Memzero(ValueData, Property->GetElementSize());
 			return;
 		}
 #else
@@ -182,7 +182,7 @@ private:
 
 	static uint32 GetIntNum(const FProperty* Property, EIntegerType IntType)
 	{
-		return Property->ElementSize / GetSizeOf(IntType);
+		return Property->GetElementSize() / GetSizeOf(IntType);
 	}
 
 	static bool CanSerializeAsZero(const FProperty* Property, EIntegerType IntType)
diff --git a/Engine/Source/Runtime/CoreUObject/Private/UObject/Class.cpp b/Engine/Source/Runtime/CoreUObject/Private/UObject/Class.cpp
--- a/Engine/Source/Runtime/CoreUObject/Private/UObject/Class.cpp
+++ b/Engine/Source/Runtime/CoreUObject/Private/UObject/Class.cpp
@@ -2837,7 +2837,7 @@ bool FindConstructorUninitialized(UStruct* BaseClass,uint8* Data,uint8* Defaults
 		}
 		else if(PB)
 		{
-			check(Size == PB->ElementSize);
+			check(Size == PB->GetElementSize());
 			if( PB->GetPropertyValue_InContainer(Data) && !PB->GetPropertyValue_InContainer(Defaults) )
 			{
 				bProblem = true;
diff --git a/Engine/Source/Runtime/CoreUObject/Private/UObject/EnumProperty.cpp b/Engine/Source/Runtime/CoreUObject/Private/UObject/EnumProperty.cpp
--- a/Engine/Source/Runtime/CoreUObject/Private/UObject/EnumProperty.cpp
+++ b/Engine/Source/Runtime/CoreUObject/Private/UObject/EnumProperty.cpp
@@ -468,7 +468,7 @@ void FEnumProperty::LinkInternal(FArchive& Ar)
 
 	UnderlyingProp->Link(Ar);
 
-	this->ElementSize = UnderlyingProp->ElementSize;
+	this->SetElementSize(UnderlyingProp->GetElementSize());
 	this->PropertyFlags |= CPF_IsPlainOldData | CPF_NoDestructor | CPF_ZeroConstructor;
 
 	PropertyFlags |= (UnderlyingProp->PropertyFlags & CPF_HasGetValueTypeHash);
@@ -630,7 +630,7 @@ void FEnumProperty::GetInnerFields(TArray<FField*>& OutFields)
 
 uint64 FEnumProperty::GetMaxNetSerializeBits() const
 {
-	const uint64 MaxBits = ElementSize * 8;
+	const uint64 MaxBits = GetElementSize() * 8;
 	const uint64 DesiredBits = FMath::CeilLogTwo64(Enum->GetMaxEnumValue() + 1);
 	
 	return FMath::Min(DesiredBits, MaxBits);
diff --git a/Engine/Source/Runtime/CoreUObject/Private/UObject/GarbageCollection.cpp b/Engine/Source/Runtime/CoreUObject/Private/UObject/GarbageCollection.cpp
--- a/Engine/Source/Runtime/CoreUObject/Private/UObject/GarbageCollection.cpp
+++ b/Engine/Source/Runtime/CoreUObject/Private/UObject/GarbageCollection.cpp
@@ -6384,7 +6384,7 @@ void FArrayProperty::EmitReferenceInfo(UE::GC::FSchemaBuilder& Schema, int32 Bas
 	{
 		bool bUsesFreezableAllocator = EnumHasAnyFlags(ArrayFlags, EArrayPropertyFlags::UsesMemoryImageAllocator);
 		EMemberType Type = bUsesFreezableAllocator ?  EMemberType::FreezableStructArray : EMemberType::StructArray;
-		FSchemaBuilder InnerSchema(Inner->ElementSize);
+		FSchemaBuilder InnerSchema(Inner->GetElementSize());
 
 		// Structs and nested arrays share the same implementation on the Garbage Collector side
 		// as arrays of structs already push the array memory into the GC stack and process its tokens
@@ -6453,7 +6453,7 @@ void FStructProperty::EmitReferenceInfo(UE::GC::FSchemaBuilder& Schema, int32 Ba
 		StructAROFn StructARO = Struct->GetCppStructOps()->AddStructReferencedObjects();	
 		for (int32 Idx = 0, Num = ArrayDim; Idx < Num; ++Idx)
 		{	
-			Schema.Add(UE::GC::DeclareMember(DebugPath, Offset + Idx * ElementSize, EMemberType::MemberARO, StructARO) );
+			Schema.Add(UE::GC::DeclareMember(DebugPath, Offset + Idx * GetElementSize(), EMemberType::MemberARO, StructARO) );
 		}
 	}
 
@@ -6477,7 +6477,7 @@ void FStructProperty::EmitReferenceInfo(UE::GC::FSchemaBuilder& Schema, int32 Ba
 		{
 			for (FProperty* Property = Struct->PropertyLink; Property; Property = Property->PropertyLinkNext)
 			{
-				Property->EmitReferenceInfo(Schema, Offset + Idx * ElementSize, EncounteredStructProps, DebugPath);
+				Property->EmitReferenceInfo(Schema, Offset + Idx * GetElementSize(), EncounteredStructProps, DebugPath);
 			}
 		}
 	}
diff --git a/Engine/Source/Runtime/CoreUObject/Private/UObject/Obj.cpp b/Engine/Source/Runtime/CoreUObject/Private/UObject/Obj.cpp
--- a/Engine/Source/Runtime/CoreUObject/Private/UObject/Obj.cpp
+++ b/Engine/Source/Runtime/CoreUObject/Private/UObject/Obj.cpp
@@ -3137,7 +3137,7 @@ void UObject::LoadConfig(UClass* ConfigClass/*=nullptr*/, const TCHAR* InFilenam
 					const FName KeyName(*Key, FNAME_Find);
 					Sec->MultiFind(KeyName, List);
 
-					const int32 Size = Array->Inner->ElementSize;
+					const int32 Size = Array->Inner->GetElementSize();
 
 					// Only override default properties if there is something to override them with.
 					if (!List.IsEmpty())
@@ -4353,7 +4353,7 @@ bool StaticExec( UWorld* InWorld, const TCHAR* Cmd, FOutputDevice& Ar )
 			{
 				FProperty* Property = *It;
 
-				Ar.Logf(TEXT("    Prop %s"), *FString::Printf(TEXT("%s at offset %d; %dx %d bytes of type %s"), *Property->GetName(), Property->GetOffset_ForDebug(), Property->ArrayDim, Property->ElementSize, *Property->GetClass()->GetName()));
+				Ar.Logf(TEXT("    Prop %s"), *FString::Printf(TEXT("%s at offset %d; %dx %d bytes of type %s"), *Property->GetName(), Property->GetOffset_ForDebug(), Property->ArrayDim, Property->GetElementSize(), *Property->GetClass()->GetName()));
 
 				for (const TCHAR* Flag : ParsePropertyFlags(Property->PropertyFlags))
 				{
@@ -4523,7 +4523,7 @@ bool StaticExec( UWorld* InWorld, const TCHAR* Cmd, FOutputDevice& Ar )
 										ExportProperty = ArrayProp->Inner;
 									}
 
-									int32 ElementSize = ExportProperty->ElementSize;
+									int32 ElementSize = ExportProperty->GetElementSize();
 									for ( int32 ArrayIndex = 0; ArrayIndex < ElementCount; ArrayIndex++ )
 									{
 										FString ResultStr;
@@ -4648,7 +4648,7 @@ bool StaticExec( UWorld* InWorld, const TCHAR* Cmd, FOutputDevice& Ar )
 						{
 							FProperty* Property = *It;
 							
-							Ar.Logf(TEXT("    Parameter %s"), *FString::Printf(TEXT("%s at offset %d; %dx %d bytes of type %s"), *Property->GetName(), Property->GetOffset_ForDebug(), Property->ArrayDim, Property->ElementSize, *Property->GetClass()->GetName()));
+							Ar.Logf(TEXT("    Parameter %s"), *FString::Printf(TEXT("%s at offset %d; %dx %d bytes of type %s"), *Property->GetName(), Property->GetOffset_ForDebug(), Property->ArrayDim, Property->GetElementSize(), *Property->GetClass()->GetName()));
 
 							for (const TCHAR* Flag : ParsePropertyFlags(Property->PropertyFlags))
 							{
@@ -4666,7 +4666,7 @@ bool StaticExec( UWorld* InWorld, const TCHAR* Cmd, FOutputDevice& Ar )
 						{							
 							FProperty* Property = *It;
 
-							Ar.Logf(TEXT("    Local %s"), *FString::Printf(TEXT("%s at offset %d; %dx %d bytes of type %s"), *Property->GetName(), Property->GetOffset_ForDebug(), Property->ArrayDim, Property->ElementSize, *Property->GetClass()->GetName()));
+							Ar.Logf(TEXT("    Local %s"), *FString::Printf(TEXT("%s at offset %d; %dx %d bytes of type %s"), *Property->GetName(), Property->GetOffset_ForDebug(), Property->ArrayDim, Property->GetElementSize(), *Property->GetClass()->GetName()));
 
 							for (const TCHAR* Flag : ParsePropertyFlags(Property->PropertyFlags))
 							{
diff --git a/Engine/Source/Runtime/CoreUObject/Private/UObject/Property.cpp b/Engine/Source/Runtime/CoreUObject/Private/UObject/Property.cpp
--- a/Engine/Source/Runtime/CoreUObject/Private/UObject/Property.cpp
+++ b/Engine/Source/Runtime/CoreUObject/Private/UObject/Property.cpp
@@ -2,6 +2,7 @@
 
 #include "UObject/UnrealType.h"
 
+#include "HAL/IConsoleManager.h"
 #include "Hash/Blake3.h"
 #include "Math/Box2D.h"
 #include "Math/InterpCurvePoint.h"
@@ -28,9 +29,16 @@
 #include "UObject/SoftObjectPath.h"
 #include "UObject/UnrealTypePrivate.h"
 #include "UObject/UObjectGlobals.h"
+#include "UObject/UObjectIterator.h"
 
 DEFINE_LOG_CATEGORY(LogProperty);
 
+namespace UE::CoreUObject::Private
+{
+	static int32 EnsureAgainstLargeProperties = 1;
+	static FAutoConsoleVariableRef CVarEnsureAgainstLargeProperties(TEXT("CoreUObject.EnsureAgainstLargeProperties"), EnsureAgainstLargeProperties, TEXT("Ensure (warn) against properties that could possibly break in future versions of the Engine due a reduction in the max size of FProperty::ElementSize."));
+}
+
 // List the core ones here as they have already been included (and can be used without CoreUObject!)
 template<typename T>
 struct TVector3StructOpsTypeTraits : public TStructOpsTypeTraitsBase2<T>
@@ -775,6 +783,13 @@ const TCHAR* FPropertyHelpers::ReadToken( const TCHAR* Buffer, FStringBuilderBas
 
 IMPLEMENT_FIELD(FProperty)
 
+void FProperty::SetElementSize(int32 NewSize)
+{
+PRAGMA_DISABLE_DEPRECATION_WARNINGS
+	ElementSize = NewSize;
+PRAGMA_ENABLE_DEPRECATION_WARNINGS
+}
+
 //
 // Constructors.
 //
@@ -837,7 +852,7 @@ FProperty::FProperty(UField* InField)
 {
 	UProperty* SourceProperty = CastChecked<UProperty>(InField);
 	ArrayDim = SourceProperty->ArrayDim;
-	ElementSize = SourceProperty->ElementSize;
+	SetElementSize(SourceProperty->ElementSize);
 	PropertyFlags = SourceProperty->PropertyFlags;
 	RepIndex = SourceProperty->RepIndex;
 	Offset_Internal = SourceProperty->Offset_Internal;
@@ -879,14 +894,22 @@ void FProperty::Serialize( FArchive& Ar )
 	Super::Serialize(Ar);
 
 	Ar << ArrayDim;
+
+PRAGMA_DISABLE_DEPRECATION_WARNINGS
 	Ar << ElementSize;
 
+	const bool EnsureAgainstLargeProperties = UE::CoreUObject::Private::EnsureAgainstLargeProperties > 0;
+	constexpr int32 ExpectedMaxSize = (1 << 24);
+	ensureMsgf(!EnsureAgainstLargeProperties || GetElementSize() < ExpectedMaxSize, TEXT("%s has ElementSize %d which will violate an upcoming change to lower the max ElementSize.  Consider breaking up the property. Disable this warning with CoreUObject.EnsureAgainstLargeProperties 0"), *GetName(), GetElementSize());
+PRAGMA_ENABLE_DEPRECATION_WARNINGS
+
 	EPropertyFlags SaveFlags = PropertyFlags & ~CPF_ComputedFlags;
 	// Archive the basic info.
 	Ar << (uint64&)SaveFlags;
 	if (Ar.IsLoading())
 	{
 		PropertyFlags = (SaveFlags & ~CPF_ComputedFlags) | (PropertyFlags & CPF_ComputedFlags);
+
 	}
 
 	if (FPlatformProperties::HasEditorOnlyData() == false)
@@ -911,7 +934,7 @@ void FProperty::PostDuplicate(const FField& InField)
 {
 	const FProperty& Source = static_cast<const FProperty&>(InField);
 	ArrayDim = Source.ArrayDim;
-	ElementSize = Source.ElementSize;
+	SetElementSize(Source.GetElementSize());
 	PropertyFlags = Source.PropertyFlags;
 	RepIndex = Source.RepIndex;
 	Offset_Internal = Source.Offset_Internal;
@@ -1468,7 +1491,7 @@ void* FProperty::GetValueAddressAtIndex_Direct(const FProperty* Inner, void* InV
 {
 	checkf(Inner == nullptr, TEXT("%s should not have an inner property or it's missing specialized GetValueAddressAtIndex_Direct override"), *GetFullName());
 	checkf(Index < ArrayDim && Index >= 0, TEXT("Array index (%d) out of range"), Index);
-	return (uint8*)InValueAddress + ElementSize * Index;
+	return (uint8*)InValueAddress + GetElementSize() * Index;
 }
 
 void FProperty::SetSingleValue_InContainer(void* OutContainer, const void* InValue, int32 ArrayIndex) const
@@ -1492,7 +1515,7 @@ void FProperty::SetSingleValue_InContainer(void* OutContainer, const void* InVal
 			uint8* ValueArray = (uint8*)AllocateAndInitializeValue();
 			GetValue_InContainer(OutContainer, ValueArray);
 			// Replace the value at the specified index in the temp array with the InValue
-			CopySingleValue(ValueArray + ArrayIndex * ElementSize, InValue);
+			CopySingleValue(ValueArray + ArrayIndex * GetElementSize(), InValue);
 			// Now call a setter to replace the entire array and then destroy the temp value
 			CallSetter(OutContainer, ValueArray);
 			DestroyAndFreeValue(ValueArray);
@@ -1521,7 +1544,7 @@ void FProperty::GetSingleValue_InContainer(const void* InContainer, void* OutVal
 			uint8* ValueArray = (uint8*)AllocateAndInitializeValue();
 			GetValue_InContainer(InContainer, ValueArray);
 			// Copy the item we care about and free the temp array
-			CopySingleValue(OutValue, ValueArray + ArrayIndex * ElementSize);
+			CopySingleValue(OutValue, ValueArray + ArrayIndex * GetElementSize());
 			DestroyAndFreeValue(ValueArray);
 		}
 	}
@@ -1905,7 +1928,7 @@ const TCHAR* FProperty::ImportSingleProperty( const TCHAR* Str, void* DestData,
 				}
 				else
 				{
-					int32 Size = ArrayProperty->Inner->ElementSize;
+					int32 Size = ArrayProperty->Inner->GetElementSize();
 
 					uint8* Temp = (uint8*)FMemory_Alloca(Size);
 					ArrayProperty->Inner->InitializeValue(Temp);
diff --git a/Engine/Source/Runtime/CoreUObject/Private/UObject/PropertyAccessUtil.cpp b/Engine/Source/Runtime/CoreUObject/Private/UObject/PropertyAccessUtil.cpp
--- a/Engine/Source/Runtime/CoreUObject/Private/UObject/PropertyAccessUtil.cpp
+++ b/Engine/Source/Runtime/CoreUObject/Private/UObject/PropertyAccessUtil.cpp
@@ -36,8 +36,8 @@ void ConvertRealNumber(const FProperty* InSrcProp, const void* InSrcValue, const
 		const FFloatProperty* DestFloatProp = CastFieldChecked<FFloatProperty>(InDestProp);
 		for (int32 Idx = 0; Idx < InCount; ++Idx)
 		{
-			const void* SrcElemValue = static_cast<const uint8*>(InSrcValue) + (InSrcProp->ElementSize * Idx);
-			void* DestElemValue = static_cast<uint8*>(InDestValue) + (InDestProp->ElementSize * Idx);
+			const void* SrcElemValue = static_cast<const uint8*>(InSrcValue) + (InSrcProp->GetElementSize() * Idx);
+			void* DestElemValue = static_cast<uint8*>(InDestValue) + (InDestProp->GetElementSize() * Idx);
 
 			const double Value = SrcDoubleProp->GetFloatingPointPropertyValue(SrcElemValue);
 			DestFloatProp->SetFloatingPointPropertyValue(DestElemValue, Value);
@@ -48,8 +48,8 @@ void ConvertRealNumber(const FProperty* InSrcProp, const void* InSrcValue, const
 		const FDoubleProperty* DestDoubleProp = CastFieldChecked<FDoubleProperty>(InDestProp);
 		for (int32 Idx = 0; Idx < InCount; ++Idx)
 		{
-			const void* SrcElemValue = static_cast<const uint8*>(InSrcValue) + (InSrcProp->ElementSize * Idx);
-			void* DestElemValue = static_cast<uint8*>(InDestValue) + (InDestProp->ElementSize * Idx);
+			const void* SrcElemValue = static_cast<const uint8*>(InSrcValue) + (InSrcProp->GetElementSize() * Idx);
+			void* DestElemValue = static_cast<uint8*>(InDestValue) + (InDestProp->GetElementSize() * Idx);
 
 			const double Value = SrcFloatProp->GetFloatingPointPropertyValue(SrcElemValue);
 			DestDoubleProp->SetFloatingPointPropertyValue(DestElemValue, Value);
@@ -230,8 +230,8 @@ bool IsCompletePropertyIdentical(const FProperty* InSrcProp, const void* InSrcVa
 	bool bIsIdentical = InSrcProp->ArrayDim == InDestProp->ArrayDim;
 	for (int32 Idx = 0; Idx < InSrcProp->ArrayDim && bIsIdentical; ++Idx)
 	{
-		const void* SrcElemValue = static_cast<const uint8*>(InSrcValue) + (InSrcProp->ElementSize * Idx);
-		const void* DestElemValue = static_cast<const uint8*>(InDestValue) + (InDestProp->ElementSize * Idx);
+		const void* SrcElemValue = static_cast<const uint8*>(InSrcValue) + (InSrcProp->GetElementSize() * Idx);
+		const void* DestElemValue = static_cast<const uint8*>(InDestValue) + (InDestProp->GetElementSize() * Idx);
 		bIsIdentical &= IsSinglePropertyIdentical(InSrcProp, SrcElemValue, InDestProp, DestElemValue);
 	}
 	return bIsIdentical;
@@ -292,7 +292,7 @@ bool CopyCompletePropertyValue(const FProperty* InSrcProp, const void* InSrcValu
 					{
 						for (int32 I = 0; I < InDestProp->ArrayDim; ++I)
 						{
-							void* DestValue = static_cast<uint8*>(InDestValue) + InDestProp->ElementSize * I;
+							void* DestValue = static_cast<uint8*>(InDestValue) + InDestProp->GetElementSize() * I;
 							CopySinglePropertyValue(SrcArray->Inner, SrcArrayHelper.GetElementPtr(I), InDestProp, DestValue);
 						}
 						return true;
@@ -311,7 +311,7 @@ bool CopyCompletePropertyValue(const FProperty* InSrcProp, const void* InSrcValu
 					DstArrayHelper.Resize(InSrcProp->ArrayDim);
 					for (int32 I = 0; I < InSrcProp->ArrayDim; ++I)
 					{
-						const void* SrcValue = static_cast<const uint8*>(InSrcValue) + InSrcProp->ElementSize * I;
+						const void* SrcValue = static_cast<const uint8*>(InSrcValue) + InSrcProp->GetElementSize() * I;
 						CopySinglePropertyValue(InSrcProp, SrcValue, DstArray->Inner, DstArrayHelper.GetElementPtr(I));
 					}
 					return true;
@@ -328,8 +328,8 @@ bool CopyCompletePropertyValue(const FProperty* InSrcProp, const void* InSrcValu
 		bool bSuccess = true;
 		for (int32 Idx = 0; Idx < InSrcProp->ArrayDim; ++Idx)
 		{
-			const void* SrcElemValue = static_cast<const uint8*>(InSrcValue) + (InSrcProp->ElementSize * Idx);
-			void* DestElemValue = static_cast<uint8*>(InDestValue) + (InDestProp->ElementSize * Idx);
+			const void* SrcElemValue = static_cast<const uint8*>(InSrcValue) + (InSrcProp->GetElementSize() * Idx);
+			void* DestElemValue = static_cast<uint8*>(InDestValue) + (InDestProp->GetElementSize() * Idx);
 
 			const int64 SrcEnumValue = GetPropertyEnumValue(InSrcProp, SrcElemValue);
 			bSuccess &= SetPropertyEnumValue(InDestProp, DestElemValue, SrcEnumValue);
@@ -342,8 +342,8 @@ bool CopyCompletePropertyValue(const FProperty* InSrcProp, const void* InSrcValu
 		const FBoolProperty* DestBoolProp = CastFieldChecked<FBoolProperty>(InDestProp);
 		for (int32 Idx = 0; Idx < InSrcProp->ArrayDim; ++Idx)
 		{
-			const void* SrcElemValue = static_cast<const uint8*>(InSrcValue) + (InSrcProp->ElementSize * Idx);
-			void* DestElemValue = static_cast<uint8*>(InDestValue) + (InDestProp->ElementSize * Idx);
+			const void* SrcElemValue = static_cast<const uint8*>(InSrcValue) + (InSrcProp->GetElementSize() * Idx);
+			void* DestElemValue = static_cast<uint8*>(InDestValue) + (InDestProp->GetElementSize() * Idx);
 
 			// Bools can be represented as bitfields, we we have to handle the copy a little differently to only extract the bool we want
 			const bool bBoolValue = SrcBoolProp->GetPropertyValue(SrcElemValue);
diff --git a/Engine/Source/Runtime/CoreUObject/Private/UObject/PropertyArray.cpp b/Engine/Source/Runtime/CoreUObject/Private/UObject/PropertyArray.cpp
--- a/Engine/Source/Runtime/CoreUObject/Private/UObject/PropertyArray.cpp
+++ b/Engine/Source/Runtime/CoreUObject/Private/UObject/PropertyArray.cpp
@@ -19,7 +19,7 @@ FArrayProperty::FArrayProperty(FFieldVariant InOwner, const FName& InName, EObje
 	, Inner(nullptr)
 {
 	ArrayFlags = InArrayPropertyFlags;
-	SetElementSize();
+	SetElementSize(TTypeFundamentals::CPPSize);
 }
 
 FArrayProperty::FArrayProperty(FFieldVariant InOwner, const UECodeGen_Private::FArrayPropertyParams& Prop)
@@ -27,7 +27,7 @@ FArrayProperty::FArrayProperty(FFieldVariant InOwner, const UECodeGen_Private::F
 	, Inner(nullptr)
 {
 	ArrayFlags = Prop.ArrayFlags;
-	SetElementSize();
+	SetElementSize(TTypeFundamentals::CPPSize);
 }
 
 #if WITH_EDITORONLY_DATA
@@ -77,7 +77,7 @@ void FArrayProperty::LinkInternal(FArchive& Ar)
 	//Ar.Preload(Inner);
 	Inner->Link(Ar);
 
-	SetElementSize();
+	SetElementSize(TTypeFundamentals::CPPSize);
 }
 bool FArrayProperty::Identical(const void* A, const void* B, uint32 PortFlags) const
 {
@@ -157,7 +157,7 @@ void FArrayProperty::SerializeItem(FStructuredArchive::FSlot Slot, void* Value,
 				ArrayHelper.EmptyAndAddUninitializedValues(ElementCount);
 			}
 
-			Stream.EnterElement().Serialize(ArrayHelper.GetRawPtr(), ElementCount * Inner->ElementSize);
+			Stream.EnterElement().Serialize(ArrayHelper.GetRawPtr(), ElementCount * Inner->GetElementSize());
 		}
 		else
 		{
@@ -321,7 +321,7 @@ void FArrayProperty::SerializeItem(FStructuredArchive::FSlot Slot, void* Value,
 				if (NumRemoved != 0)
 				{
 					TArray<int32> IndicesToRemove;
-					TempValueStorage = (uint8*)FMemory::Malloc(InnerObjectProperty->ElementSize);
+					TempValueStorage = (uint8*)FMemory::Malloc(InnerObjectProperty->GetElementSize());
 					InnerObjectProperty->InitializeValue(TempValueStorage);
 
 					for (int32 i = 0; i < NumRemoved; ++i)
@@ -367,7 +367,7 @@ void FArrayProperty::SerializeItem(FStructuredArchive::FSlot Slot, void* Value,
 				{
 					if (!TempValueStorage)
 					{
-						TempValueStorage = (uint8*)FMemory::Malloc(InnerObjectProperty->ElementSize);
+						TempValueStorage = (uint8*)FMemory::Malloc(InnerObjectProperty->GetElementSize());
 						InnerObjectProperty->InitializeValue(TempValueStorage);
 					}
 
@@ -395,7 +395,7 @@ void FArrayProperty::SerializeItem(FStructuredArchive::FSlot Slot, void* Value,
 				{
 					if (!TempValueStorage)
 					{
-						TempValueStorage = (uint8*)FMemory::Malloc(InnerObjectProperty->ElementSize);
+						TempValueStorage = (uint8*)FMemory::Malloc(InnerObjectProperty->GetElementSize());
 						InnerObjectProperty->InitializeValue(TempValueStorage);
 					}
 
@@ -935,7 +935,7 @@ void FArrayProperty::ExportTextInnerItem(FString& ValueStr, const FProperty* Inn
 			ValueStr += FString::Printf(TEXT("[%i] "), i);
 		}
 
-		uint8* PropData = (uint8*)PropertyValue + i * Inner->ElementSize;
+		uint8* PropData = (uint8*)PropertyValue + i * Inner->GetElementSize();
 
 		// Always use struct defaults if the inner is a struct, for symmetry with the import of array inner struct defaults
 		uint8* PropDefault = nullptr;
@@ -951,7 +951,7 @@ void FArrayProperty::ExportTextInnerItem(FString& ValueStr, const FProperty* Inn
 		{
 			if (DefaultValue && DefaultSize > i)
 			{
-				PropDefault = (uint8*)DefaultValue + i * Inner->ElementSize;
+				PropDefault = (uint8*)DefaultValue + i * Inner->GetElementSize();
 			}
 		}
 
@@ -1046,7 +1046,7 @@ const TCHAR* FArrayProperty::ImportTextInnerItem(const TCHAR* Buffer, const FPro
 
 		if (*Buffer != TCHAR(','))
 		{
-			uint8* Address = ArrayHelper ? ArrayHelper->GetRawPtr(Index) : ((uint8*)Data + Inner->ElementSize * Index);
+			uint8* Address = ArrayHelper ? ArrayHelper->GetRawPtr(Index) : ((uint8*)Data + Inner->GetElementSize() * Index);
 			// Parse the item
 			checkf(ArrayHelper == nullptr || Inner->GetOffset_ForInternal() == 0, TEXT("Expected the Inner property of the FArrayProperty."));
 			Buffer = Inner->ImportText_Direct(Buffer, Address, Parent, PortFlags | PPF_Delimited, ErrorText);
@@ -1125,7 +1125,7 @@ void FArrayProperty::CopyValuesInternal(void* Dest, void const* Src, int32 Count
 	}
 	if (Num)
 	{
-		size_t Size = Inner->ElementSize;
+		size_t Size = Inner->GetElementSize();
 		uint8* SrcData = (uint8*)SrcArrayHelper.GetRawPtr();
 		uint8* DestData = (uint8*)DestArrayHelper.GetRawPtr();
 		if (!(Inner->PropertyFlags & CPF_IsPlainOldData))
@@ -1197,7 +1197,7 @@ void FArrayProperty::InstanceSubobjects(void* Data, void const* DefaultData, UOb
 		FScriptArrayHelper ArrayHelper(this, Data);
 		FScriptArrayHelper DefaultArrayHelper(this, DefaultData);
 
-		int32 InnerElementSize = Inner->ElementSize;
+		int32 InnerElementSize = Inner->GetElementSize();
 		void* TempElement = FMemory_Alloca(InnerElementSize);
 
 		for (int32 ElementIndex = 0; ElementIndex < ArrayHelper.Num(); ElementIndex++)
diff --git a/Engine/Source/Runtime/CoreUObject/Private/UObject/PropertyBaseObject.cpp b/Engine/Source/Runtime/CoreUObject/Private/UObject/PropertyBaseObject.cpp
--- a/Engine/Source/Runtime/CoreUObject/Private/UObject/PropertyBaseObject.cpp
+++ b/Engine/Source/Runtime/CoreUObject/Private/UObject/PropertyBaseObject.cpp
@@ -63,11 +63,11 @@ void FObjectPropertyBase::InstanceSubobjects(void* Data, void const* DefaultData
 {
 	for ( int32 ArrayIndex = 0; ArrayIndex < ArrayDim; ArrayIndex++ )
 	{
-		TObjectPtr<UObject> CurrentObjectPtr = GetObjectPtrPropertyValue((uint8*)Data + ArrayIndex * ElementSize);
+		TObjectPtr<UObject> CurrentObjectPtr = GetObjectPtrPropertyValue((uint8*)Data + ArrayIndex * GetElementSize());
 		UObject* CurrentValue = CurrentObjectPtr.Get();
 		if (CurrentObjectPtr.IsResolved() && CurrentValue)
 		{
-			TObjectPtr<UObject> SubobjectTemplate = DefaultData ? GetObjectPtrPropertyValue((uint8*)DefaultData + ArrayIndex * ElementSize): nullptr;
+			TObjectPtr<UObject> SubobjectTemplate = DefaultData ? GetObjectPtrPropertyValue((uint8*)DefaultData + ArrayIndex * GetElementSize()): nullptr;
 			EInstancePropertyValueFlags Flags = EInstancePropertyValueFlags::None;
 			if (HasAnyPropertyFlags(CPF_InstancedReference))
 			{
@@ -78,7 +78,7 @@ void FObjectPropertyBase::InstanceSubobjects(void* Data, void const* DefaultData
 				Flags |= EInstancePropertyValueFlags::AllowSelfReference;
 			}
 			UObject* NewValue = InstanceGraph->InstancePropertyValue(SubobjectTemplate, CurrentValue, InOwner, Flags);
-			SetObjectPropertyValue((uint8*)Data + ArrayIndex * ElementSize, NewValue);
+			SetObjectPropertyValue((uint8*)Data + ArrayIndex * GetElementSize(), NewValue);
 		}
 	}
 }
diff --git a/Engine/Source/Runtime/CoreUObject/Private/UObject/PropertyBool.cpp b/Engine/Source/Runtime/CoreUObject/Private/UObject/PropertyBool.cpp
--- a/Engine/Source/Runtime/CoreUObject/Private/UObject/PropertyBool.cpp
+++ b/Engine/Source/Runtime/CoreUObject/Private/UObject/PropertyBool.cpp
@@ -96,8 +96,8 @@ void FBoolProperty::SetBoolSize( const uint32 InSize, const bool bIsNativeBool,
 		PropertyFlags |= CPF_NoDestructor;
 	}
 	uint32 TestBitmask = InBitMask ? InBitMask : 1;
-	ElementSize = InSize;
-	FieldSize = (uint8)ElementSize;
+	SetElementSize(InSize);
+	FieldSize = (uint8)GetElementSize();
 	ByteOffset = 0;
 	if (bIsNativeBool)
 	{		
@@ -110,8 +110,8 @@ void FBoolProperty::SetBoolSize( const uint32 InSize, const bool bIsNativeBool,
 		for (ByteOffset = 0; ByteOffset < InSize && ((ByteMask = *((uint8*)&TestBitmask + ByteOffset)) == 0); ByteOffset++);
 		FieldMask = ByteMask;
 	}
-	check((int32)FieldSize == ElementSize);
-	check(ElementSize != 0);
+	check((int32)FieldSize == GetElementSize());
+	check(GetElementSize() != 0);
 	check(FieldMask != 0);
 	check(ByteMask != 0);
 }
@@ -119,7 +119,7 @@ void FBoolProperty::SetBoolSize( const uint32 InSize, const bool bIsNativeBool,
 int32 FBoolProperty::GetMinAlignment() const
 {
 	int32 Alignment = 0;
-	switch(ElementSize)
+	switch(GetElementSize())
 	{
 	case sizeof(uint8):
 		Alignment = alignof(uint8); break;
@@ -130,14 +130,14 @@ int32 FBoolProperty::GetMinAlignment() const
 	case sizeof(uint64):
 		Alignment = alignof(uint64); break;
 	default:
-		UE_LOG(LogProperty, Fatal, TEXT("Unsupported FBoolProperty %s size %d."), *GetName(), (int32)ElementSize);
+		UE_LOG(LogProperty, Fatal, TEXT("Unsupported FBoolProperty %s size %d."), *GetName(), (int32)GetElementSize());
 	}
 	return Alignment;
 }
 void FBoolProperty::LinkInternal(FArchive& Ar)
 {
 	check(FieldSize != 0);
-	ElementSize = FieldSize;
+	SetElementSize(FieldSize);
 	if (IsNativeBool())
 	{
 		PropertyFlags |= (CPF_IsPlainOldData | CPF_NoDestructor | CPF_ZeroConstructor);
@@ -159,7 +159,7 @@ void FBoolProperty::Serialize( FArchive& Ar )
 	Ar << FieldMask;
 
 	// Serialize additional flags which will help to identify this FBoolProperty type and size.
-	uint8 BoolSize = (uint8)ElementSize;
+	uint8 BoolSize = (uint8)GetElementSize();
 	Ar << BoolSize;
 	uint8 NativeBool = false;
 	if( Ar.IsLoading())
@@ -187,7 +187,7 @@ FString FBoolProperty::GetCPPType( FString* ExtendedTypeText/*=NULL*/, uint32 CP
 	else
 	{
 		// Bitfields
-		switch(ElementSize)
+		switch(GetElementSize())
 		{
 		case sizeof(uint64):
 			return TEXT("uint64");
@@ -198,7 +198,7 @@ FString FBoolProperty::GetCPPType( FString* ExtendedTypeText/*=NULL*/, uint32 CP
 		case sizeof(uint8):
 			return TEXT("uint8");
 		default:
-			UE_LOG(LogProperty, Fatal, TEXT("Unsupported FBoolProperty %s size %d."), *GetName(), ElementSize);
+			UE_LOG(LogProperty, Fatal, TEXT("Unsupported FBoolProperty %s size %d."), *GetName(), GetElementSize());
 			break;
 		}
 	}
@@ -219,7 +219,7 @@ FString FBoolProperty::GetCPPMacroType( FString& ExtendedTypeText ) const
 	}
 	else
 	{
-		switch(ElementSize)
+		switch(GetElementSize())
 		{
 		case sizeof(uint64):
 			return TEXT("UBOOL64");
@@ -230,7 +230,7 @@ FString FBoolProperty::GetCPPMacroType( FString& ExtendedTypeText ) const
 		case sizeof(uint8):
 			return TEXT("UBOOL8");
 		default:
-			UE_LOG(LogProperty, Fatal, TEXT("Unsupported FBoolProperty %s size %d."), *GetName(), ElementSize);
+			UE_LOG(LogProperty, Fatal, TEXT("Unsupported FBoolProperty %s size %d."), *GetName(), GetElementSize());
 			break;
 		}
 	}
@@ -422,8 +422,8 @@ void FBoolProperty::CopyValuesInternal( void* Dest, void const* Src, int32 Count
 	check(FieldSize != 0 && !IsNativeBool());
 	for (int32 Index = 0; Index < Count; Index++)
 	{
-		uint8* DestByteValue = (uint8*)Dest + Index * ElementSize + ByteOffset;
-		uint8* SrcByteValue = (uint8*)Src + Index * ElementSize + ByteOffset;
+		uint8* DestByteValue = (uint8*)Dest + Index * GetElementSize() + ByteOffset;
+		uint8* SrcByteValue = (uint8*)Src + Index * GetElementSize() + ByteOffset;
 		*DestByteValue = (*DestByteValue & ~FieldMask) | (*SrcByteValue & FieldMask);
 	}
 }
diff --git a/Engine/Source/Runtime/CoreUObject/Private/UObject/PropertyMap.cpp b/Engine/Source/Runtime/CoreUObject/Private/UObject/PropertyMap.cpp
--- a/Engine/Source/Runtime/CoreUObject/Private/UObject/PropertyMap.cpp
+++ b/Engine/Source/Runtime/CoreUObject/Private/UObject/PropertyMap.cpp
@@ -1103,7 +1103,7 @@ const TCHAR* FMapProperty::ImportText_Internal(const TCHAR* Buffer, void* Contai
 		return Buffer + 1;
 	}
 
-	uint8* TempPairStorage   = (uint8*)FMemory::Malloc(MapLayout.ValueOffset + ValueProp->ElementSize);
+	uint8* TempPairStorage   = (uint8*)FMemory::Malloc(MapLayout.ValueOffset + ValueProp->GetElementSize());
 
 	bool bSuccess = false;
 	ON_SCOPE_EXIT
diff --git a/Engine/Source/Runtime/CoreUObject/Private/UObject/PropertyOptional.cpp b/Engine/Source/Runtime/CoreUObject/Private/UObject/PropertyOptional.cpp
--- a/Engine/Source/Runtime/CoreUObject/Private/UObject/PropertyOptional.cpp
+++ b/Engine/Source/Runtime/CoreUObject/Private/UObject/PropertyOptional.cpp
@@ -161,7 +161,7 @@ void FOptionalProperty::LinkInternal(FArchive& Ar)
 	ValueProperty->Link(Ar);
 	
 	// After ValueProperty's size has been computed, compute the size of this property.
-	ElementSize = CalcSize();
+	SetElementSize(CalcSize());
 
 	// Standard non-intrusive optional properties can always be initialized by zeroing memory.
 	if (!ValueProperty->HasIntrusiveUnsetOptionalState())
diff --git a/Engine/Source/Runtime/CoreUObject/Private/UObject/PropertySet.cpp b/Engine/Source/Runtime/CoreUObject/Private/UObject/PropertySet.cpp
--- a/Engine/Source/Runtime/CoreUObject/Private/UObject/PropertySet.cpp
+++ b/Engine/Source/Runtime/CoreUObject/Private/UObject/PropertySet.cpp
@@ -674,7 +674,7 @@ const TCHAR* FSetProperty::ImportText_Internal(const TCHAR* Buffer, void* Contai
 		return Buffer + 1;
 	}
 
-	TempElementStorage = (uint8*)FMemory::Malloc(ElementProp->ElementSize);
+	TempElementStorage = (uint8*)FMemory::Malloc(ElementProp->GetElementSize());
 	// From this point failure should empty the set
 	bSuccess = false;
 
diff --git a/Engine/Source/Runtime/CoreUObject/Private/UObject/PropertyStruct.cpp b/Engine/Source/Runtime/CoreUObject/Private/UObject/PropertyStruct.cpp
--- a/Engine/Source/Runtime/CoreUObject/Private/UObject/PropertyStruct.cpp
+++ b/Engine/Source/Runtime/CoreUObject/Private/UObject/PropertyStruct.cpp
@@ -47,7 +47,7 @@ FStructProperty::FStructProperty(FFieldVariant InOwner, const FName& InName, EOb
 	: Super(InOwner, InName, InObjectFlags)
 	, Struct(nullptr)
 {
-	ElementSize = 0;
+	SetElementSize(0);
 }
 
 static EPropertyFlags GetStructComputedPropertyFlags(const UECodeGen_Private::FStructPropertyParams& Prop)
@@ -73,7 +73,7 @@ FStructProperty::FStructProperty(UField* InField)
 {
 	UStructProperty* SourceProperty = CastChecked<UStructProperty>(InField);
 	Struct = SourceProperty->Struct;
-	check(ElementSize == SourceProperty->ElementSize); // this should've been set by FProperty
+	check(GetElementSize() == SourceProperty->ElementSize); // this should've been set by FProperty
 }
 #endif // WITH_EDITORONLY_DATA
 
@@ -109,7 +109,7 @@ void FStructProperty::LinkInternal(FArchive& Ar)
 	}
 	PreloadInnerStructMembers(this);
 	
-	ElementSize = Align(Struct->PropertiesSize, Struct->GetMinAlignment());
+	SetElementSize(Align(Struct->PropertiesSize, Struct->GetMinAlignment()));
 	if(UScriptStruct::ICppStructOps* Ops = Struct->GetCppStructOps())
 	{
 		PropertyFlags |= Ops->GetComputedPropertyFlags();
@@ -478,7 +478,7 @@ void FStructProperty::InstanceSubobjects( void* Data, void const* DefaultData, U
 {
 	for (int32 Index = 0; Index < ArrayDim; Index++)
 	{
-		Struct->InstanceSubobjectTemplates( (uint8*)Data + ElementSize * Index, DefaultData ? (uint8*)DefaultData + ElementSize * Index : NULL, Struct, InOwner, InstanceGraph );
+		Struct->InstanceSubobjectTemplates( (uint8*)Data + GetElementSize() * Index, DefaultData ? (uint8*)DefaultData + GetElementSize() * Index : NULL, Struct, InOwner, InstanceGraph );
 	}
 }
 
diff --git a/Engine/Source/Runtime/CoreUObject/Private/UObject/ScriptCore.cpp b/Engine/Source/Runtime/CoreUObject/Private/UObject/ScriptCore.cpp
--- a/Engine/Source/Runtime/CoreUObject/Private/UObject/ScriptCore.cpp
+++ b/Engine/Source/Runtime/CoreUObject/Private/UObject/ScriptCore.cpp
@@ -905,7 +905,7 @@ void UObject::SkipFunction(FFrame& Stack, RESULT_DECL, UFunction* Function)
 		// destroy old value if necessary
 		ReturnProp->DestroyValue(RESULT_PARAM);
 		// copy zero value for return property into Result
-		FMemory::Memzero(RESULT_PARAM, ReturnProp->ArrayDim * ReturnProp->ElementSize);
+		FMemory::Memzero(RESULT_PARAM, ReturnProp->ArrayDim * ReturnProp->GetElementSize());
 	}
 }
 
@@ -1156,7 +1156,7 @@ void ClearReturnValue(FProperty* ReturnProp, RESULT_DECL)
 	if (ReturnProp != NULL)
 	{
 		uint8* Data = (uint8*)RESULT_PARAM;
-		for (int32 ArrayIdx = 0; ArrayIdx < ReturnProp->ArrayDim; ArrayIdx++, Data += ReturnProp->ElementSize)
+		for (int32 ArrayIdx = 0; ArrayIdx < ReturnProp->ArrayDim; ArrayIdx++, Data += ReturnProp->GetElementSize())
 		{
 			// Clear the property. This assumes that it has already been initialized, and that the caller will destroy it.
 			ReturnProp->ClearValue(Data);
@@ -2153,7 +2153,7 @@ void UObject::ProcessEvent( UFunction* Function, void* Parms )
 				}
 				else if (!(P->PropertyFlags & CPF_OutParm))
 				{
-					FMemory::Memcpy(P->ContainerPtrToValuePtr<uint8>(Parms), P->ContainerPtrToValuePtr<uint8>(NewStack.Locals), P->ArrayDim * P->ElementSize);
+					FMemory::Memcpy(P->ContainerPtrToValuePtr<uint8>(Parms), P->ContainerPtrToValuePtr<uint8>(NewStack.Locals), P->ArrayDim * P->GetElementSize());
 				}
 			}
 		}
diff --git a/Engine/Source/Runtime/CoreUObject/Public/UObject/UnrealType.h b/Engine/Source/Runtime/CoreUObject/Public/UObject/UnrealType.h
--- a/Engine/Source/Runtime/CoreUObject/Public/UObject/UnrealType.h
+++ b/Engine/Source/Runtime/CoreUObject/Public/UObject/UnrealType.h
@@ -179,7 +179,9 @@ class FProperty : public FField
 
 	// Persistent variables.
 	int32			ArrayDim;
+	UE_DEPRECATED(5.5, "Use GetElementSize/SetElementSize instead.")
 	int32			ElementSize;
+public:
 	EPropertyFlags	PropertyFlags;
 	uint16			RepIndex;
 
@@ -228,6 +230,15 @@ public:
 	COREUOBJECT_API explicit FProperty(UField* InField);
 #endif // WITH_EDITORONLY_DATA
 
+	// ElementSize accessors to facilitate underlying type change
+	COREUOBJECT_API void SetElementSize(int32 NewSize);
+	int32 	GetElementSize() const
+	{
+PRAGMA_DISABLE_DEPRECATION_WARNINGS
+		return ElementSize;
+PRAGMA_ENABLE_DEPRECATION_WARNINGS
+	}
+
 	// UObject interface
 	COREUOBJECT_API virtual void Serialize( FArchive& Ar ) override;
 	// End of UObject interface
@@ -669,7 +680,7 @@ private:
 			check(!GetOwner<UClass>()); // Check we are _not_ calling this on a direct child property of a UClass, you should pass in a UObject* in that case
 		}
 
-		return (uint8*)ContainerPtr + Offset_Internal + static_cast<size_t>(ElementSize) * ArrayIndex;
+		return (uint8*)ContainerPtr + Offset_Internal + static_cast<size_t>(GetElementSize()) * ArrayIndex;
 	}
 
 	FORCEINLINE void* ContainerUObjectPtrToValuePtrInternal(UObject* ContainerPtr, int32 ArrayIndex) const
@@ -696,7 +707,7 @@ private:
 			check(!GetOwner<UClass>()); // Check we are _not_ calling this on a direct child property of a UClass, you should pass in a UObject* in that case
 		}
 
-		return (uint8*)ContainerPtr + Offset_Internal + static_cast<size_t>(ElementSize) * ArrayIndex;
+		return (uint8*)ContainerPtr + Offset_Internal + static_cast<size_t>(GetElementSize()) * ArrayIndex;
 	}
 
 protected:
@@ -708,7 +719,7 @@ protected:
 	{
 		if (PropertyPointerType == EPropertyPointerType::Container)
 		{
-			return (uint8*)ContainerOrPropertyPtr + Offset_Internal + static_cast<size_t>(ElementSize) * ArrayIndex;
+			return (uint8*)ContainerOrPropertyPtr + Offset_Internal + static_cast<size_t>(GetElementSize()) * ArrayIndex;
 		}
 		else
 		{
@@ -812,7 +823,7 @@ public:
 		{
 			if (PropertyFlags & CPF_IsPlainOldData)
 			{
-				FMemory::Memcpy( Dest, Src, ElementSize );
+				FMemory::Memcpy( Dest, Src, GetElementSize() );
 			}
 			else
 			{
@@ -846,7 +857,7 @@ public:
 		{
 			if (PropertyFlags & CPF_IsPlainOldData)
 			{
-				FMemory::Memcpy( Dest, Src, static_cast<size_t>(ElementSize) * ArrayDim );
+				FMemory::Memcpy( Dest, Src, static_cast<size_t>(GetElementSize()) * ArrayDim );
 			}
 			else
 			{
@@ -919,7 +930,7 @@ public:
 	{
 		if (HasAllPropertyFlags(CPF_NoDestructor | CPF_ZeroConstructor))
 		{
-			FMemory::Memzero( Data, ElementSize );
+			FMemory::Memzero( Data, GetElementSize() );
 		}
 		else
 		{
@@ -936,7 +947,7 @@ public:
 	{
 		if (HasAllPropertyFlags(CPF_NoDestructor | CPF_ZeroConstructor))
 		{
-			FMemory::Memzero( ContainerPtrToValuePtr<void>(Data, ArrayIndex), ElementSize );
+			FMemory::Memzero( ContainerPtrToValuePtr<void>(Data, ArrayIndex), GetElementSize() );
 		}
 		else
 		{
@@ -1039,7 +1050,7 @@ public:
 	{
 		if (PropertyFlags & CPF_ZeroConstructor)
 		{
-			FMemory::Memzero(Dest, static_cast<size_t>(ElementSize) * ArrayDim);
+			FMemory::Memzero(Dest, static_cast<size_t>(GetElementSize()) * ArrayDim);
 		}
 		else
 		{
@@ -1057,7 +1068,7 @@ public:
 	{
 		if (PropertyFlags & CPF_ZeroConstructor)
 		{
-			FMemory::Memzero(ContainerPtrToValuePtr<void>(Dest), static_cast<size_t>(ElementSize) * ArrayDim);
+			FMemory::Memzero(ContainerPtrToValuePtr<void>(Dest), static_cast<size_t>(GetElementSize()) * ArrayDim);
 		}
 		else
 		{
@@ -1135,7 +1146,7 @@ public:
     // need to audit all callers to make such a change.
 	FORCEINLINE int32 GetSize() const
 	{
-		return ArrayDim * ElementSize;
+		return ArrayDim * GetElementSize();
 	}
 	COREUOBJECT_API bool ShouldSerializeValue( FArchive& Ar ) const;
 
@@ -1486,7 +1497,7 @@ public:
 	TProperty(FFieldVariant InOwner, const FName& InName, EObjectFlags InObjectFlags)
 		: Super(InOwner, InName, InObjectFlags)
 	{
-		SetElementSize();
+		this->SetElementSize(TTypeFundamentals::CPPSize);
 	}
 
 	/**
@@ -1498,7 +1509,7 @@ public:
 	TProperty(FFieldVariant InOwner, PropertyParamsType& Prop)
 		: Super(InOwner, Prop, TTypeFundamentals::GetComputedFlagsPropertyFlags())
 	{
-		SetElementSize();
+		this->SetElementSize(TTypeFundamentals::CPPSize);
 	}
 
 public:
@@ -1507,7 +1518,7 @@ public:
 	explicit TProperty(UField* InField)
 		: Super(InField)
 	{
-		SetElementSize();
+		this->SetElementSize(TTypeFundamentals::CPPSize);
 	}
 #endif // WITH_EDITORONLY_DATA
 
@@ -1531,7 +1542,7 @@ public:
 	}
 	virtual void LinkInternal(FArchive& Ar) override
 	{
-		SetElementSize();
+		this->SetElementSize(TTypeFundamentals::CPPSize);
 		this->PropertyFlags |= TTypeFundamentals::GetComputedFlagsPropertyFlags();
 
 	}
@@ -1550,14 +1561,14 @@ public:
 	{
 		for (int32 i = 0; i < this->ArrayDim; ++i)
 		{
-			TTypeFundamentals::InitializePropertyValue((uint8*)Dest + i * static_cast<size_t>(this->ElementSize));
+			TTypeFundamentals::InitializePropertyValue((uint8*)Dest + i * static_cast<size_t>(this->GetElementSize()));
 		}
 	}
 	virtual void DestroyValueInternal( void* Dest ) const override
 	{
 		for (int32 i = 0; i < this->ArrayDim; ++i)
 		{
-			TTypeFundamentals::DestroyPropertyValue((uint8*)Dest + i * static_cast<size_t>(this->ElementSize));
+			TTypeFundamentals::DestroyPropertyValue((uint8*)Dest + i * static_cast<size_t>(this->GetElementSize()));
 		}
 	}
 
@@ -1616,12 +1627,6 @@ public:
 	{
 		TTypeFundamentals::ClearIntrusiveOptionalValue(Data);
 	}
-
-protected:
-	FORCEINLINE void SetElementSize()
-	{
-		this->ElementSize = TTypeFundamentals::CPPSize;
-	}
 	// End of FProperty interface
 };
 
@@ -2689,7 +2694,7 @@ protected:
 					FProperty::GetValue_InContainer(SrcAddress, ValueArray);
 
 					// Grab the items we care about and free the temp array
-					int32 LocalElementSize = ElementSize;
+					int32 LocalElementSize = GetElementSize();
 					for (int32 OutIndex = 0; OutIndex != ArrayCount; ++OutIndex)
 					{
 						OutObjects[OutIndex] = ValueArray[ArrayIndex + OutIndex].Get();
@@ -2710,7 +2715,7 @@ protected:
 		}
 		else
 		{
-			int32 LocalElementSize = ElementSize;
+			int32 LocalElementSize = GetElementSize();
 			for (int32 OutIndex = 0; OutIndex != ArrayCount; ++OutIndex)
 			{
 				OutObjects[OutIndex] = GetObjectPropertyValue((const uint8*)SrcAddress + OutIndex * LocalElementSize);
@@ -2743,7 +2748,7 @@ protected:
 					FProperty::GetValue_InContainer(DestAddress, ValueArray);
 
 					// Replace the items we care about
-					int32 LocalElementSize = ElementSize;
+					int32 LocalElementSize = GetElementSize();
 					for (int32 OutIndex = 0; OutIndex != ArrayCount; ++OutIndex)
 					{
 						ValueArray[ArrayIndex + OutIndex] = InValues[OutIndex];
@@ -2767,7 +2772,7 @@ protected:
 		}
 		else
 		{
-			int32 LocalElementSize = ElementSize;
+			int32 LocalElementSize = GetElementSize();
 			for (int32 OutIndex = 0; OutIndex != ArrayCount; ++OutIndex)
 			{
 				SetObjectPropertyValue((uint8*)DestAddress + OutIndex * LocalElementSize, InValues[OutIndex]);
@@ -3477,14 +3482,14 @@ public:
 
 			for (int32 i = 0; i < this->ArrayDim; ++i)
 			{
-				new ((uint8*)Dest + i * static_cast<size_t>(this->ElementSize)) FFreezableScriptArray;
+				new ((uint8*)Dest + i * static_cast<size_t>(this->GetElementSize())) FFreezableScriptArray;
 			}
 		}
 		else
 		{
 			for (int32 i = 0; i < this->ArrayDim; ++i)
 			{
-				new ((uint8*)Dest + i * static_cast<size_t>(this->ElementSize)) FScriptArray;
+				new ((uint8*)Dest + i * static_cast<size_t>(this->GetElementSize())) FScriptArray;
 			}
 		}
 	}
@@ -3524,12 +3529,6 @@ public:
 #if WITH_EDITORONLY_DATA
 	virtual void AppendSchemaHash(FBlake3& Builder, bool bSkipEditorOnly) const override;
 #endif
-
-private:
-	FORCEINLINE void SetElementSize()
-	{
-		this->ElementSize = CPPSize;
-	}
 };
 
 using FFreezableScriptMap = TScriptMap<FMemoryImageSetAllocator>;
@@ -3617,14 +3616,14 @@ public:
 
 			for (int32 i = 0; i < this->ArrayDim; ++i)
 			{
-				new ((uint8*)Dest + i * static_cast<size_t>(this->ElementSize)) FFreezableScriptMap;
+				new ((uint8*)Dest + i * static_cast<size_t>(this->GetElementSize())) FFreezableScriptMap;
 			}
 		}
 		else
 		{
 			for (int32 i = 0; i < this->ArrayDim; ++i)
 			{
-				new ((uint8*)Dest + i * static_cast<size_t>(this->ElementSize)) FScriptMap;
+				new ((uint8*)Dest + i * static_cast<size_t>(this->GetElementSize())) FScriptMap;
 			}
 		}
 	}
@@ -3853,7 +3852,7 @@ public:
 	 *	@param	InArray: pointer to raw memory that corresponds to this array. This can be NULL, and sometimes is, but in that case almost all operations will crash.
 	**/
 	FORCEINLINE FScriptArrayHelper(const FArrayProperty* InProperty, const void* InArray)
-		: FScriptArrayHelper(Internal, InProperty->Inner, InArray, InProperty->Inner->ElementSize, InProperty->Inner->GetMinAlignment(), InProperty->ArrayFlags)
+		: FScriptArrayHelper(Internal, InProperty->Inner, InArray, InProperty->Inner->GetElementSize(), InProperty->Inner->GetMinAlignment(), InProperty->ArrayFlags)
 	{
 	}
 
@@ -4096,7 +4095,7 @@ public:
 
 	static FScriptArrayHelper CreateHelperFormInnerProperty(const FProperty* InInnerProperty, const void *InArray, EArrayPropertyFlags InArrayFlags = EArrayPropertyFlags::None)
 	{
-		return FScriptArrayHelper(Internal, InInnerProperty, InArray, InInnerProperty->ElementSize, InInnerProperty->GetMinAlignment(), InArrayFlags);
+		return FScriptArrayHelper(Internal, InInnerProperty, InArray, InInnerProperty->GetElementSize(), InInnerProperty->GetMinAlignment(), InArrayFlags);
 	}
 
 private:
diff --git a/Engine/Source/Runtime/Engine/Classes/Kismet/BlueprintMapLibrary.h b/Engine/Source/Runtime/Engine/Classes/Kismet/BlueprintMapLibrary.h
--- a/Engine/Source/Runtime/Engine/Classes/Kismet/BlueprintMapLibrary.h
+++ b/Engine/Source/Runtime/Engine/Classes/Kismet/BlueprintMapLibrary.h
@@ -132,7 +132,7 @@ public:
 
 		// Since Key and Value aren't really an int, step the stack manually
 		const FProperty* CurrKeyProp = MapProperty->KeyProp;
-		const int32 KeyPropertySize = CurrKeyProp->ElementSize * CurrKeyProp->ArrayDim;
+		const int32 KeyPropertySize = CurrKeyProp->GetElementSize() * CurrKeyProp->ArrayDim;
 		void* KeyStorageSpace = FMemory_Alloca(KeyPropertySize);
 		CurrKeyProp->InitializeValue(KeyStorageSpace);
 
@@ -141,7 +141,7 @@ public:
 		Stack.StepCompiledIn<FProperty>(KeyStorageSpace);
 		
 		const FProperty* CurrValueProp = MapProperty->ValueProp;
-		const int32 ValuePropertySize = CurrValueProp->ElementSize * CurrValueProp->ArrayDim;
+		const int32 ValuePropertySize = CurrValueProp->GetElementSize() * CurrValueProp->ArrayDim;
 		void* ValueStorageSpace = FMemory_Alloca(ValuePropertySize);
 		CurrValueProp->InitializeValue(ValueStorageSpace);
 		
@@ -173,7 +173,7 @@ public:
 
 		// Since Key and Value aren't really an int, step the stack manually
 		const FProperty* CurrKeyProp = MapProperty->KeyProp;
-		const int32 KeyPropertySize = CurrKeyProp->ElementSize * CurrKeyProp->ArrayDim;
+		const int32 KeyPropertySize = CurrKeyProp->GetElementSize() * CurrKeyProp->ArrayDim;
 		void* KeyStorageSpace = FMemory_Alloca(KeyPropertySize);
 		CurrKeyProp->InitializeValue(KeyStorageSpace);
 
@@ -203,7 +203,7 @@ public:
 
 		// Since Key and Value aren't really an int, step the stack manually
 		const FProperty* CurrKeyProp = MapProperty->KeyProp;
-		const int32 KeyPropertySize = CurrKeyProp->ElementSize * CurrKeyProp->ArrayDim;
+		const int32 KeyPropertySize = CurrKeyProp->GetElementSize() * CurrKeyProp->ArrayDim;
 		void* KeyStorageSpace = FMemory_Alloca(KeyPropertySize);
 		CurrKeyProp->InitializeValue(KeyStorageSpace);
 
@@ -212,7 +212,7 @@ public:
 		Stack.StepCompiledIn<FProperty>(KeyStorageSpace);
 		
 		const FProperty* CurrValueProp = MapProperty->ValueProp;
-		const int32 ValuePropertySize = CurrValueProp->ElementSize * CurrValueProp->ArrayDim;
+		const int32 ValuePropertySize = CurrValueProp->GetElementSize() * CurrValueProp->ArrayDim;
 		void* ValueStorageSpace = FMemory_Alloca(ValuePropertySize);
 		CurrValueProp->InitializeValue(ValueStorageSpace);
 		
@@ -223,7 +223,7 @@ public:
 		const FFieldClass* MostRecentPropClass = Stack.MostRecentProperty->GetClass();
 		void* ItemPtr;
 		// If the destination and the inner type are identical in size and their field classes derive from one another, then permit the writing out of the array element to the destination memory
-		if (Stack.MostRecentPropertyAddress != NULL && (ValuePropertySize == Stack.MostRecentProperty->ElementSize*Stack.MostRecentProperty->ArrayDim) &&
+		if (Stack.MostRecentPropertyAddress != NULL && (ValuePropertySize == Stack.MostRecentProperty->GetElementSize()*Stack.MostRecentProperty->ArrayDim) &&
 			(MostRecentPropClass->IsChildOf(CurrValuePropClass) || CurrValuePropClass->IsChildOf(MostRecentPropClass)))
 		{
 			ItemPtr = Stack.MostRecentPropertyAddress;
@@ -313,7 +313,7 @@ public:
 
 		// Since Key and Value aren't really an int, step the stack manually
 		const FProperty* CurrKeyProp = MapProperty->KeyProp;
-		const int32 KeyPropertySize = CurrKeyProp->ElementSize * CurrKeyProp->ArrayDim;
+		const int32 KeyPropertySize = CurrKeyProp->GetElementSize() * CurrKeyProp->ArrayDim;
 		void* KeyStorageSpace = FMemory_Alloca(KeyPropertySize);
 		CurrKeyProp->InitializeValue(KeyStorageSpace);
 
diff --git a/Engine/Source/Runtime/Engine/Classes/Kismet/BlueprintSetLibrary.h b/Engine/Source/Runtime/Engine/Classes/Kismet/BlueprintSetLibrary.h
--- a/Engine/Source/Runtime/Engine/Classes/Kismet/BlueprintSetLibrary.h
+++ b/Engine/Source/Runtime/Engine/Classes/Kismet/BlueprintSetLibrary.h
@@ -165,7 +165,7 @@ public:
 
 		// Since ItemPtr isn't really an int, step the stack manually
 		const FProperty* ElementProp = SetProperty->ElementProp;
-		const int32 PropertySize = ElementProp->ElementSize * ElementProp->ArrayDim;
+		const int32 PropertySize = ElementProp->GetElementSize() * ElementProp->ArrayDim;
 		void* StorageSpace = FMemory_Alloca(PropertySize);
 		ElementProp->InitializeValue(StorageSpace);
 
@@ -226,7 +226,7 @@ public:
 
 		// Since ItemPtr isn't really an int, step the stack manually
 		const FProperty* ElementProp = SetProperty->ElementProp;
-		const int32 PropertySize = ElementProp->ElementSize * ElementProp->ArrayDim;
+		const int32 PropertySize = ElementProp->GetElementSize() * ElementProp->ArrayDim;
 		void* StorageSpace = FMemory_Alloca(PropertySize);
 		ElementProp->InitializeValue(StorageSpace);
 
@@ -357,7 +357,7 @@ public:
 
 		// Since ItemToFind isn't really an int, step the stack manually
 		const FProperty* ElementProp = SetProperty->ElementProp;
-		const int32 PropertySize = ElementProp->ElementSize * ElementProp->ArrayDim;
+		const int32 PropertySize = ElementProp->GetElementSize() * ElementProp->ArrayDim;
 		void* StorageSpace = FMemory_Alloca(PropertySize);
 		ElementProp->InitializeValue(StorageSpace);
 
diff --git a/Engine/Source/Runtime/Engine/Classes/Kismet/KismetArrayLibrary.h b/Engine/Source/Runtime/Engine/Classes/Kismet/KismetArrayLibrary.h
--- a/Engine/Source/Runtime/Engine/Classes/Kismet/KismetArrayLibrary.h
+++ b/Engine/Source/Runtime/Engine/Classes/Kismet/KismetArrayLibrary.h
@@ -390,7 +390,7 @@ public:
  
  		// Since NewItem isn't really an int, step the stack manually
  		const FProperty* InnerProp = ArrayProperty->Inner;
- 		const int32 PropertySize = InnerProp->ElementSize * InnerProp->ArrayDim;
+ 		const int32 PropertySize = InnerProp->GetElementSize() * InnerProp->ArrayDim;
  		void* StorageSpace = FMemory_Alloca(PropertySize);
  		InnerProp->InitializeValue(StorageSpace);
  
@@ -420,7 +420,7 @@ public:
 
 		// Since NewItem isn't really an int, step the stack manually
 		const FProperty* InnerProp = ArrayProperty->Inner;
-		const int32 PropertySize = InnerProp->ElementSize * InnerProp->ArrayDim;
+		const int32 PropertySize = InnerProp->GetElementSize() * InnerProp->ArrayDim;
 		void* StorageSpace = FMemory_Alloca(PropertySize);
 		InnerProp->InitializeValue(StorageSpace);
 
@@ -551,7 +551,7 @@ public:
 
 		// Since NewItem isn't really an int, step the stack manually
 		const FProperty* InnerProp = ArrayProperty->Inner;
-		const int32 PropertySize = InnerProp->ElementSize * InnerProp->ArrayDim;
+		const int32 PropertySize = InnerProp->GetElementSize() * InnerProp->ArrayDim;
 		void* StorageSpace = FMemory_Alloca(PropertySize);
 		InnerProp->InitializeValue(StorageSpace);
 
@@ -601,7 +601,7 @@ public:
 		}
 		// Since Item isn't really an int, step the stack manually
 		const FProperty* InnerProp = ArrayProperty->Inner;
-		const int32 PropertySize = InnerProp->ElementSize * InnerProp->ArrayDim;
+		const int32 PropertySize = InnerProp->GetElementSize() * InnerProp->ArrayDim;
 		void* StorageSpace = FMemory_Alloca(PropertySize);
 		InnerProp->InitializeValue(StorageSpace);
 
@@ -759,7 +759,7 @@ public:
 
 		// Since Item isn't really an int, step the stack manually
 		const FProperty* InnerProp = ArrayProperty->Inner;
-		const int32 PropertySize = InnerProp->ElementSize * InnerProp->ArrayDim;
+		const int32 PropertySize = InnerProp->GetElementSize() * InnerProp->ArrayDim;
 		void* StorageSpace = FMemory_Alloca(PropertySize);
 		InnerProp->InitializeValue(StorageSpace);
 
@@ -770,7 +770,7 @@ public:
 		const FFieldClass* MostRecentPropClass = Stack.MostRecentProperty->GetClass();
 		void* ItemPtr;
 		// If the destination and the inner type are identical in size and their field classes derive from one another, then permit the writing out of the array element to the destination memory
-		if (Stack.MostRecentPropertyAddress != NULL && (PropertySize == Stack.MostRecentProperty->ElementSize*Stack.MostRecentProperty->ArrayDim) &&
+		if (Stack.MostRecentPropertyAddress != NULL && (PropertySize == Stack.MostRecentProperty->GetElementSize()*Stack.MostRecentProperty->ArrayDim) &&
 			(MostRecentPropClass->IsChildOf(InnerPropClass) || InnerPropClass->IsChildOf(MostRecentPropClass)))
 		{
 			ItemPtr = Stack.MostRecentPropertyAddress;
@@ -802,7 +802,7 @@ public:
 
 		// Since NewItem isn't really an int, step the stack manually
 		const FProperty* InnerProp = ArrayProperty->Inner;
-		const int32 PropertySize = InnerProp->ElementSize * InnerProp->ArrayDim;
+		const int32 PropertySize = InnerProp->GetElementSize() * InnerProp->ArrayDim;
 		void* StorageSpace = FMemory_Alloca(PropertySize);
 		InnerProp->InitializeValue(StorageSpace);
 
@@ -857,7 +857,7 @@ public:
 		}
 		// Since ItemToFind isn't really an int, step the stack manually
 		const FProperty* InnerProp = ArrayProperty->Inner;
-		const int32 PropertySize = InnerProp->ElementSize * InnerProp->ArrayDim;
+		const int32 PropertySize = InnerProp->GetElementSize() * InnerProp->ArrayDim;
 		void* StorageSpace = FMemory_Alloca(PropertySize);
 		InnerProp->InitializeValue(StorageSpace);
 
@@ -892,7 +892,7 @@ public:
 		}
 		// Since ItemToFind isn't really an int, step the stack manually
 		const FProperty* InnerProp = ArrayProperty->Inner;
-		const int32 PropertySize = InnerProp->ElementSize * InnerProp->ArrayDim;
+		const int32 PropertySize = InnerProp->GetElementSize() * InnerProp->ArrayDim;
 		void* StorageSpace = FMemory_Alloca(PropertySize);
 		InnerProp->InitializeValue(StorageSpace);
 
diff --git a/Engine/Source/Runtime/Engine/Private/Animation/AnimClassInterface.cpp b/Engine/Source/Runtime/Engine/Private/Animation/AnimClassInterface.cpp
--- a/Engine/Source/Runtime/Engine/Private/Animation/AnimClassInterface.cpp
+++ b/Engine/Source/Runtime/Engine/Private/Animation/AnimClassInterface.cpp
@@ -129,7 +129,7 @@ static const UObject* ValuePtrToContainerUObjectPtr(FProperty* Property, const v
 	check(ArrayIndex < Property->ArrayDim);
 	check(ValuePtr);
 
-	const uint8* ContainerPtr = (const uint8*)ValuePtr - (Property->GetOffset_ForInternal() + Property->ElementSize * ArrayIndex);
+	const uint8* ContainerPtr = (const uint8*)ValuePtr - (Property->GetOffset_ForInternal() + Property->GetElementSize() * ArrayIndex);
 
 	check(((const UObject*)ContainerPtr)->IsValidLowLevel()); // Check its a valid UObject that was passed in
 	check(((const UObject*)ContainerPtr)->GetClass() != NULL);
diff --git a/Engine/Source/Runtime/Engine/Private/BlueprintGeneratedClass.cpp b/Engine/Source/Runtime/Engine/Private/BlueprintGeneratedClass.cpp
--- a/Engine/Source/Runtime/Engine/Private/BlueprintGeneratedClass.cpp
+++ b/Engine/Source/Runtime/Engine/Private/BlueprintGeneratedClass.cpp
@@ -1339,7 +1339,7 @@ void UBlueprintGeneratedClass::InitArrayPropertyFromCustomList(const FArrayPrope
 			uint8* DstArrayItemValue = DstArrayValueHelper.GetRawPtr(ArrayIndex);
 			const uint8* SrcArrayItemValue = SrcArrayValueHelper.GetRawPtr(ArrayIndex);
 
-			FMemory::Memcpy(DstArrayItemValue, SrcArrayItemValue, (SrcNum - ArrayIndex) * ArrayProperty->Inner->ElementSize);
+			FMemory::Memcpy(DstArrayItemValue, SrcArrayItemValue, (SrcNum - ArrayIndex) * ArrayProperty->Inner->GetElementSize());
 		}
 		else
 		{
diff --git a/Engine/Source/Runtime/Engine/Private/KismetArrayLibrary.cpp b/Engine/Source/Runtime/Engine/Private/KismetArrayLibrary.cpp
--- a/Engine/Source/Runtime/Engine/Private/KismetArrayLibrary.cpp
+++ b/Engine/Source/Runtime/Engine/Private/KismetArrayLibrary.cpp
@@ -597,7 +597,7 @@ void UKismetArrayLibrary::GenericArray_HandleBool(const FProperty* Property, voi
 	const FBoolProperty* BoolProperty = CastField<const FBoolProperty>(Property);
 	if (BoolProperty)
 	{
-		ensure((BoolProperty->ElementSize * BoolProperty->ArrayDim) == sizeof(uint8));
+		ensure((BoolProperty->GetElementSize() * BoolProperty->ArrayDim) == sizeof(uint8));
 		BoolProperty->SetPropertyValue(ItemPtr, 0 != *(reinterpret_cast<uint8*>(ItemPtr)));
 	}
 }
diff --git a/Engine/Source/Runtime/Engine/Private/NetDriver.cpp b/Engine/Source/Runtime/Engine/Private/NetDriver.cpp
--- a/Engine/Source/Runtime/Engine/Private/NetDriver.cpp
+++ b/Engine/Source/Runtime/Engine/Private/NetDriver.cpp
@@ -232,7 +232,7 @@ namespace UE::Net::Private
 
 				void* Dest = It->ContainerPtrToValuePtr<void>(LocalParms);
 
-				const int32 CopySize = It->ElementSize * It->ArrayDim;
+				const int32 CopySize = It->GetElementSize() * It->ArrayDim;
 
 				check(((uint8*)Dest - (uint8*)LocalParms) + CopySize <= Function->ParmsSize);
 
diff --git a/Engine/Source/Runtime/Engine/Private/PropertyAccess.cpp b/Engine/Source/Runtime/Engine/Private/PropertyAccess.cpp
--- a/Engine/Source/Runtime/Engine/Private/PropertyAccess.cpp
+++ b/Engine/Source/Runtime/Engine/Private/PropertyAccess.cpp
@@ -33,7 +33,7 @@ struct FPropertyAccessSystem
 
 	static uint32 GetPropertyOffset(const FProperty* InProperty, int32 InArrayIndex = 0)
 	{
-		return (uint32)(InProperty->GetOffset_ForInternal() + InProperty->ElementSize * InArrayIndex);
+		return (uint32)(InProperty->GetOffset_ForInternal() + InProperty->GetElementSize() * InArrayIndex);
 	}
 
 	// Called on load to resolve all path segments to indirections
@@ -268,7 +268,7 @@ struct FPropertyAccessSystem
 		case EPropertyAccessCopyType::Plain:
 			checkSlow(InSrcProperty->PropertyFlags & CPF_IsPlainOldData);
 			checkSlow(InDestProperty->PropertyFlags & CPF_IsPlainOldData);
-			FMemory::Memcpy(InDestAddr, InSrcAddr, InSrcProperty->ElementSize);
+			FMemory::Memcpy(InDestAddr, InSrcAddr, InSrcProperty->GetElementSize());
 			break;
 		case EPropertyAccessCopyType::Complex:
 			InSrcProperty->CopyCompleteValue(InDestAddr, InSrcAddr);
diff --git a/Engine/Source/Runtime/Engine/Private/RepLayout.cpp b/Engine/Source/Runtime/Engine/Private/RepLayout.cpp
--- a/Engine/Source/Runtime/Engine/Private/RepLayout.cpp
+++ b/Engine/Source/Runtime/Engine/Private/RepLayout.cpp
@@ -887,7 +887,7 @@ static void SerializeReadWritePropertyChecksum(
 		Cmd.Property->ArrayDim = 1;
 
 		TArray<uint8> TempPropMemory;
-		TempPropMemory.AddZeroed(Cmd.Property->ElementSize + 4);
+		TempPropMemory.AddZeroed(Cmd.Property->GetElementSize() + 4);
 		uint32* Guard = (uint32*)&TempPropMemory[TempPropMemory.Num() - 4];
 		const uint32 TAG_VALUE = 0xABADF00D;
 		*Guard = TAG_VALUE;
@@ -5386,7 +5386,7 @@ static uint32 AddPropertyCmd(
 	Cmd.Property = StackParams.Property;
 	Cmd.Type = ERepLayoutCmdType::Property;		// Initially set to generic type
 	Cmd.Offset = StackParams.Offset;
-	Cmd.ElementSize = Cmd.Property->ElementSize;
+	Cmd.ElementSize = Cmd.Property->GetElementSize();
 	Cmd.RelativeHandle = StackParams.RelativeHandle;
 	Cmd.ParentIndex = SharedParams.ParentIndex;
 	Cmd.CompatibleChecksum = GetRepLayoutCmdCompatibleChecksum(SharedParams, StackParams);
@@ -5527,7 +5527,7 @@ static FORCEINLINE uint32 AddArrayCmd(
 	Cmd.Type = ERepLayoutCmdType::DynamicArray;
 	Cmd.Property = StackParams.Property;
 	Cmd.Offset = StackParams.Offset;
-	Cmd.ElementSize = static_cast<FArrayProperty*>(StackParams.Property)->Inner->ElementSize;
+	Cmd.ElementSize = static_cast<FArrayProperty*>(StackParams.Property)->Inner->GetElementSize();
 	Cmd.RelativeHandle = StackParams.RelativeHandle;
 	Cmd.ParentIndex = SharedParams.ParentIndex;
 	Cmd.CompatibleChecksum = GetRepLayoutCmdCompatibleChecksum(SharedParams, StackParams);
@@ -5611,7 +5611,7 @@ static int32 InitFromStructProperty(
 	{
 		for (int32 j = 0; j < NetProperties[i]->ArrayDim; j++)
 		{
-			const int32 ArrayElementOffset = j * NetProperties[i]->ElementSize;
+			const int32 ArrayElementOffset = j * NetProperties[i]->GetElementSize();
 
 			FInitFromPropertyStackParams NewStackParams{
 				/*Property=*/NetProperties[i],
@@ -6089,7 +6089,7 @@ void FRepLayout::InitFromClass(
 		check(ParentHandle == i);
 		check(Parents[i].Property->RepIndex + Parents[i].ArrayIndex == i);
 
-		const int32 ParentOffset = Property->ElementSize * ArrayIdx;
+		const int32 ParentOffset = Property->GetElementSize() * ArrayIdx;
 
 		FInitFromPropertySharedParams SharedParams
 		{
@@ -6423,7 +6423,7 @@ void FRepLayout::InitFromFunction(
 			FInitFromPropertyStackParams StackParams
 			{
 				/*Property=*/*It,
-				/*Offset=*/It->ElementSize* ArrayIdx,
+				/*Offset=*/It->GetElementSize()* ArrayIdx,
 				/*RelativeHandle=*/RelativeHandle,
 				/*ParentChecksum=*/0,
 				/*StaticArrayIndex=*/ArrayIdx
@@ -6489,7 +6489,7 @@ void FRepLayout::InitFromStruct(
 			FInitFromPropertyStackParams StackParams
 			{
 				/*Property=*/*It,
-				/*Offset=*/It->ElementSize * ArrayIdx,
+				/*Offset=*/It->GetElementSize() * ArrayIdx,
 				/*RelativeHandle=*/RelativeHandle,
 				/*ParentChecksum=*/0,
 				/*StaticArrayIndex=*/ArrayIdx
diff --git a/Engine/Source/Runtime/Experimental/Dataflow/Core/Private/Dataflow/DataflowNode.cpp b/Engine/Source/Runtime/Experimental/Dataflow/Core/Private/Dataflow/DataflowNode.cpp
--- a/Engine/Source/Runtime/Experimental/Dataflow/Core/Private/Dataflow/DataflowNode.cpp
+++ b/Engine/Source/Runtime/Experimental/Dataflow/Core/Private/Dataflow/DataflowNode.cpp
@@ -30,7 +30,7 @@ static uint32 GetArrayElementOffsetFromReference(const FArrayProperty* const Arr
 	if (const void* const AddressAtIndex = ArrayProperty->GetValueAddressAtIndex_Direct(ArrayProperty->Inner, const_cast<void*>(Reference.ContainerReference), Reference.Index))
 	{
 		check((size_t)Reference.Reference >= (size_t)AddressAtIndex);
-		check((int32)((size_t)Reference.Reference - (size_t)AddressAtIndex) < ArrayProperty->Inner->ElementSize);
+		check((int32)((size_t)Reference.Reference - (size_t)AddressAtIndex) < ArrayProperty->Inner->GetElementSize());
 		return (uint32)((size_t)Reference.Reference - (size_t)AddressAtIndex);
 	}
 	return INDEX_NONE;
diff --git a/Engine/Source/Runtime/Experimental/Iris/Core/Private/Iris/ReplicationState/InternalPropertyReplicationState.cpp b/Engine/Source/Runtime/Experimental/Iris/Core/Private/Iris/ReplicationState/InternalPropertyReplicationState.cpp
--- a/Engine/Source/Runtime/Experimental/Iris/Core/Private/Iris/ReplicationState/InternalPropertyReplicationState.cpp
+++ b/Engine/Source/Runtime/Experimental/Iris/Core/Private/Iris/ReplicationState/InternalPropertyReplicationState.cpp
@@ -257,7 +257,7 @@ bool InternalCompareStructProperty(const FReplicationStateDescriptor* StructDesc
 		const FReplicationStateMemberDescriptor& MemberDescriptor = MemberDescriptors[StructMemberIt];
 		const FReplicationStateMemberPropertyDescriptor& MemberPropertyDescriptor = MemberPropertyDescriptors[StructMemberIt];
 		const FProperty* MemberProperty = MemberProperties[StructMemberIt];
-		const SIZE_T MemberOffset = MemberProperty->GetOffset_ForGC() + MemberProperty->ElementSize*MemberPropertyDescriptor.ArrayIndex;
+		const SIZE_T MemberOffset = MemberProperty->GetOffset_ForGC() + MemberProperty->GetElementSize()*MemberPropertyDescriptor.ArrayIndex;
 		if (!InternalCompareMember(StructDescriptor, StructMemberIt, static_cast<const uint8*>(ValueA) + MemberOffset, static_cast<const uint8*>(ValueB) + MemberOffset))
 		{
 			return false;
@@ -302,7 +302,7 @@ void InternalApplyStructProperty(const FReplicationStateDescriptor* StructDescri
 		const FReplicationStateMemberDescriptor& MemberDescriptor = MemberDescriptors[StructMemberIt];
 		const FReplicationStateMemberPropertyDescriptor& MemberPropertyDescriptor = MemberPropertyDescriptors[StructMemberIt];
 		const FProperty* MemberProperty = MemberProperties[StructMemberIt];
-		const SIZE_T MemberOffset = MemberProperty->GetOffset_ForGC() + MemberProperty->ElementSize*MemberPropertyDescriptor.ArrayIndex;
+		const SIZE_T MemberOffset = MemberProperty->GetOffset_ForGC() + MemberProperty->GetElementSize()*MemberPropertyDescriptor.ArrayIndex;
 		InternalApplyPropertyValue(StructDescriptor, StructMemberIt, static_cast<uint8*>(Dst) + MemberOffset, static_cast<const uint8*>(Src) + MemberOffset);
 	}
 }
@@ -411,7 +411,7 @@ void InternalCopyStructProperty(const FReplicationStateDescriptor* StructDescrip
 		const FReplicationStateMemberDescriptor& MemberDescriptor = MemberDescriptors[StructMemberIt];
 		const FReplicationStateMemberPropertyDescriptor& MemberPropertyDescriptor = MemberPropertyDescriptors[StructMemberIt];
 		const FProperty* MemberProperty = MemberProperties[StructMemberIt];
-		const SIZE_T MemberOffset = MemberProperty->GetOffset_ForGC() + MemberProperty->ElementSize*MemberPropertyDescriptor.ArrayIndex;
+		const SIZE_T MemberOffset = MemberProperty->GetOffset_ForGC() + MemberProperty->GetElementSize()*MemberPropertyDescriptor.ArrayIndex;
 		InternalCopyPropertyValue(StructDescriptor, StructMemberIt, static_cast<uint8*>(Dst) + MemberOffset, static_cast<const uint8*>(Src) + MemberOffset);
 	}
 }
diff --git a/Engine/Source/Runtime/Experimental/Iris/Core/Private/Iris/ReplicationState/PropertyReplicationState.cpp b/Engine/Source/Runtime/Experimental/Iris/Core/Private/Iris/ReplicationState/PropertyReplicationState.cpp
--- a/Engine/Source/Runtime/Experimental/Iris/Core/Private/Iris/ReplicationState/PropertyReplicationState.cpp
+++ b/Engine/Source/Runtime/Experimental/Iris/Core/Private/Iris/ReplicationState/PropertyReplicationState.cpp
@@ -284,7 +284,7 @@ bool FPropertyReplicationState::PollPropertyReplicationState(const void* RESTRIC
 			const FProperty* Property = MemberProperties[MemberIt];
 
 			//$TODO: make special version to avoid unnecessary overhead.
-			PollPropertyValue(MemberIt, SrcBuffer + Property->GetOffset_ForGC() + Property->ElementSize*MemberPropertyDescriptor.ArrayIndex);
+			PollPropertyValue(MemberIt, SrcBuffer + Property->GetOffset_ForGC() + Property->GetElementSize()*MemberPropertyDescriptor.ArrayIndex);
 		}
 	}
 
@@ -317,9 +317,9 @@ bool FPropertyReplicationState::StoreCurrentPropertyReplicationStateForRepNotifi
 				const FProperty* Property = MemberProperties[MemberIt];
 
 				void* DstValue = StateBuffer + Descriptor->MemberDescriptors[MemberIt].ExternalMemberOffset;
-				const void* SrcValue = SrcBuffer + Property->GetOffset_ForGC() + Property->ElementSize*MemberPropertyDescriptor.ArrayIndex;
+				const void* SrcValue = SrcBuffer + Property->GetOffset_ForGC() + Property->GetElementSize()*MemberPropertyDescriptor.ArrayIndex;
 
-				Private::InternalCopyPropertyValue(Descriptor, MemberIt, DstValue, SrcBuffer + Property->GetOffset_ForGC() + Property->ElementSize*MemberPropertyDescriptor.ArrayIndex);
+				Private::InternalCopyPropertyValue(Descriptor, MemberIt, DstValue, SrcBuffer + Property->GetOffset_ForGC() + Property->GetElementSize()*MemberPropertyDescriptor.ArrayIndex);
 			}
 		}
 	}
@@ -359,7 +359,7 @@ void FPropertyReplicationState::PushPropertyReplicationState(const UObject* Owne
 				const FReplicationStateMemberPropertyDescriptor& MemberPropertyDescriptor = MemberPropertyDescriptors[MemberIt];
 				const FProperty* Property = MemberProperties[MemberIt];
 
-				PushPropertyValue(MemberIt, DstBuffer + Property->GetOffset_ForGC() + Property->ElementSize*MemberPropertyDescriptor.ArrayIndex);
+				PushPropertyValue(MemberIt, DstBuffer + Property->GetOffset_ForGC() + Property->GetElementSize()*MemberPropertyDescriptor.ArrayIndex);
 
 #if WITH_PUSH_MODEL
 				if (MemberPropertyDescriptor.ArrayIndex == 0)
@@ -418,7 +418,7 @@ bool FPropertyReplicationState::PollObjectReferences(const void* RESTRICT SrcSta
 				const FReplicationStateMemberPropertyDescriptor& MemberPropertyDescriptor = MemberPropertyDescriptors[MemberIt];
 				const FProperty* Property = MemberProperties[MemberIt];
 
-				PollPropertyValue(MemberIt, SrcBuffer + Property->GetOffset_ForGC() + Property->ElementSize*MemberPropertyDescriptor.ArrayIndex);
+				PollPropertyValue(MemberIt, SrcBuffer + Property->GetOffset_ForGC() + Property->GetElementSize()*MemberPropertyDescriptor.ArrayIndex);
 			}
 		}
 	}
diff --git a/Engine/Source/Runtime/Experimental/Iris/Core/Private/Iris/ReplicationState/ReplicationStateDescriptorBuilder.cpp b/Engine/Source/Runtime/Experimental/Iris/Core/Private/Iris/ReplicationState/ReplicationStateDescriptorBuilder.cpp
--- a/Engine/Source/Runtime/Experimental/Iris/Core/Private/Iris/ReplicationState/ReplicationStateDescriptorBuilder.cpp
+++ b/Engine/Source/Runtime/Experimental/Iris/Core/Private/Iris/ReplicationState/ReplicationStateDescriptorBuilder.cpp
@@ -574,7 +574,7 @@ void FPropertyReplicationStateDescriptorBuilder::BuildMemberCache(FBuilderContex
 			  * Use ElementSize instead of GetSize() since GetSize() will multiply by array size and currently we treat
 			  * each element individually. 
 			  */
-			CurrentCacheEntry->ExternalSizeAndAlignment = { (SIZE_T)Member.Property->ElementSize, (SIZE_T)Member.Property->GetMinAlignment() };
+			CurrentCacheEntry->ExternalSizeAndAlignment = { (SIZE_T)Member.Property->GetElementSize(), (SIZE_T)Member.Property->GetMinAlignment() };
 			CurrentCacheEntry->InternalSizeAndAlignment = { Descriptor->InternalSize, Descriptor->InternalAlignment };
 
 			CurrentCacheEntry->bIsStruct = 1;
@@ -607,7 +607,7 @@ void FPropertyReplicationStateDescriptorBuilder::BuildMemberCache(FBuilderContex
 				Context.BuildParams.bAllMembersAreReplicated = false;
 			}
 
-			CurrentCacheEntry->ExternalSizeAndAlignment = { (SIZE_T)Member.Property->ElementSize, (SIZE_T)Member.Property->GetMinAlignment() };
+			CurrentCacheEntry->ExternalSizeAndAlignment = { (SIZE_T)Member.Property->GetElementSize(), (SIZE_T)Member.Property->GetMinAlignment() };
 			CurrentCacheEntry->InternalSizeAndAlignment = { Serializer->QuantizedTypeSize, Serializer->QuantizedTypeAlignment };
 
 			CurrentCacheEntry->bIsDynamicArray = 1;
@@ -616,7 +616,7 @@ void FPropertyReplicationStateDescriptorBuilder::BuildMemberCache(FBuilderContex
 		{
 			bSomeMembersAreProperties = true;
 
-			CurrentCacheEntry->ExternalSizeAndAlignment = { (SIZE_T)Member.Property->ElementSize, (SIZE_T)Member.Property->GetMinAlignment() };
+			CurrentCacheEntry->ExternalSizeAndAlignment = { (SIZE_T)Member.Property->GetElementSize(), (SIZE_T)Member.Property->GetMinAlignment() };
 			CurrentCacheEntry->InternalSizeAndAlignment = { Serializer->QuantizedTypeSize, Serializer->QuantizedTypeAlignment };
 		}
 		else
@@ -785,7 +785,7 @@ void FPropertyReplicationStateDescriptorBuilder::AllocateAndInitializeDefaultInt
 		FNetQuantizeArgs Args;
 		Args.Version = 0;
 		Args.NetSerializerConfig = MemberSerializerDescriptor.SerializerConfig;
-		Args.Source = reinterpret_cast<NetSerializerValuePointer>(SrcBuffer + Property->GetOffset_ForGC() + Property->ElementSize * MemberPropertyDescriptor.ArrayIndex);
+		Args.Source = reinterpret_cast<NetSerializerValuePointer>(SrcBuffer + Property->GetOffset_ForGC() + Property->GetElementSize() * MemberPropertyDescriptor.ArrayIndex);
 		Args.Target = reinterpret_cast<NetSerializerValuePointer>(DstStateBuffer + MemberDescriptor.InternalMemberOffset);
 
 		MemberSerializerDescriptor.Serializer->Quantize(Context, Args);
@@ -928,7 +928,7 @@ void FPropertyReplicationStateDescriptorBuilder::BuildMemberDescriptorsForStruct
 		LastProperty = Property;
 
 		ExternalBufferAlignment = FMath::Max(ExternalBufferAlignment, MemberCacheEntry->ExternalSizeAndAlignment.Alignment);
-		CurrentMemberDescriptor->ExternalMemberOffset = Property ? (Property->GetOffset_ForGC() + (ArrayIndex * (uint32)Property->ElementSize)) : 0U;
+		CurrentMemberDescriptor->ExternalMemberOffset = Property ? (Property->GetOffset_ForGC() + (ArrayIndex * (uint32)Property->GetElementSize())) : 0U;
 
 		// Internal
 		const FNetSerializer* MemberSerializer = MemberCacheEntry->Serializer;
diff --git a/Engine/Source/Runtime/JsonUtilities/Private/JsonObjectConverter.cpp b/Engine/Source/Runtime/JsonUtilities/Private/JsonObjectConverter.cpp
--- a/Engine/Source/Runtime/JsonUtilities/Private/JsonObjectConverter.cpp
+++ b/Engine/Source/Runtime/JsonUtilities/Private/JsonObjectConverter.cpp
@@ -226,7 +226,7 @@ TSharedPtr<FJsonValue> FJsonObjectConverter::UPropertyToJsonValue(FProperty* Pro
 	TArray< TSharedPtr<FJsonValue> > Array;
 	for (int Index = 0; Index != Property->ArrayDim; ++Index)
 	{
-		Array.Add(ConvertScalarFPropertyToJsonValue(Property, (char*)Value + Index * Property->ElementSize, CheckFlags, SkipFlags, ExportCb, OuterProperty, ConversionFlags));
+		Array.Add(ConvertScalarFPropertyToJsonValue(Property, (char*)Value + Index * Property->GetElementSize(), CheckFlags, SkipFlags, ExportCb, OuterProperty, ConversionFlags));
 	}
 	return MakeShared<FJsonValueArray>(Array);
 }
@@ -915,7 +915,7 @@ namespace
 		const int32 ItemsToRead = FMath::Clamp(ArrayValue.Num(), 0, Property->ArrayDim);
 		for (int Index = 0; Index != ItemsToRead; ++Index)
 		{
-			if (!ConvertScalarJsonValueToFPropertyWithContainer(ArrayValue[Index], Property, static_cast<char*>(OutValue) + Index * Property->ElementSize, ContainerStruct, Container, CheckFlags, SkipFlags, bStrictMode, OutFailReason, ImportCb))
+			if (!ConvertScalarJsonValueToFPropertyWithContainer(ArrayValue[Index], Property, static_cast<char*>(OutValue) + Index * Property->GetElementSize(), ContainerStruct, Container, CheckFlags, SkipFlags, bStrictMode, OutFailReason, ImportCb))
 			{
 				return false;
 			}
diff --git a/Engine/Source/Runtime/MovieScene/Public/MovieSceneCommonHelpers.h b/Engine/Source/Runtime/MovieScene/Public/MovieSceneCommonHelpers.h
--- a/Engine/Source/Runtime/MovieScene/Public/MovieSceneCommonHelpers.h
+++ b/Engine/Source/Runtime/MovieScene/Public/MovieSceneCommonHelpers.h
@@ -593,7 +593,7 @@ void FTrackInstancePropertyBindings::InvokeSetterFunction(UObject* InRuntimeObje
 				// The first encountered property is assumed to be the input value so initialize this with the user-specified value from InPropertyValue
 				if (Property->HasAnyPropertyFlags(CPF_Parm) && !Property->HasAnyPropertyFlags(CPF_ReturnParm) && bFirstProperty)
 				{
-					const bool bIsValid = ensureMsgf(sizeof(T) == Property->ElementSize, TEXT("Property type does not match for Sequencer setter function %s::%s (%ibytes != %ibytes"), *InRuntimeObject->GetName(), *Setter->GetName(), sizeof(T), Property->ElementSize);
+					const bool bIsValid = ensureMsgf(sizeof(T) == Property->GetElementSize(), TEXT("Property type does not match for Sequencer setter function %s::%s (%ibytes != %ibytes"), *InRuntimeObject->GetName(), *Setter->GetName(), sizeof(T), Property->GetElementSize());
 					if (bIsValid)
 					{
 						Property->CopyCompleteValue(Property->ContainerPtrToValuePtr<void>(Params), &InPropertyValue);
diff --git a/Engine/Source/Runtime/MovieSceneTracks/Private/MovieSceneTracksComponentTypes.cpp b/Engine/Source/Runtime/MovieSceneTracks/Private/MovieSceneTracksComponentTypes.cpp
--- a/Engine/Source/Runtime/MovieSceneTracks/Private/MovieSceneTracksComponentTypes.cpp
+++ b/Engine/Source/Runtime/MovieSceneTracks/Private/MovieSceneTracksComponentTypes.cpp
@@ -492,7 +492,7 @@ struct FBoolHandler : TPropertyComponentHandler<FBoolPropertyTraits, bool>
 					auto FieldMask = BoundProperty->GetFieldMask();
 					static_assert(std::is_same_v<decltype(FieldMask), uint8>, "Unexpected size of field mask returned from FBoolProperty::FieldMask");
 
-					OutMetaData.BitFieldSize = static_cast<uint8>(BoundProperty->ElementSize);
+					OutMetaData.BitFieldSize = static_cast<uint8>(BoundProperty->GetElementSize());
 					OutMetaData.BitIndex     = static_cast<uint8>(FMath::CountTrailingZeros(FieldMask));
 				}
 			}
diff --git a/Engine/Source/Runtime/PropertyPath/Private/PropertyPathHelpers.cpp b/Engine/Source/Runtime/PropertyPath/Private/PropertyPathHelpers.cpp
--- a/Engine/Source/Runtime/PropertyPath/Private/PropertyPathHelpers.cpp
+++ b/Engine/Source/Runtime/PropertyPath/Private/PropertyPathHelpers.cpp
@@ -180,7 +180,7 @@ namespace PropertyPathHelpersInternal
 					{
 						// Create and init a buffer for the function to write to
 						TArray<uint8> TempBuffer;
-						TempBuffer.AddUninitialized(ReturnProperty->ElementSize);
+						TempBuffer.AddUninitialized(ReturnProperty->GetElementSize());
 						ReturnProperty->InitializeValue(TempBuffer.GetData());
 
 						InContainer->ProcessEvent(InFunction, TempBuffer.GetData());
@@ -276,7 +276,7 @@ namespace PropertyPathHelpersInternal
 					{
 						// Create and init a buffer for the function to read from
 						TArray<uint8> TempBuffer;
-						TempBuffer.AddUninitialized(ParamProperty->ElementSize);
+						TempBuffer.AddUninitialized(ParamProperty->GetElementSize());
 						ParamProperty->InitializeValue(TempBuffer.GetData());
 
 						ParamProperty->ImportText_Direct(*InValue, TempBuffer.GetData(), nullptr, 0);
diff --git a/Engine/Source/Runtime/PropertyPath/Public/PropertyPathHelpers.h b/Engine/Source/Runtime/PropertyPath/Public/PropertyPathHelpers.h
--- a/Engine/Source/Runtime/PropertyPath/Public/PropertyPathHelpers.h
+++ b/Engine/Source/Runtime/PropertyPath/Public/PropertyPathHelpers.h
@@ -937,7 +937,7 @@ namespace PropertyPathHelpersInternal
 						const FPropertyStructView& InStuctView = static_cast<const FPropertyStructView&>(InValue);
 
 						// Ensure that the element sizes are the same, prevents the user from doing something terribly wrong.
-						if (ParamProperty->ElementSize == InStuctView.ScriptStruct->GetStructureSize() && !InContainer->IsUnreachable())
+						if (ParamProperty->GetElementSize() == InStuctView.ScriptStruct->GetStructureSize() && !InContainer->IsUnreachable())
 						{
 							InContainer->ProcessEvent(InFunction, const_cast<uint8*>(InStuctView.Memory));
 							return true;
@@ -980,7 +980,7 @@ namespace PropertyPathHelpersInternal
 
 				// Ensure that the element sizes are the same, prevents the user from doing something terribly wrong.
 				ArrayIndex = ArrayIndex == INDEX_NONE ? 0 : ArrayIndex;
-				if ( Property->ElementSize == InStuctView.ScriptStruct->GetStructureSize() && ArrayIndex < Property->ArrayDim)
+				if ( Property->GetElementSize() == InStuctView.ScriptStruct->GetStructureSize() && ArrayIndex < Property->ArrayDim)
 				{
 					if (Property->HasSetter())
 					{
@@ -1120,7 +1120,7 @@ namespace PropertyPathHelpersInternal
 						const FPropertyStructView& InStuctView = static_cast<const FPropertyStructView&>(InValue);
 
 						// Ensure that the element sizes are the same, prevents the user from doing something terribly wrong.
-						if (ArrayProp->Inner->ElementSize == InStuctView.ScriptStruct->GetStructureSize())
+						if (ArrayProp->Inner->GetElementSize() == InStuctView.ScriptStruct->GetStructureSize())
 						{
 							if (void* Address = static_cast<void*>(ArrayHelper.GetRawPtr(ArrayIndex)))
 							{
diff --git a/Engine/Source/Runtime/PropertyPath/Public/PropertyTypeCompatibility.h b/Engine/Source/Runtime/PropertyPath/Public/PropertyTypeCompatibility.h
--- a/Engine/Source/Runtime/PropertyPath/Public/PropertyTypeCompatibility.h
+++ b/Engine/Source/Runtime/PropertyPath/Public/PropertyTypeCompatibility.h
@@ -403,10 +403,10 @@ inline bool PropertySizesMatch_Impl(FProperty* InProperty)
 {
 	if (FByteProperty* ByteProperty = CastField<FByteProperty>(InProperty))
 	{
-		return InProperty->ElementSize == sizeof(uint8);
+		return InProperty->GetElementSize() == sizeof(uint8);
 	}
 
-	return InProperty->ElementSize == sizeof(T);
+	return InProperty->GetElementSize() == sizeof(T);
 }
 
 template<>
@@ -417,7 +417,7 @@ inline bool PropertySizesMatch_Impl<uint8>(FProperty* InProperty)
 		return !BoolProperty->IsNativeBool();
 	}
 
-	return InProperty->ElementSize == sizeof(uint8);
+	return InProperty->GetElementSize() == sizeof(uint8);
 }
 
 template<>
@@ -428,7 +428,7 @@ inline bool PropertySizesMatch_Impl<uint16>(FProperty* InProperty)
 		return !BoolProperty->IsNativeBool();
 	}
 
-	return InProperty->ElementSize == sizeof(uint16);
+	return InProperty->GetElementSize() == sizeof(uint16);
 }
 
 template<>
@@ -439,7 +439,7 @@ inline bool PropertySizesMatch_Impl<uint32>(FProperty* InProperty)
 		return !BoolProperty->IsNativeBool();
 	}
 
-	return InProperty->ElementSize == sizeof(uint32);
+	return InProperty->GetElementSize() == sizeof(uint32);
 }
 
 template<>
@@ -450,7 +450,7 @@ inline bool PropertySizesMatch_Impl<uint64>(FProperty* InProperty)
 		return !BoolProperty->IsNativeBool();
 	}
 
-	return InProperty->ElementSize == sizeof(uint64);
+	return InProperty->GetElementSize() == sizeof(uint64);
 }
 
 template<typename T>
