[RegressionFix] IDO Impersonation broke on save causing unresolved loose properties to never get saved

diff --git a/Engine/Source/Runtime/CoreUObject/Internal/UObject/PropertyBagRepository.h b/Engine/Source/Runtime/CoreUObject/Internal/UObject/PropertyBagRepository.h
--- a/Engine/Source/Runtime/CoreUObject/Internal/UObject/PropertyBagRepository.h
+++ b/Engine/Source/Runtime/CoreUObject/Internal/UObject/PropertyBagRepository.h
@@ -229,6 +229,7 @@ struct FScopedIDOSerializationContext
 #if WITH_EDITORONLY_DATA
 	COREUOBJECT_API FScopedIDOSerializationContext(UObject* InObject, FArchive& Archive);
 	COREUOBJECT_API FScopedIDOSerializationContext(UObject* InObject, bool bImpersonate); // assumes save
+	COREUOBJECT_API FScopedIDOSerializationContext(bool bImpersonate); // assumes save
 	COREUOBJECT_API ~FScopedIDOSerializationContext();
 
 	bool bCreateIDO = false;
@@ -249,6 +250,7 @@ private:
 #else
 	inline FScopedIDOSerializationContext(UObject* InObject, FArchive& Archive) {}
 	inline explicit FScopedIDOSerializationContext(UObject* InObject, bool bImpersonate) {}
+	inline explicit FScopedIDOSerializationContext(bool bImpersonate) {}
 #endif
 
 	FScopedIDOSerializationContext(const FScopedIDOSerializationContext&) = delete;
diff --git a/Engine/Source/Runtime/CoreUObject/Private/UObject/PropertyBagRepository.cpp b/Engine/Source/Runtime/CoreUObject/Private/UObject/PropertyBagRepository.cpp
--- a/Engine/Source/Runtime/CoreUObject/Private/UObject/PropertyBagRepository.cpp
+++ b/Engine/Source/Runtime/CoreUObject/Private/UObject/PropertyBagRepository.cpp
@@ -959,6 +959,20 @@ FScopedIDOSerializationContext::FScopedIDOSerializationContext(UObject* InObject
 	}
 }
 
+FScopedIDOSerializationContext::FScopedIDOSerializationContext(bool bImpersonate)
+	: bCreateIDO(false)
+	, Archive(nullptr)
+	, Object(nullptr)
+	, PreSerializeOffset(0)
+{
+	FUObjectSerializeContext* SerializeContext = FUObjectThreadContext::Get().GetSerializeContext();
+	bool bHasIDOSupport = IsInstanceDataObjectSupportEnabled();
+	if (bHasIDOSupport)
+	{
+		ScopedImpersonateProperties.Emplace(SerializeContext->bImpersonateProperties, bImpersonate);
+	}
+}
+
 FScopedIDOSerializationContext::~FScopedIDOSerializationContext()
 {
 	if (bCreateIDO)
diff --git a/Engine/Source/Runtime/CoreUObject/Private/UObject/SavePackage2.cpp b/Engine/Source/Runtime/CoreUObject/Private/UObject/SavePackage2.cpp
--- a/Engine/Source/Runtime/CoreUObject/Private/UObject/SavePackage2.cpp
+++ b/Engine/Source/Runtime/CoreUObject/Private/UObject/SavePackage2.cpp
@@ -2997,7 +2997,7 @@ ESavePackageResult InnerSave(FSaveContext& SaveContext)
 {
 	SaveContext.SetEDLCookChecker(&FEDLCookCheckerThreadState::Get());
 
-	UE::FScopedIDOSerializationContext IDOSaveContext(nullptr, !SaveContext.IsCooking());
+	UE::FScopedIDOSerializationContext IDOSaveContext(!SaveContext.IsCooking());
 
 	// Create slow task dialog if needed
 	const int32 TotalSaveSteps = 3;
