Reduce verbosity (Display -> Log) of loader log message detailing how many exports were removed as part of RemoveUnreachableObjects

#rnx

diff --git a/Engine/Source/Runtime/CoreUObject/Private/Serialization/AsyncLoading2.cpp b/Engine/Source/Runtime/CoreUObject/Private/Serialization/AsyncLoading2.cpp
--- a/Engine/Source/Runtime/CoreUObject/Private/Serialization/AsyncLoading2.cpp
+++ b/Engine/Source/Runtime/CoreUObject/Private/Serialization/AsyncLoading2.cpp
@@ -8965,7 +8965,7 @@ void FAsyncLoadingThread2::RemoveUnreachableObjects(FUnreachableObjects& Objects
 	const int32 RemovedPublicExportCount = OldPublicExportCount - NewPublicExportCount;
 
 	const double StopTime = FPlatformTime::Seconds();
-	UE_LOG(LogStreaming, Display,
+	UE_LOG(LogStreaming, Log,
 		TEXT("%.3f ms for processing %d objects in RemoveUnreachableObjects(Queued=%d, Async=%d). ")
 		TEXT("Removed %d (%d->%d) packages and %d (%d->%d) public exports."),
 		(StopTime - StartTime) * 1000,
