[Backout] - CL35428516 (but keeping the comment and variable typo fixes)

Original CL Desc
-----------------------------------------------------------------
- Changes DoesPackageExist to look at the IoDispatcher (if initialized) before looking at the filesystem to match path casing. Checking the OS path casing is very expensive compared to checking the IOStore, so this ends up being a good win for cooked editor builds in particular since cooked editors will make frequent Filesystem filter requests but many assets will reside in chunk data. Opening a large project goes from 15s -> 10s spent in DoesPackageExist (more specifically TryMatchCaseOnDisk)
-- Non-cooked editor builds won't touch this code path in most cases since earlier code will have checked the AssetRegistry for the package case.
-- Runtime code will not commonly request from the filesystem and thus won't be impacted by this change
- Minor typo fixes in comments

#rnx

diff --git a/Engine/Source/Runtime/CoreUObject/Private/Misc/PackageName.cpp b/Engine/Source/Runtime/CoreUObject/Private/Misc/PackageName.cpp
--- a/Engine/Source/Runtime/CoreUObject/Private/Misc/PackageName.cpp
+++ b/Engine/Source/Runtime/CoreUObject/Private/Misc/PackageName.cpp
@@ -2059,13 +2059,31 @@ FPackageName::EPackageLocationFilter FPackageName::InternalDoesPackageExistEx(co
 
 	EPackageLocationFilter Result = EPackageLocationFilter::None;
 
-	bool bCheckIoDispatcher = (uint8)Filter & (uint8)EPackageLocationFilter::IoDispatcher;
-	bool bCheckFilesystem = (uint8)Filter & (uint8)EPackageLocationFilter::FileSystem;
-	bool bAttemptSkipTryMatchCase = bCheckFilesystem && bMatchCaseOnDisk;
+	if ((uint8)Filter & (uint8)EPackageLocationFilter::FileSystem)
+	{
+		bool bFoundInFileSystem = false;
+		if (bMatchCaseOnDisk)
+		{
+			bFoundInFileSystem = IPackageResourceManager::Get().TryMatchCaseOnDisk(PackagePath, OutPackagePath);
+		}
+		else
+		{
+			bFoundInFileSystem = IPackageResourceManager::Get().DoesPackageExist(PackagePath, OutPackagePath);
+		}
 
-	// Matching path case on disk is expensive. Even if the package filter doesn't include IoDispatcher, check 
-	// to see if the IoDispatcher contains this path since if it's there we can skip matching the case on disk.
-	if (bCheckIoDispatcher || bAttemptSkipTryMatchCase)
+		if (bFoundInFileSystem)
+		{
+			Result = EPackageLocationFilter((uint8)Result | (uint8)EPackageLocationFilter::FileSystem);
+
+			// if we just want to find any existence, then we are done
+			if (Filter == EPackageLocationFilter::Any)
+			{
+				return Result;
+			}
+		}
+	}
+
+	if (((uint8)Filter & (uint8)EPackageLocationFilter::IoDispatcher))
 	{
 		bool bFoundInIoDispatcher = false;
 		if (DoesPackageExistOverrideDelegate.IsBound())
@@ -2091,36 +2109,7 @@ FPackageName::EPackageLocationFilter FPackageName::InternalDoesPackageExistEx(co
 				}
 			}
 
-			static_assert(uint8(EPackageLocationFilter::None) == 0);
-			const uint8 CheckIoDispatcherBitMask = ~uint8(bCheckIoDispatcher) + 1;
-			Result = EPackageLocationFilter(uint8(Result) | (uint8(EPackageLocationFilter::IoDispatcher) & CheckIoDispatcherBitMask));
-
-			// if we just want to find any existence, then we are done
-			if (Filter == EPackageLocationFilter::Any)
-			{
-				return Result;
-			}
-
-			// We do not support assets in both the filesystem and IoDispatcher simultaneously, so skip checking the filesystem if requested
-			bCheckFilesystem = false;
-		}
-	}
-
-	if (bCheckFilesystem)
-	{
-		bool bFoundInFileSystem = false;
-		if (bMatchCaseOnDisk)
-		{
-			bFoundInFileSystem = IPackageResourceManager::Get().TryMatchCaseOnDisk(PackagePath, OutPackagePath);
-		}
-		else
-		{
-			bFoundInFileSystem = IPackageResourceManager::Get().DoesPackageExist(PackagePath, OutPackagePath);
-		}
-
-		if (bFoundInFileSystem)
-		{
-			Result = EPackageLocationFilter((uint8)Result | (uint8)EPackageLocationFilter::FileSystem);
+			Result = EPackageLocationFilter((uint8)Result | (uint8)EPackageLocationFilter::IoDispatcher);
 
 			// if we just want to find any existence, then we are done
 			if (Filter == EPackageLocationFilter::Any)
