Fix compile error when NO_LOGGING=1

#rnx

- The method ::GetCategoryName() does not exist when NO_LOGGING is enabled.

diff --git a/Engine/Source/Runtime/CoreUObject/Private/UObject/Class.cpp b/Engine/Source/Runtime/CoreUObject/Private/UObject/Class.cpp
--- a/Engine/Source/Runtime/CoreUObject/Private/UObject/Class.cpp
+++ b/Engine/Source/Runtime/CoreUObject/Private/UObject/Class.cpp
@@ -4267,6 +4267,7 @@ static FString GetSuggestedPathNameForTypeShortName(const FString& ShortName)
 template <typename T>
 static void LogMetaDataShortTypeName(const T* Field, FName MetaDataKey, const FString& MetaDataValue)
 {
+#if !NO_LOGGING
 	static struct FLogShortTypeNameInMetaDataCheckSettings
 	{
 		ELogVerbosity::Type LogVerbosity = ELogVerbosity::Warning;
@@ -4322,6 +4323,7 @@ static void LogMetaDataShortTypeName(const T* Field, FName MetaDataKey, const FS
 			*GetSuggestedPathNameForTypeShortName(MetaDataValue),
 			*GetFieldLocation(OwnerStruct));
 	}
+#endif //# !NO_LOGGING
 }
 
 /** Checks MetaData values for known MetaData keys to see if they contain short type names **/
