- Enables -nozenloader as the default for AsyncPackageLoader until we can resolve errors it causes in FindPackageRefChecked regarding "Package xxxx has been deleted"
- Prevents finishing ready for removal from LoadedPackagesToProcess until we are no longer recursively processing the package. If a package has a deferred PostLoad that initiates a syncronous load, we may attempt to remove the same package more than once.

#rnx

diff --git a/Engine/Source/Runtime/CoreUObject/Private/Serialization/AsyncLoading.cpp b/Engine/Source/Runtime/CoreUObject/Private/Serialization/AsyncLoading.cpp
--- a/Engine/Source/Runtime/CoreUObject/Private/Serialization/AsyncLoading.cpp
+++ b/Engine/Source/Runtime/CoreUObject/Private/Serialization/AsyncLoading.cpp
@@ -4791,8 +4791,9 @@ EAsyncPackageState::Type FAsyncLoadingThread::ProcessLoadedPackages(bool bUseTim
 			{
 				break;
 			}
-			// if we are done with post load and there isn't any package currently depending on this package we can start to clean it up
-			else if (Package->GetDependencyRefCount() == 0)
+			// if we are done with post load (including not recursively handling one currently) 
+			// and there isn't any package currently depending on this package we can start to clean it up
+			else if (Package->GetDependencyRefCount() == 0 && !Package->IsBeingProcessedRecursively())
 			{
 				check(Result == EAsyncPackageState::Complete);
 				// Remove the package from the list before we trigger the callbacks, 
@@ -4805,10 +4806,7 @@ EAsyncPackageState::Type FAsyncLoadingThread::ProcessLoadedPackages(bool bUseTim
 					if (FPlatformProperties::RequiresCookedData())
 					{
 						// Emulates ResetLoaders on the package linker's linkerroot.
-						if (!Package->IsBeingProcessedRecursively())
-						{
-							Package->ResetLoader();
-						}
+						Package->ResetLoader();
 					}
 					else
 					{
diff --git a/Engine/Source/Runtime/CoreUObject/Private/Serialization/AsyncPackageLoader.cpp b/Engine/Source/Runtime/CoreUObject/Private/Serialization/AsyncPackageLoader.cpp
--- a/Engine/Source/Runtime/CoreUObject/Private/Serialization/AsyncPackageLoader.cpp
+++ b/Engine/Source/Runtime/CoreUObject/Private/Serialization/AsyncPackageLoader.cpp
@@ -15,7 +15,7 @@
 #include "IO/IoDispatcherInternal.h"
 #include "HAL/IConsoleManager.h"
 
-static bool GUseOldLoaderAsFallback = false;
+static bool GUseOldLoaderAsFallback = true;
 static FAutoConsoleVariableRef CVarGUseOldLoaderAsFallback(
 	TEXT("s.UseOldLoaderAsFallback"),
 	GUseOldLoaderAsFallback,
