Cooker: Add instrumentation for error PackageSplitter found the Generator package still in memory after it should have been deleted by GC.
Allow specifying a logverbosity in FReferenceChainSearch.
Fix declaration of dependencies from each UPackage to its public UObjects to still be active when reference chain diagnostics are run.
When the error occurs, rerun the GC with FGCHistory to provide FReferenceChainSearch   with information about hidden dependencies.
#rnx

diff --git a/Engine/Source/Editor/UnrealEd/Classes/CookOnTheSide/CookOnTheFlyServer.h b/Engine/Source/Editor/UnrealEd/Classes/CookOnTheSide/CookOnTheFlyServer.h
--- a/Engine/Source/Editor/UnrealEd/Classes/CookOnTheSide/CookOnTheFlyServer.h
+++ b/Engine/Source/Editor/UnrealEd/Classes/CookOnTheSide/CookOnTheFlyServer.h
@@ -176,6 +176,7 @@ namespace UE::Cook
 	class FBuildDefinitions;
 	class FCachedDependencies;
 	class FCookDirector;
+	class FCookGCDiagnosticContext;
 	class FCookSandbox;
 	class FCookWorkerClient;
 	class FCookWorkerServer;
@@ -827,11 +828,14 @@ public:
 	UNREALED_API void SetGarbageCollectType(uint32 ResultFlagsFromTick);
 	UNREALED_API void ClearGarbageCollectType();
 
+	UNREALED_API void OnCookerStartCollectGarbage();
+	UNREALED_API void OnCookerEndCollectGarbage();
 	UNREALED_API void EvaluateGarbageCollectionResults(bool bWasDueToOOM, bool bWasPartialGC, uint32 ResultFlags,
 		int32 NumObjectsBeforeGC, const FPlatformMemoryStats& MemStatsBeforeGC,
 		const FGenericMemoryStats& AllocatorStatsBeforeGC,
 		int32 NumObjectsAfterGC, const FPlatformMemoryStats& MemStatsAfterGC,
 		const FGenericMemoryStats& AllocatorStatsAfterGC);
+	UNREALED_API bool NeedsDiagnosticSecondGC() const;
 
 	/**
 	 * RequestPackage to be cooked
@@ -1575,6 +1579,7 @@ private:
 	TUniquePtr<UE::Cook::FCookWorkerClient> CookWorkerClient;
 	TUniquePtr<FLayeredCookArtifactReader> AllContextArtifactReader;
 	TSharedPtr<FLooseFilesCookArtifactReader> SharedLooseFilesCookArtifactReader;
+	TUniquePtr<UE::Cook::FCookGCDiagnosticContext> GCDiagnosticContext;
 
 	TArray<UE::Cook::FCookSavePackageContext*> SavePackageContexts;
 	/**
diff --git a/Engine/Source/Editor/UnrealEd/Private/Commandlets/CookCommandlet.cpp b/Engine/Source/Editor/UnrealEd/Private/Commandlets/CookCommandlet.cpp
--- a/Engine/Source/Editor/UnrealEd/Private/Commandlets/CookCommandlet.cpp
+++ b/Engine/Source/Editor/UnrealEd/Private/Commandlets/CookCommandlet.cpp
@@ -734,7 +734,15 @@ void UCookCommandlet::ConditionalCollectGarbage(uint32 TickResults, UCookOnTheFl
 	UE_LOG(LogCookCommandlet, Display, TEXT("GarbageCollection...%s (%s)"), *GCType, *GCReason);
 	{
 		TGuardValue<bool> SoftGCGuard(UPackage::bSupportCookerSoftGC, true);
+		COTFS.OnCookerStartCollectGarbage();
 		CollectGarbage(RF_NoFlags);
+		COTFS.OnCookerEndCollectGarbage();
+		if (COTFS.NeedsDiagnosticSecondGC())
+		{
+			COTFS.OnCookerStartCollectGarbage();
+			CollectGarbage(RF_NoFlags);
+			COTFS.OnCookerEndCollectGarbage();
+		}
 	}
 
 	COTFS.ClearGarbageCollectType();
diff --git a/Engine/Source/Editor/UnrealEd/Private/CookOnTheFlyServer.cpp b/Engine/Source/Editor/UnrealEd/Private/CookOnTheFlyServer.cpp
--- a/Engine/Source/Editor/UnrealEd/Private/CookOnTheFlyServer.cpp
+++ b/Engine/Source/Editor/UnrealEd/Private/CookOnTheFlyServer.cpp
@@ -24,6 +24,7 @@
 #include "Cooker/CookConfigAccessTracker.h"
 #include "Cooker/CookDiagnostics.h"
 #include "Cooker/CookDirector.h"
+#include "Cooker/CookGarbageCollect.h"
 #include "Cooker/CookGenerationHelper.h"
 #include "Cooker/CookOnTheFlyServerInterface.h"
 #include "Cooker/CookPackageData.h"
@@ -5558,15 +5559,11 @@ void UCookOnTheFlyServer::PostGarbageCollect()
 	// Remove objects that were deleted by garbage collection from our containers that track raw object pointers
 	PackageDatas->CachedCookedPlatformDataObjectsPostGarbageCollect(SaveQueueObjectsThatStillExist);
 
-	GCKeepObjects.Empty();
-	UPackage::SoftGCPackageToObjectList.Empty();
-	SoftGCPackageToObjectListBuffer.Empty();
-
-	PackageDatas->LockAndEnumeratePackageDatas([](FPackageData* PackageData)
+	PackageDatas->LockAndEnumeratePackageDatas([this](FPackageData* PackageData)
 	{
 		if (TRefCountPtr<FGenerationHelper> GenerationHelper = PackageData->GetGenerationHelper())
 		{
-			GenerationHelper->PostGarbageCollect(GenerationHelper);
+			GenerationHelper->PostGarbageCollect(GenerationHelper, *GCDiagnosticContext);
 		}
 	});
 	PackageDatas->LockAndEnumeratePackageDatas([](FPackageData* PackageData)
@@ -5574,6 +5571,12 @@ void UCookOnTheFlyServer::PostGarbageCollect()
 		PackageData->SetKeepReferencedDuringGC(false);
 	});
 
+	// Only after running all possible callbacks that need our links for diagnostics, clear the list of temporary
+	// references that we created for the garbage collection.
+	GCKeepObjects.Empty();
+	UPackage::SoftGCPackageToObjectList.Empty();
+	SoftGCPackageToObjectListBuffer.Empty();
+
 	CookedPackageCountSinceLastGC = 0;
 
 	// Whenever we collect garbage, reset the counter for how many busy reports with an
@@ -5581,6 +5584,21 @@ void UCookOnTheFlyServer::PostGarbageCollect()
 	bShaderCompilerWasActiveeOnPreviousBusyReport = true;
 }
 
+bool UCookOnTheFlyServer::NeedsDiagnosticSecondGC() const
+{
+	return GCDiagnosticContext->NeedsDiagnosticSecondGC();
+}
+
+void UCookOnTheFlyServer::OnCookerStartCollectGarbage()
+{
+	GCDiagnosticContext->OnCookerStartCollectGarbage();
+}
+
+void UCookOnTheFlyServer::OnCookerEndCollectGarbage()
+{
+	GCDiagnosticContext->OnCookerEndCollectGarbage();
+}
+
 void UCookOnTheFlyServer::EvaluateGarbageCollectionResults(bool bWasDueToOOM, bool bWasPartialGC, uint32 ResultFlags,
 	int32 NumObjectsBeforeGC, const FPlatformMemoryStats& MemStatsBeforeGC,
 	const FGenericMemoryStats& AllocatorStatsBeforeGC,
@@ -5592,6 +5610,7 @@ void UCookOnTheFlyServer::EvaluateGarbageCollectionResults(bool bWasDueToOOM, bo
 	ON_SCOPE_EXIT
 	{
 		ExpectedFreedPackageNames.Empty();
+		GCDiagnosticContext->OnEvaluateResultsComplete();
 	};
 	bWarnedExceededMaxMemoryWithinGCCooldown = false;
 	LastGCTime = FPlatformTime::Seconds();
@@ -6907,6 +6926,7 @@ void UCookOnTheFlyServer::Initialize( ECookMode::Type DesiredCookMode, ECookInit
 	CookOnTheFlyOptions = MakeUnique<UE::Cook::FCookOnTheFlyOptions>();
 	AssetRegistry = IAssetRegistry::Get();
 	CachedDependencies = MakeUnique<UE::Cook::FCachedDependencies>();
+	GCDiagnosticContext = MakeUnique<UE::Cook::FCookGCDiagnosticContext>();
 
 	if (!IsCookWorkerMode())
 	{
diff --git a/Engine/Source/Editor/UnrealEd/Private/Cooker/CookGarbageCollect.cpp b/Engine/Source/Editor/UnrealEd/Private/Cooker/CookGarbageCollect.cpp
new file mode 100644
--- /dev/null
+++ b/Engine/Source/Editor/UnrealEd/Private/Cooker/CookGarbageCollect.cpp
@@ -0,0 +1,80 @@
+// Copyright Epic Games, Inc. All Rights Reserved.
+
+#include "Cooker/CookGarbageCollect.h"
+
+#include "UObject/GarbageCollectionHistory.h"
+
+namespace UE::Cook
+{
+
+FCookGCDiagnosticContext::~FCookGCDiagnosticContext()
+{
+	SetGCWithHistoryRequested(false);
+}
+
+bool FCookGCDiagnosticContext::NeedsDiagnosticSecondGC() const
+{
+	return bRequestGCWithHistory;
+}
+
+bool FCookGCDiagnosticContext::CurrentGCHasHistory() const
+{
+	return bCurrentGCHasHistory;
+}
+
+bool FCookGCDiagnosticContext::TryRequestGCWithHistory()
+{
+	if (!bRequestsAvailable || !bGCInProgress || bCurrentGCHasHistory)
+	{
+		return false;
+	}
+	SetGCWithHistoryRequested(true);
+	return true;
+}
+
+void FCookGCDiagnosticContext::OnCookerStartCollectGarbage()
+{
+	bRequestsAvailable = true;
+
+	bGCInProgress = true;
+	bCurrentGCHasHistory = FGCHistory::Get().GetHistorySize() > 0;
+}
+
+void FCookGCDiagnosticContext::OnCookerEndCollectGarbage()
+{
+	bGCInProgress = false;
+	bCurrentGCHasHistory = false;
+}
+
+void FCookGCDiagnosticContext::OnEvaluateResultsComplete()
+{
+	SetGCWithHistoryRequested(false);
+}
+
+void FCookGCDiagnosticContext::SetGCWithHistoryRequested(bool bValue)
+{
+	if (bValue == bRequestGCWithHistory)
+	{
+		return;
+	}
+
+	if (bValue)
+	{
+		SavedGCHistorySize = FGCHistory::Get().GetHistorySize();
+		if (SavedGCHistorySize < 1)
+		{
+			FGCHistory::Get().SetHistorySize(1);
+		}
+	}
+	else
+	{
+		if (SavedGCHistorySize != FGCHistory::Get().GetHistorySize())
+		{
+			FGCHistory::Get().SetHistorySize(SavedGCHistorySize);
+		}
+		SavedGCHistorySize = 0;
+	}
+	bRequestGCWithHistory = bValue;
+}
+
+} // namespace UE::Cook
\ No newline at end of file
diff --git a/Engine/Source/Editor/UnrealEd/Private/Cooker/CookGarbageCollect.h b/Engine/Source/Editor/UnrealEd/Private/Cooker/CookGarbageCollect.h
new file mode 100644
--- /dev/null
+++ b/Engine/Source/Editor/UnrealEd/Private/Cooker/CookGarbageCollect.h
@@ -0,0 +1,44 @@
+// Copyright Epic Games, Inc. All Rights Reserved.
+
+#pragma once
+
+#include "CoreTypes.h"
+
+namespace UE::Cook
+{
+
+/**
+ * Holds information about the cookers garbage collection status, and communicates requests from low level structures
+ * back up to the CookCommandlet that is capable of acting on those requests with additional garbage collection
+ * commands.
+ */
+class FCookGCDiagnosticContext
+{
+public:
+	~FCookGCDiagnosticContext();
+
+	bool NeedsDiagnosticSecondGC() const;
+	bool CurrentGCHasHistory() const;
+
+	/**
+	 * Add a request to reexecute the current GC after all of the PostGarbageCollect calls run and
+	 * control returns back to the caller of CollectGarbage. Returns false if not currently in post-GC,
+	 * or the garbage collect that just ran already had history.
+	 */
+	bool TryRequestGCWithHistory();
+
+	void OnCookerStartCollectGarbage();
+	void OnCookerEndCollectGarbage();
+	void OnEvaluateResultsComplete();
+
+private:
+	void SetGCWithHistoryRequested(bool bValue);
+	
+	int32 SavedGCHistorySize = 0;
+	bool bRequestsAvailable = false;
+	bool bGCInProgress = false;
+	bool bRequestGCWithHistory = false;
+	bool bCurrentGCHasHistory = false;
+};
+
+} // namespace UE::Cook
diff --git a/Engine/Source/Editor/UnrealEd/Private/Cooker/CookGenerationHelper.cpp b/Engine/Source/Editor/UnrealEd/Private/Cooker/CookGenerationHelper.cpp
--- a/Engine/Source/Editor/UnrealEd/Private/Cooker/CookGenerationHelper.cpp
+++ b/Engine/Source/Editor/UnrealEd/Private/Cooker/CookGenerationHelper.cpp
@@ -5,6 +5,7 @@
 #include "Algo/Unique.h"
 #include "AssetRegistry/AssetData.h"
 #include "Cooker/CookDirector.h"
+#include "Cooker/CookGarbageCollect.h"
 #include "Cooker/CookPlatformManager.h"
 #include "Cooker/CookWorkerServer.h"
 #include "Cooker/IWorkerRequests.h"
@@ -1517,11 +1518,11 @@ void FGenerationHelper::PreGarbageCollectGCLifetimeData()
 	bNeedConfirmGeneratorPackageDestroyed = IsRequiresGeneratorPackageDestructBeforeResplit();
 }
 
-void FGenerationHelper::PostGarbageCollectGCLifetimeData()
+void FGenerationHelper::PostGarbageCollectGCLifetimeData(FCookGCDiagnosticContext& Context)
 {
 	if (bNeedConfirmGeneratorPackageDestroyed)
 	{
-		VerifyGeneratorPackageGarbageCollected();
+		VerifyGeneratorPackageGarbageCollected(Context);
 		bNeedConfirmGeneratorPackageDestroyed = false;
 	}
 
@@ -1592,9 +1593,9 @@ void FGenerationHelper::MarkPackageIterativelySkipped(FPackageData& PackageData)
 	}
 }
 
-void FGenerationHelper::PostGarbageCollect(const TRefCountPtr<FGenerationHelper>& RefcountHeldByCaller)
+void FGenerationHelper::PostGarbageCollect(const TRefCountPtr<FGenerationHelper>& RefcountHeldByCaller, FCookGCDiagnosticContext& Context)
 {
-	PostGarbageCollectGCLifetimeData();
+	PostGarbageCollectGCLifetimeData(Context);
 	if (!IsInitialized())
 	{
 		return;
@@ -1622,7 +1623,7 @@ void FGenerationHelper::PostGarbageCollect(const TRefCountPtr<FGenerationHelper>
 		// The splitter can opt-out of this contract and keep it referenced itself if it desires.
 		if (!Owner.IsInProgress() && !Owner.IsKeepReferencedDuringGC())
 		{
-			VerifyGeneratorPackageGarbageCollected();
+			VerifyGeneratorPackageGarbageCollected(Context);
 		}
 	}
 
@@ -1670,32 +1671,34 @@ void FGenerationHelper::PostGarbageCollect(const TRefCountPtr<FGenerationHelper>
 	}
 }
 
-void FGenerationHelper::VerifyGeneratorPackageGarbageCollected()
+void FGenerationHelper::VerifyGeneratorPackageGarbageCollected(FCookGCDiagnosticContext& Context)
 {
 	FString GeneratorPackageName = GetOwner().GetPackageName().ToString();
 	UPackage* LocalOwnerPackage = FindObject<UPackage>(nullptr, *GeneratorPackageName);
 	if (LocalOwnerPackage)
 	{
-		// Might be called when uninitialized, so do not call GetSplitDataObjectNameIfAvailable
-		FString Identifier;
-		if (!SplitDataObjectName.IsNone())
+		if (!Context.TryRequestGCWithHistory())
 		{
-			Identifier = FString::Printf(TEXT("Splitter=%s"), *SplitDataObjectName.ToString());
-		}
-		else
-		{
-			Identifier = FString::Printf(TEXT("GeneratorPackage=%s"), *GeneratorPackageName);
-		}
-		UE_LOG(LogCook, Error,
-			TEXT("PackageSplitter found the Generator package still in memory after it should have been deleted by GC.")
-			TEXT("\n\tThis is unexpected since garbage has been collected and the package should have been unreferenced so it should have been collected, and will break population of Generated packages.")
+			// Might be called when uninitialized, so do not call GetSplitDataObjectNameIfAvailable
+			FString Identifier;
+			if (!SplitDataObjectName.IsNone())
+			{
+				Identifier = FString::Printf(TEXT("Splitter=%s"), *SplitDataObjectName.ToString());
+			}
+			else
+			{
+				Identifier = FString::Printf(TEXT("GeneratorPackage=%s"), *GeneratorPackageName);
+			}
+			UE_LOG(LogCook, Error,
+				TEXT("PackageSplitter found the Generator package still in memory after it should have been deleted by GC.")
+				TEXT("\n\tThis is unexpected since garbage has been collected and the package should have been unreferenced so it should have been collected, and will break population of Generated packages.")
 			TEXT("\n\tSplitter=%s"), *Identifier);
-		EReferenceChainSearchMode SearchMode = EReferenceChainSearchMode::Shortest
-			| EReferenceChainSearchMode::PrintAllResults
-			| EReferenceChainSearchMode::FullChain;
-		FReferenceChainSearch RefChainSearch(LocalOwnerPackage, SearchMode);
+			EReferenceChainSearchMode SearchMode = EReferenceChainSearchMode::Shortest
+				| EReferenceChainSearchMode::PrintAllResults
+				| EReferenceChainSearchMode::FullChain;
+			FReferenceChainSearch RefChainSearch(LocalOwnerPackage, SearchMode, ELogVerbosity::Display);
+		}
 	}
-
 }
 
 void FGenerationHelper::UpdateSaveAfterGarbageCollect(const FPackageData& PackageData, bool& bInOutDemote)
diff --git a/Engine/Source/Editor/UnrealEd/Private/Cooker/CookGenerationHelper.h b/Engine/Source/Editor/UnrealEd/Private/Cooker/CookGenerationHelper.h
--- a/Engine/Source/Editor/UnrealEd/Private/Cooker/CookGenerationHelper.h
+++ b/Engine/Source/Editor/UnrealEd/Private/Cooker/CookGenerationHelper.h
@@ -25,6 +25,7 @@
 
 class ITargetPlatform;
 class UCookOnTheFlyServer;
+namespace UE::Cook { class FCookGCDiagnosticContext; }
 namespace UE::TargetDomain { struct FGeneratedPackageResultStruct; }
 
 namespace UE::Cook
@@ -369,7 +370,8 @@ public:
 	 * Callback during garbage collection. Does not call initialize. Caller must pass in a refcount to show
 	 * a guarantee that clearing internal references will not delete before function return.
 	 */
-	void PostGarbageCollect(const TRefCountPtr<FGenerationHelper>& RefcountHeldByCaller);
+	void PostGarbageCollect(const TRefCountPtr<FGenerationHelper>& RefcountHeldByCaller,
+		FCookGCDiagnosticContext& Context);
 	/**
 	 * Called from PackageData function of the same name to decide whether to demote the package out of save.
 	 * Does not call Initialize.
@@ -459,10 +461,10 @@ private:
 	FCookGenerationInfo* FindInfoNoInitialize(const FPackageData& PackageData);
 	void NotifyCompletion(ICookPackageSplitter::ETeardown Status);
 	void PreGarbageCollectGCLifetimeData();
-	void PostGarbageCollectGCLifetimeData();
+	void PostGarbageCollectGCLifetimeData(FCookGCDiagnosticContext& Context);
 	void Uninitialize();
 	void ModifyNumSaved(int32 Delta);
-	void VerifyGeneratorPackageGarbageCollected();
+	void VerifyGeneratorPackageGarbageCollected(FCookGCDiagnosticContext& Context);
 	void DemoteStalledPackages(UCookOnTheFlyServer& COTFS);
 
 private:
diff --git a/Engine/Source/Runtime/CoreUObject/Private/UObject/ReferenceChainSearch.cpp b/Engine/Source/Runtime/CoreUObject/Private/UObject/ReferenceChainSearch.cpp
--- a/Engine/Source/Runtime/CoreUObject/Private/UObject/ReferenceChainSearch.cpp
+++ b/Engine/Source/Runtime/CoreUObject/Private/UObject/ReferenceChainSearch.cpp
@@ -8,6 +8,7 @@
 #include "Experimental/Graph/GraphConvert.h"
 #include "HAL/PlatformStackWalk.h"
 #include "HAL/ThreadHeartBeat.h"
+#include "Misc/FeedbackContext.h"
 #include "UObject/FastReferenceCollector.h"
 #include "UObject/GCObject.h"
 #include "UObject/GarbageCollectionInternalFlags.h"
@@ -1355,7 +1356,8 @@ static void ConvertStackFramesToCallstack(
 	int32 NumStackFrames,
 	int32 Indent,
 	FOutputDevice& Out,
-	TMap<uint64, FString>& Cache)
+	TMap<uint64, FString>& Cache,
+	ELogVerbosity::Type VerbosityForPrint)
 {
 	// Convert the stack trace to text
 	FStringView View;
@@ -1385,7 +1387,7 @@ static void ConvertStackFramesToCallstack(
 			{
 				View = View.RightChop(Index + 1);
 			}
-			Out.Logf(ELogVerbosity::Log, TEXT("%*s   ^ %.*s"), Indent, TEXT(""), View.Len(), View.GetData());
+			Out.Logf(VerbosityForPrint, TEXT("%*s   ^ %.*s"), Indent, TEXT(""), View.Len(), View.GetData());
 		}
 	}
 }
@@ -1393,7 +1395,8 @@ static void ConvertStackFramesToCallstack(
 void FReferenceChainSearch::DumpChain(FReferenceChainSearch::FReferenceChain* Chain,
 	TFunctionRef<bool(FCallbackParams& Params)> ReferenceCallback,
 	TMap<uint64, FString>& CallstackCache,
-	FOutputDevice& Out)
+	FOutputDevice& Out,
+	ELogVerbosity::Type InVerbosityForPrint)
 {
 	if (Chain->Num())
 	{
@@ -1409,7 +1412,7 @@ void FReferenceChainSearch::DumpChain(FReferenceChainSearch::FReferenceChain* Ch
 			Params.Indent = FMath::Min<int32>(TCStringSpcHelper<TCHAR>::MAX_SPACES, Chain->Num() - RootIndex);
 			Params.Out = &Out;
 
-			Out.Logf(ELogVerbosity::Log, TEXT("%s%s %s"),
+			Out.Logf(InVerbosityForPrint, TEXT("%s%s %s"),
 				FCString::Spc(Params.Indent),
 				*GetObjectFlags(*ReferencerObject),
 				*ReferencerObject->GetFullName());
@@ -1471,7 +1474,7 @@ void FReferenceChainSearch::DumpChain(FReferenceChainSearch::FReferenceChain* Ch
 					ReferencingPropertyName = FString::Printf(TEXT("UObject* %s::%s"), *ClassName, *ReferenceInfo->ReferencerName.ToString());
 				}
 
-				Out.Logf(ELogVerbosity::Log, TEXT("%s-> %s = %s %s"),
+				Out.Logf(InVerbosityForPrint, TEXT("%s-> %s = %s %s"),
 					FCString::Spc(Params.Indent),
 					*ReferencingPropertyName,
 					*GetObjectFlags(*Object),
@@ -1498,7 +1501,7 @@ void FReferenceChainSearch::DumpChain(FReferenceChainSearch::FReferenceChain* Ch
 					UObjectOrGCObjectName = ReferenceInfo->ReferencerName.ToString();
 				}
 
-				Out.Logf(ELogVerbosity::Log, TEXT("%s-> %s::AddReferencedObjects(%s %s)"),
+				Out.Logf(InVerbosityForPrint, TEXT("%s-> %s::AddReferencedObjects(%s %s)"),
 					FCString::Spc(Params.Indent),
 					*UObjectOrGCObjectName,
 					*GetObjectFlags(*Object),
@@ -1510,12 +1513,13 @@ void FReferenceChainSearch::DumpChain(FReferenceChainSearch::FReferenceChain* Ch
 						ReferenceInfo->StackFrames.Num(),
 						Params.Indent,
 						Out,
-						CallstackCache);
+						CallstackCache,
+						InVerbosityForPrint);
 				}
 			}
 			else if (ReferenceInfo && ReferenceInfo->Type == EReferenceType::OuterChain)
 			{
-				Out.Logf(ELogVerbosity::Log, TEXT("%s-> %s = %s %s"),
+				Out.Logf(InVerbosityForPrint, TEXT("%s-> %s = %s %s"),
 					FCString::Spc(Params.Indent),
 					TEXT("Outer Chain"),
 					*GetObjectFlags(*Object),
@@ -1523,7 +1527,7 @@ void FReferenceChainSearch::DumpChain(FReferenceChainSearch::FReferenceChain* Ch
 			}
 			else
 			{
-				Out.Logf(ELogVerbosity::Log, TEXT("%s-> %s = %s %s"),
+				Out.Logf(InVerbosityForPrint, TEXT("%s-> %s = %s %s"),
 					FCString::Spc(Params.Indent),
 					TEXT("UNKNOWN"),
 					*GetObjectFlags(*Object),
@@ -1535,7 +1539,7 @@ void FReferenceChainSearch::DumpChain(FReferenceChainSearch::FReferenceChain* Ch
 			ReferencerObject = Object;
 			ReferenceInfo = Chain->GetReferenceInfo(NodeIndex);
 		}
-		Out.Logf(ELogVerbosity::Log, TEXT("  "));
+		Out.Logf(InVerbosityForPrint, TEXT("  "));
 	}
 }
 
@@ -1553,14 +1557,17 @@ bool FReferenceChainSearch::FReferenceChain::IsExternal() const
 }
 
 FReferenceChainSearch::FReferenceChainSearch(UObject* InObjectToFindReferencesTo,
-	EReferenceChainSearchMode Mode /*= EReferenceChainSearchMode::PrintResults*/)
-	: FReferenceChainSearch(TConstArrayView<UObject*>(&InObjectToFindReferencesTo, 1), Mode)
+	EReferenceChainSearchMode Mode /*= EReferenceChainSearchMode::PrintResults*/,
+	ELogVerbosity::Type InVerbosityForPrint /* = ELogVerbosity::Log */)
+	: FReferenceChainSearch(TConstArrayView<UObject*>(&InObjectToFindReferencesTo, 1), Mode, InVerbosityForPrint)
 {
 }
 
 FReferenceChainSearch::FReferenceChainSearch(TConstArrayView<UObject*> InObjectsToFindReferencesTo,
-	EReferenceChainSearchMode Mode /*= EReferenceChainSearchMode::PrintResults*/)
+	EReferenceChainSearchMode Mode /*= EReferenceChainSearchMode::PrintResults*/,
+	ELogVerbosity::Type InVerbosityForPrint /* = ELogVerbosity::Log */)
 	: SearchMode(Mode)
+	, VerbosityForPrint(InVerbosityForPrint)
 {
 	FSlowHeartBeatScope DisableHangDetection; // This function can be very slow
 
@@ -1677,6 +1684,11 @@ void FReferenceChainSearch::PerformSearchFromGCSnapshot(TConstArrayView<UObject*
 }
 #endif // ENABLE_GC_HISTORY
 
+void FReferenceChainSearch::SetVerbosityForPrint(ELogVerbosity::Type Verbosity)
+{
+	VerbosityForPrint = Verbosity;
+}
+
 int32 FReferenceChainSearch::PrintResults(bool bDumpAllChains /*= false*/, UObject* TargetObject /*= nullptr*/) const
 {
 	return PrintResults([](FCallbackParams& Params) { return true; }, bDumpAllChains, TargetObject);
@@ -1699,12 +1711,15 @@ int32 FReferenceChainSearch::PrintResults(TFunctionRef<bool(FCallbackParams& Par
 
 		if (bDumpAllChains || NumPrintedChains < MaxChainsToPrint)
 		{
-			DumpChain(Chain, ReferenceCallback, CallstackCache, *GLog);
+			DumpChain(Chain, ReferenceCallback, CallstackCache, *GWarn, VerbosityForPrint);
 			NumPrintedChains++;
 		}
 		else
 		{
-			UE_LOG(LogReferenceChain, Log, TEXT("Referenced by %d more reference chain(s)."), ReferenceChains.Num() - NumPrintedChains);
+#if !NO_LOGGING
+			GWarn->CategorizedLogf(LogReferenceChain.GetCategoryName(), VerbosityForPrint,
+				TEXT("Referenced by %d more reference chain(s)."), ReferenceChains.Num() - NumPrintedChains);
+#endif
 			break;
 		}
 	}
@@ -1714,15 +1729,24 @@ int32 FReferenceChainSearch::PrintResults(TFunctionRef<bool(FCallbackParams& Par
 		auto LogUnreachableObject = [this](const FGCObjectInfo& ObjInfo) {
 			if (ObjInfo.HasAnyInternalFlags(EInternalObjectFlags_RootFlags))
 			{
-				UE_LOG(LogReferenceChain, Log, TEXT("%s%s is not currently reachable but it does have some of EInternalObjectFlags_RootFlags set."), *GetObjectFlags(ObjInfo), *ObjInfo.GetFullName());
+#if !NO_LOGGING
+				GWarn->CategorizedLogf(LogReferenceChain.GetCategoryName(), VerbosityForPrint,
+					TEXT("%s%s is not currently reachable but it does have some of EInternalObjectFlags_RootFlags set."), *GetObjectFlags(ObjInfo), *ObjInfo.GetFullName());
+#endif
 			}
 			else if (ObjInfo.HasAnyFlags(GARBAGE_COLLECTION_KEEPFLAGS))
 			{
-				UE_LOG(LogReferenceChain, Log, TEXT("%s%s is not currently reachable but it does have some of GARBAGE_COLLECTION_KEEPFLAGS set."), *GetObjectFlags(ObjInfo), *ObjInfo.GetFullName());
+#if !NO_LOGGING
+				GWarn->CategorizedLogf(LogReferenceChain.GetCategoryName(), VerbosityForPrint,
+					TEXT("%s%s is not currently reachable but it does have some of GARBAGE_COLLECTION_KEEPFLAGS set."), *GetObjectFlags(ObjInfo), *ObjInfo.GetFullName());
+#endif
 			}
 			else
 			{
-				UE_LOG(LogReferenceChain, Log, TEXT("%s%s is not currently reachable. Try using GC history to debug transient leaks with 'gc.historysize 1'"), *GetObjectFlags(ObjInfo), *ObjInfo.GetFullName());
+#if !NO_LOGGING
+				GWarn->CategorizedLogf(LogReferenceChain.GetCategoryName(), VerbosityForPrint,
+					TEXT("%s%s is not currently reachable. Try using GC history to debug transient leaks with 'gc.historysize 1'"), *GetObjectFlags(ObjInfo), *ObjInfo.GetFullName());
+#endif
 			}
 		};
 		if (TargetObject)
@@ -1766,7 +1790,7 @@ FString FReferenceChainSearch::GetRootPath(TFunctionRef<bool(FCallbackParams& Pa
 		FStringOutputDevice OutString;
 		OutString.SetAutoEmitLineTerminator(true);
 		TMap<uint64, FString> CallstackCache;
-		DumpChain(Chain, ReferenceCallback, CallstackCache, OutString);
+		DumpChain(Chain, ReferenceCallback, CallstackCache, OutString, VerbosityForPrint);
 		return MoveTemp(OutString);
 	}
 	else
diff --git a/Engine/Source/Runtime/CoreUObject/Public/UObject/ReferenceChainSearch.h b/Engine/Source/Runtime/CoreUObject/Public/UObject/ReferenceChainSearch.h
--- a/Engine/Source/Runtime/CoreUObject/Public/UObject/ReferenceChainSearch.h
+++ b/Engine/Source/Runtime/CoreUObject/Public/UObject/ReferenceChainSearch.h
@@ -286,8 +286,12 @@ public:
 	};
 
 	/** Constructs a new search engine and finds references to the specified object */
-	COREUOBJECT_API explicit FReferenceChainSearch(UObject* InObjectToFindReferencesTo, EReferenceChainSearchMode Mode = EReferenceChainSearchMode::PrintResults);
-	COREUOBJECT_API explicit FReferenceChainSearch(TConstArrayView<UObject*> InObjectsToFindReferencesTo, EReferenceChainSearchMode Mode = EReferenceChainSearchMode::PrintResults);
+	COREUOBJECT_API explicit FReferenceChainSearch(UObject* InObjectToFindReferencesTo,
+		EReferenceChainSearchMode Mode = EReferenceChainSearchMode::PrintResults,
+		ELogVerbosity::Type InVerbosityForPrint = ELogVerbosity::Log);
+	COREUOBJECT_API explicit FReferenceChainSearch(TConstArrayView<UObject*> InObjectsToFindReferencesTo,
+		EReferenceChainSearchMode Mode = EReferenceChainSearchMode::PrintResults,
+		ELogVerbosity::Type InVerbosityForPrint = ELogVerbosity::Log);
 
 	/** Constructs a new search engine but does not find references to any objects until one of the PerformSearch*() functions is called */
 	COREUOBJECT_API explicit FReferenceChainSearch(EReferenceChainSearchMode Mode);
@@ -307,6 +311,7 @@ public:
 	 * @returns The number of results printed.
 	 */
 	COREUOBJECT_API int32 PrintResults(bool bDumpAllChains = false, UObject* TargetObject=nullptr) const;
+	COREUOBJECT_API void SetVerbosityForPrint(ELogVerbosity::Type Verbosity);
 
 	/**
 	 * Dumps results to log
@@ -356,6 +361,8 @@ private:
 	/** Maps UObject pointers to object info structs */
 	TMap<const UObject*, FGCObjectInfo*> ObjectToInfoMap;
 
+	ELogVerbosity::Type VerbosityForPrint = ELogVerbosity::Log;
+
 	/** Frees memory */
 	void Cleanup();
 
@@ -380,5 +387,5 @@ private:
 
 	/** Dumps a reference chain to log */
 	static void DumpChain(FReferenceChainSearch::FReferenceChain* Chain, TFunctionRef<bool(FCallbackParams& Params)> ReferenceCallback,
-		TMap<uint64, FString>& CallstackCache, FOutputDevice& Out);
+		TMap<uint64, FString>& CallstackCache, FOutputDevice& Out, ELogVerbosity::Type InVerbosityForPrint);
 };
