Restrict creation of placeholder types on editor load to missing Verse class imports only (temp).

#rnx

diff --git a/Engine/Source/Runtime/CoreUObject/Internal/UObject/InstanceDataObjectUtils.h b/Engine/Source/Runtime/CoreUObject/Internal/UObject/InstanceDataObjectUtils.h
--- a/Engine/Source/Runtime/CoreUObject/Internal/UObject/InstanceDataObjectUtils.h
+++ b/Engine/Source/Runtime/CoreUObject/Internal/UObject/InstanceDataObjectUtils.h
@@ -26,6 +26,9 @@ bool IsInstanceDataObjectSupportEnabled();
 bool IsInstanceDataObjectSupportEnabled(const UObject* Object);
 bool StructContainsLooseProperties(const UStruct* Struct);
 
+/** Query if placeholder support is enabled for a specific import class type. */
+bool CanCreatePropertyBagPlaceholderTypeForImportClass(const UClass* ImportType);
+
 /** Helper to check if a class is an IDO class. */
 bool IsClassOfInstanceDataObjectClass(UStruct* Class);
 
diff --git a/Engine/Source/Runtime/CoreUObject/Private/UObject/InstanceDataObjectUtils.cpp b/Engine/Source/Runtime/CoreUObject/Private/UObject/InstanceDataObjectUtils.cpp
--- a/Engine/Source/Runtime/CoreUObject/Private/UObject/InstanceDataObjectUtils.cpp
+++ b/Engine/Source/Runtime/CoreUObject/Private/UObject/InstanceDataObjectUtils.cpp
@@ -208,6 +208,12 @@ namespace UE
 		return !!ObjClass;
 	}
 
+	bool CanCreatePropertyBagPlaceholderTypeForImportClass(const UClass* ImportClass)
+	{
+		// @todo - Expand to other import types (e.g. prefab BPs) later; for now restricted to Verse class objects only.
+		return ImportClass && ImportClass->GetFName() == NAME_VerseClass;
+	}
+
 
 	bool IsClassOfInstanceDataObjectClass(UStruct* Class)
 	{
diff --git a/Engine/Source/Runtime/CoreUObject/Private/UObject/LinkerLoad.cpp b/Engine/Source/Runtime/CoreUObject/Private/UObject/LinkerLoad.cpp
--- a/Engine/Source/Runtime/CoreUObject/Private/UObject/LinkerLoad.cpp
+++ b/Engine/Source/Runtime/CoreUObject/Private/UObject/LinkerLoad.cpp
@@ -4096,7 +4096,7 @@ UClass* FLinkerLoad::TryCreatePlaceholderClassImport(int32 ImportIndex)
 	{
 		if (UClass* ImportClass = FindObjectFast<UClass>(ImportClassPackage, Import.ClassName, /*bExactClass =*/ false))
 		{
-			if (ImportClass->IsChildOf<UClass>())
+			if (UE::CanCreatePropertyBagPlaceholderTypeForImportClass(ImportClass))
 			{
 				// If the outer package import is also missing, create it now so that the full path remains the same. 
 				UPackage* ClassObjectPackage = Cast<UPackage>(IndexToObject(Import.OuterIndex));
diff --git a/Engine/Source/Runtime/CoreUObject/Private/UObject/PropertyBagRepository.cpp b/Engine/Source/Runtime/CoreUObject/Private/UObject/PropertyBagRepository.cpp
--- a/Engine/Source/Runtime/CoreUObject/Private/UObject/PropertyBagRepository.cpp
+++ b/Engine/Source/Runtime/CoreUObject/Private/UObject/PropertyBagRepository.cpp
@@ -1073,7 +1073,7 @@ bool FPropertyBagRepository::IsPropertyBagPlaceholderObjectSupportEnabled()
 UStruct* FPropertyBagRepository::CreatePropertyBagPlaceholderType(UObject* Outer, UClass* Class, FName Name, EObjectFlags Flags, UStruct* SuperStruct)
 {
 	// Generate and link a new type object using the given SuperStruct as its base.
-	UStruct* PlaceholderType = NewObject<UClass>(Outer, Class, Name, Flags);
+	UStruct* PlaceholderType = NewObject<UStruct>(Outer, Class, Name, Flags);
 	PlaceholderType->SetSuperStruct(SuperStruct);
 	PlaceholderType->Bind();
 	PlaceholderType->StaticLink(/*bRelinkExistingProperties =*/ true);
