Fix CIS error due to nested condition.

diff --git a/Engine/Source/Runtime/CoreUObject/Private/Misc/PackageName.cpp b/Engine/Source/Runtime/CoreUObject/Private/Misc/PackageName.cpp
--- a/Engine/Source/Runtime/CoreUObject/Private/Misc/PackageName.cpp
+++ b/Engine/Source/Runtime/CoreUObject/Private/Misc/PackageName.cpp
@@ -837,7 +837,7 @@ bool FPackageName::TryConvertFilenameToLongPackageName(FStringView InFilename, F
 	const bool bContainsColon = LongPackageName.FindChar(TEXT(':'), CharacterIndex);
 
 	// For Verse files, the filenames can have dots in them to allow for `*.*.verse` to be considered as valid Verse snippets.
-	if (!(!EnumHasAnyFlags(Flags, EConvertFlags::AllowDots) && bContainsDot || bContainsBackslash || bContainsColon))
+	if (!(!EnumHasAnyFlags(Flags, EConvertFlags::AllowDots) && bContainsDot) || bContainsBackslash || bContainsColon)
 	{
 		OutPackageName = LongPackageName;
 		return true;
