Ensure we do not clear the linker when renaming an object to a new outer within the same package, should that new out not yet have a linker (e.g. we created the outer in-memory and intend to save out a new package). Otherwise we will potentially detach exports in a package which are intended to not be saved by implicitly calling ResetObject as part of SetLinker's export detaching (clearing filter and load failed flags)

#rnx

diff --git a/Engine/Source/Runtime/CoreUObject/Private/UObject/Obj.cpp b/Engine/Source/Runtime/CoreUObject/Private/UObject/Obj.cpp
--- a/Engine/Source/Runtime/CoreUObject/Private/UObject/Obj.cpp
+++ b/Engine/Source/Runtime/CoreUObject/Private/UObject/Obj.cpp
@@ -249,9 +249,16 @@ bool UObject::Rename(const TCHAR* InName, UObject* NewOuter, ERenameFlags Flags)
 		}
 
 		// If moving the object to a new package, remove its linker and detach the object
-		if (((Flags & REN_AllowPackageLinkerMismatch) == 0) && GetLinker() != NewOuter->GetLinker())
+		const FLinkerLoad* Linker = GetLinker();
+		const FLinkerLoad* NewOuterLinker = NewOuter->GetLinker();
+		if (((Flags & REN_AllowPackageLinkerMismatch) == 0) && Linker != NewOuterLinker)
 		{
-			SetLinker(nullptr, INDEX_NONE);
+			// It's possible we are moving to a new outer in the same package which doesn't
+			// have a linker (in which case we don't want to clear the existing linker)
+			if (NewOuterLinker || GetPackage() != NewOuter->GetPackage())
+			{
+				SetLinker(nullptr, INDEX_NONE);
+			}
 		}
 	}
 
