Passing in a verbosity argument in FReferenceChainSearch::PrintResults

This is to fix an issue where important referencing information was not printed with the default 'Log' level. This essentialy fixes the issue where the Verbosity passed in PrintStaleReferenceChainsAndFindReferencingObjects is not used in PrintResults.

diff --git a/Engine/Source/Runtime/CoreUObject/Private/UObject/ReferenceChainSearch.cpp b/Engine/Source/Runtime/CoreUObject/Private/UObject/ReferenceChainSearch.cpp
--- a/Engine/Source/Runtime/CoreUObject/Private/UObject/ReferenceChainSearch.cpp
+++ b/Engine/Source/Runtime/CoreUObject/Private/UObject/ReferenceChainSearch.cpp
@@ -1590,7 +1590,7 @@ FReferenceChainSearch::FReferenceChainSearch(TConstArrayView<UObject*> InObjects
 
 	if (!!(Mode & (EReferenceChainSearchMode::PrintResults|EReferenceChainSearchMode::PrintAllResults)))
 	{
-		PrintResults(!!(Mode & EReferenceChainSearchMode::PrintAllResults));
+		PrintResults(!!(Mode & EReferenceChainSearchMode::PrintAllResults), /*TargetObject*/nullptr, VerbosityForPrint);
 	}
 
 	UE_LOG(LogReferenceChain, Display, TEXT("Post-search memory usage: %.2f"), static_cast<double>(ReferenceGraph.GetAllocatedSize() + Paths.GetAllocatedSize() + GetAllocatedSize()) / 1024.0 / 1024.0);
@@ -1679,7 +1679,7 @@ void FReferenceChainSearch::PerformSearchFromGCSnapshot(TConstArrayView<UObject*
 
 	if (!!(SearchMode & (EReferenceChainSearchMode::PrintResults | EReferenceChainSearchMode::PrintAllResults)))
 	{
-		PrintResults(!!(SearchMode & EReferenceChainSearchMode::PrintAllResults));
+		PrintResults(!!(SearchMode & EReferenceChainSearchMode::PrintAllResults), /*TargetObject*/nullptr, VerbosityForPrint);
 	}
 }
 #endif // ENABLE_GC_HISTORY
@@ -1689,12 +1689,12 @@ void FReferenceChainSearch::SetVerbosityForPrint(ELogVerbosity::Type Verbosity)
 	VerbosityForPrint = Verbosity;
 }
 
-int32 FReferenceChainSearch::PrintResults(bool bDumpAllChains /*= false*/, UObject* TargetObject /*= nullptr*/) const
+int32 FReferenceChainSearch::PrintResults(bool bDumpAllChains /*= false*/, UObject* TargetObject /*= nullptr*/, ELogVerbosity::Type InVerbosityForPrint /*= ELogVerbosity::Log*/) const
 {
-	return PrintResults([](FCallbackParams& Params) { return true; }, bDumpAllChains, TargetObject);
+	return PrintResults([](FCallbackParams& Params) { return true; }, bDumpAllChains, TargetObject, InVerbosityForPrint);
 }
 
-int32 FReferenceChainSearch::PrintResults(TFunctionRef<bool(FCallbackParams& Params)> ReferenceCallback, bool bDumpAllChains /*= false*/, UObject* TargetObject /*= nullptr*/) const
+int32 FReferenceChainSearch::PrintResults(TFunctionRef<bool(FCallbackParams& Params)> ReferenceCallback, bool bDumpAllChains /*= false*/, UObject* TargetObject /*= nullptr*/, ELogVerbosity::Type InVerbosityForPrint /*= ELogVerbosity::Log*/) const
 {
 	FSlowHeartBeatScope DisableHangDetection; // This function can be very slow
 
@@ -1711,13 +1711,13 @@ int32 FReferenceChainSearch::PrintResults(TFunctionRef<bool(FCallbackParams& Par
 
 		if (bDumpAllChains || NumPrintedChains < MaxChainsToPrint)
 		{
-			DumpChain(Chain, ReferenceCallback, CallstackCache, *GWarn, VerbosityForPrint);
+			DumpChain(Chain, ReferenceCallback, CallstackCache, *GWarn, InVerbosityForPrint);
 			NumPrintedChains++;
 		}
 		else
 		{
 #if !NO_LOGGING
-			GWarn->CategorizedLogf(LogReferenceChain.GetCategoryName(), VerbosityForPrint,
+			GWarn->CategorizedLogf(LogReferenceChain.GetCategoryName(), InVerbosityForPrint,
 				TEXT("Referenced by %d more reference chain(s)."), ReferenceChains.Num() - NumPrintedChains);
 #endif
 			break;
@@ -1726,25 +1726,25 @@ int32 FReferenceChainSearch::PrintResults(TFunctionRef<bool(FCallbackParams& Par
 
 	if (NumPrintedChains == 0)
 	{
-		auto LogUnreachableObject = [this](const FGCObjectInfo& ObjInfo) {
+		auto LogUnreachableObject = [InVerbosityForPrint, this](const FGCObjectInfo& ObjInfo) {
 			if (ObjInfo.HasAnyInternalFlags(EInternalObjectFlags_RootFlags))
 			{
 #if !NO_LOGGING
-				GWarn->CategorizedLogf(LogReferenceChain.GetCategoryName(), VerbosityForPrint,
+				GWarn->CategorizedLogf(LogReferenceChain.GetCategoryName(), InVerbosityForPrint,
 					TEXT("%s%s is not currently reachable but it does have some of EInternalObjectFlags_RootFlags set."), *GetObjectFlags(ObjInfo), *ObjInfo.GetFullName());
 #endif
 			}
 			else if (ObjInfo.HasAnyFlags(GARBAGE_COLLECTION_KEEPFLAGS))
 			{
 #if !NO_LOGGING
-				GWarn->CategorizedLogf(LogReferenceChain.GetCategoryName(), VerbosityForPrint,
+				GWarn->CategorizedLogf(LogReferenceChain.GetCategoryName(), InVerbosityForPrint,
 					TEXT("%s%s is not currently reachable but it does have some of GARBAGE_COLLECTION_KEEPFLAGS set."), *GetObjectFlags(ObjInfo), *ObjInfo.GetFullName());
 #endif
 			}
 			else
 			{
 #if !NO_LOGGING
-				GWarn->CategorizedLogf(LogReferenceChain.GetCategoryName(), VerbosityForPrint,
+				GWarn->CategorizedLogf(LogReferenceChain.GetCategoryName(), InVerbosityForPrint,
 					TEXT("%s%s is not currently reachable. Try using GC history to debug transient leaks with 'gc.historysize 1'"), *GetObjectFlags(ObjInfo), *ObjInfo.GetFullName());
 #endif
 			}
@@ -1868,7 +1868,7 @@ static bool PrintStaleReferenceChainsAndFindReferencingObjects(UObject* ObjectTo
 				}
 				return true;
 			}
-		}, false, ObjectToFindReferencesTo) != 0;
+		}, false, ObjectToFindReferencesTo, Verbosity) != 0;
 }
 
 static FString GetPathToStaleObjectReferencer(UObject* ObjectToFindReferencesTo, FReferenceChainSearch& RefChainSearch)
diff --git a/Engine/Source/Runtime/CoreUObject/Public/UObject/ReferenceChainSearch.h b/Engine/Source/Runtime/CoreUObject/Public/UObject/ReferenceChainSearch.h
--- a/Engine/Source/Runtime/CoreUObject/Public/UObject/ReferenceChainSearch.h
+++ b/Engine/Source/Runtime/CoreUObject/Public/UObject/ReferenceChainSearch.h
@@ -310,7 +310,7 @@ public:
 	 * @param bDumpAllChains - if set to false, the output will be trimmed to the first 100 reference chains
 	 * @returns The number of results printed.
 	 */
-	COREUOBJECT_API int32 PrintResults(bool bDumpAllChains = false, UObject* TargetObject=nullptr) const;
+	COREUOBJECT_API int32 PrintResults(bool bDumpAllChains = false, UObject* TargetObject = nullptr, ELogVerbosity::Type VerbosityForPrint = ELogVerbosity::Log) const;
 	COREUOBJECT_API void SetVerbosityForPrint(ELogVerbosity::Type Verbosity);
 
 	/**
@@ -319,7 +319,7 @@ public:
 	 * @param bDumpAllChains - if set to false, the output will be trimmed to the first 100 reference chains
 	 * @returns The number of results printed.
 	 */
-	COREUOBJECT_API int32 PrintResults(TFunctionRef<bool(FCallbackParams& Params)> ReferenceCallback, bool bDumpAllChains = false, UObject* TargetObject = nullptr) const;
+	COREUOBJECT_API int32 PrintResults(TFunctionRef<bool(FCallbackParams& Params)> ReferenceCallback, bool bDumpAllChains = false, UObject* TargetObject = nullptr, ELogVerbosity::Type VerbosityForPrint = ELogVerbosity::Log) const;
 
 	/** Returns a string with a short report explaining the root path, will contain newlines */
 	COREUOBJECT_API FString GetRootPath(UObject* TargetObject = nullptr) const;
