Change object config save/load to pass PPF_SerializedAsImportText so types like GameplayTag know that the ImportText is actually a persistent one.
This fixes issues where gameplay tags in config files would be cleared if they used native tags from plugins that load late.

diff --git a/Engine/Source/Runtime/CoreUObject/Private/UObject/Obj.cpp b/Engine/Source/Runtime/CoreUObject/Private/UObject/Obj.cpp
--- a/Engine/Source/Runtime/CoreUObject/Private/UObject/Obj.cpp
+++ b/Engine/Source/Runtime/CoreUObject/Private/UObject/Obj.cpp
@@ -3075,7 +3075,7 @@ void UObject::LoadConfig(UClass* ConfigClass/*=nullptr*/, const TCHAR* InFilenam
 		const FString& PropFileName = *PropFileNamePtr;
 
 		FString Key = Property->GetName();
-		int32 PortFlags = 0;
+		int32 PortFlags = EPropertyPortFlags::PPF_SerializedAsImportText;
 
 #if WITH_EDITOR
 		const FString& CVarName = Property->GetMetaData(GConsoleVariableFName);
@@ -3378,7 +3378,7 @@ void UObject::SaveConfig(uint64 RequiredPropertyFlags, const TCHAR* InFilename,
 			}
 
 			FString Key				= Property->GetName();
-			int32 PortFlags			= 0;
+			int32 PortFlags			= EPropertyPortFlags::PPF_SerializedAsImportText;
 
 #if WITH_EDITOR
 			const FString& CVarName = Property->GetMetaData(GConsoleVariableFName);
