Fixed typos.

#rnx

diff --git a/Engine/Source/Runtime/CoreUObject/Public/UObject/UnrealType.h b/Engine/Source/Runtime/CoreUObject/Public/UObject/UnrealType.h
--- a/Engine/Source/Runtime/CoreUObject/Public/UObject/UnrealType.h
+++ b/Engine/Source/Runtime/CoreUObject/Public/UObject/UnrealType.h
@@ -1289,7 +1289,7 @@ public:
 	COREUOBJECT_API virtual void InitializeIntrusiveUnsetOptionalValue(void* Data) const;
 
 	/**
-	 * Returns whether an optional value of this inner type is unset. Only valid to call if HasIntrusiveOptionalState returns true.
+	 * Returns whether an optional value of this inner type is unset. Only valid to call if HasIntrusiveUnsetOptionalState returns true.
 	 * Equivalent to TOptional<T>::IsSet()
 	 * @see operator==(FIntrusiveUnsetOptionalState)
 	 * 
@@ -1299,7 +1299,7 @@ public:
 	COREUOBJECT_API virtual bool IsIntrusiveOptionalValueSet(const void* Data) const;
 
 	/**
-	 * Set the value to it's special unset state. Equivalent to TOptional<T>::Reset. Only valid to call if HasIntrusiveOptionalState returns true.
+	 * Set the value to it's special unset state. Equivalent to TOptional<T>::Reset. Only valid to call if HasIntrusiveUnsetOptionalState returns true.
 	 * @see operator=(FIntrusiveUnsetOptionalState)
 	 * 
 	 * @param Data Address of the alue, already offset.
