Mark _InitializedValues and _SerializedValues properties as SkipSerialization so that they don't lead to dirtied assets when changed.

diff --git a/Engine/Source/Runtime/CoreUObject/Private/UObject/InstanceDataObjectUtils.cpp b/Engine/Source/Runtime/CoreUObject/Private/UObject/InstanceDataObjectUtils.cpp
--- a/Engine/Source/Runtime/CoreUObject/Private/UObject/InstanceDataObjectUtils.cpp
+++ b/Engine/Source/Runtime/CoreUObject/Private/UObject/InstanceDataObjectUtils.cpp
@@ -619,8 +619,8 @@ namespace UE
 		FByteProperty* InitializedValuesProperty = CastFieldChecked<FByteProperty>(FByteProperty::Construct(Result, NAME_InitializedValues, RF_Transient | RF_MarkAsNative));
 		FByteProperty* SerializedValuesProperty = CastFieldChecked<FByteProperty>(FByteProperty::Construct(Result, NAME_SerializedValues, RF_Transient | RF_MarkAsNative));
 		{
-			InitializedValuesProperty->SetPropertyFlags(CPF_Transient | CPF_EditorOnly | CPF_NativeAccessSpecifierPrivate);
-			SerializedValuesProperty->SetPropertyFlags(CPF_Transient | CPF_EditorOnly | CPF_NativeAccessSpecifierPrivate);
+			InitializedValuesProperty->SetPropertyFlags(CPF_Transient | CPF_EditorOnly | CPF_SkipSerialization | CPF_NativeAccessSpecifierPrivate);
+			SerializedValuesProperty->SetPropertyFlags(CPF_Transient | CPF_EditorOnly | CPF_SkipSerialization | CPF_NativeAccessSpecifierPrivate);
 			Result->AddCppProperty(InitializedValuesProperty);
 			Result->AddCppProperty(SerializedValuesProperty);
 		}
