Changed hardcoded MaxTimeForFinishDestroy variable to a proper cvar

#rnx

diff --git a/Engine/Source/Runtime/CoreUObject/Private/UObject/GarbageCollection.cpp b/Engine/Source/Runtime/CoreUObject/Private/UObject/GarbageCollection.cpp
--- a/Engine/Source/Runtime/CoreUObject/Private/UObject/GarbageCollection.cpp
+++ b/Engine/Source/Runtime/CoreUObject/Private/UObject/GarbageCollection.cpp
@@ -4562,6 +4562,14 @@ static FAutoConsoleVariableRef CVarGCEnableTimeoutOnPendingDestroyedObjectInShip
 );
 #endif
 
+static float GMaxTimeForFinishDestroyGC = 10.0f;
+static FAutoConsoleVariableRef CVarGCMaxTimeForFinishDestroy(
+	TEXT("gc.MaxTimeForFinishDestroyGC"),
+	GMaxTimeForFinishDestroyGC,
+	TEXT("Max time in seconds GC waits on pending objects to be destroyed (default is 10 seconds)."),
+	ECVF_Default
+);
+
 static int32 GAdditionalFinishDestroyTimeGC = 40;
 static FAutoConsoleVariableRef CVarAdditionalFinishDestroyTimeGC(
 	TEXT("gc.AdditionalFinishDestroyTimeGC"),
@@ -4662,7 +4670,7 @@ bool IncrementalDestroyGarbage(bool bUseTimeLimit, double TimeLimit)
 		if (GObjCurrentPurgeObjectIndex >= GUnreachableObjects.Num())
 		{
 			TRACE_CPUPROFILER_EVENT_SCOPE(ConditionalFinishDestroyDeferred);
-			double MaxTimeForFinishDestroy = 10.00;
+			double MaxTimeForFinishDestroy = double(GMaxTimeForFinishDestroyGC);
 			bool bFinishDestroyTimeExtended = false;
 			FString FirstObjectNotReadyWhenTimeExtended;
 			int32 StartObjectsPendingDestructionCount = GGCObjectsPendingDestructionCount;
