Ensure that changes to serialized subobject instance data are reflected on associated IDOs at PostLoad() time.

#rnx

diff --git a/Engine/Source/Runtime/CoreUObject/Private/UObject/Obj.cpp b/Engine/Source/Runtime/CoreUObject/Private/UObject/Obj.cpp
--- a/Engine/Source/Runtime/CoreUObject/Private/UObject/Obj.cpp
+++ b/Engine/Source/Runtime/CoreUObject/Private/UObject/Obj.cpp
@@ -1290,15 +1290,6 @@ void UObject::ConditionalPostLoad()
 
 		ConditionalPostLoadSubobjects();
 
-#if WITH_EDITORONLY_DATA
-		// Object has been deserialized, if IDO is enabled, generate it
-		if (UE::IsInstanceDataObjectSupportEnabled(this) && !FUObjectThreadContext::Get().GetSerializeContext()->bImpersonateProperties)
-		{
-			UE::FPropertyBagRepository& PropertyBagRepository = UE::FPropertyBagRepository::Get();
-			PropertyBagRepository.PostLoadInstanceDataObject(this);
-		}
-#endif
-
 		{
 			FExclusiveLoadPackageTimeTracker::FScopedPostLoadTracker Tracker(this);
 
@@ -1324,6 +1315,17 @@ void UObject::ConditionalPostLoad()
 			}
 		}
 
+#if WITH_EDITORONLY_DATA
+		// Object has been deserialized, so update its associated IDO. Note that this should be done
+		// *after* PostLoad() is called, since that could mutate already-serialized fields on this object,
+		// and we need those changes to also be reflected on the IDO, which was created at serialization time.
+		if (UE::IsInstanceDataObjectSupportEnabled(this) && !FUObjectThreadContext::Get().GetSerializeContext()->bImpersonateProperties)
+		{
+			UE::FPropertyBagRepository& PropertyBagRepository = UE::FPropertyBagRepository::Get();
+			PropertyBagRepository.PostLoadInstanceDataObject(this);
+		}
+#endif
+
 #if !(UE_BUILD_SHIPPING || UE_BUILD_TEST)
 		if (ThreadContext.DebugPostLoad.Contains(this))
 		{
diff --git a/Engine/Source/Runtime/CoreUObject/Private/UObject/PropertyBagRepository.cpp b/Engine/Source/Runtime/CoreUObject/Private/UObject/PropertyBagRepository.cpp
--- a/Engine/Source/Runtime/CoreUObject/Private/UObject/PropertyBagRepository.cpp
+++ b/Engine/Source/Runtime/CoreUObject/Private/UObject/PropertyBagRepository.cpp
@@ -5,6 +5,7 @@
 #if WITH_EDITORONLY_DATA
 
 #include "Containers/Queue.h"
+#include "InstancedReferenceSubobjectHelper.h"
 #include "Misc/StringBuilder.h"
 #include "Serialization/ObjectReader.h"
 #include "Serialization/ObjectWriter.h"
@@ -694,6 +695,19 @@ void FPropertyBagRepository::PostLoadInstanceDataObject(const UObject* Owner)
 		{
 			// copy data from owner to IDO
 			CopyTaggedProperties(Owner, BagData->InstanceDataObject);
+
+			// the owner's PostLoad() may have mutated its instanced subobjects as well (e.g. pointer fixup). to handle
+			// that case, we look for any instanced subobjects that have already had their PostLoad() called, as those will
+			// not have a chance to get their IDO data fixed up to match changes potentially made by its owner's PostLoad().
+			TSet<UObject*> InstancedSubObjects;
+			FFindInstancedReferenceSubobjectHelper::GetInstancedSubObjects(Owner, InstancedSubObjects);
+			for (UObject* InstancedSubObject : InstancedSubObjects)
+			{
+				if (!InstancedSubObject->HasAnyFlags(RF_NeedPostLoad) && InstancedSubObject->IsInOuter(Owner))
+				{
+					PostLoadInstanceDataObject(InstancedSubObject);
+				}
+			}
 		}
 	}
 }
