Treat packages containing only placeholder exports as editor-only data for cooking purposes.

#rnx

diff --git a/Engine/Source/Runtime/CoreUObject/Private/UObject/SavePackage2.cpp b/Engine/Source/Runtime/CoreUObject/Private/UObject/SavePackage2.cpp
--- a/Engine/Source/Runtime/CoreUObject/Private/UObject/SavePackage2.cpp
+++ b/Engine/Source/Runtime/CoreUObject/Private/UObject/SavePackage2.cpp
@@ -625,13 +625,35 @@ ESavePackageResult ValidateExports(FSaveContext& SaveContext)
 #if WITH_EDITOR
 	if (UE::FPropertyBagRepository::IsPropertyBagPlaceholderObjectSupportEnabled())
 	{
-		bool bSuccess = true;
-		for (const FTaggedExport& Export : SaveContext.GetExports())
+		TArray<const UObject*, TInlineAllocator<16>> PlaceholderExports;
+		const TSet<FTaggedExport>& AllExports = SaveContext.GetExports();
+		for (const FTaggedExport& Export : AllExports)
 		{
 			if (UE::FPropertyBagRepository::IsPropertyBagPlaceholderObject(Export.Obj))
 			{
-				UObject* Outer = Export.Obj->GetOuter();
-				FString ErrorMessage = FString::Printf(TEXT("Saving illegal placeholder object in package '%s'. Object Name '%s', Class '%s', Outer '%s', Outer Class '%s', Full Path '%s'. You might have unresolved classes or compilation failures to fix."), *SaveContext.GetPackage()->GetName(), *Export.Obj->GetName(), *Export.Obj->GetClass()->GetName(), Outer ? *Outer->GetName() : TEXT("None"), Outer ? *Outer->GetClass()->GetName() : TEXT("None"), *Export.Obj->GetFullName());
+				PlaceholderExports.Add(Export.Obj);
+			}
+			else if (const UStruct* TypeExport = Cast<UStruct>(Export.Obj))
+			{
+				if (UE::FPropertyBagRepository::IsPropertyBagPlaceholderType(TypeExport))
+				{
+					PlaceholderExports.Add(TypeExport);
+				}
+			}
+		}
+
+		if (PlaceholderExports.Num() == AllExports.Num())
+		{
+			// If all exports are placeholders, allow this to pass as editor-only data (e.g. a placeholder type package).
+			return ESavePackageResult::ContainsEditorOnlyData;
+		}
+		else if (PlaceholderExports.Num() > 0)
+		{
+			for (const UObject* PlaceholderExport : PlaceholderExports)
+			{
+				check(PlaceholderExport);
+				UObject* Outer = PlaceholderExport->GetOuter();
+				FString ErrorMessage = FString::Printf(TEXT("Saving illegal placeholder object in package '%s'. Object Name '%s', Class '%s', Outer '%s', Outer Class '%s', Full Path '%s'. You might have unresolved classes or compilation failures to fix."), *SaveContext.GetPackage()->GetName(), *PlaceholderExport->GetName(), *PlaceholderExport->GetClass()->GetName(), Outer ? *Outer->GetName() : TEXT("None"), Outer ? *Outer->GetClass()->GetName() : TEXT("None"), *PlaceholderExport->GetFullName());
 				if (SaveContext.IsGenerateSaveError())
 				{
 					SaveContext.GetError()->Logf(ELogVerbosity::Warning, TEXT("%s"), *ErrorMessage);
@@ -640,11 +662,8 @@ ESavePackageResult ValidateExports(FSaveContext& SaveContext)
 				{
 					UE_LOG(LogSavePackage, Error, TEXT("%s"), *ErrorMessage);
 				}
-				bSuccess = false;
 			}
-		}
-		if (!bSuccess)
-		{
+
 			return ESavePackageResult::Error;
 		}
 	}
