Relax ClassWithin constraint while reinstancing and avoid marking discarded ICH templates as garbage, as they creates reachability warnings since we are maintaining a reference to these objects

[RN] Fixed reachable garbage issue involving UInheritableComponentHandler::UnnecessaryComponents

diff --git a/Engine/Source/Runtime/CoreUObject/Private/UObject/UObjectGlobals.cpp b/Engine/Source/Runtime/CoreUObject/Private/UObject/UObjectGlobals.cpp
--- a/Engine/Source/Runtime/CoreUObject/Private/UObject/UObjectGlobals.cpp
+++ b/Engine/Source/Runtime/CoreUObject/Private/UObject/UObjectGlobals.cpp
@@ -3385,7 +3385,13 @@ bool StaticAllocateObjectErrorTests( const UClass* InClass, UObject* InOuter, FN
 		}
 	}
 
-	if ( (InFlags & (RF_ClassDefaultObject|RF_ArchetypeObject)) == 0 )
+	// When reinstancing, allow any within violations as they were likely caused by users renaming
+	// objects to different outers, hopefully with intent:
+	if ( (InFlags & (RF_ClassDefaultObject|RF_ArchetypeObject)) == 0 
+#if WITH_EDITORONLY_DATA
+		&& !GIsReinstancing
+#endif
+	)
 	{
 		if ( InOuter != NULL && !InOuter->IsA(InClass->ClassWithin) )
 		{
diff --git a/Engine/Source/Runtime/Engine/Private/InheritableComponentHandler.cpp b/Engine/Source/Runtime/Engine/Private/InheritableComponentHandler.cpp
--- a/Engine/Source/Runtime/Engine/Private/InheritableComponentHandler.cpp
+++ b/Engine/Source/Runtime/Engine/Private/InheritableComponentHandler.cpp
@@ -262,7 +262,6 @@ void UInheritableComponentHandler::ValidateTemplates()
 						Record.ComponentTemplate->Rename(nullptr, GetTransientPackage(), REN_DoNotDirty | REN_DontCreateRedirectors | REN_NonTransactional);
 						Record.ComponentTemplate->ClearFlags(RF_Standalone);
 						Record.ComponentTemplate->RemoveFromRoot();
-						Record.ComponentTemplate->MarkAsGarbage();
 #endif // WITH_EDITOR
 						UnnecessaryComponents.AddUnique(Record.ComponentTemplate);
 					}
