Move some UObject locks to transactionally safe locks

diff --git a/Engine/Source/Runtime/CoreUObject/Private/UObject/UObjectArray.cpp b/Engine/Source/Runtime/CoreUObject/Private/UObject/UObjectArray.cpp
--- a/Engine/Source/Runtime/CoreUObject/Private/UObject/UObjectArray.cpp
+++ b/Engine/Source/Runtime/CoreUObject/Private/UObject/UObjectArray.cpp
@@ -129,7 +129,7 @@ void FUObjectArray::OpenDisregardForGC()
 void FUObjectArray::CloseDisregardForGC()
 {
 #if THREADSAFE_UOBJECTS
-	FScopeLock ObjObjectsLock(&ObjObjectsCritical);
+	FTransactionallySafeScopeLock ObjObjectsLock(&ObjObjectsCritical);
 #else
 	// Disregard from GC pool is only available from the game thread, at least for now
 	check(IsInGameThread());
diff --git a/Engine/Source/Runtime/CoreUObject/Public/UObject/UObjectArray.h b/Engine/Source/Runtime/CoreUObject/Public/UObject/UObjectArray.h
--- a/Engine/Source/Runtime/CoreUObject/Public/UObject/UObjectArray.h
+++ b/Engine/Source/Runtime/CoreUObject/Public/UObject/UObjectArray.h
@@ -11,6 +11,7 @@
 #include "Misc/TransactionallySafeCriticalSection.h"
 #include "UObject/GarbageCollectionGlobals.h"
 #include "UObject/UObjectBase.h"
+#include "Misc/TransactionallySafeScopeLock.h"
 
 #if WITH_VERSE_VM || defined(__INTELLISENSE__)
 namespace Verse
@@ -1304,7 +1305,7 @@ private:
 	/** Array of all live objects.											*/
 	TUObjectArray ObjObjects;
 	/** Synchronization object for all live objects.											*/
-	mutable FCriticalSection ObjObjectsCritical;
+	mutable FTransactionallySafeCriticalSection ObjObjectsCritical;
 	/** Available object indices.											*/
 	TArray<int32> ObjAvailableList;
 
@@ -1341,7 +1342,7 @@ public:
 
     SIZE_T GetAllocatedSize() const
     {
-		FScopeLock ObjListLock(&ObjObjectsCritical);
+		FTransactionallySafeScopeLock ObjListLock(&ObjObjectsCritical);
 #if THREADSAFE_UOBJECTS
 		FTransactionallySafeScopeLock ListenersLock(&UObjectDeleteListenersCritical);
 #endif
