[BugFix] Name collision fix for InstanceDataObject Classes

diff --git a/Engine/Source/Runtime/CoreUObject/Private/UObject/InstanceDataObjectUtils.cpp b/Engine/Source/Runtime/CoreUObject/Private/UObject/InstanceDataObjectUtils.cpp
--- a/Engine/Source/Runtime/CoreUObject/Private/UObject/InstanceDataObjectUtils.cpp
+++ b/Engine/Source/Runtime/CoreUObject/Private/UObject/InstanceDataObjectUtils.cpp
@@ -312,7 +312,7 @@ namespace UE
 
 		// Construct a transient type that impersonates the original type.
 		const FName InstanceDataObjectName(WriteToString<128>(EnumName, TEXTVIEW("_InstanceDataObject")));
-		UEnum* Enum = NewObject<UEnum>(Outer, MakeUniqueObjectName(nullptr, UEnum::StaticClass(), InstanceDataObjectName));
+		UEnum* Enum = NewObject<UEnum>(Outer, MakeUniqueObjectName(Outer, UEnum::StaticClass(), InstanceDataObjectName));
 		Enum->SetEnums(EnumValueNames, TemplateEnum->GetCppForm(), bHasFlags ? EEnumFlags::Flags : EEnumFlags::None, /*bAddMaxKeyIfMissing*/ false);
 		Enum->SetMetaData(*WriteToString<32>(NAME_OriginalType), *WriteToString<128>(EnumTypeName));
 
@@ -571,7 +571,7 @@ namespace UE
 		}
 
 		const FName InstanceDataObjectName = (TemplateStruct) ? FName(WriteToString<128>(TemplateStruct->GetName(), TEXTVIEW("_InstanceDataObject"))) : FName(TEXTVIEW("InstanceDataObject"));
-		UStruct* Result = NewObject<UStruct>(Outer, StructClass, MakeUniqueObjectName(nullptr, StructClass, InstanceDataObjectName));
+		UStruct* Result = NewObject<UStruct>(Outer, StructClass, MakeUniqueObjectName(Outer, StructClass, InstanceDataObjectName));
 		Result->SetSuperStruct(Super);
 
 		// inherit ContainsLooseProperties metadata
