[AutoRTFM] Use the FTransactionallySafeRWLock for FUObjectAnnotationChunked::AnnotationArrayCritical

Drop the explicit open / on-abort logic in UObjectArchetype as it should not be needed any more.

diff --git a/Engine/Source/Runtime/CoreUObject/Private/UObject/UObjectArchetype.cpp b/Engine/Source/Runtime/CoreUObject/Private/UObject/UObjectArchetype.cpp
--- a/Engine/Source/Runtime/CoreUObject/Private/UObject/UObjectArchetype.cpp
+++ b/Engine/Source/Runtime/CoreUObject/Private/UObject/UObjectArchetype.cpp
@@ -234,16 +234,7 @@ UObject* GetArchetypeImpl(const UObject* InObject, const FObjectArchetypeHelper:
 		if (Archetype && !(InObject->GetOuter() && InObject->GetOuter()->HasAnyFlags(RF_NeedLoad)))
 		{
 			ArchetypeIndex = GUObjectArray.ObjectToIndex(Archetype);
-			UE_AUTORTFM_OPEN2
-			{
-				ArchetypeAnnotation.AddAnnotation(InObject, FArchetypeInfo{ ArchetypeIndex, GUObjectArray.AllocateSerialNumber(ArchetypeIndex) });
-			};
-
-			// If we abort we need to remove this annotation
-			UE_AUTORTFM_ONABORT2(=)
-			{
-				ArchetypeAnnotation.RemoveAnnotation(InObject);
-			};
+			ArchetypeAnnotation.AddAnnotation(InObject, FArchetypeInfo{ ArchetypeIndex, GUObjectArray.AllocateSerialNumber(ArchetypeIndex) });
 		}
 	}
 	else
diff --git a/Engine/Source/Runtime/CoreUObject/Public/UObject/UObjectAnnotation.h b/Engine/Source/Runtime/CoreUObject/Public/UObject/UObjectAnnotation.h
--- a/Engine/Source/Runtime/CoreUObject/Public/UObject/UObjectAnnotation.h
+++ b/Engine/Source/Runtime/CoreUObject/Public/UObject/UObjectAnnotation.h
@@ -7,7 +7,7 @@
 #pragma once
 
 #include "UObject/UObjectArray.h"
-#include "Misc/TransactionallySafeScopeLock.h"
+#include "Misc/TransactionallySafeRWLock.h"
 #include "Misc/TransactionallySafeRWScopeLock.h"
 
 /**
@@ -588,7 +588,7 @@ class FUObjectAnnotationChunked : public FUObjectArray::FUObjectDeleteListener
 	uint32 MaxAllocatedMemory;
 
 	/** Mutex */
-	FRWLock AnnotationArrayCritical;
+	FTransactionallySafeRWLock AnnotationArrayCritical;
 
 	/**
 	* Makes sure we have enough chunks to fit the new index
@@ -760,13 +760,13 @@ public:
 	 */
 	void AddAnnotation(int32 Index, const TAnnotation& Annotation)
 	{
-		FRWScopeLock AnnotationArrayLock(AnnotationArrayCritical, SLT_Write);
+		FTransactionallySafeWriteScopeLock AnnotationArrayLock(AnnotationArrayCritical);
 		AddAnnotationInternal(Index, Annotation);
 	}
 
 	void AddAnnotation(int32 Index, TAnnotation&& Annotation)
 	{
-		FRWScopeLock AnnotationArrayLock(AnnotationArrayCritical, SLT_Write);
+		FTransactionallySafeWriteScopeLock AnnotationArrayLock(AnnotationArrayCritical);
 		AddAnnotationInternal(Index, MoveTemp(Annotation));
 	}
 
@@ -783,7 +783,7 @@ public:
 	 */
 	TAnnotation& AddOrGetAnnotation(int32 Index, TFunctionRef<TAnnotation()> NewAnnotationFn)
 	{		
-		FRWScopeLock AnnotationArrayLock(AnnotationArrayCritical, SLT_Write);
+		FTransactionallySafeWriteScopeLock AnnotationArrayLock(AnnotationArrayCritical);
 		
 		if (NumAnnotations == 0 && Chunks.Num() == 0)
 		{
@@ -836,7 +836,7 @@ public:
 	 */
 	void RemoveAnnotation(int32 Index)
 	{
-		FRWScopeLock AnnotationArrayLock(AnnotationArrayCritical, SLT_Write);
+		FTransactionallySafeWriteScopeLock AnnotationArrayLock(AnnotationArrayCritical);
 		FreeAnnotation(Index);
 	}
 
@@ -864,7 +864,7 @@ public:
 
 		UE_AUTORTFM_OPEN2
 		{
-			FRWScopeLock AnnotationArrayLock(AnnotationArrayCritical, SLT_ReadOnly);
+			FTransactionallySafeReadScopeLock AnnotationArrayLock(AnnotationArrayCritical);
 
 			const int32 ChunkIndex = Index / NumAnnotationsPerChunk;
 			if (ChunkIndex < Chunks.Num())
@@ -930,7 +930,7 @@ public:
 	void RemoveAllAnnotations()
 	{
 		bool bHadAnnotations = (NumAnnotations > 0);	
-		FRWScopeLock AnnotationArrayLock(AnnotationArrayCritical, SLT_Write);
+		FTransactionallySafeWriteScopeLock AnnotationArrayLock(AnnotationArrayCritical);
 		FreeAllAnnotations();
 		if (bHadAnnotations)
 		{
@@ -949,7 +949,7 @@ public:
 	 */
 	void TrimAnnotations()
 	{
-		FRWScopeLock AnnotationArrayLock(AnnotationArrayCritical, SLT_Write);
+		FTransactionallySafeWriteScopeLock AnnotationArrayLock(AnnotationArrayCritical);
 		for (TAnnotationChunk& Chunk : Chunks)
 		{
 			if (Chunk.Num == 0 && Chunk.Items)
