Removed explicit ARO of IDO namespaces because these are referenced through outers

Fixed missing removal of namespace mappings.

#rnx

diff --git a/Engine/Source/Runtime/CoreUObject/Private/UObject/PropertyBagRepository.cpp b/Engine/Source/Runtime/CoreUObject/Private/UObject/PropertyBagRepository.cpp
--- a/Engine/Source/Runtime/CoreUObject/Private/UObject/PropertyBagRepository.cpp
+++ b/Engine/Source/Runtime/CoreUObject/Private/UObject/PropertyBagRepository.cpp
@@ -48,7 +48,7 @@ public:
 	void AddReferencedObjects(FReferenceCollector& Collector)
 	{
 		ConsumePendingPlaceholderTypes();
-		Collector.AddReferencedObjects(PlaceholderTypes);
+		Collector.AddStableReferenceSet(&PlaceholderTypes);
 	}
 
 	void Add(UStruct* Type)
@@ -754,6 +754,10 @@ void FPropertyBagRepository::MarkAsFixedUp(const UObject* Object)
 
 bool FPropertyBagRepository::RemoveAssociationUnsafe(const UObject* Owner)
 {
+	// note: RemoveAssociationUnsafe is called on every object regardless of whether it has a property bag.
+	// in that scenario, there's a chance we have a namespace associated with it. Remove that namespace.
+	Namespaces.Remove(Owner);
+
 	FPropertyBagAssociationData OldData;
 	if(AssociatedData.RemoveAndCopyValue(Owner, OldData))
 	{
@@ -761,10 +765,6 @@ bool FPropertyBagRepository::RemoveAssociationUnsafe(const UObject* Owner)
 		OldData.Destroy();
 		return true;
 	}
-
-	// note: RemoveAssociationUnsafe is called on every object regardless of whether it has a property bag.
-	// in that scenario, there's a chance we have a namespace associated with it. Remove that namespace.
-	Namespaces.Remove(Owner);
 	return false;
 }
 
@@ -837,11 +837,6 @@ bool FPropertyBagRepository::WasPropertyValueSerialized(const UStruct* Struct, c
 
 void FPropertyBagRepository::AddReferencedObjects(FReferenceCollector& Collector)
 {
-	for (TPair<const UObject*, TObjectPtr<UObject>>& Element : Namespaces)
-	{
-		Collector.AddReferencedObject(Element.Value);
-	}
-
 	FPropertyBagPlaceholderTypeRegistry::Get().AddReferencedObjects(Collector);
 }
 
