Change the invalid ClassWithin check to not crash in the editor, it now ensures like the abstract class check.
Clean up handling of GIsEditor for object allocation checks to avoid redundant work and better indicate which checks are editor only.
Change Rename() to skip within and linker checks if the new outer is the same as the current one, this avoids potential crashes and is faster.

diff --git a/Engine/Source/Runtime/CoreUObject/Private/UObject/Obj.cpp b/Engine/Source/Runtime/CoreUObject/Private/UObject/Obj.cpp
--- a/Engine/Source/Runtime/CoreUObject/Private/UObject/Obj.cpp
+++ b/Engine/Source/Runtime/CoreUObject/Private/UObject/Obj.cpp
@@ -226,7 +226,7 @@ bool UObject::Rename(const TCHAR* InName, UObject* NewOuter, ERenameFlags Flags)
 	FMetaDataUtilities::FMoveMetadataHelperContext MoveMetaData(this, true);
 #endif //WITH_EDITOR
 
-	if (NewOuter)
+	if (NewOuter && NewOuter != GetOuter())
 	{
 		// Renaming the CDO of a Blueprint is a special case so we do not validate what would otherwise be incorrect use of Rename.
 		// Moving objects to the transient package is commonly used halfway through destroying them so that is also fine, otherwise
@@ -1455,6 +1455,18 @@ void UObject::PreSave(FObjectPreSaveContext SaveContext)
 	FCoreUObjectDelegates::OnObjectSaved.Broadcast(this);
 	PRAGMA_ENABLE_DEPRECATION_WARNINGS
 	FCoreUObjectDelegates::OnObjectPreSave.Broadcast(this, SaveContext);
+
+	// Validate ClassWithin before save, matches code in StaticAllocateObjectErrorTests
+	UClass* ObjClass = GetClass();
+	if (ObjClass->ClassWithin != nullptr)
+	{
+		if (!HasAnyFlags(RF_ClassDefaultObject | RF_ArchetypeObject) && GetOuter() != nullptr && !GetOuter()->IsA(ObjClass->ClassWithin))
+		{
+			const FString ErrorMsg = FString::Printf(TEXT("Object %s with ClassWithin of %s is being saved in invalid Outer of class %s!"), *GetFullName(nullptr, EObjectFullNameFlags::IncludeClassPackage), *ObjClass->ClassWithin->GetPathName(), *GetOuter()->GetClass()->GetPathName());
+			UE_LOG(LogUObjectGlobals, Error, TEXT("%s"), *ErrorMsg);
+			ensureMsgf(false, TEXT("%s"), *ErrorMsg);
+		}
+	}
 #endif
 	CollectSaveOverrides(FObjectCollectSaveOverridesContext(SaveContext.Data));
 }
diff --git a/Engine/Source/Runtime/CoreUObject/Private/UObject/UObjectGlobals.cpp b/Engine/Source/Runtime/CoreUObject/Private/UObject/UObjectGlobals.cpp
--- a/Engine/Source/Runtime/CoreUObject/Private/UObject/UObjectGlobals.cpp
+++ b/Engine/Source/Runtime/CoreUObject/Private/UObject/UObjectGlobals.cpp
@@ -3350,30 +3350,22 @@ bool FScopedAllowAbstractClassAllocation::IsDisallowedAbstractClass(const UClass
 	return false;
 }
 
+#if WITH_EDITOR
 bool StaticAllocateObjectErrorTests( const UClass* InClass, UObject* InOuter, FName InName, EObjectFlags InFlags)
 {
-	// Validation checks.
+	// Editor-only validation checks, StaticAllocateObject has the simpler checks for packaged builds
 	if( !InClass )
 	{
 		UE_LOG(LogUObjectGlobals, Fatal, TEXT("Empty class for object %s"), *InName.ToString() );
 		return true;
 	}
 
-	// for abstract classes that are being loaded NOT in the editor we want to error.  If they are in the editor we do not want to have an error
 	if (FScopedAllowAbstractClassAllocation::IsDisallowedAbstractClass(InClass, InFlags))
 	{
-		if ( GIsEditor )
-		{
-			const FString ErrorMsg = FString::Printf(TEXT("Class which was marked abstract was trying to be loaded in Outer %s.  It will be nulled out on save. %s %s"), *GetPathNameSafe(InOuter), *InName.ToString(), *InClass->GetName());
-			// if we are trying instantiate an abstract class in the editor we'll warn the user that it will be nulled out on save
-			UE_LOG(LogUObjectGlobals, Warning, TEXT("%s"), *ErrorMsg);
-			ensureMsgf(false, TEXT("%s"), *ErrorMsg);
-		}
-		else
-		{
-			UE_LOG(LogUObjectGlobals, Fatal, TEXT("%s"), *FString::Printf( TEXT("Can't create object %s in Outer %s: class %s is abstract"), *InName.ToString(), *GetPathNameSafe(InOuter), *InClass->GetName()));
-			return true;
-		}
+		const FString ErrorMsg = FString::Printf(TEXT("Class which was marked abstract was trying to be loaded in Outer %s.  It will be nulled out on save. %s %s"), *GetPathNameSafe(InOuter), *InName.ToString(), *InClass->GetName());
+		// if we are trying instantiate an abstract class in the editor we'll warn the user that it will be nulled out on save
+		UE_LOG(LogUObjectGlobals, Warning, TEXT("%s"), *ErrorMsg);
+		ensureMsgf(false, TEXT("%s"), *ErrorMsg);
 	}
 
 	if( InOuter == NULL )
@@ -3392,20 +3384,19 @@ bool StaticAllocateObjectErrorTests( const UClass* InClass, UObject* InOuter, FN
 
 	// When reinstancing, allow any within violations as they were likely caused by users renaming
 	// objects to different outers, hopefully with intent:
-	if ( (InFlags & (RF_ClassDefaultObject|RF_ArchetypeObject)) == 0 
-#if WITH_EDITORONLY_DATA
-		&& !GIsReinstancing
-#endif
-	)
+	if ( (InFlags & (RF_ClassDefaultObject|RF_ArchetypeObject)) == 0 && !GIsReinstancing )
 	{
-		if ( InOuter != NULL && !InOuter->IsA(InClass->ClassWithin) )
+		if ( InOuter != nullptr && !InOuter->IsA(InClass->ClassWithin) )
 		{
-			UE_LOG(LogUObjectGlobals, Fatal, TEXT("%s"), *FString::Printf( TEXT("Object %s %s created in %s instead of %s"), *InClass->GetName(), *InName.ToString(), *InOuter->GetClass()->GetName(), *InClass->ClassWithin->GetName()) );
-			return true;
+			// This is also validated in UObject::PreSave
+			const FString ErrorMsg = FString::Printf(TEXT("Object %s of class %s with ClassWithin of %s was created in invalid Outer %s!"), *InName.ToString(), *InClass->GetPathName(), *InClass->ClassWithin->GetPathName(), *InOuter->GetClass()->GetPathName());
+			UE_LOG(LogUObjectGlobals, Warning, TEXT("%s"), *ErrorMsg);
+			ensureMsgf(false, TEXT("%s"), *ErrorMsg);
 		}
 	}
 	return false;
 }
+#endif // WITH_EDITOR
 
 /**
 * For object overwrites, the class may want to persist some info over the re-initialize
@@ -3434,7 +3425,10 @@ UObject* StaticAllocateObject
 
 	SCOPE_CYCLE_COUNTER(STAT_AllocateObject);
 	checkSlow(InOuter != INVALID_OBJECT); // not legal
-	check(!InClass || (InClass->ClassWithin && InClass->ClassConstructor));
+	check(InClass && InClass->ClassWithin && InClass->ClassConstructor);
+
+	const bool bCreatingCDO = (InFlags & RF_ClassDefaultObject) != 0;
+	const bool bCreatingArchetype = (InFlags & RF_ArchetypeObject) != 0;
 #if WITH_EDITOR
 	if (GIsEditor)
 	{
@@ -3443,17 +3437,16 @@ UObject* StaticAllocateObject
 			return NULL;
 		}
 	}
+	else
 #endif // WITH_EDITOR
-	const bool bCreatingCDO = (InFlags & RF_ClassDefaultObject) != 0;
-	const bool bCreatingArchetype = (InFlags & RF_ArchetypeObject) != 0;
+	{
+		// In the editor these are handled inside StaticAllocateObjectErrorTests and they may be temporary warnings
+		checkf(!FScopedAllowAbstractClassAllocation::IsDisallowedAbstractClass(InClass, InFlags), TEXT("Unable to create new object: %s %s.%s. Creating an instance of an abstract class is not allowed!"),
+			*GetNameSafe(InClass), *GetPathNameSafe(InOuter), *InName.ToString());
+		check(bCreatingCDO || bCreatingArchetype || !InOuter || InOuter->IsA(InClass->ClassWithin));
+		check(InOuter || (InClass == UPackage::StaticClass() && InName != NAME_None)); // only packages can not have an outer, and they must be named explicitly	
+	}
 
-	check(InClass);
-	check(InOuter || (InClass == UPackage::StaticClass() && InName != NAME_None)); // only packages can not have an outer, and they must be named explicitly
-	// this is a warning in the editor, otherwise it is illegal to create an abstract class, except the CDO
-	checkf(GIsEditor || !FScopedAllowAbstractClassAllocation::IsDisallowedAbstractClass(InClass, InFlags), TEXT("Unable to create new object: %s %s.%s. Creating an instance of an abstract class is not allowed!"),
-		*GetNameSafe(InClass), *GetPathNameSafe(InOuter), *InName.ToString());
-	//checkf(InClass != UPackage::StaticClass() || !InOuter || bCreatingCDO, TEXT("Creating nested packages is not allowed: Outer=%s, Package=%s"), *GetNameSafe(InOuter), *InName.ToString());
-	check(bCreatingCDO || bCreatingArchetype || !InOuter || InOuter->IsA(InClass->ClassWithin));
 	checkf(!IsGarbageCollectingAndLockingUObjectHashTables(), TEXT("Unable to create new object: %s %s.%s. Creating UObjects while Collecting Garbage is not allowed!"),
 		*GetNameSafe(InClass), *GetPathNameSafe(InOuter), *InName.ToString());
 
