- Ensure when creating exports using a archetype/template object that the template object has been preloaded (Zenloader only). In cases where the template object is a blueprint generated class, we may have only created the archetype export as part of a recursive call but have not yet preloaded the object. When we resolve deferred exports we may then inadvertently postload an export that has this non-preloaded archetype set and begin tripped asserts.
- When patching Blueprint CDO subobjects, if we cannot find a sub-object on disk in our lookup table to patch, we invalidate the export as it is no longer valid. This prevents the export from loading in future requests and prevents subsequent dependent loads the export would have caused from firing saving time since those loads are all wasted work.

#rnx

diff --git a/Engine/Source/Editor/UnrealEd/Private/Kismet2/BlueprintEditorUtils.cpp b/Engine/Source/Editor/UnrealEd/Private/Kismet2/BlueprintEditorUtils.cpp
--- a/Engine/Source/Editor/UnrealEd/Private/Kismet2/BlueprintEditorUtils.cpp
+++ b/Engine/Source/Editor/UnrealEd/Private/Kismet2/BlueprintEditorUtils.cpp
@@ -1502,6 +1502,12 @@ void FBlueprintEditorUtils::PatchCDOSubobjectsIntoExport(UObject* PreviousCDO, U
 
 							PatchSubObjects(OldSubObj, NewSubObj, AlreadyPatched);
 						}
+						// Matches check in FPackageHarvester::ProcessImport. Invalidate these objects
+						// to avoid generation of implicit dependencies
+						else if (OldSubObj->HasAnyFlags(RF_DefaultSubObject | RF_ArchetypeObject))
+						{
+							FLinkerLoad::InvalidateExport(OldSubObj);
+						}
 					}
 				}
 			}
diff --git a/Engine/Source/Runtime/CoreUObject/Private/UObject/LinkerLoad.cpp b/Engine/Source/Runtime/CoreUObject/Private/UObject/LinkerLoad.cpp
--- a/Engine/Source/Runtime/CoreUObject/Private/UObject/LinkerLoad.cpp
+++ b/Engine/Source/Runtime/CoreUObject/Private/UObject/LinkerLoad.cpp
@@ -5327,11 +5327,15 @@ UObject* FLinkerLoad::CreateExport( int32 Index )
 			return nullptr;
 		}
 
-		// Find the Archetype object for the one we are loading.
+		// Find the Archetype object for the one we are loading and ensure it is preloaded since recursively creating exports 
+		// may find newly created archetype exports that have not been preloaded yet
 		UObject* Template = UObject::GetArchetypeFromRequiredInfo(LoadClass, ThisParent, Export.ObjectName, Export.ObjectFlags);
-
 		checkf(Template, TEXT("Failed to get template for class %s. ExportName=%s"), *LoadClass->GetPathName(), *Export.ObjectName.ToString());
 		checkfSlow(((Export.ObjectFlags&RF_ClassDefaultObject)!=0 || Template->IsA(LoadClass)), TEXT("Mismatch between template %s and load class %s.  If this is a legacy blueprint or map, it may need to be resaved with bRecompileOnLoad turned off."), *Template->GetPathName(), *LoadClass->GetPathName());
+		if (GetLoaderType() == ELoaderType::ZenLoader)
+		{
+			Preload(Template);
+		}
 		
 		// we also need to ensure that the template has set up any instances
 		Template->ConditionalPostLoadSubobjects();
