Fix incorrect call to AddStableReference on a stack variable

#udn 12363184
#rnx

diff --git a/Engine/Source/Runtime/CoreUObject/Private/UObject/UObjectGlobals.cpp b/Engine/Source/Runtime/CoreUObject/Private/UObject/UObjectGlobals.cpp
--- a/Engine/Source/Runtime/CoreUObject/Private/UObject/UObjectGlobals.cpp
+++ b/Engine/Source/Runtime/CoreUObject/Private/UObject/UObjectGlobals.cpp
@@ -5251,16 +5251,13 @@ void FReferenceCollector::CreateVerySlowReferenceCollectorArchive()
 
 FArchive& FReferenceCollectorArchive::operator<<(UObject*& Object)
 {
-	Collector.AddStableReference(&ObjectPtrWrap(Object));
+	Collector.AddReferencedObject(ObjectPtrWrap(Object));
 	return *this;
 }
 
 FArchive& FReferenceCollectorArchive::operator<<(FObjectPtr& Object)
 {
-	if (Object.IsResolved())
-	{
-		Collector.AddStableReference(reinterpret_cast<TObjectPtr<UObject>*>(&Object));
-	}
+	Collector.AddReferencedObject(reinterpret_cast<TObjectPtr<UObject>&>(Object));
 	return *this;
 }
 
