IncrementalCook: Add FCookDependency::AssetRegistryQuery.
#rnx

diff --git a/Engine/Source/Runtime/AssetRegistry/Private/AssetRegistry.cpp b/Engine/Source/Runtime/AssetRegistry/Private/AssetRegistry.cpp
--- a/Engine/Source/Runtime/AssetRegistry/Private/AssetRegistry.cpp
+++ b/Engine/Source/Runtime/AssetRegistry/Private/AssetRegistry.cpp
@@ -285,7 +285,7 @@ public:
 
 	virtual UE::AssetRegistry::EExists TryGetAssetByObjectPath(const FSoftObjectPath& ObjectPath, FAssetData& OutAssetData) const override
 	{
-		auto AssetRegistry = IAssetRegistry::Get();
+		IAssetRegistry* AssetRegistry = IAssetRegistry::Get();
 		if (!AssetRegistry)
 		{
 			return UE::AssetRegistry::EExists::Unknown;
@@ -301,13 +301,24 @@ public:
 	
 	virtual UE::AssetRegistry::EExists TryGetAssetPackageData(FName PackageName, class FAssetPackageData& OutPackageData, FName& OutCorrectCasePackageName) const override
 	{
-		auto AssetRegistry = IAssetRegistry::Get();
+		IAssetRegistry* AssetRegistry = IAssetRegistry::Get();
 		if (!AssetRegistry)
 		{
 			return UE::AssetRegistry::EExists::Unknown;
 		}
 		return AssetRegistry->TryGetAssetPackageData(PackageName, OutPackageData, OutCorrectCasePackageName);
 	}
+
+	virtual bool EnumerateAssets(const FARFilter& Filter, TFunctionRef<bool(const FAssetData&)> Callback,
+		UE::AssetRegistry::EEnumerateAssetsFlags InEnumerateFlags) const override
+	{
+		IAssetRegistry* AssetRegistry = IAssetRegistry::Get();
+		if (!AssetRegistry)
+		{
+			return false;;
+		}
+		return AssetRegistry->EnumerateAssets(Filter, Callback, InEnumerateFlags);
+	}
 };
 FAssetRegistryInterface GAssetRegistryInterface;
 
diff --git a/Engine/Source/Runtime/AssetRegistry/Public/AssetRegistry/AssetRegistryModule.h b/Engine/Source/Runtime/AssetRegistry/Public/AssetRegistry/AssetRegistryModule.h
--- a/Engine/Source/Runtime/AssetRegistry/Public/AssetRegistry/AssetRegistryModule.h
+++ b/Engine/Source/Runtime/AssetRegistry/Public/AssetRegistry/AssetRegistryModule.h
@@ -126,4 +126,15 @@ public:
 		}
 		return AssetRegistry->TryGetAssetPackageData(PackageName, OutAssetPackageData, OutCorrectCasePackageName);
 	}
+
+	virtual bool EnumerateAssets(const FARFilter& Filter, TFunctionRef<bool(const FAssetData&)> Callback,
+		UE::AssetRegistry::EEnumerateAssetsFlags InEnumerateFlags) const override
+	{
+		IAssetRegistry* AssetRegistry = IAssetRegistry::Get();
+		if (!AssetRegistry)
+		{
+			return false;
+		}
+		return AssetRegistry->EnumerateAssets(Filter, Callback, InEnumerateFlags);
+	}
 };
diff --git a/Engine/Source/Runtime/AssetRegistry/Public/AssetRegistry/IAssetRegistry.h b/Engine/Source/Runtime/AssetRegistry/Public/AssetRegistry/IAssetRegistry.h
--- a/Engine/Source/Runtime/AssetRegistry/Public/AssetRegistry/IAssetRegistry.h
+++ b/Engine/Source/Runtime/AssetRegistry/Public/AssetRegistry/IAssetRegistry.h
@@ -72,15 +72,6 @@ ENUM_CLASS_FLAGS(EScanFlags);
 
 ASSETREGISTRY_API FString LexToString(EScanFlags Flags);
 
-enum class EEnumerateAssetsFlags : uint32
-{
-	None = 0,							// No flags
-	OnlyOnDiskAssets = (1 << 0),		// should only assets on disk be included in the enumeration. When set only DiskGatheredData will be used, does not calculate from UObjects.
-	AllowUnmountedPaths = (1 << 1),		// should unmounted asset paths be allowed
-	AllowUnfilteredArAssets = (1 << 2),	// skip the filtering of UE::AssetRegistry::FFiltering
-};
-ENUM_CLASS_FLAGS(EEnumerateAssetsFlags);
-
 } // namespace UE::AssetRegistry
 
 USTRUCT(BlueprintType)
@@ -326,6 +317,7 @@ public:
 	 * @param Callback function to call for each asset data enumerated
 	 * @param InEnumerateFlags flags to control enumeration and filtering.
 	 *        @see EEnumerateAssetsFlags.
+	 * @return False if filter is invalid, otherwise true.
 	 */
 	virtual bool EnumerateAssets(const FARFilter& Filter, TFunctionRef<bool(const FAssetData&)> Callback) const = 0;
 	virtual bool EnumerateAssets(const FARCompiledFilter& Filter, TFunctionRef<bool(const FAssetData&)> Callback) const = 0;
diff --git a/Engine/Source/Runtime/CoreUObject/Private/AssetRegistry/ARFilter.cpp b/Engine/Source/Runtime/CoreUObject/Private/AssetRegistry/ARFilter.cpp
new file mode 100644
--- /dev/null
+++ b/Engine/Source/Runtime/CoreUObject/Private/AssetRegistry/ARFilter.cpp
@@ -0,0 +1,426 @@
+// Copyright Epic Games, Inc. All Rights Reserved.
+
+#include "AssetRegistry/ARFilter.h"
+
+#include "Algo/Compare.h"
+#include "Serialization/CompactBinarySerialization.h"
+#include "Serialization/CompactBinaryWriter.h"
+#include "Templates/UnrealTemplate.h"
+
+namespace UE::AssetRegistry
+{
+	constexpr int32 FARCompactBinaryVersion = 1;
+}
+
+namespace Algo
+{
+	// TODO: Think of a name for this Compare function that returns -1,0,1 and put it in Compare.h
+	template <typename InAT, typename InBT, typename LessThanType>
+	constexpr int CompareAsInt(InAT&& InputA, InBT&& InputB, LessThanType LessThan)
+	{
+		const SIZE_T SizeA = GetNum(InputA);
+		const SIZE_T SizeB = GetNum(InputB);
+
+		if (SizeA != SizeB)
+		{
+			return SizeA < SizeB ? -1 : 1;
+		}
+
+		auto* A = GetData(InputA);
+		auto* B = GetData(InputB);
+
+		for (SIZE_T Count = SizeA; Count; --Count)
+		{
+			if (LessThan(*A, *B))
+			{
+				return -1;
+			}
+			if (LessThan(*A, *B))
+			{
+				return 1;
+			}
+			A++;
+			B++;
+		}
+		return 0;
+	}
+
+	// TODO: Write this without copying to an array and put it in Compare.h
+	template <typename MapType, typename KeyLessThanType, typename ValueLessThanType>
+	int CompareMultiMap(const MapType& A, const MapType& B, KeyLessThanType KeyLessThan, ValueLessThanType ValueLessThan)
+	{
+		// TODO: Add KeyType and ValueType to TMultiMap
+		using KeyType = decltype(DeclVal<typename MapType::ElementType>().Key);
+		using ValueType = decltype(DeclVal<typename MapType::ElementType>().Value);
+
+		if (A.Num() != B.Num())
+		{
+			return A.Num() < B.Num() ? -1 : 1;
+		}
+		if (A.Num() == 0)
+		{
+			return 0;
+		}
+
+		auto PairLessThan = [&KeyLessThan, &ValueLessThan]
+		(const TPair<KeyType, ValueType>& PairA, const TPair<KeyType, ValueType>& PairB)
+			{
+				if (KeyLessThan(PairA.Key, PairB.Key))
+				{
+					return true;
+				}
+				if (KeyLessThan(PairB.Key, PairA.Key))
+				{
+					return false;
+				}
+				if (ValueLessThan(PairA.Value, PairB.Value))
+				{
+					return true;
+				}
+				if (ValueLessThan(PairB.Value, PairA.Value))
+				{
+					return false;
+				}
+				return false; // equal
+			};
+		TArray<TPair<KeyType, ValueType>> AVals = A.Array();
+		TArray<TPair<KeyType, ValueType>> BVals = B.Array();
+		Algo::Sort(AVals, PairLessThan);
+		Algo::Sort(BVals, PairLessThan);
+		return CompareAsInt(AVals, BVals, PairLessThan);
+	}
+}
+
+// CompactBinary for FSoftObjectPath; we can move these into public if necessary
+static FCbWriter& operator<<(FCbWriter& Writer, const FSoftObjectPath& SoftObjectPath)
+{
+	Writer.BeginArray();
+	Writer << SoftObjectPath.GetAssetPath();
+	Writer << SoftObjectPath.GetSubPathString();
+	Writer.EndArray();
+	return Writer;
+}
+
+bool LoadFromCompactBinary(FCbFieldView Field, FSoftObjectPath& SoftObjectPath)
+{
+	FTopLevelAssetPath AssetPath;
+	FString SubPath;
+	FCbFieldViewIterator Iter = Field.CreateViewIterator();
+	bool bOk = true;
+	bOk = LoadFromCompactBinary(Iter++, AssetPath) & bOk;
+	bOk = LoadFromCompactBinary(Iter++, SubPath) & bOk;
+	SoftObjectPath = FSoftObjectPath(MoveTemp(AssetPath), MoveTemp(SubPath));
+	return bOk;
+}
+
+// Special case CompactBinary that could be replaced with general versions
+static FCbWriter& operator<<(FCbWriter& Writer, const TPair<FName, TOptional<FString>>& Value)
+{
+	bool bHasValue = Value.Value.IsSet();
+	Writer.BeginArray();
+	Writer << Value.Key;
+	Writer << bHasValue;
+	if (bHasValue)
+	{
+		Writer << *Value.Value;
+	}
+	Writer.EndArray();
+	return Writer;
+}
+
+static bool LoadFromCompactBinary(FCbFieldView Field, TPair<FName, TOptional<FString>>& Value)
+{
+	FCbFieldViewIterator Iter = Field.CreateViewIterator();
+	bool bOk = LoadFromCompactBinary(*Iter++, Value.Key);
+	bool bHasValue;
+	if (LoadFromCompactBinary(*Iter++, bHasValue))
+	{
+		if (bHasValue)
+		{
+			Value.Value.Emplace();
+			bOk = LoadFromCompactBinary(*Iter++, *Value.Value);
+		}
+	}
+	else
+	{
+		bOk = false;
+	}
+	return bOk;
+}
+
+struct FOptionalFStringIgnoreCaseLess
+{
+	bool operator()(const TOptional<FString>& A, const TOptional<FString>& B) const
+	{
+		if (A.IsSet() != B.IsSet())
+		{
+			return !A.IsSet();
+		}
+		if (!A.IsSet())
+		{
+			return false; // Both not set, equal
+		}
+		return *A < *B;
+	}
+};
+
+void FARFilter::SortForSaving()
+{
+	Algo::Sort(PackageNames, FNameLexicalLess());
+	Algo::Sort(PackagePaths, FNameLexicalLess());
+	Algo::Sort(SoftObjectPaths, FSoftObjectPathLexicalLess());
+	Algo::Sort(ClassPaths, FTopLevelAssetPathLexicalLess());
+	TagsAndValues.KeySort(FNameLexicalLess());
+	TagsAndValues.ValueSort(FOptionalFStringIgnoreCaseLess());
+	RecursiveClassPathsExclusionSet.Sort(FTopLevelAssetPathLexicalLess());
+}
+
+bool FARFilter::operator==(const FARFilter& Other) const
+{
+	// Arrays are compared without sorting; caller is responsible for calling Sort
+	int32 Compare;
+	Compare = Algo::CompareAsInt(PackageNames, Other.PackageNames, FNameLexicalLess());
+	if (Compare != 0)
+	{
+		return false;
+	}
+	Compare = Algo::CompareAsInt(PackagePaths, Other.PackagePaths, FNameLexicalLess());
+	if (Compare != 0)
+	{
+		return false;
+	}
+	Compare = Algo::CompareAsInt(SoftObjectPaths, Other.SoftObjectPaths, FSoftObjectPathLexicalLess());
+	if (Compare != 0)
+	{
+		return false;
+	}
+	Compare = Algo::CompareAsInt(ClassPaths, Other.ClassPaths, FTopLevelAssetPathLexicalLess());
+	if (Compare != 0)
+	{
+		return false;
+	}
+	TArray<TPair<FName, TOptional<FString>>> A;
+	Compare = Algo::CompareMultiMap(TagsAndValues, Other.TagsAndValues, FNameLexicalLess(),
+		FOptionalFStringIgnoreCaseLess());
+	if (Compare != 0)
+	{
+		return false;
+	}
+	Compare = Algo::CompareSet(RecursiveClassPathsExclusionSet, Other.RecursiveClassPathsExclusionSet,
+		FTopLevelAssetPathLexicalLess());
+	if (Compare != 0)
+	{
+		return false;
+	}
+	if (bRecursivePaths != Other.bRecursivePaths)
+	{
+		return false;
+	}
+	if (bRecursiveClasses != Other.bRecursiveClasses)
+	{
+		return false;
+	}
+	if (bIncludeOnlyOnDiskAssets != Other.bIncludeOnlyOnDiskAssets)
+	{
+		return false;
+	}
+	if (WithoutPackageFlags != Other.WithoutPackageFlags)
+	{
+		return false;
+	}
+	if (WithPackageFlags != Other.WithPackageFlags)
+	{
+		return false;
+	}
+	return true;
+}
+
+bool FARFilter::operator<(const FARFilter& Other) const
+{
+	int32 Compare;
+	Compare = Algo::CompareAsInt(PackageNames, Other.PackageNames, FNameLexicalLess());
+	if (Compare != 0)
+	{
+		return Compare < 0;
+	}
+	Compare = Algo::CompareAsInt(PackagePaths, Other.PackagePaths, FNameLexicalLess());
+	if (Compare != 0)
+	{
+		return Compare < 0;
+	}
+	Compare = Algo::CompareAsInt(SoftObjectPaths, Other.SoftObjectPaths, FSoftObjectPathLexicalLess());
+	if (Compare != 0)
+	{
+		return Compare < 0;
+	}
+	Compare = Algo::CompareAsInt(ClassPaths, Other.ClassPaths, FTopLevelAssetPathLexicalLess());
+	if (Compare != 0)
+	{
+		return Compare < 0;
+	}
+	Compare = Algo::CompareMultiMap(TagsAndValues, Other.TagsAndValues, FNameLexicalLess(),
+		FOptionalFStringIgnoreCaseLess());
+	if (Compare != 0)
+	{
+		return Compare < 0;
+	}
+	Compare = Algo::CompareSet(RecursiveClassPathsExclusionSet, Other.RecursiveClassPathsExclusionSet,
+		FTopLevelAssetPathLexicalLess());
+	if (Compare != 0)
+	{
+		return Compare < 0;
+	}
+	if (bRecursivePaths != Other.bRecursivePaths)
+	{
+		return !bRecursivePaths;
+	}
+	if (bRecursiveClasses != Other.bRecursiveClasses)
+	{
+		return !bRecursiveClasses;
+	}
+	if (bIncludeOnlyOnDiskAssets != Other.bIncludeOnlyOnDiskAssets)
+	{
+		return !bIncludeOnlyOnDiskAssets;
+	}
+	if (WithoutPackageFlags != Other.WithoutPackageFlags)
+	{
+		return WithoutPackageFlags < Other.WithoutPackageFlags;
+	}
+	if (WithPackageFlags != Other.WithPackageFlags)
+	{
+		return WithPackageFlags < Other.WithPackageFlags;
+	}
+	return false;
+}
+
+void FARFilter::Save(FCbWriter& Writer) const
+{
+	Writer.BeginObject();
+	Writer << "Version" << UE::AssetRegistry::FARCompactBinaryVersion;
+	Writer << "PackageNames" << PackageNames;
+	Writer << "PackagePaths" << PackagePaths;
+	Writer << "SoftObjectPaths" << SoftObjectPaths;
+	Writer << "ClassPaths" << ClassPaths;
+	TArray<TPair<FName, TOptional<FString>>> TagsAndValuesArray = TagsAndValues.Array();
+	Writer << "TagsAndValues" << TagsAndValuesArray;
+	TArray<FTopLevelAssetPath> RecursiveClassPathsExclusionSetArray = RecursiveClassPathsExclusionSet.Array();
+	Writer << "RecursiveClassPathsExclusionSet" << RecursiveClassPathsExclusionSetArray;
+	Writer << "bRecursivePaths" << bRecursivePaths;
+	Writer << "bRecursiveClasses" << bRecursiveClasses;
+	Writer << "bIncludeOnlyOnDiskAssets" << bIncludeOnlyOnDiskAssets;
+	Writer << "WithoutPackageFlags" << WithoutPackageFlags;
+	Writer << "WithPackageFlags" << WithPackageFlags;
+	Writer.EndObject();
+}
+
+bool FARFilter::TryLoad(const FCbFieldView& Field)
+{
+	int32 Version = -1;
+	for (FCbFieldViewIterator ElementView(Field.CreateViewIterator()); ElementView; )
+	{
+		const FCbFieldViewIterator Last = ElementView;
+		if (ElementView.GetName().Equals(UTF8TEXTVIEW("Version")))
+		{
+			if (!LoadFromCompactBinary(ElementView++, Version))
+			{
+				return false;
+			}
+		}
+		if (ElementView.GetName().Equals(UTF8TEXTVIEW("PackageNames")))
+		{
+			if (!LoadFromCompactBinary(ElementView++, PackageNames))
+			{
+				return false;
+			}
+		}
+		if (ElementView.GetName().Equals(UTF8TEXTVIEW("PackagePaths")))
+		{
+			if (!LoadFromCompactBinary(ElementView++, PackagePaths))
+			{
+				return false;
+			}
+		}
+		if (ElementView.GetName().Equals(UTF8TEXTVIEW("SoftObjectPaths")))
+		{
+			if (!LoadFromCompactBinary(ElementView++, SoftObjectPaths))
+			{
+				return false;
+			}
+		}
+		if (ElementView.GetName().Equals(UTF8TEXTVIEW("ClassPaths")))
+		{
+			if (!LoadFromCompactBinary(ElementView++, ClassPaths))
+			{
+				return false;
+			}
+		}
+		if (ElementView.GetName().Equals(UTF8TEXTVIEW("TagsAndValues")))
+		{
+			TArray<TPair<FName, TOptional<FString>>> TagsAndValuesArray;
+			if (!LoadFromCompactBinary(ElementView++, TagsAndValuesArray))
+			{
+				return false;
+			}
+			TagsAndValues.Empty(TagsAndValuesArray.Num());
+			for (TPair<FName, TOptional<FString>>& Pair : TagsAndValuesArray)
+			{
+				TagsAndValues.Add(Pair.Key, Pair.Value);
+			}
+		}
+		if (ElementView.GetName().Equals(UTF8TEXTVIEW("RecursiveClassPathsExclusionSet")))
+		{
+			TArray<FTopLevelAssetPath> RecursiveClassPathsExclusionSetArray;
+			if (!LoadFromCompactBinary(ElementView++, RecursiveClassPathsExclusionSetArray))
+			{
+				return false;
+			}
+			RecursiveClassPathsExclusionSet.Empty(RecursiveClassPathsExclusionSetArray.Num());
+			RecursiveClassPathsExclusionSet.Append(RecursiveClassPathsExclusionSetArray);
+		}
+		if (ElementView.GetName().Equals(UTF8TEXTVIEW("bRecursivePaths")))
+		{
+			if (!LoadFromCompactBinary(ElementView++, bRecursivePaths))
+			{
+				return false;
+			}
+		}
+		if (ElementView.GetName().Equals(UTF8TEXTVIEW("bRecursiveClasses")))
+		{
+			if (!LoadFromCompactBinary(ElementView++, bRecursiveClasses))
+			{
+				return false;
+			}
+		}
+		if (ElementView.GetName().Equals(UTF8TEXTVIEW("bIncludeOnlyOnDiskAssets")))
+		{
+			if (!LoadFromCompactBinary(ElementView++, bIncludeOnlyOnDiskAssets))
+			{
+				return false;
+			}
+		}
+		if (ElementView.GetName().Equals(UTF8TEXTVIEW("WithoutPackageFlags")))
+		{
+			if (!LoadFromCompactBinary(ElementView++, WithoutPackageFlags))
+			{
+				return false;
+			}
+		}
+		if (ElementView.GetName().Equals(UTF8TEXTVIEW("WithPackageFlags")))
+		{
+			if (!LoadFromCompactBinary(ElementView++, WithPackageFlags))
+			{
+				return false;
+			}
+		}
+
+		if (ElementView == Last)
+		{
+			++ElementView;
+		}
+	}
+	if (Version != UE::AssetRegistry::FARCompactBinaryVersion)
+	{
+		return false;
+	}
+	return true;
+}
diff --git a/Engine/Source/Runtime/CoreUObject/Private/Cooker/CookDependency.cpp b/Engine/Source/Runtime/CoreUObject/Private/Cooker/CookDependency.cpp
--- a/Engine/Source/Runtime/CoreUObject/Private/Cooker/CookDependency.cpp
+++ b/Engine/Source/Runtime/CoreUObject/Private/Cooker/CookDependency.cpp
@@ -3,6 +3,9 @@
 #include "Cooker/CookDependency.h"
 
 #if WITH_EDITOR
+#include "Algo/Unique.h"
+#include "AssetRegistry/ARFilter.h"
+#include "AssetRegistry/AssetData.h"
 #include "Containers/Array.h"
 #include "Containers/Map.h"
 #include "CoreGlobals.h"
@@ -11,6 +14,7 @@
 #include "HAL/PlatformMath.h"
 #include "Hash/Blake3.h"
 #include "Misc/AssertionMacros.h"
+#include "Misc/AssetRegistryInterface.h"
 #include "Misc/ConfigAccessData.h"
 #include "Misc/StringBuilder.h"
 #include "Serialization/Archive.h"
@@ -166,6 +170,14 @@ FCookDependency FCookDependency::NativeClass(FStringView ClassPath)
 	return Result;
 }
 
+FCookDependency FCookDependency::AssetRegistryQuery(FARFilter Filter)
+{
+	FCookDependency Result(ECookDependency::AssetRegistryQuery);
+	Filter.SortForSaving();
+	Result.ARFilter = MakeUnique<FARFilter>(MoveTemp(Filter));
+	return Result;
+}
+
 FCookDependency::FCookDependency()
 	: Type(ECookDependency::None)
 {
@@ -225,6 +237,9 @@ FCookDependency& FCookDependency::operator=(const FCookDependency& Other)
 	case ECookDependency::SettingsObject:
 		ObjectPtr = Other.ObjectPtr;
 		break;
+	case ECookDependency::AssetRegistryQuery:
+		ARFilter.Reset(Other.ARFilter.IsValid() ? new FARFilter(*Other.ARFilter) : nullptr);
+		break;
 	default:
 		checkNoEntry();
 		break;
@@ -261,6 +276,9 @@ FCookDependency& FCookDependency::operator=(FCookDependency&& Other)
 	case ECookDependency::SettingsObject:
 		ObjectPtr = Other.ObjectPtr;
 		break;
+	case ECookDependency::AssetRegistryQuery:
+		ARFilter = MoveTemp(Other.ARFilter);
+		break;
 	default:
 		checkNoEntry();
 		break;
@@ -366,6 +384,30 @@ void FCookDependency::UpdateHash(FCookDependencyContext& Context) const
 			TEXT("FCookDependency::NativeClass('%s') failed to UpdateHash: NativeClass dependencies do not implement UpdateHash and it should not be called on them."),
 			*StringData));
 		return;
+	case ECookDependency::AssetRegistryQuery:
+		if (ARFilter.IsValid())
+		{
+			IAssetRegistryInterface* AssetRegistry = IAssetRegistryInterface::GetPtr();
+			if (AssetRegistry)
+			{
+				TArray<FName> PackageNames;
+				AssetRegistry->EnumerateAssets(*ARFilter, [&Context, &PackageNames](const FAssetData& AssetData)
+					{
+						PackageNames.Add(AssetData.PackageName);
+						return true;
+					}, UE::AssetRegistry::EEnumerateAssetsFlags::None);
+				PackageNames.Sort(FNameLexicalLess());
+				PackageNames.SetNum(Algo::Unique(PackageNames));
+				TStringBuilder<256> PackageNameStr;
+				for (FName PackageName : PackageNames)
+				{
+					PackageNameStr.Reset();
+					PackageNameStr << PackageName;
+					Context.Update(*PackageNameStr, sizeof(**PackageNameStr) * PackageNameStr.Len());
+				}
+			}
+		}
+		return;
 	default:
 		checkNoEntry();
 		return;
@@ -399,6 +441,10 @@ void FCookDependency::Construct()
 	case ECookDependency::SettingsObject:
 		ObjectPtr = nullptr;
 		break;
+	case ECookDependency::AssetRegistryQuery:
+		// Set the union's bytes equal to a TUniquePtr with a null internal pointer
+		new(&ARFilter) TUniquePtr<FARFilter>(nullptr);
+		break;
 	default:
 		checkNoEntry();
 		break;
@@ -432,6 +478,10 @@ void FCookDependency::Destruct()
 	case ECookDependency::SettingsObject:
 		ObjectPtr = nullptr;
 		break;
+	case ECookDependency::AssetRegistryQuery:
+		// TUniquePtr's destructor will also destruct its internal pointer
+		ARFilter.~TUniquePtr<FARFilter>();
+		break;
 	default:
 		checkNoEntry();
 		break;
@@ -471,6 +521,16 @@ void FCookDependency::Save(FCbWriter& Writer) const
 	case ECookDependency::SettingsObject:
 		// Settings objects are not persistable; save out an empty SettingsObject dependency
 		break;
+	case ECookDependency::AssetRegistryQuery:
+	{
+		bool bValid = ARFilter.IsValid();
+		Writer << bValid;
+		if (bValid)
+		{
+			Writer << *ARFilter;
+		}
+		break;
+	}
 	default:
 		checkNoEntry();
 		break;
@@ -588,6 +648,28 @@ bool FCookDependency::Load(FCbFieldView Value)
 		*this = FCookDependency::NativeClass(LocalClassPath);
 		return true;
 	}
+	case ECookDependency::AssetRegistryQuery:
+	{
+		bool bValid;
+		if (!LoadFromCompactBinary(Field++, bValid))
+		{
+			return false;
+		}
+		if (!bValid)
+		{
+			*this = FCookDependency(ECookDependency::AssetRegistryQuery);
+		}
+		else
+		{
+			FARFilter LocalARFilter;
+			if (!LoadFromCompactBinary(Field++, LocalARFilter))
+			{
+				return false;
+			}
+			*this = FCookDependency::AssetRegistryQuery(MoveTemp(LocalARFilter));
+		}
+		return true;
+	}
 	default:
 		break;
 	}
@@ -607,6 +689,16 @@ bool FCookDependency::ConfigAccessDataEqual(const UE::ConfigAccessTracking::FCon
 	return A == B;
 }
 
+bool FCookDependency::ARFilterLessThan(const FARFilter& A, const FARFilter& B)
+{
+	return A < B;
+}
+
+bool FCookDependency::ARFilterEqual(const FARFilter& A, const FARFilter& B)
+{
+	return A == B;
+}
+
 void FCookDependencyContext::Update(const void* Data, uint64 Size)
 {
 	FBlake3& Blake3Hasher = *(reinterpret_cast<FBlake3*>(Hasher));
diff --git a/Engine/Source/Runtime/CoreUObject/Public/AssetRegistry/ARFilter.h b/Engine/Source/Runtime/CoreUObject/Public/AssetRegistry/ARFilter.h
--- a/Engine/Source/Runtime/CoreUObject/Public/AssetRegistry/ARFilter.h
+++ b/Engine/Source/Runtime/CoreUObject/Public/AssetRegistry/ARFilter.h
@@ -11,6 +11,8 @@
 #include "CoreMinimal.h"
 #endif
 
+class FCbField;
+class FCbWriter;
 struct FAssetData;
 
 /**
@@ -136,6 +138,29 @@ PRAGMA_ENABLE_DEPRECATION_WARNINGS
 	}
 
 	void PostSerialize(const FArchive& Ar);
+
+	/**
+	 * Sort the structures (e.g. PackageNames) on this filter to be deterministic for saving the structure during
+	 * cooking. This is also required for correctness of operator== when comparing the order-independent arrays.
+	 */
+	COREUOBJECT_API void SortForSaving();
+	COREUOBJECT_API bool operator==(const FARFilter& Other) const;
+	COREUOBJECT_API bool operator<(const FARFilter& Other) const;
+
+private:
+	COREUOBJECT_API void Save(FCbWriter& Writer) const;
+	COREUOBJECT_API bool TryLoad(const FCbFieldView& Field);
+
+	// Hidden friends for our argument-dependent lookup functions
+	friend FCbWriter& operator<<(FCbWriter& Writer, const FARFilter& Filter)
+	{
+		Filter.Save(Writer);
+		return Writer;
+	}
+	friend bool LoadFromCompactBinary(const FCbFieldView& Field, FARFilter& Filter)
+	{
+		return Filter.TryLoad(Field);
+	}
 };
 
 template<>
diff --git a/Engine/Source/Runtime/CoreUObject/Public/Cooker/CookDependency.h b/Engine/Source/Runtime/CoreUObject/Public/Cooker/CookDependency.h
--- a/Engine/Source/Runtime/CoreUObject/Public/Cooker/CookDependency.h
+++ b/Engine/Source/Runtime/CoreUObject/Public/Cooker/CookDependency.h
@@ -16,6 +16,7 @@ class FCbWriter;
 class UObject;
 namespace UE::ConfigAccessTracking { enum class ELoadType : uint8; }
 namespace UE::ConfigAccessTracking { struct FConfigAccessData; }
+struct FARFilter;
 #endif
 
 #if WITH_EDITOR
@@ -93,6 +94,7 @@ enum class ECookDependency : uint8
 	Config					= 6,
 	SettingsObject			= 7,
 	NativeClass				= 8,
+	AssetRegistryQuery		= 9,
 
 	Count,
 };
@@ -168,6 +170,9 @@ public:
 	COREUOBJECT_API static FCookDependency NativeClass(const UClass* InClass);
 	COREUOBJECT_API static FCookDependency NativeClass(FStringView ClassPath);
 
+	/** Adds a dependency on the results reported by an AssetRegistry query. */
+	COREUOBJECT_API static FCookDependency AssetRegistryQuery(FARFilter Filter);
+
 	/** Construct an empty dependency; it will never be invalidated. */
 	COREUOBJECT_API FCookDependency();
 
@@ -211,6 +216,9 @@ public:
 	 */
 	FStringView GetClassPath() const;
 
+	/** Returns the FARFilter if GetType() == AssetRegistryQuery, else nullptr. */
+	const FARFilter* GetARFilter() const;
+
 	/**
 	 * Comparison operator for e.g. deterministic ordering of dependencies.
 	 * Uses persistent comparison data and is somewhat expensive.
@@ -233,6 +241,8 @@ private:
 		const UE::ConfigAccessTracking::FConfigAccessData& B);
 	static COREUOBJECT_API bool ConfigAccessDataEqual(const UE::ConfigAccessTracking::FConfigAccessData& A,
 		const UE::ConfigAccessTracking::FConfigAccessData& B);
+	static COREUOBJECT_API bool ARFilterLessThan(const FARFilter& A, const FARFilter& B);
+	static COREUOBJECT_API bool ARFilterEqual(const FARFilter& A, const FARFilter& B);
 
 	/** Public hidden friend for operator<< into an FCbWriter. */
 	friend FCbWriter& operator<<(FCbWriter& Writer, const FCookDependency& CookDependencies)
@@ -266,6 +276,7 @@ private:
 		FName NameData;
 		const UObject* ObjectPtr;
 		TUniquePtr<UE::ConfigAccessTracking::FConfigAccessData> ConfigAccessData;
+		TUniquePtr<FARFilter> ARFilter;
 	};
 };
 
@@ -391,6 +402,11 @@ inline FStringView FCookDependency::GetClassPath() const
 	return Type == ECookDependency::NativeClass ? StringData : FStringView();
 }
 
+inline const FARFilter* FCookDependency::GetARFilter() const
+{
+	return Type == ECookDependency::AssetRegistryQuery ? ARFilter.Get() : nullptr;
+}
+
 inline bool FCookDependency::operator<(const FCookDependency& Other) const
 {
 	if (static_cast<uint8>(Type) != static_cast<uint8>(Other.Type))
@@ -452,6 +468,16 @@ inline bool FCookDependency::operator<(const FCookDependency& Other) const
 	case ECookDependency::SettingsObject:
 		// SettingsObjects are not persistable, so we do not use a persistent sort key; just the object ptr.
 		return ObjectPtr < Other.ObjectPtr;
+	case ECookDependency::AssetRegistryQuery:
+		if (ARFilter.IsValid() != Other.ARFilter.IsValid())
+		{
+			return !ARFilter.IsValid();
+		}
+		if (!ARFilter.IsValid())
+		{
+			return false; // equal
+		}
+		return ARFilterLessThan(*ARFilter, *Other.ARFilter);
 	default:
 		checkNoEntry();
 		return false;
@@ -514,6 +540,19 @@ inline bool FCookDependency::operator==(const FCookDependency& Other) const
 	case ECookDependency::SettingsObject:
 		// SettingsObjects are not persistable, so we do not use a persistent sort key; just the object ptr.
 		return ObjectPtr == Other.ObjectPtr;
+	case ECookDependency::AssetRegistryQuery:
+	{
+		if (ARFilter.IsValid() != Other.ARFilter.IsValid())
+		{
+			return false;
+		}
+		if (!ARFilter.IsValid())
+		{
+			return true;
+		}
+
+		return ARFilterEqual(*ARFilter, *Other.ARFilter);
+	}
 	default:
 		checkNoEntry();
 		return false;
diff --git a/Engine/Source/Runtime/CoreUObject/Public/Misc/AssetRegistryInterface.h b/Engine/Source/Runtime/CoreUObject/Public/Misc/AssetRegistryInterface.h
--- a/Engine/Source/Runtime/CoreUObject/Public/Misc/AssetRegistryInterface.h
+++ b/Engine/Source/Runtime/CoreUObject/Public/Misc/AssetRegistryInterface.h
@@ -9,6 +9,7 @@
 #include "Misc/EnumClassFlags.h"
 #include "Modules/ModuleInterface.h"
 #include "Modules/ModuleManager.h"
+#include "Templates/Function.h"
 #include "UObject/NameTypes.h"
 #include "UObject/TopLevelAssetPath.h"
 
@@ -22,6 +23,7 @@ class ITargetPlatform;
 class UObject;
 class UPackage;
 struct FArchiveCookContext;
+struct FARFilter;
 struct FAssetData;
 
 namespace EAssetRegistryDependencyType
@@ -166,6 +168,15 @@ namespace UE::AssetRegistry
 		Unknown,		// Not known. AssetRegistry might still be indexing
 	};
 
+	enum class EEnumerateAssetsFlags : uint32
+	{
+		None = 0,							// No flags
+		OnlyOnDiskAssets = (1 << 0),		// should only assets on disk be included in the enumeration. When set only DiskGatheredData will be used, does not calculate from UObjects.
+		AllowUnmountedPaths = (1 << 1),		// should unmounted asset paths be allowed
+		AllowUnfilteredArAssets = (1 << 2),	// skip the filtering of UE::AssetRegistry::FFiltering
+	};
+	ENUM_CLASS_FLAGS(EEnumerateAssetsFlags);
+
 	/**
 	 * A struct that is equivalent to EDependencyQuery, but is more useful for performance in filtering operations.
 	 * This is used by the filter implementations inside of GetDependency/GetReferencer calls; callers of those functions can instead use the more convenient values in EDependencyQuery.
@@ -407,6 +418,20 @@ public:
 	 */
 	virtual UE::AssetRegistry::EExists TryGetAssetPackageData(FName PackageName, class FAssetPackageData& OutPackageData, FName& OutCorrectCasePackageName) const = 0;
 
+	/**
+	 * Enumerate asset data for all assets that match the filter.
+	 * Assets returned must satisfy every filter component if there is at least one element in the component's array.
+	 * Assets will satisfy a component if they match any of the elements in it.
+	 *
+	 * @param Filter filter to apply to the assets in the AssetRegistry
+	 * @param Callback function to call for each asset data enumerated
+	 * @param InEnumerateFlags flags to control enumeration and filtering.
+	 *        @see EEnumerateAssetsFlags.
+	 * @return False if the AssetRegistry is not available or the filter is invalid, otherwise true.
+	 */
+	virtual bool EnumerateAssets(const FARFilter& Filter, TFunctionRef<bool(const FAssetData&)> Callback,
+		UE::AssetRegistry::EEnumerateAssetsFlags InEnumerateFlags) const = 0;
+
 protected:
 
 	static COREUOBJECT_API IAssetRegistryInterface* Default;
diff --git a/Engine/Source/Runtime/CoreUObject/Tests/ObjectPropertyTest.cpp b/Engine/Source/Runtime/CoreUObject/Tests/ObjectPropertyTest.cpp
--- a/Engine/Source/Runtime/CoreUObject/Tests/ObjectPropertyTest.cpp
+++ b/Engine/Source/Runtime/CoreUObject/Tests/ObjectPropertyTest.cpp
@@ -281,6 +281,12 @@ public:
 		return UE::AssetRegistry::EExists::Exists;
 	}
 
+	virtual bool EnumerateAssets(const FARFilter& Filter, TFunctionRef<bool(const FAssetData&)> Callback,
+		UE::AssetRegistry::EEnumerateAssetsFlags InEnumerateFlags) const override
+	{
+		return true;
+	}
+
 	IAssetRegistryInterface* Old;
 	TMap<FSoftObjectPath, FAssetData> AssetData;
 };
