[verse] run GC barriers whenever TPtr points to something of UObject (sub)type

#rnx

diff --git a/Engine/Source/Runtime/CoreUObject/Public/UObject/ObjectPtr.h b/Engine/Source/Runtime/CoreUObject/Public/UObject/ObjectPtr.h
--- a/Engine/Source/Runtime/CoreUObject/Public/UObject/ObjectPtr.h
+++ b/Engine/Source/Runtime/CoreUObject/Public/UObject/ObjectPtr.h
@@ -1796,6 +1796,16 @@ public:
 		return Object;
 	}	
 
+	FORCEINLINE bool IsInitialized() const
+	{
+		return Object != nullptr;
+	}
+
+	/**
+	 * Use IsInitialized if needed
+	 */
+	explicit operator bool() const = delete;
+
 private:
 
 	/** The object we're holding a reference to. */
