Re-added explicit ARO of IDO namespaces because of an edge case with the transient package

#rnx

diff --git a/Engine/Source/Runtime/CoreUObject/Private/UObject/PropertyBagRepository.cpp b/Engine/Source/Runtime/CoreUObject/Private/UObject/PropertyBagRepository.cpp
--- a/Engine/Source/Runtime/CoreUObject/Private/UObject/PropertyBagRepository.cpp
+++ b/Engine/Source/Runtime/CoreUObject/Private/UObject/PropertyBagRepository.cpp
@@ -837,6 +837,11 @@ bool FPropertyBagRepository::WasPropertyValueSerialized(const UStruct* Struct, c
 
 void FPropertyBagRepository::AddReferencedObjects(FReferenceCollector& Collector)
 {
+	for (TPair<const UObject*, TObjectPtr<UObject>>& Element : Namespaces)
+	{
+		Collector.AddReferencedObject(Element.Value);
+	}
+
 	FPropertyBagPlaceholderTypeRegistry::Get().AddReferencedObjects(Collector);
 }
 
