#UE Improve warning about property serialization type mismatch to include the serialized object to make it easier to find the asset that has the property that will be nulled

diff --git a/Engine/Source/Runtime/CoreUObject/Private/UObject/PropertyBaseObject.cpp b/Engine/Source/Runtime/CoreUObject/Private/UObject/PropertyBaseObject.cpp
--- a/Engine/Source/Runtime/CoreUObject/Private/UObject/PropertyBaseObject.cpp
+++ b/Engine/Source/Runtime/CoreUObject/Private/UObject/PropertyBaseObject.cpp
@@ -736,12 +736,14 @@ void FObjectPropertyBase::CheckValidObject(void* ValueAddress, TObjectPtr<UObjec
 		bool bIsReplacingClassRefs = PropertyClass && PropertyClass->HasAnyClassFlags(CLASS_NewerVersionExists) != ObjectClass->HasAnyClassFlags(CLASS_NewerVersionExists);
 		if (!bIsReplacingClassRefs && !IsDeferringValueLoad())
 		{
+			FUObjectSerializeContext* SerializeContext = FUObjectThreadContext::Get().GetSerializeContext();
 			if (!HasAnyPropertyFlags(CPF_NonNullable))
 			{
 				UE_LOG(LogProperty, Warning,
-					TEXT("Serialized %s for a property of %s. Reference will be nullptred.\n    Property = %s\n    Item = %s"),
+					TEXT("Serialized %s for a property of %s. Reference will be nullptred.\n    ReferencingObject = %s\n    Property = %s\n    Item = %s"),
 					*ObjectClass->GetFullName(),
 					*PropertyClass->GetFullName(),
+					*GetFullNameSafe(SerializeContext->SerializedObject),
 					*GetFullName(),
 					*Object.GetFullName()
 				);
@@ -752,10 +754,11 @@ void FObjectPropertyBase::CheckValidObject(void* ValueAddress, TObjectPtr<UObjec
 				UObject* DefaultValue = ConstructDefaultObjectValueIfNecessary(OldValue);
 
 				UE_LOG(LogProperty, Warning,
-					TEXT("Serialized %s for a non-nullable property of %s. Reference will be defaulted to %s.\n    Property = %s\n    Item = %s"),
+					TEXT("Serialized %s for a non-nullable property of %s. Reference will be defaulted to %s.\n    ReferencingObject = %s\n    Property = %s\n    Item = %s"),
 					*ObjectClass->GetFullName(),
 					*PropertyClass->GetFullName(),
 					DefaultValue ? *DefaultValue->GetFullName() : TEXT("None"),
+					*GetFullNameSafe(SerializeContext->SerializedObject),
 					*GetFullName(),
 					*Object.GetFullName()
 				);
