#UE Null check SerializeContext since this is just for reporting in a warning

diff --git a/Engine/Source/Runtime/CoreUObject/Private/UObject/PropertyBaseObject.cpp b/Engine/Source/Runtime/CoreUObject/Private/UObject/PropertyBaseObject.cpp
--- a/Engine/Source/Runtime/CoreUObject/Private/UObject/PropertyBaseObject.cpp
+++ b/Engine/Source/Runtime/CoreUObject/Private/UObject/PropertyBaseObject.cpp
@@ -743,7 +743,7 @@ void FObjectPropertyBase::CheckValidObject(void* ValueAddress, TObjectPtr<UObjec
 					TEXT("Serialized %s for a property of %s. Reference will be nullptred.\n    ReferencingObject = %s\n    Property = %s\n    Item = %s"),
 					*ObjectClass->GetFullName(),
 					*PropertyClass->GetFullName(),
-					*GetFullNameSafe(SerializeContext->SerializedObject),
+					*GetFullNameSafe(SerializeContext ? SerializeContext->SerializedObject : nullptr),
 					*GetFullName(),
 					*Object.GetFullName()
 				);
@@ -758,7 +758,7 @@ void FObjectPropertyBase::CheckValidObject(void* ValueAddress, TObjectPtr<UObjec
 					*ObjectClass->GetFullName(),
 					*PropertyClass->GetFullName(),
 					DefaultValue ? *DefaultValue->GetFullName() : TEXT("None"),
-					*GetFullNameSafe(SerializeContext->SerializedObject),
+					*GetFullNameSafe(SerializeContext ? SerializeContext->SerializedObject : nullptr),
 					*GetFullName(),
 					*Object.GetFullName()
 				);
