Fix a regression that prevents placeholder types from being created on load for missing imports when IDO features are enabled.

#rnx

diff --git a/Engine/Source/Runtime/CoreUObject/Private/UObject/InstanceDataObjectUtils.cpp b/Engine/Source/Runtime/CoreUObject/Private/UObject/InstanceDataObjectUtils.cpp
--- a/Engine/Source/Runtime/CoreUObject/Private/UObject/InstanceDataObjectUtils.cpp
+++ b/Engine/Source/Runtime/CoreUObject/Private/UObject/InstanceDataObjectUtils.cpp
@@ -217,7 +217,7 @@ namespace UE
 	bool CanCreatePropertyBagPlaceholderTypeForImportClass(const UClass* ImportClass)
 	{
 		// @todo - Expand to other import types (e.g. prefab BPs) later; for now restricted to Verse class objects only.
-		return ImportClass && ImportClass->CanCreateInstanceDataObject();
+		return ImportClass && ImportClass->GetFName() == NAME_VerseClass;
 	}
 
 	bool IsClassOfInstanceDataObjectClass(UStruct* Class)
