In zenloader, do not call editor callbacks while handling a PostLoad. We don't allow callbacks when recursively loading already and we should avoid callbacks when handling postloads since we can't be sure what an OnAssetLoaded delegate might assume is completed when issuing a callback _while_ loading.  Instead the callbacks will be deferred until the next loader Tick or Flush from the game thread.

#rnx

diff --git a/Engine/Source/Runtime/CoreUObject/Private/Serialization/AsyncLoading2.cpp b/Engine/Source/Runtime/CoreUObject/Private/Serialization/AsyncLoading2.cpp
--- a/Engine/Source/Runtime/CoreUObject/Private/Serialization/AsyncLoading2.cpp
+++ b/Engine/Source/Runtime/CoreUObject/Private/Serialization/AsyncLoading2.cpp
@@ -9497,15 +9497,16 @@ void FAsyncPackage2::FinishUPackage()
 void FAsyncLoadingThread2::ConditionalProcessEditorCallbacks()
 {
 	check(IsInGameThread());
-	
-	// Prevent objects still being loaded from being accessible from the game thread.
-	TGuardValue GuardVisibilityFilter(FUObjectThreadContext::Get().AsyncVisibilityFilter, EInternalObjectFlags::AsyncLoadingPhase1);
 
-	if (!GameThreadState->SyncLoadContextStack.IsEmpty())
+	FUObjectThreadContext& ThreadContext = FUObjectThreadContext::Get();
+	if (ThreadContext.IsRoutingPostLoad || !GameThreadState->SyncLoadContextStack.IsEmpty())
 	{
 		return;
 	}
 
+	// Prevent objects still being loaded from being accessible from the game thread.
+	TGuardValue GuardVisibilityFilter(ThreadContext.AsyncVisibilityFilter, EInternalObjectFlags::AsyncLoadingPhase1);
+
 	FBlueprintSupport::FlushReinstancingQueue();
 
 	while (!EditorCompletedUPackages.IsEmpty() || !EditorLoadedAssets.IsEmpty())
