For zenloader, when waiting on a dependent package, increase the refcount when we store package references (and releaseref when removing the waiting package reference). This is to prevent a PackageA waiting on PackageB to not read from a free'd PackageB should another thread fully load PackageB and process the deferred deletion package list (freeing fully loaded packages)

#rnx

diff --git a/Engine/Source/Runtime/CoreUObject/Private/Serialization/AsyncLoading2.cpp b/Engine/Source/Runtime/CoreUObject/Private/Serialization/AsyncLoading2.cpp
--- a/Engine/Source/Runtime/CoreUObject/Private/Serialization/AsyncLoading2.cpp
+++ b/Engine/Source/Runtime/CoreUObject/Private/Serialization/AsyncLoading2.cpp
@@ -3182,7 +3182,9 @@ private:
 				check(!WaitListPackageState.PackagesWaitingForThisHead);
 				WaitListPackageState.PackagesWaitingForThisHead = PackageToAdd;
 			}
+
 			WaitListPackageState.PackagesWaitingForThisTail = PackageToAdd;
+			WaitListPackage->AddRef();
 			PackageToAddState.WaitingForPackage = WaitListPackage;
 		}
 
@@ -3205,6 +3207,7 @@ private:
 				check(WaitListPackageState.PackagesWaitingForThisHead == PackageToRemove);
 				WaitListPackageState.PackagesWaitingForThisHead = PackageToRemoveState.NextLink;
 			}
+
 			if (PackageToRemoveState.NextLink)
 			{
 				FAllDependenciesState& NextLinkState = PackageToRemoveState.NextLink->*StateMemberPtr;
@@ -3215,9 +3218,12 @@ private:
 				check(WaitListPackageState.PackagesWaitingForThisTail == PackageToRemove);
 				WaitListPackageState.PackagesWaitingForThisTail = PackageToRemoveState.PrevLink;
 			}
+
 			PackageToRemoveState.PrevLink = nullptr;
 			PackageToRemoveState.NextLink = nullptr;
+			FAsyncPackage2* WaitingForPackage = PackageToRemoveState.WaitingForPackage;
 			PackageToRemoveState.WaitingForPackage = nullptr;
+			WaitingForPackage->ReleaseRef();
 		}
 	};
 	FAllDependenciesState		AllDependenciesSetupState;
