Treat always-overridden properties as a 'replace' operation in the TPS data stream on save/write.

This ensures the data will be written back out again if the underlying object type is a placeholder and/or if the original property is missing (loose).

diff --git a/Engine/Source/Runtime/CoreUObject/Private/UObject/PropertyTag.cpp b/Engine/Source/Runtime/CoreUObject/Private/UObject/PropertyTag.cpp
--- a/Engine/Source/Runtime/CoreUObject/Private/UObject/PropertyTag.cpp
+++ b/Engine/Source/Runtime/CoreUObject/Private/UObject/PropertyTag.cpp
@@ -135,7 +135,17 @@ static EPropertyTagExtension CalculatePropertyExtensionFlags(FArchive& Underlyin
 		Tag.bExperimentalOverridableLogic = Tag.GetProperty()->HasAnyPropertyFlags(CPF_ExperimentalOverridableLogic);
 		if (OverriddenProperties || Tag.bExperimentalOverridableLogic)
 		{
-			Tag.OverrideOperation = OverriddenProperties ? OverriddenProperties->GetOverriddenPropertyOperation(UnderlyingArchive.GetSerializedPropertyChain(), Tag.GetProperty()) : EOverriddenPropertyOperation::None;
+			// If this property is marked as always overridden, emit op to replace any inherited/default value with our own. If this
+			// value is later read back in as a loose property, it will be properly treated as an override during write impersonation.
+			if (Tag.GetProperty()->HasAnyPropertyFlags(CPF_ExperimentalAlwaysOverriden))
+			{
+				Tag.OverrideOperation = EOverriddenPropertyOperation::Replace;
+			}
+			else
+			{
+				Tag.OverrideOperation = OverriddenProperties ? OverriddenProperties->GetOverriddenPropertyOperation(UnderlyingArchive.GetSerializedPropertyChain(), Tag.GetProperty()) : EOverriddenPropertyOperation::None;
+			}
+
 			PropertyTagExtensions |= EPropertyTagExtension::OverridableInformation;
 		}
 	}
