SavePackage/Cooke/FArchiver: Add FArchiveSavePackageData to FArchive, giving Serialize functions access to FObjectSaveContextData. This is necessary to allow structs (which do not have a PreSave function) the ability to add manage dependencies in editor saves and build dependencies in cook saves.
#rnx

diff --git a/Engine/Plugins/Runtime/USDCore/Source/USDClasses/Private/USDAssetCache2.cpp b/Engine/Plugins/Runtime/USDCore/Source/USDClasses/Private/USDAssetCache2.cpp
--- a/Engine/Plugins/Runtime/USDCore/Source/USDClasses/Private/USDAssetCache2.cpp
+++ b/Engine/Plugins/Runtime/USDCore/Source/USDClasses/Private/USDAssetCache2.cpp
@@ -253,7 +253,7 @@ namespace UE::AssetCache::Private
 
 		// Transfer over whether we're cooking or not, but otherwise leave at the default settings, as that is
 		// how we'll try to deserialize on the other end, which is potentially at runtime
-		Ar.SetCookData(DestArchive.GetCookData());
+		Ar.SetSavePackageData(DestArchive.GetSavePackageData());
 		Ar.SetFilterEditorOnly(DestArchive.IsFilterEditorOnly());
 		Ar.SetUEVer(DestArchive.UEVer());
 		Ar.SetLicenseeUEVer(DestArchive.LicenseeUEVer());
diff --git a/Engine/Source/Runtime/Core/Private/Serialization/Archive.cpp b/Engine/Source/Runtime/Core/Private/Serialization/Archive.cpp
--- a/Engine/Source/Runtime/Core/Private/Serialization/Archive.cpp
+++ b/Engine/Source/Runtime/Core/Private/Serialization/Archive.cpp
@@ -193,7 +193,7 @@ void FArchiveState::Reset()
 	ArUseCustomPropertyList				= false;
 	ArMergeOverrides					= false;
 	ArShouldSkipUpdateCustomVersion		= false;
-	CookData							= nullptr;
+	SavePackageData						= nullptr;
 	SerializedProperty					= nullptr;
 
 	delete SerializedPropertyChain;
@@ -253,7 +253,7 @@ void FArchiveState::CopyTrivialFArchiveStatusMembers(const FArchiveState& Archiv
 	ArUseCustomPropertyList              = ArchiveToCopy.ArUseCustomPropertyList;
 	ArMergeOverrides					 = ArchiveToCopy.ArMergeOverrides;
 	ArShouldSkipUpdateCustomVersion		 = ArchiveToCopy.ArShouldSkipUpdateCustomVersion;
-	CookData							 = ArchiveToCopy.CookData;
+	SavePackageData						 = ArchiveToCopy.SavePackageData;
 	SerializedProperty					 = ArchiveToCopy.SerializedProperty;
 #if USE_STABLE_LOCALIZATION_KEYS
 	SetBaseLocalizationNamespace(ArchiveToCopy.GetBaseLocalizationNamespace());
diff --git a/Engine/Source/Runtime/Core/Public/Serialization/Archive.h b/Engine/Source/Runtime/Core/Public/Serialization/Archive.h
--- a/Engine/Source/Runtime/Core/Public/Serialization/Archive.h
+++ b/Engine/Source/Runtime/Core/Public/Serialization/Archive.h
@@ -13,6 +13,7 @@
 #include "Misc/EngineVersionBase.h"
 #include "Misc/VarArgs.h"
 #include "Serialization/ArchiveCookData.h"
+#include "Serialization/ArchiveSavePackageData.h"
 #include "Templates/EnableIf.h"
 #include "Templates/IsArrayOrRefOfTypeByPredicate.h"
 #include "Templates/IsEnumClass.h"
@@ -210,20 +211,6 @@ public:
 		return ArLicenseeUEVer;
 	}
 
-	/** Returns the global engine serialization version used for this archive. */
-	UE_DEPRECATED(5.0, "Use UEVer instead which returns the version as a FPackageFileVersion. See the @FPackageFileVersion documentation for further details")
-	FORCEINLINE int32 UE4Ver() const
-	{
-		return ArUEVer.FileVersionUE4;
-	}
-
-	/** Returns the licensee-specific version used for this archive, will be 0 by default. */
-	UE_DEPRECATED(5.0, "Use LicenseeUEVer instead")
-	FORCEINLINE int32 LicenseeUE4Ver() const
-	{
-		return LicenseeUEVer();
-	}
-
 	/** Returns the compiled engine version used for this archive. */
 	FORCEINLINE FEngineVersionBase EngineVer() const
 	{
@@ -600,47 +587,75 @@ public:
 	 * @return true if the archive is used for cooking, false otherwise.
 	 */
 	FORCEINLINE bool IsCooking() const
-	{		
-		return CookData != nullptr;
+	{
+		return SavePackageData && SavePackageData->CookContext;
 	}
 
 	/**
-	 * Marks that this archive is "cooking" by providing the cook data bundle.
-	 * Must be set after loading/saving/transacting.
+	 * Marks that this archive is one of the archives used by SavePackage, and provides access
+	 * to data and functions used only during savepackage, such as the cook data bundle for information about
+	 * the cook being conducted (if applicable) and writing build dependencies for assetregistry.
+	 * Must be set only after setting the loading/saving/transacting flags.
 	 */
-	void SetCookData(FArchiveCookData* InCookData)
+	void SetSavePackageData(FArchiveSavePackageData* InSavePackageData)
 	{
-		checkf(!(InCookData == nullptr && CookData), TEXT("Can't turn off cooking once you turn it on!"));
+		checkf(!(InSavePackageData == nullptr && SavePackageData), TEXT("Can't turn off cooking once you turn it on!"));
 
-		if (InCookData)
+		if (InSavePackageData)
 		{
 			check(!IsLoading() && !IsTransacting() && IsSaving());
-			CookData = InCookData;
+			SavePackageData = InSavePackageData;
 		}
-		
+
+	}
+
+	FArchiveSavePackageData* GetSavePackageData()
+	{
+		return SavePackageData;
+	}
+
+	UE_DEPRECATED(5.5, "Use SetSavePackageData.")
+	void SetCookData(FArchiveCookData* InCookData)
+	{
+		// Backwards compatibility for this function is not supported because of the low amount of public usage and the
+		// cost of supporting the backwards compatability (extra memory on FArchiveSavePackageData). Change your code
+		// to use SetSavePackageData, or contact Epic for assistance in adding the backwards compatibility.
+		check(false);
 	}
 
+	UE_DEPRECATED(5.5, "Use GetCookContext or CookingTarget or GetSavePackageData.")
 	FArchiveCookData* GetCookData()
 	{
-		return CookData;
+		// Backwards compatibility for this function is not supported because of the low amount of public usage and the
+		// cost of supporting the backwards compatability (extra memory on FArchiveSavePackageData). Change your code
+		// to use GetCookContext or CookingTarget or GetSavePackageData, or contact Epic for assistance.
+		check(false);
+		return nullptr;
 	}
 
 	FORCEINLINE FArchiveCookContext* GetCookContext()
 	{
-		return CookData ? &CookData->CookContext : nullptr;
+		return SavePackageData ? SavePackageData->CookContext : nullptr;
 	}
 
-	UE_DEPRECATED(5.1, "CookingTarget has been moved to FArchiveCookData.")
-	void SetCookingTarget(const ITargetPlatform*) {}
-
 	/**
 	 * Returns the cooking target platform.
+	 * Returns null if not being serialized from Savepackage, or if it is an editor save rather than a cook save.
 	 *
 	 * @return Target platform.
 	 */
 	FORCEINLINE const ITargetPlatform* CookingTarget() const
 	{
-		return CookData ? &CookData->TargetPlatform : nullptr;
+		return SavePackageData ? SavePackageData->TargetPlatform : nullptr;
+	}
+
+	/**
+	 * Return the API object used to record extra data for SavePackage calls (e.g. cook build dependencies).
+	 * Returns null if not being serialized from SavePackage.
+	 */
+	FObjectSavePackageSerializeContext* GetSavePackageSerializeContext()
+	{
+		return SavePackageData ? &SavePackageData->SavePackageContext : nullptr;
 	}
 
 	/**
@@ -941,13 +956,6 @@ public:
 	 */
 	CORE_API virtual void SetUEVer(FPackageFileVersion InVer);
 
-	UE_DEPRECATED(5.0, "Use SetUEVer instead which takes the version as a FPackageFileVersion. See the @FPackageFileVersion documentation for further details")
-	FORCEINLINE void SetUE4Ver(int32 InVer)
-	{
-		FPackageFileVersion PackageFileVersion = FPackageFileVersion::CreateUE4Version((EUnrealEngineObjectUE4Version)InVer);
-		SetUEVer(PackageFileVersion);
-	}
-
 	/**
 	 * Sets the archive licensee version number. Used by the code that makes sure that FLinkerLoad's 
 	 * internal archive versions match the file reader it creates.
@@ -956,12 +964,6 @@ public:
 	 */
 	CORE_API virtual void SetLicenseeUEVer(int32 InVer);
 
-	UE_DEPRECATED(5.0, "Use SetLicenseeUEVer instead")
-	FORCEINLINE void SetLicenseeUE4Ver(int32 InVer)
-	{
-		SetLicenseeUEVer(InVer);
-	}
-
 	/**
 	 * Sets the archive engine version. Used by the code that makes sure that FLinkerLoad's
 	 * internal archive versions match the file reader it creates.
@@ -1009,8 +1011,11 @@ public:
 
 // These will be private in FArchive
 protected:
-	/** Holds data for cooking. Required if cooking, nullptr means not cooking */
-	FArchiveCookData* CookData = nullptr;
+	/**
+	 * Provider of data and API specific to UPackage::Save2. Required by archives used from UPackage::Save2, and required for cook saves of packages.
+	 * nullptr means the archive is not one from UPackage::Save2 and is not being called from a cook save.
+	 */
+	FArchiveSavePackageData* SavePackageData = nullptr;
 	
 	/** Holds the pointer to the property that is currently being serialized */
 	FProperty* SerializedProperty;
@@ -1846,7 +1851,6 @@ public:
 	}
 
 	using FArchiveState::UEVer;
-	using FArchiveState::UE4Ver;
 	using FArchiveState::LicenseeUEVer;
 	using FArchiveState::EngineVer;
 	using FArchiveState::EngineNetVer;
@@ -1919,8 +1923,25 @@ public:
 	using FArchiveState::IsNetArchive;
 	using FArchiveState::IsCooking;
 	using FArchiveState::CookingTarget;
-	using FArchiveState::SetCookData;
-	using FArchiveState::GetCookData;
+	UE_DEPRECATED(5.5, "Use SetSavePackageData.")
+	void SetCookData(FArchiveCookData* InCookData)
+	{
+		// Backwards compatibility for this function is not supported because of the low amount of public usage and the
+		// cost of supporting the backwards compatability (extra memory on FArchiveSavePackageData). Change your code
+		// to use SetSavePackageData, or contact Epic for assistance.
+		check(false);
+	}
+	using FArchiveState::SetSavePackageData;
+	using FArchiveState::GetSavePackageData;
+	UE_DEPRECATED(5.5, "Use GetCookContext or CookingTarget or GetSavePackageData")
+	FArchiveCookData* GetCookData()
+	{
+		// Backwards compatibility for this function is not supported because of the low amount of public usage and the
+		// cost of supporting the backwards compatability (extra memory on FArchiveSavePackageData). Change your code
+		// to use GetCookContext or CookingTarget or GetSavePackageData, or contact Epic for assistance.
+		check(false);
+		return nullptr;
+	}
 	using FArchiveState::GetCookContext;
 	using FArchiveState::UseToResolveEnumerators;
 	using FArchiveState::ShouldSkipProperty;
@@ -2094,7 +2115,6 @@ public:
 	using FArchiveState::SetForceUnicode;
 	using FArchiveState::SetIsPersistent;
 	using FArchiveState::SetUEVer;
-	using FArchiveState::SetUE4Ver;
 	using FArchiveState::SetLicenseeUEVer;
 	using FArchiveState::SetEngineVer;
 	using FArchiveState::SetEngineNetVer;
diff --git a/Engine/Source/Runtime/Core/Public/Serialization/ArchiveSavePackageData.h b/Engine/Source/Runtime/Core/Public/Serialization/ArchiveSavePackageData.h
new file mode 100644
--- /dev/null
+++ b/Engine/Source/Runtime/Core/Public/Serialization/ArchiveSavePackageData.h
@@ -0,0 +1,25 @@
+// Copyright Epic Games, Inc. All Rights Reserved.
+
+#pragma once
+
+class FObjectSavePackageSerializeContext;
+class ITargetPlatform;
+struct FArchiveCookContext;
+
+/** Holds archive data only relevant for archives used during SavePackage. */
+struct FArchiveSavePackageData
+{
+	FArchiveSavePackageData(FObjectSavePackageSerializeContext& InSavePackageContext,
+		const ITargetPlatform* InTargetPlatform, FArchiveCookContext* InCookContext)
+		: SavePackageContext(InSavePackageContext)
+		, TargetPlatform(InTargetPlatform)
+		, CookContext(InCookContext)
+	{
+
+	}
+	FObjectSavePackageSerializeContext& SavePackageContext;
+	/** Non-null if a cook save. Null if an editor save. */
+	const ITargetPlatform* TargetPlatform = nullptr;
+	/** Non-null if a cook save. Null if an editor save. */
+	FArchiveCookContext* CookContext = nullptr;
+};
diff --git a/Engine/Source/Runtime/CoreUObject/Private/UObject/SavePackage/PackageHarvester.cpp b/Engine/Source/Runtime/CoreUObject/Private/UObject/SavePackage/PackageHarvester.cpp
--- a/Engine/Source/Runtime/CoreUObject/Private/UObject/SavePackage/PackageHarvester.cpp
+++ b/Engine/Source/Runtime/CoreUObject/Private/UObject/SavePackage/PackageHarvester.cpp
@@ -419,7 +419,7 @@ FPackageHarvester::FPackageHarvester(FSaveContext& InContext)
 
 	this->SetPortFlags(SaveContext.GetPortFlags());
 	this->SetFilterEditorOnly(SaveContext.IsFilterEditorOnly());
-	this->SetCookData(SaveContext.GetSaveArgs().ArchiveCookData);
+	this->SetSavePackageData(&SaveContext.GetArchiveSavePackageData());
 	this->SetUseUnversionedPropertySerialization(SaveContext.IsSaveUnversionedProperties());
 
 	ResolveOverrides();
diff --git a/Engine/Source/Runtime/CoreUObject/Private/UObject/SavePackage/SaveContext.h b/Engine/Source/Runtime/CoreUObject/Private/UObject/SavePackage/SaveContext.h
--- a/Engine/Source/Runtime/CoreUObject/Private/UObject/SavePackage/SaveContext.h
+++ b/Engine/Source/Runtime/CoreUObject/Private/UObject/SavePackage/SaveContext.h
@@ -8,6 +8,7 @@
 #include "Misc/AssetRegistryInterface.h"
 #include "Misc/Optional.h"
 #include "Misc/PackageName.h"
+#include "Serialization/ArchiveSavePackageData.h"
 #include "Serialization/CustomVersion.h"
 #include "Serialization/LargeMemoryWriter.h"
 #include "Serialization/PropertyLocalizationDataGathering.h"
@@ -478,6 +479,8 @@ public:
 		, Filename(InFilename)
 		, SaveArgs(InSaveArgs)
 		, PackageWriter(InSaveArgs.SavePackageContext ? InSaveArgs.SavePackageContext->PackageWriter : nullptr)
+		, ObjectSavePackageSerializeContext(ObjectSaveContext)
+		, ArchiveSavePackageData(ObjectSavePackageSerializeContext, nullptr, nullptr)
 		, GameRealmExcludedObjectMarks(GetExcludedObjectMarksForGameRealm(SaveArgs.GetTargetPlatform()))
 	{
 		// Assumptions & checks
@@ -527,6 +530,8 @@ public:
 			ObjectSaveContext.CookType = SaveArgs.ArchiveCookData->CookContext.GetCookType();
 			ObjectSaveContext.CookingDLC = SaveArgs.ArchiveCookData->CookContext.GetCookingDLC();
 		}
+		ArchiveSavePackageData.TargetPlatform = ObjectSaveContext.TargetPlatform;
+		ArchiveSavePackageData.CookContext = SaveArgs.ArchiveCookData ? &SaveArgs.ArchiveCookData->CookContext : nullptr;
 		if (SaveArgs.InOutSaveOverrides)
 		{
 			ObjectSaveContext.SaveOverrides = MoveTemp(*SaveArgs.InOutSaveOverrides);
@@ -566,9 +571,9 @@ public:
 		return SaveArgs;
 	}
 
-	FArchiveCookData* GetCookData()
+	FArchiveSavePackageData& GetArchiveSavePackageData()
 	{
-		return SaveArgs.ArchiveCookData;
+		return ArchiveSavePackageData;
 	}
 
 	const ITargetPlatform* GetTargetPlatform() const
@@ -1220,6 +1225,8 @@ private:
 
 	// State context
 	FObjectSaveContextData ObjectSaveContext;
+	FObjectSavePackageSerializeContext ObjectSavePackageSerializeContext;
+	FArchiveSavePackageData ArchiveSavePackageData;
 	bool bCanUseUnversionedPropertySerialization = false;
 	bool bTextFormat = false;
 	bool bIsProcessingPrestreamPackages = false;
diff --git a/Engine/Source/Runtime/CoreUObject/Private/UObject/SavePackage/SavePackageUtilities.cpp b/Engine/Source/Runtime/CoreUObject/Private/UObject/SavePackage/SavePackageUtilities.cpp
--- a/Engine/Source/Runtime/CoreUObject/Private/UObject/SavePackage/SavePackageUtilities.cpp
+++ b/Engine/Source/Runtime/CoreUObject/Private/UObject/SavePackage/SavePackageUtilities.cpp
@@ -19,6 +19,7 @@
 #include "Misc/ScopeExit.h"
 #include "Misc/ScopedSlowTask.h"
 #include "SaveContext.h"
+#include "Serialization/ArchiveSavePackageDataBuffer.h"
 #include "Serialization/BulkData.h"
 #include "Serialization/EditorBulkData.h"
 #include "Serialization/CompactBinarySerialization.h"
@@ -768,30 +769,29 @@ FObjectSaveContextData::FObjectSaveContextData(UPackage* Package, const ITargetP
 }
 
 #if WITH_EDITOR
-void FObjectPreSaveContext::AddCookBuildDependency(UE::Cook::FCookDependency BuildDependency)
-{
-	Data.CookBuildDependencies.Add(MoveTemp(BuildDependency));
-}
-void FObjectPreSaveContext::AddCookRuntimeDependency(FSoftObjectPath RuntimeDependency)
+namespace UE::SavePackageUtilities
 {
-	Data.CookRuntimeDependencies.Add(MoveTemp(RuntimeDependency));
-}
-void FObjectPreSaveContext::HarvestCookRuntimeDependencies(UObject* HarvestReferencesFrom)
+
+void HarvestCookRuntimeDependencies(FObjectSaveContextData& Data, UObject* HarvestReferencesFrom)
 {
 	if (!HarvestReferencesFrom)
 	{
 		return;
 	}
-	if (!GetTargetPlatform())
+	if (!Data.TargetPlatform)
 	{
 		return;
 	}
 
 	UPackage* PackageBeingSaved = nullptr; // We don't have a pointer for this, so set it to null
-	FArchiveCookContext CookContext(PackageBeingSaved, Data.CookType, Data.CookingDLC, GetTargetPlatform());
-	FArchiveCookData CookData(*GetTargetPlatform(), CookContext);
+	// Don't store the input Data on the ArchiveSavePackageData; we just want to harvest the serialized
+	// FSoftObjectPaths, not allow direct writes to its cookdependencies or other data. We only set
+	// ArchiveSavePackageData to provide access to the CookContext.
+	FArchiveCookContext CookContext(PackageBeingSaved, Data.CookType, Data.CookingDLC, Data.TargetPlatform);
+	FArchiveSavePackageDataBuffer SavePackageData(CookContext);
+
 	FImportExportCollector Collector(HarvestReferencesFrom->GetPackage());
-	Collector.SetCookData(&CookData);
+	Collector.SetSavePackageData(&SavePackageData);
 	Collector.SerializeObjectAndReferencedExports(HarvestReferencesFrom);
 	for (const TPair<FName, ESoftObjectPathCollectType>& Pair : Collector.GetImportedPackages())
 	{
@@ -806,12 +806,26 @@ void FObjectPreSaveContext::HarvestCookRuntimeDependencies(UObject* HarvestRefer
 			continue;
 		}
 		FSoftObjectPath PackageSoftPath(PackageName, NAME_None, FString());
-		AddCookRuntimeDependency(PackageSoftPath);
+		Data.CookRuntimeDependencies.Add(MoveTemp(PackageSoftPath));
 	}
+}
 
+} // namespace UE::SavePackageUtilities
+
+void FObjectPreSaveContext::AddCookBuildDependency(UE::Cook::FCookDependency BuildDependency)
+{
+	Data.CookBuildDependencies.Add(MoveTemp(BuildDependency));
+}
+void FObjectPreSaveContext::AddCookRuntimeDependency(FSoftObjectPath RuntimeDependency)
+{
+	Data.CookRuntimeDependencies.Add(MoveTemp(RuntimeDependency));
+}
+void FObjectPreSaveContext::HarvestCookRuntimeDependencies(UObject* HarvestReferencesFrom)
+{
+	UE::SavePackageUtilities::HarvestCookRuntimeDependencies(Data, HarvestReferencesFrom);
 }
 
-bool FObjectPreSaveContext::IsDeterminismDebug()
+bool FObjectPreSaveContext::IsDeterminismDebug() const
 {
 	return Data.bDeterminismDebug;
 }
@@ -825,6 +839,61 @@ void FObjectPreSaveContext::RegisterDeterminismHelper(
 	}
 }
 
+void FObjectSavePackageSerializeContext::AddCookBuildDependency(UE::Cook::FCookDependency BuildDependency)
+{
+	if (GetPhase() != EObjectSaveContextPhase::Harvest)
+	{
+		UE_LOG(LogSavePackage, Error,
+			TEXT("AddCookBuildDependency called when GetPhase() != EObjectSaveContextPhase::Harvest. This is invalid and will be ignored."));
+		FDebug::DumpStackTraceToLog(ELogVerbosity::Warning);
+		return;
+	}
+	Data.CookBuildDependencies.Add(MoveTemp(BuildDependency));
+}
+void FObjectSavePackageSerializeContext::AddCookRuntimeDependency(FSoftObjectPath RuntimeDependency)
+{
+	if (GetPhase() != EObjectSaveContextPhase::Harvest)
+	{
+		UE_LOG(LogSavePackage, Error,
+			TEXT("AddCookRuntimeDependency called when GetPhase() != EObjectSaveContextPhase::Harvest. This is invalid and will be ignored."));
+		FDebug::DumpStackTraceToLog(ELogVerbosity::Warning);
+		return;
+	}
+	Data.CookRuntimeDependencies.Add(MoveTemp(RuntimeDependency));
+}
+void FObjectSavePackageSerializeContext::HarvestCookRuntimeDependencies(UObject* HarvestReferencesFrom)
+{
+	if (GetPhase() != EObjectSaveContextPhase::Harvest)
+	{
+		UE_LOG(LogSavePackage, Error,
+			TEXT("HarvestCookRuntimeDependencies called when GetPhase() != EObjectSaveContextPhase::Harvest. This is invalid and will be ignored."));
+		FDebug::DumpStackTraceToLog(ELogVerbosity::Warning);
+		return;
+	}
+	UE::SavePackageUtilities::HarvestCookRuntimeDependencies(Data, HarvestReferencesFrom);
+}
+
+bool FObjectSavePackageSerializeContext::IsDeterminismDebug() const
+{
+	return Data.bDeterminismDebug;
+}
+
+void FObjectSavePackageSerializeContext::RegisterDeterminismHelper(
+	const TRefCountPtr<UE::Cook::IDeterminismHelper>& DeterminismHelper)
+{
+	if (GetPhase() != EObjectSaveContextPhase::Harvest)
+	{
+		UE_LOG(LogSavePackage, Error,
+			TEXT("RegisterDeterminismHelper called when GetPhase() != EObjectSaveContextPhase::Harvest. This is invalid and will be ignored."));
+		FDebug::DumpStackTraceToLog(ELogVerbosity::Warning);
+		return;
+	}
+	if (Data.PackageWriter)
+	{
+		Data.PackageWriter->RegisterDeterminismHelper(Data.Object, DeterminismHelper);
+	}
+}
+
 #endif
 
 void FObjectSaveContextData::Set(UPackage* Package, const ITargetPlatform* InTargetPlatform, const TCHAR* InTargetFilename, uint32 InSaveFlags)
diff --git a/Engine/Source/Runtime/CoreUObject/Private/UObject/SavePackage2.cpp b/Engine/Source/Runtime/CoreUObject/Private/UObject/SavePackage2.cpp
--- a/Engine/Source/Runtime/CoreUObject/Private/UObject/SavePackage2.cpp
+++ b/Engine/Source/Runtime/CoreUObject/Private/UObject/SavePackage2.cpp
@@ -1160,7 +1160,7 @@ ESavePackageResult BuildLinker(FSaveContext& SaveContext)
 		{
 			SaveContext.GetLinker()->SetFilterEditorOnly(SaveContext.IsFilterEditorOnly());
 		}
-		SaveContext.GetLinker()->SetCookData(SaveContext.GetCookData());
+		SaveContext.GetLinker()->SetSavePackageData(&SaveContext.GetArchiveSavePackageData());
 
 		bool bUseUnversionedProperties = SaveContext.IsSaveUnversionedProperties();
 		SaveContext.GetLinker()->SetUseUnversionedPropertySerialization(bUseUnversionedProperties);
@@ -1999,7 +1999,7 @@ ESavePackageResult WritePackageHeader(FStructuredArchive::FRecord& StructuredArc
 	{
 		// Save asset registry data so the editor can search for information about assets in this package
 		SCOPED_SAVETIMER(UPackage_Save_SaveAssetRegistryData);
-		FArchiveCookData* CookData = SaveContext.GetCookData();
+		FArchiveSavePackageData& ArchiveSavePackageData = SaveContext.GetArchiveSavePackageData();
 		UE::AssetRegistry::FWritePackageDataArgs WriteARArgs;
 		WriteARArgs.ParentRecord = &StructuredArchiveRoot;
 		WriteARArgs.Package = SaveContext.GetPackage();
@@ -2008,7 +2008,7 @@ ESavePackageResult WritePackageHeader(FStructuredArchive::FRecord& StructuredArc
 		WriteARArgs.SoftPackagesUsedInGame = &SaveContext.GetSoftPackagesUsedInGame();
 		WriteARArgs.PackageBuildDependencies = &SaveContext.GetPackageBuildDependencies();
 		WriteARArgs.bProceduralSave = SaveContext.IsProceduralSave();
-		WriteARArgs.CookContext = CookData ? &CookData->CookContext : nullptr;;
+		WriteARArgs.CookContext = ArchiveSavePackageData.CookContext;
 		WriteARArgs.OutAssetDatas = &SaveContext.GetSavedAssets();
 		UE::AssetRegistry::WritePackageData(WriteARArgs);
 	}
@@ -3038,6 +3038,7 @@ ESavePackageResult InnerSave(FSaveContext& SaveContext)
 
 	// Harvest Package
 	SlowTask.EnterProgressFrame();
+	SaveContext.GetObjectSaveContext().ObjectSaveContextPhase = EObjectSaveContextPhase ::Harvest;
 	SaveContext.Result = HarvestPackage(SaveContext);
 	if (SaveContext.Result != ESavePackageResult::Success)
 	{
@@ -3068,6 +3069,7 @@ ESavePackageResult InnerSave(FSaveContext& SaveContext)
 	//										-> .o.ubulk
 	//										-> etc
 	SlowTask.EnterProgressFrame();
+	SaveContext.GetObjectSaveContext().ObjectSaveContextPhase = EObjectSaveContextPhase ::Write;
 	for (ESaveRealm HarvestingContext : SaveContext.GetHarvestedRealmsToSave())
 	{
 		SaveContext.Result = SaveHarvestedRealms(SaveContext, HarvestingContext);
@@ -3187,6 +3189,7 @@ FSavePackageResultStruct UPackage::Save2(UPackage* InPackage, UObject* InAsset,
 		IPackageWriter* PackageWriter = SaveContext.GetPackageWriter();
 		if (!PackageWriter || !PackageWriter->IsPreSaveCompleted())
 		{
+			SaveContext.GetObjectSaveContext().ObjectSaveContextPhase = EObjectSaveContextPhase::PreSave;
 			SaveContext.Result = RoutePresave(SaveContext);
 			if (SaveContext.Result != ESavePackageResult::Success)
 			{
@@ -3239,6 +3242,7 @@ FSavePackageResultStruct UPackage::Save2(UPackage* InPackage, UObject* InAsset,
 	SlowTask.EnterProgressFrame();
 	if (SaveContext.GetPostSaveRootRequired() && SaveContext.GetAsset())
 	{
+		SaveContext.GetObjectSaveContext().ObjectSaveContextPhase = EObjectSaveContextPhase::PostSave;
 		UE::SavePackageUtilities::CallPostSaveRoot(SaveContext.GetAsset(), SaveContext.GetObjectSaveContext(), SaveContext.GetPreSaveCleanup());
 		SaveContext.SetPostSaveRootRequired(false);
 	}
diff --git a/Engine/Source/Runtime/CoreUObject/Public/Serialization/ArchiveSavePackageDataBuffer.h b/Engine/Source/Runtime/CoreUObject/Public/Serialization/ArchiveSavePackageDataBuffer.h
new file mode 100644
--- /dev/null
+++ b/Engine/Source/Runtime/CoreUObject/Public/Serialization/ArchiveSavePackageDataBuffer.h
@@ -0,0 +1,63 @@
+// Copyright Epic Games, Inc. All Rights Reserved.
+
+#pragma once
+
+#include "Misc/Optional.h"
+#include "Serialization/ArchiveSavePackageData.h"
+#include "UObject/ArchiveCookContext.h"
+#include "UObject/ObjectSaveContext.h"
+
+/**
+ * This is the structure that should be used by most callers of an archive (other than SavePackage which handles
+ * the full complexity) that need to provide SavePackage or Cook information to the archive.
+ * 
+ * FArchiveSavePackageData is a minimalist struct that holds pointers to types it cannot have defined,
+ * because they are defined in modules higher in the dependency graph.
+ * Because of that minimalism, it needs pointers to several pieces of related data.
+ * Collect all of those pieces of related data in this module that knows all of the types, and keep them
+ * in a single amalgamated structure the sets the base class pointers to the appropriate internal structures.
+ */
+struct FArchiveSavePackageDataBuffer : public FArchiveSavePackageData
+{
+public:
+	FArchiveSavePackageDataBuffer(const ITargetPlatform* InTargetPlatform = nullptr)
+		: FArchiveSavePackageData(ObjectSavePackageSerializeContextBuffer, nullptr, nullptr)
+		, ObjectSaveContextData(FObjectSaveContextData())
+		, ObjectSavePackageSerializeContextBuffer(*ObjectSaveContextData)
+		, CookContextBuffer(nullptr /*InPackage*/, UE::Cook::ECookType::Unknown, UE::Cook::ECookingDLC::Unknown,
+			InTargetPlatform)
+	{
+		SetConstructorTargetPlatform(InTargetPlatform);
+	}
+
+	FArchiveSavePackageDataBuffer(FArchiveCookContext InContext)
+		: FArchiveSavePackageData(ObjectSavePackageSerializeContextBuffer, nullptr, nullptr)
+		, ObjectSaveContextData(FObjectSaveContextData())
+		, ObjectSavePackageSerializeContextBuffer(*ObjectSaveContextData)
+		, CookContextBuffer(MoveTemp(InContext))
+	{
+		SetConstructorTargetPlatform(InContext.GetTargetPlatform());
+	}
+
+	FArchiveSavePackageDataBuffer(FObjectSaveContextData& InData, UPackage* InPackage = nullptr)
+		: FArchiveSavePackageData(ObjectSavePackageSerializeContextBuffer, nullptr, nullptr)
+		, ObjectSavePackageSerializeContextBuffer(InData)
+		, CookContextBuffer(InPackage, InData.CookType, InData.CookingDLC, InData.TargetPlatform)
+	{
+		SetConstructorTargetPlatform(InData.TargetPlatform);
+	}
+
+	TOptional<FObjectSaveContextData> ObjectSaveContextData;
+	FObjectSavePackageSerializeContext ObjectSavePackageSerializeContextBuffer;
+	FArchiveCookContext CookContextBuffer;
+
+private:
+	void SetConstructorTargetPlatform(const ITargetPlatform* InTargetPlatform)
+	{
+		if (InTargetPlatform)
+		{
+			TargetPlatform = InTargetPlatform;
+			CookContext = &CookContextBuffer;
+		}
+	}
+};
diff --git a/Engine/Source/Runtime/CoreUObject/Public/UObject/ObjectSaveContext.h b/Engine/Source/Runtime/CoreUObject/Public/UObject/ObjectSaveContext.h
--- a/Engine/Source/Runtime/CoreUObject/Public/UObject/ObjectSaveContext.h
+++ b/Engine/Source/Runtime/CoreUObject/Public/UObject/ObjectSaveContext.h
@@ -4,9 +4,13 @@
 
 #include "Containers/UnrealString.h"
 #include "HAL/Platform.h"
+#include "Serialization/ArchiveSavePackageData.h"
+#include "Templates/RefCounting.h"
 #include "UObject/CookEnums.h"
 #include "UObject/ObjectSaveOverride.h"
-#include "Templates/RefCounting.h"
+#if WITH_EDITOR
+#include "Cooker/CookDependency.h"
+#endif
 
 class FPackagePath;
 class IPackageWriter;
@@ -14,10 +18,40 @@ class ITargetPlatform;
 class UPackage;
 struct FSoftObjectPath;
 #if WITH_EDITOR
-namespace UE::Cook { class FCookDependency; }
 namespace UE::Cook { class IDeterminismHelper; }
 #endif
 
+/** SavePackage calls PreSave and Serialize hooks on each object, and Serialize is called more than once. */
+enum class EObjectSaveContextPhase : uint8
+{
+	/**
+	 * Phase has not been set. This can be returned from GetPhase functions when the serialize calls need to report
+	 * a cook targetplatform but are not called from SavePackage.
+	 */
+	Invalid,
+	/**
+	 * SavePackage is calling PreSave on objects in the package.
+	 * Build dependencies are valid for writing during this phase.
+	 */
+	PreSave,
+	/**
+	 * SavePackage is calling Serialize(FArchive&) on objects in the package. The FArchive is an object collector
+	 * and the save is collecting imports, exports, and names.
+	 * Build dependencies are valid for writing during this phase.
+	 */
+	Harvest,
+	/**
+	 * SavePackage is calling Serialize(FArchive&) on objects in the package. The FArchive is recording the blobs of
+	 * data from each exported UObject that will be written to disk.
+	 * It is not valid and will cause an error if build dependencies are written during this phase.
+	 */
+	Write,
+	/**
+	 * SavePackage is calling PostSaveRoot on objects in the package.
+	 */
+	PostSave,
+};
+
 /** Data used to provide information about the save parameters during PreSave/PostSave. */
 struct FObjectSaveContextData
 {
@@ -73,6 +107,9 @@ struct FObjectSaveContextData
 	UE::Cook::ECookType CookType = UE::Cook::ECookType::Unknown;
 	UE::Cook::ECookingDLC CookingDLC = UE::Cook::ECookingDLC::Unknown;
 
+	/** Set to the appropriate phase when calling Serialize during SavePackage. */
+	EObjectSaveContextPhase ObjectSaveContextPhase = EObjectSaveContextPhase::Invalid;
+
 	/**
 	 * Set to true when the package is being saved due to a procedural save.
 	 * Any save without the possibility of user-generated edits to the package is a procedural save (Cooking, EditorDomain).
@@ -211,6 +248,9 @@ public:
 	UE::Cook::ECookType GetCookType() const { return Data.CookType; }
 	UE::Cook::ECookingDLC GetCookingDLC() const { return Data.CookingDLC; }
 
+	/** Return which phase of SavePackage callbacks are active. @see EObjectSaveContextPhase. */
+	EObjectSaveContextPhase GetPhase() const { return Data.ObjectSaveContextPhase; }
+
 #if WITH_EDITOR
 	/**
 	 * Add the given FCookDependency to the build dependencies for the package being cook-saved. Iterative cooks will
@@ -233,7 +273,7 @@ public:
 	 * Applicable only to cook saves: True if the SavePackage call should write extra debug data for debugging
 	 * cook determinism or incremental cook issues.
 	 */
-	COREUOBJECT_API bool IsDeterminismDebug();
+	COREUOBJECT_API bool IsDeterminismDebug() const;
 	/**
 	 * Ignored unless IsDeterminismDebug()=true. An object should call this function to register
 	 * their callback class for adding determinism diagnostics the package save.
@@ -277,6 +317,98 @@ protected:
 	friend class UObject;
 };
 
+/** Interface used by FArchiveSavePackageData during Serialize(FArchive& Ar) to access the save parameters. */
+class FObjectSavePackageSerializeContext
+{
+public:
+	explicit FObjectSavePackageSerializeContext(FObjectSaveContextData& InData)
+		: Data(InData)
+	{
+	}
+
+	FObjectSavePackageSerializeContext(const FObjectSavePackageSerializeContext& Other)
+		: Data(Other.Data)
+	{
+	}
+
+	/**
+	 * The target Filename being saved into (not the temporary file for saving).
+	 * The path is in the standard UnrealEngine form - it is as a relative path from the process binary directory.
+	 * Empty string if the saved bytes are not being saved to a file. Never null.
+	 */
+	const TCHAR* GetTargetFilename() const { return *Data.TargetFilename; }
+
+	/** Report whether this is a save into a target-specific cooked format. */
+	bool IsCooking() const { return Data.TargetPlatform != nullptr; }
+
+	/** Return the targetplatform of the save, if cooking. Null if not cooking. */
+	const ITargetPlatform* GetTargetPlatform() const { return Data.TargetPlatform; }
+
+	bool IsCookByTheBook() const { return GetCookType() == UE::Cook::ECookType::ByTheBook; }
+	bool IsCookOnTheFly() const { return GetCookType() == UE::Cook::ECookType::OnTheFly; }
+	bool IsCookTypeUnknown() const { return GetCookType() == UE::Cook::ECookType::Unknown; }
+	UE::Cook::ECookType GetCookType() const { return Data.CookType; }
+	UE::Cook::ECookingDLC GetCookingDLC() const { return Data.CookingDLC; }
+
+	/** Return which phase of SavePackage callbacks are active. @see EObjectSaveContextPhase. */
+	EObjectSaveContextPhase GetPhase() const { return Data.ObjectSaveContextPhase; }
+
+#if WITH_EDITOR
+	/**
+	 * Add the given FCookDependency to the build dependencies for the package being cook-saved. Iterative cooks will
+	 * invalidate the package and recook it if the CookDependency changes.
+	 * CookBuildDependencies of type FCookDependency::CreatePackage are also used when saving an editor package rather
+	 * than a cooked package. They are saved into the AssetRegistry dependency data for the package, as
+	 * EDependencyProperty::Build properties.
+	 */
+	COREUOBJECT_API void AddCookBuildDependency(UE::Cook::FCookDependency BuildDependency);
+	/**
+	 * Add the given UObject's package as a runtime dependency for the package being cook-saved. It will
+	 * be cooked for the current platform even if it is not otherwise referenced, and even if the object
+	 * being PreSaved does not end up saving into the package's exports for the current platform.
+	 */
+	COREUOBJECT_API void AddCookRuntimeDependency(FSoftObjectPath Dependency);
+	/** Serialize an object to find all packages that it references, and AddCookRuntimeDependency for each one. */
+	COREUOBJECT_API void HarvestCookRuntimeDependencies(UObject* HarvestReferencesFrom);
+
+	/**
+	 * Applicable only to cook saves: True if the SavePackage call should write extra debug data for debugging
+	 * cook determinism or incremental cook issues.
+	 */
+	COREUOBJECT_API bool IsDeterminismDebug() const;
+	/**
+	 * Ignored unless IsDeterminismDebug()=true. An object should call this function to register
+	 * their callback class for adding determinism diagnostics the package save.
+	 */
+	COREUOBJECT_API void RegisterDeterminismHelper(const TRefCountPtr<UE::Cook::IDeterminismHelper>& DeterminismHelper);
+#endif
+
+	/**
+	 * Return whether the package is being saved due to a procedural save.
+	 * Any save without the possibility of user-generated edits to the package is a procedural save (Cooking, EditorDomain).
+	 * This allows us to execute transforms that only need to be executed in response to new user data.
+	 */
+	bool IsProceduralSave() const { return Data.bProceduralSave; }
+
+	/**
+	 * Return whether LoadedPath of the package being saved is being updated.
+	 * This allows us to update the in-memory package when it is saved in editor to match its new save file.
+	 */
+	bool IsUpdatingLoadedPath() const { return Data.bUpdatingLoadedPath; }
+
+	/** Return the save flags (ESaveFlags) of the save. */
+	uint32 GetSaveFlags() const { return Data.SaveFlags; }
+
+	/**
+	 * Always true normally. When a system is executing multiple PreSaves concurrently before a single save,
+	 * will return false for all but the first PreSave.
+	 */
+	bool IsFirstConcurrentSave() const { return Data.bOuterConcurrentSave; }
+
+protected:
+	FObjectSaveContextData& Data;
+};
+
 /** Interface used by PostSave to access the save parameters. */
 class FObjectPostSaveContext
 {
@@ -340,6 +472,9 @@ public:
 	 */
 	bool IsLastConcurrentSave() const { return Data.bOuterConcurrentSave; }
 
+	/** Return which phase of SavePackage callbacks are active. @see EObjectSaveContextPhase. */
+	EObjectSaveContextPhase GetPhase() const { return Data.ObjectSaveContextPhase; }
+
 protected:
 	FObjectSaveContextData& Data;
 	friend class UObject;
diff --git a/Engine/Source/Runtime/Engine/Private/BlueprintGeneratedClass.cpp b/Engine/Source/Runtime/Engine/Private/BlueprintGeneratedClass.cpp
--- a/Engine/Source/Runtime/Engine/Private/BlueprintGeneratedClass.cpp
+++ b/Engine/Source/Runtime/Engine/Private/BlueprintGeneratedClass.cpp
@@ -2561,7 +2561,7 @@ void UBlueprintGeneratedClass::Serialize(FArchive& Ar)
 		// Find all imported packages from the Blueprint and its subobjects and declare them as
 		// used-in-game imports of the cooked package by serializing them as SoftObjectPaths.
 		FImportExportCollector Collector(this->GetPackage());
-		Collector.SetCookData(Ar.GetCookData());
+		Collector.SetSavePackageData(Ar.GetSavePackageData());
 		Collector.AddExportToIgnore(this);
 		Collector.SetFilterEditorOnly(Ar.IsFilterEditorOnly());
 		UObject* ClassGeneratedByPtr = ClassGeneratedBy;
diff --git a/Engine/Source/Runtime/Engine/Private/Materials/MaterialShared.cpp b/Engine/Source/Runtime/Engine/Private/Materials/MaterialShared.cpp
--- a/Engine/Source/Runtime/Engine/Private/Materials/MaterialShared.cpp
+++ b/Engine/Source/Runtime/Engine/Private/Materials/MaterialShared.cpp
@@ -5197,7 +5197,7 @@ FMaterialResourceMemoryWriter::FMaterialResourceMemoryWriter(FArchive& Ar) :
 {
 	check(Ar.IsSaving());
 	this->SetByteSwapping(Ar.IsByteSwapping());
-	this->SetCookData(Ar.GetCookData());
+	this->SetSavePackageData(Ar.GetSavePackageData());
 }
 
 FMaterialResourceMemoryWriter::~FMaterialResourceMemoryWriter()
diff --git a/Engine/Source/Runtime/Engine/Private/ShaderCompiler/ShaderCompiler.cpp b/Engine/Source/Runtime/Engine/Private/ShaderCompiler/ShaderCompiler.cpp
--- a/Engine/Source/Runtime/Engine/Private/ShaderCompiler/ShaderCompiler.cpp
+++ b/Engine/Source/Runtime/Engine/Private/ShaderCompiler/ShaderCompiler.cpp
@@ -55,7 +55,7 @@
 #include "ColorManagement/ColorSpace.h"
 
 #if WITH_EDITOR
-#include "UObject/ArchiveCookContext.h"
+#include "Serialization/ArchiveSavePackageDataBuffer.h"
 #include "UObject/UObjectGlobals.h"
 #include "DerivedDataCache.h"
 #include "DerivedDataRequestOwner.h"
@@ -4755,18 +4755,15 @@ FString SaveGlobalShaderFile(EShaderPlatform Platform, FString SavePath, class I
 	TArray<uint8> GlobalShaderData;
 	{
 #if WITH_EDITOR
-		TOptional<FArchiveCookContext> CookContext;
-		TOptional<FArchiveCookData> CookData;
+		TOptional<FArchiveSavePackageDataBuffer> ArchiveSavePackageData;
 #endif
 		FMemoryWriter MemoryWriter(GlobalShaderData, true);
 
 #if WITH_EDITOR
 		if (TargetPlatform != nullptr)
 		{
-			CookContext.Emplace(nullptr /*InPackage*/, UE::Cook::ECookType::Unknown,
-				UE::Cook::ECookingDLC::Unknown, TargetPlatform);
-			CookData.Emplace(*TargetPlatform, *CookContext);
-			MemoryWriter.SetCookData(CookData.GetPtrOrNull());
+			ArchiveSavePackageData.Emplace(TargetPlatform);
+			MemoryWriter.SetSavePackageData(&ArchiveSavePackageData.GetValue());
 		}
 #endif // WITH_EDITOR
 
diff --git a/Engine/Source/Runtime/Engine/Private/ShaderCompiler/ShaderCompilerEditor.cpp b/Engine/Source/Runtime/Engine/Private/ShaderCompiler/ShaderCompilerEditor.cpp
--- a/Engine/Source/Runtime/Engine/Private/ShaderCompiler/ShaderCompilerEditor.cpp
+++ b/Engine/Source/Runtime/Engine/Private/ShaderCompiler/ShaderCompilerEditor.cpp
@@ -29,7 +29,7 @@
 #include "UObject/UObjectIterator.h"
 #include "UnrealEngine.h"
 
-#include "UObject/ArchiveCookContext.h"
+#include "Serialization/ArchiveSavePackageDataBuffer.h"
 #include "UObject/UObjectGlobals.h"
 #include "DerivedDataCache.h"
 #include "DerivedDataRequestOwner.h"
@@ -237,17 +237,14 @@ static void CompileGlobalShaderMapForRemote(
 	FinishRecompileGlobalShaders();
 
 	// Write the shader compilation info to memory, converting FName to strings
-	TOptional<FArchiveCookContext> CookContext;
-	TOptional<FArchiveCookData> CookData;
+	TOptional<FArchiveSavePackageDataBuffer> ArchiveSavePackageData;
 	FMemoryWriter MemWriter(*OutArray, true);
 	FNameAsStringProxyArchive Ar(MemWriter);
 
 	if (TargetPlatform != nullptr)
 	{
-		CookContext.Emplace(nullptr /*InPackage*/, UE::Cook::ECookType::Unknown,
-			UE::Cook::ECookingDLC::Unknown, TargetPlatform);
-		CookData.Emplace(*TargetPlatform, *CookContext);
-		Ar.SetCookData(CookData.GetPtrOrNull());
+		ArchiveSavePackageData.Emplace(TargetPlatform);
+		Ar.SetSavePackageData(&ArchiveSavePackageData.GetValue());
 	}
 
 	// save out the global shader map to the byte array
@@ -257,17 +254,14 @@ static void CompileGlobalShaderMapForRemote(
 static void SaveShaderMapsForRemote(ITargetPlatform* TargetPlatform, const TMap<FString, TArray<TRefCountPtr<FMaterialShaderMap>>>& CompiledShaderMaps, TArray<uint8>* OutArray)
 {
 	// write the shader compilation info to memory, converting fnames to strings
-	TOptional<FArchiveCookContext> CookContext;
-	TOptional<FArchiveCookData> CookData;
+	TOptional<FArchiveSavePackageDataBuffer> ArchiveSavePackageData;
 	FMemoryWriter MemWriter(*OutArray, true);
 	FNameAsStringProxyArchive Ar(MemWriter);
 
 	if (TargetPlatform != nullptr)
 	{
-		CookContext.Emplace(nullptr /*InPackage*/, UE::Cook::ECookType::Unknown,
-			UE::Cook::ECookingDLC::Unknown, TargetPlatform);
-		CookData.Emplace(*TargetPlatform, *CookContext);
-		Ar.SetCookData(CookData.GetPtrOrNull());
+		ArchiveSavePackageData.Emplace(TargetPlatform);
+		Ar.SetSavePackageData(&ArchiveSavePackageData.GetValue());
 	}
 
 	// save out the shader map to the byte array
diff --git a/Engine/Source/Runtime/Engine/Private/SkeletalMeshLODRenderData.cpp b/Engine/Source/Runtime/Engine/Private/SkeletalMeshLODRenderData.cpp
--- a/Engine/Source/Runtime/Engine/Private/SkeletalMeshLODRenderData.cpp
+++ b/Engine/Source/Runtime/Engine/Private/SkeletalMeshLODRenderData.cpp
@@ -982,7 +982,7 @@ void FSkeletalMeshLODRenderData::Serialize(FArchive& Ar, UObject* Owner, int32 I
 		if (Ar.IsSaving())
 		{
 			FSkeletalMeshLODSizeCounter LODSizeCounter;
-			LODSizeCounter.SetCookData(Ar.GetCookData());
+			LODSizeCounter.SetSavePackageData(Ar.GetSavePackageData());
 			LODSizeCounter.SetByteSwapping(Ar.IsByteSwapping());
 			SerializeStreamedData(LODSizeCounter, OwnerMesh, Idx, ClassDataStripFlags, bNeedsCPUAccess, bForceKeepCPUResources);
 			BuffersSize = LODSizeCounter.TotalSize();
@@ -1010,7 +1010,7 @@ void FSkeletalMeshLODRenderData::Serialize(FArchive& Ar, UObject* Owner, int32 I
 				if (!bDiscardBulkData)
 				{
 					FMemoryWriter MemWriter(TmpBuff, true);
-					MemWriter.SetCookData(Ar.GetCookData());
+					MemWriter.SetSavePackageData(Ar.GetSavePackageData());
 					MemWriter.SetByteSwapping(Ar.IsByteSwapping());
 					SerializeStreamedData(MemWriter, OwnerMesh, Idx, ClassDataStripFlags, bNeedsCPUAccess, bForceKeepCPUResources);
 				}
diff --git a/Engine/Source/Runtime/Engine/Private/StaticMesh.cpp b/Engine/Source/Runtime/Engine/Private/StaticMesh.cpp
--- a/Engine/Source/Runtime/Engine/Private/StaticMesh.cpp
+++ b/Engine/Source/Runtime/Engine/Private/StaticMesh.cpp
@@ -801,7 +801,7 @@ void FStaticMeshLODResources::Serialize(FArchive& Ar, UObject* Owner, int32 Inde
 				if (!bDiscardBulkData)
 				{
 					FMemoryWriter MemWriter(TmpBuff, true);
-					MemWriter.SetCookData(Ar.GetCookData());
+					MemWriter.SetSavePackageData(Ar.GetSavePackageData());
 					MemWriter.SetByteSwapping(Ar.IsByteSwapping());
 					SerializeBuffers(MemWriter, OwnerStaticMesh, ClassDataStripFlags, TmpBuffersSize);
 				}
@@ -1747,7 +1747,7 @@ void FStaticMeshRayTracingProxyLOD::Serialize(FArchive& Ar, UObject* Owner, int3
 
 			{
 				FMemoryWriter MemWriter(TmpBuff, /*bIsPersistent=*/ true);
-				MemWriter.SetCookData(Ar.GetCookData());
+				MemWriter.SetSavePackageData(Ar.GetSavePackageData());
 				MemWriter.SetByteSwapping(Ar.IsByteSwapping());
 
 				BuffersSize = SerializeBuffers(MemWriter, OwnerStaticMesh, 0);
