ArchiveSavePackageDataBuffer: Suppress static analysis warning.
#rnx

diff --git a/Engine/Source/Runtime/CoreUObject/Public/Serialization/ArchiveSavePackageDataBuffer.h b/Engine/Source/Runtime/CoreUObject/Public/Serialization/ArchiveSavePackageDataBuffer.h
--- a/Engine/Source/Runtime/CoreUObject/Public/Serialization/ArchiveSavePackageDataBuffer.h
+++ b/Engine/Source/Runtime/CoreUObject/Public/Serialization/ArchiveSavePackageDataBuffer.h
@@ -21,7 +21,9 @@ struct FArchiveSavePackageDataBuffer : public FArchiveSavePackageData
 {
 public:
 	FArchiveSavePackageDataBuffer(const ITargetPlatform* InTargetPlatform = nullptr)
-		: FArchiveSavePackageData(ObjectSavePackageSerializeContextBuffer, nullptr, nullptr)
+		// warning V1050 : The uninitialized class member 'ObjectSavePackageSerializeContextBuffer' is used when initializing the base class 'FArchiveSavePackageData'.
+		// This warning can be ignored because the base class only records the pointer and does not dereference it
+		: FArchiveSavePackageData(ObjectSavePackageSerializeContextBuffer, nullptr, nullptr) // -V1050
 		, ObjectSaveContextData(FObjectSaveContextData())
 		, ObjectSavePackageSerializeContextBuffer(*ObjectSaveContextData)
 		, CookContextBuffer(nullptr /*InPackage*/, UE::Cook::ECookType::Unknown, UE::Cook::ECookingDLC::Unknown,
