ArchiveSavePackageData: Fix deprecation warning: Replace use of deprecated SetCookData with new SetSavePackageData.
#rnx

diff --git a/Engine/Source/Runtime/CoreUObject/Public/Serialization/ArchiveSavePackageDataBuffer.h b/Engine/Source/Runtime/CoreUObject/Public/Serialization/ArchiveSavePackageDataBuffer.h
--- a/Engine/Source/Runtime/CoreUObject/Public/Serialization/ArchiveSavePackageDataBuffer.h
+++ b/Engine/Source/Runtime/CoreUObject/Public/Serialization/ArchiveSavePackageDataBuffer.h
@@ -20,13 +20,13 @@
 struct FArchiveSavePackageDataBuffer : public FArchiveSavePackageData
 {
 public:
-	FArchiveSavePackageDataBuffer(const ITargetPlatform* InTargetPlatform = nullptr)
+	FArchiveSavePackageDataBuffer(const ITargetPlatform* InTargetPlatform = nullptr, UPackage* Package = nullptr)
 		// warning V1050 : The uninitialized class member 'ObjectSavePackageSerializeContextBuffer' is used when initializing the base class 'FArchiveSavePackageData'.
 		// This warning can be ignored because the base class only records the pointer and does not dereference it
 		: FArchiveSavePackageData(ObjectSavePackageSerializeContextBuffer, nullptr, nullptr) // -V1050
 		, ObjectSaveContextData(FObjectSaveContextData())
 		, ObjectSavePackageSerializeContextBuffer(*ObjectSaveContextData)
-		, CookContextBuffer(nullptr /*InPackage*/, UE::Cook::ECookType::Unknown, UE::Cook::ECookingDLC::Unknown,
+		, CookContextBuffer(Package, UE::Cook::ECookType::Unknown, UE::Cook::ECookingDLC::Unknown,
 			InTargetPlatform)
 	{
 		SetConstructorTargetPlatform(InTargetPlatform);
