#UE Add linker root to warning about scriptstruct advance to identify the asset

diff --git a/Engine/Source/Runtime/CoreUObject/Private/StructUtils/InstancedStruct.cpp b/Engine/Source/Runtime/CoreUObject/Private/StructUtils/InstancedStruct.cpp
--- a/Engine/Source/Runtime/CoreUObject/Private/StructUtils/InstancedStruct.cpp
+++ b/Engine/Source/Runtime/CoreUObject/Private/StructUtils/InstancedStruct.cpp
@@ -164,7 +164,7 @@ bool FInstancedStruct::Serialize(FArchive& Ar)
 			// In this case we manually seek in the archive to skip its serialized content. 
 			// We don't want to rely on TaggedSerialization that will mark an error in the archive that
 			// may cause other serialization to fail (e.g. FArchive& operator<<(FArchive& Ar, TArray& A))
-			UE_LOG(LogCore, Warning, TEXT("Unable to find serialized UScriptStruct -> Advance %u bytes in the archive and reset to empty FInstancedStruct. SerializedProperty:%s"), SerialSize, *Ar.GetSerializedProperty()->GetPathName());
+			UE_LOG(LogCore, Warning, TEXT("Unable to find serialized UScriptStruct -> Advance %u bytes in the archive and reset to empty FInstancedStruct. SerializedProperty:%s LinkerRoot:%s"), SerialSize, *Ar.GetSerializedProperty()->GetPathName(), Ar.GetLinker() ? *GetPathNameSafe(Ar.GetLinker()->LinkerRoot) : TEXT("NoLinker"));
 			Ar.Seek(Ar.Tell() + SerialSize);
 		}
 		else if (NonConstStruct != nullptr && ensureMsgf(GetMutableMemory() != nullptr, TEXT("A valid script struct should always have allocated memory")))
diff --git a/Engine/Source/Runtime/CoreUObject/Private/StructUtils/InstancedStructContainer.cpp b/Engine/Source/Runtime/CoreUObject/Private/StructUtils/InstancedStructContainer.cpp
--- a/Engine/Source/Runtime/CoreUObject/Private/StructUtils/InstancedStructContainer.cpp
+++ b/Engine/Source/Runtime/CoreUObject/Private/StructUtils/InstancedStructContainer.cpp
@@ -568,7 +568,7 @@ bool FInstancedStructContainer::Serialize(FArchive& Ar)
 					// In this case we manually seek in the archive to skip its serialized content. 
 					// We don't want to rely on TaggedSerialization that will mark an error in the archive that
 					// may cause other serialization to fail (e.g. FArchive& operator<<(FArchive& Ar, TArray& A))
-					UE_LOG(LogCore, Warning, TEXT("Unable to find serialized UScriptStruct -> Advance %u bytes in the archive and reset to empty FInstancedStructContainer"), SerialSize);
+					UE_LOG(LogCore, Warning, TEXT("Unable to find serialized UScriptStruct -> Advance %u bytes in the archive and reset to empty FInstancedStructContainer. LinkerRoot:%s"), SerialSize, Ar.GetLinker() ? *GetPathNameSafe(Ar.GetLinker()->LinkerRoot) : TEXT("NoLinker"));
 					Ar.Seek(Ar.Tell() + SerialSize);
 				}
 				else if (NonConstStruct != nullptr)
