#UE Nonunity fix

diff --git a/Engine/Source/Runtime/CoreUObject/Private/StructUtils/InstancedStruct.cpp b/Engine/Source/Runtime/CoreUObject/Private/StructUtils/InstancedStruct.cpp
--- a/Engine/Source/Runtime/CoreUObject/Private/StructUtils/InstancedStruct.cpp
+++ b/Engine/Source/Runtime/CoreUObject/Private/StructUtils/InstancedStruct.cpp
@@ -2,6 +2,7 @@
 #include "StructUtils/InstancedStruct.h"
 #include "StructUtils/StructView.h"
 #include "Serialization/CustomVersion.h"
+#include "UObject/Linker.h"
 
 #if WITH_EDITOR
 #include "Serialization/MemoryReader.h"
diff --git a/Engine/Source/Runtime/CoreUObject/Private/StructUtils/InstancedStructContainer.cpp b/Engine/Source/Runtime/CoreUObject/Private/StructUtils/InstancedStructContainer.cpp
--- a/Engine/Source/Runtime/CoreUObject/Private/StructUtils/InstancedStructContainer.cpp
+++ b/Engine/Source/Runtime/CoreUObject/Private/StructUtils/InstancedStructContainer.cpp
@@ -1,6 +1,7 @@
 // Copyright Epic Games, Inc. All Rights Reserved.
 
 #include "StructUtils/InstancedStructContainer.h"
+#include "UObject/Linker.h"
 
 #if WITH_EDITOR
 #include "Serialization/MemoryReader.h"
