FBulkMetaData::FromSerialized now overwrites all fields in its output object, to ensure any stale state is cleared.

Also updated some checks in that function to ensure that zero size bulk data is loaded correctly.

This fixes a bug where loading an FBulkData with a zero element count wouldn't update the Size field, so if the same FBulkData had previously been used to load data with a non-zero size, the non-zero value for Size would incorrectly persist after loading the zero size data. I've added a regression test to the existing suite to prove that this case is now handled correctly.

#tests Ran FoundationTests and went through bug repro in editor

diff --git a/Engine/Source/Runtime/CoreUObject/Private/Serialization/BulkData.cpp b/Engine/Source/Runtime/CoreUObject/Private/Serialization/BulkData.cpp
--- a/Engine/Source/Runtime/CoreUObject/Private/Serialization/BulkData.cpp
+++ b/Engine/Source/Runtime/CoreUObject/Private/Serialization/BulkData.cpp
@@ -271,9 +271,16 @@ FString GetDebugNameFromArchive(FArchive& Ar)
 
 bool FBulkMetaData::FromSerialized(FArchive& Ar, int64 ElementSize, FBulkMetaData& OutMetaData, int64& OutDuplicateOffset)
 {
+	check(Ar.IsLoading());
+
+	// Ensure all fields are initialized.
+	//
+	// Note that this will reset the lock status as well, so callers may want to first check that 
+	// the bulk data is unlocked before the lock state is lost.
+	OutMetaData = FBulkMetaData();
+
 	if (Ar.IsError())
 	{
-		OutMetaData = FBulkMetaData();
 		return false;
 	}
 
@@ -285,14 +292,12 @@ bool FBulkMetaData::FromSerialized(FArchive& Ar, int64 ElementSize, FBulkMetaDat
 		// Note that setting the error flag on the archive is not enough to stop the package from being loaded so for now we 
 		// need to fatal error to prevent the process from continuing to use the corrupted package.
 		UE_LOG(LogSerialization, Fatal, TEXT("Bulkdata error when serializing '%s', could not serialize FBulkMetaResource correctly"), *GetDebugNameFromArchive(Ar));
-		OutMetaData = FBulkMetaData();
 		return false;
 	}
 
 	if (Resource.ElementCount > 0)
 	{
-		// TODO: This would be a good use case for FGuardedInt64 once it is moved to core
-		FGuardedInt64 MetadataSize = FGuardedInt64(Resource.ElementCount) * ElementSize;
+		const FGuardedInt64 MetadataSize = FGuardedInt64(Resource.ElementCount) * ElementSize;
 		if (MetadataSize.IsValid())
 		{
 			OutMetaData.SetSize(MetadataSize.Get(0));
@@ -309,8 +314,7 @@ bool FBulkMetaData::FromSerialized(FArchive& Ar, int64 ElementSize, FBulkMetaDat
 				*GetDebugNameFromArchive(Ar),
 				Resource.ElementCount,
 				ElementSize);
-
-			OutMetaData = FBulkMetaData();
+			
 			return false;
 		}
 	}
@@ -319,12 +323,12 @@ bool FBulkMetaData::FromSerialized(FArchive& Ar, int64 ElementSize, FBulkMetaDat
 	OutMetaData.SetOffset(Resource.Offset);
 	OutMetaData.SetFlags(Resource.Flags);
 
-	check(Resource.ElementCount <= 0 || OutMetaData.GetSize() == Resource.ElementCount * ElementSize);
+	check(Resource.ElementCount < 0 || OutMetaData.GetSize() == Resource.ElementCount * ElementSize);
 	check(OutMetaData.GetOffset() == Resource.Offset);
 	check(OutMetaData.GetFlags() == Resource.Flags);
 
 #if !USE_RUNTIME_BULKDATA
-	check(Resource.ElementCount <= 0 || OutMetaData.GetSizeOnDisk() == Resource.SizeOnDisk);
+	check(Resource.ElementCount < 0 || OutMetaData.GetSizeOnDisk() == Resource.SizeOnDisk);
 #endif
 
 	OutDuplicateOffset = Resource.DuplicateOffset;
diff --git a/Engine/Source/Runtime/CoreUObject/Public/Serialization/BulkData.h b/Engine/Source/Runtime/CoreUObject/Public/Serialization/BulkData.h
--- a/Engine/Source/Runtime/CoreUObject/Public/Serialization/BulkData.h
+++ b/Engine/Source/Runtime/CoreUObject/Public/Serialization/BulkData.h
@@ -416,6 +416,8 @@ public:
 	/** 
 	 * Serializes FBulkMetaResource from the given archive and builds the returned FBulkMetaData from it. 
 	 * The offset for duplicated data will also be returned
+	 * 
+	 * All fields of OutMetaData will be overwritten, including lock status.
 	 */
 	COREUOBJECT_API static bool FromSerialized(FArchive& Ar, int64 ElementSize, FBulkMetaData& OutMetaData, int64& OutDuplicateOffset);
 
diff --git a/Engine/Source/Runtime/CoreUObject/Tests/Serialization/BulkDataTests.cpp b/Engine/Source/Runtime/CoreUObject/Tests/Serialization/BulkDataTests.cpp
--- a/Engine/Source/Runtime/CoreUObject/Tests/Serialization/BulkDataTests.cpp
+++ b/Engine/Source/Runtime/CoreUObject/Tests/Serialization/BulkDataTests.cpp
@@ -380,6 +380,72 @@ TEST_CASE("CoreUObject::Serialization::FBulkData::Serialize", "[CoreUObject][Ser
 			CHECK(BulkData.GetBulkDataSize() == UncompressedPayloadSize );
 		}
 	}
+
+	SECTION("Reuse FBulkData for serialization")
+	{
+		// Zero size was chosen for the second payload, because there was a bug where the bulk data
+		// state wouldn't be initialized correctly if the payload was empty, leaving stale state
+		// from a previous serialization.
+		//
+		// This acts as a regression test for that bug.
+		const int64 FullPayloadSize = 128;
+		const int64 ZeroPayloadSize = 0;
+
+		FUniqueBuffer Payload = CreatePayload(FullPayloadSize);
+
+		FLargeMemoryWriter ArFullPayload;
+		FLargeMemoryWriter ArZeroPayload;
+		ArFullPayload.SetIsPersistent(true);
+		ArZeroPayload.SetIsPersistent(true);
+
+		// Write two different payloads with the same FBulkData
+		{
+			FBulkData BulkData;
+			CopyPayload(BulkData, Payload);
+			BulkData.Serialize(ArFullPayload, nullptr, false, 1, EFileRegionType::None);
+			CHECK(ArFullPayload.TotalSize() > BulkData.GetBulkDataSize()); // Bulk meta data + payload
+
+			// Remove the payload and zero ElementCount
+			BulkData.RemoveBulkData();
+
+			BulkData.Serialize(ArZeroPayload, nullptr, false, 1, EFileRegionType::None);
+			CHECK(ArZeroPayload.TotalSize() > BulkData.GetBulkDataSize()); // Bulk meta data + payload
+		}
+
+		// Read the two payloads with the same FBulkData
+		{
+			FBulkData BulkData;
+
+			{
+				FMemoryReaderView ReaderArFullPayload(ArFullPayload.GetView());
+				ReaderArFullPayload.SetIsPersistent(true);
+
+				BulkData.Serialize(ReaderArFullPayload, nullptr, false, 1, EFileRegionType::None);
+				FMemoryView PayloadView(BulkData.LockReadOnly(), BulkData.GetBulkDataSize());
+
+				CHECK(BulkData.GetBulkDataSize() == FullPayloadSize);
+				CHECK(BulkData.GetBulkDataOffsetInFile() == ReaderArFullPayload.Tell() - FullPayloadSize);
+				CHECK(BulkData.GetBulkDataFlags() == 0);
+				CHECK(BulkData.IsInlined());
+				CHECK(PayloadView.EqualBytes(Payload.GetView()));
+
+				BulkData.Unlock();
+			}
+
+			{
+				FMemoryReaderView ReaderArZeroPayload(ArZeroPayload.GetView());
+				ReaderArZeroPayload.SetIsPersistent(true);
+
+				BulkData.Serialize(ReaderArZeroPayload, nullptr, false, 1, EFileRegionType::None);
+
+				// Should be able to access these without locking
+				CHECK(BulkData.GetBulkDataSize() == ZeroPayloadSize);
+				CHECK(BulkData.GetBulkDataOffsetInFile() == ReaderArZeroPayload.Tell() - ZeroPayloadSize);
+				CHECK(BulkData.GetBulkDataFlags() == 0);
+				CHECK(BulkData.IsInlined());
+			}
+		}
+	}
 }
 
 #endif // WITH_EDITORONLY_DATA
