Update the scoped IDO serialization context to disable unnecessary features too

It previously only enabled features, which would mean leaving impersonation and other features and their overhead enabled for objects that have do not have IDO support enabled.

#rnx

diff --git a/Engine/Source/Runtime/CoreUObject/Internal/UObject/PropertyBagRepository.h b/Engine/Source/Runtime/CoreUObject/Internal/UObject/PropertyBagRepository.h
--- a/Engine/Source/Runtime/CoreUObject/Internal/UObject/PropertyBagRepository.h
+++ b/Engine/Source/Runtime/CoreUObject/Internal/UObject/PropertyBagRepository.h
@@ -224,27 +224,29 @@ struct FScopedIDOSerializationContext
 #if WITH_EDITORONLY_DATA
 	COREUOBJECT_API FScopedIDOSerializationContext(UObject* InObject, FArchive& Archive);
 	COREUOBJECT_API FScopedIDOSerializationContext(UObject* InObject, bool bImpersonate); // assumes save
-	COREUOBJECT_API FScopedIDOSerializationContext(bool bImpersonate); // assumes save
+	COREUOBJECT_API explicit FScopedIDOSerializationContext(bool bImpersonate); // assumes save
 	COREUOBJECT_API ~FScopedIDOSerializationContext();
 
-	bool bCreateIDO = false;
 	FArchive* Archive = nullptr;
 	UObject* const Object = nullptr;
-	const int64 PreSerializeOffset;
-	TOptional<TGuardValue<bool>> ScopedTrackSerializedPropertyPath;
-	TOptional<TGuardValue<bool>> ScopedSerializeUnknownProperties;
-	TOptional<TGuardValue<bool>> ScopedSerializeUnknownEnumNames;
-	TOptional<TGuardValue<bool>> ScopedImpersonateProperties;
-	TOptional<TGuardValue<bool>> ScopedTrackInitializedProperties;
-	TOptional<TGuardValue<bool>> ScopedTrackSerializedProperties;
-	TOptional<TGuardValue<UObject*>> ScopedSerializedObject;
+	const int64 PreSerializeOffset = 0;
+	UObject* SavedSerializedObject;
+	bool bSavedTrackSerializedPropertyPath;
+	bool bSavedTrackInitializedProperties;
+	bool bSavedTrackSerializedProperties;
+	bool bSavedTrackUnknownProperties;
+	bool bSavedTrackUnknownEnumNames;
+	bool bSavedImpersonateProperties;
+	bool bCreateIDO = false;
 
 private:
+	void SaveSerializeContext(FUObjectSerializeContext* SerializeContext);
+	void RestoreSerializeContext(FUObjectSerializeContext* SerializeContext) const;
 	// if we're loading and IDO should be created, this will be called when the context falls out of scope
 	void FinishCreatingInstanceDataObject() const;
 #else
 	inline FScopedIDOSerializationContext(UObject* InObject, FArchive& Archive) {}
-	inline explicit FScopedIDOSerializationContext(UObject* InObject, bool bImpersonate) {}
+	inline FScopedIDOSerializationContext(UObject* InObject, bool bImpersonate) {}
 	inline explicit FScopedIDOSerializationContext(bool bImpersonate) {}
 #endif
 
diff --git a/Engine/Source/Runtime/CoreUObject/Private/UObject/PropertyBagRepository.cpp b/Engine/Source/Runtime/CoreUObject/Private/UObject/PropertyBagRepository.cpp
--- a/Engine/Source/Runtime/CoreUObject/Private/UObject/PropertyBagRepository.cpp
+++ b/Engine/Source/Runtime/CoreUObject/Private/UObject/PropertyBagRepository.cpp
@@ -936,60 +936,47 @@ FScopedIDOSerializationContext::FScopedIDOSerializationContext(UObject* InObject
 	, PreSerializeOffset(InArchive.Tell())
 {
 	FUObjectSerializeContext* SerializeContext = FUObjectThreadContext::Get().GetSerializeContext();
-	bool bHasIDOSupport = IsInstanceDataObjectSupportEnabled(Object);
+	SaveSerializeContext(SerializeContext);
 
-	bool bHasReinstancedClass = InObject->GetClass()->HasAnyClassFlags(CLASS_NewerVersionExists);
-	bCreateIDO = bHasIDOSupport && !SerializeContext->bImpersonateProperties && Archive->IsLoading() && !bHasReinstancedClass;
+	const bool bHasIDOSupport = IsInstanceDataObjectSupportEnabled(Object);
 
-	if (bHasIDOSupport)
+	if (Archive->IsLoading())
 	{
-		if (Archive->IsLoading())
-		{
-			// Enable creation of a property path name tree to track any property that does not match the current class schema,
-			// except when impersonation is enabled because that implies we are deserializing an IDO.
-			ScopedTrackSerializedPropertyPath.Emplace(SerializeContext->bTrackSerializedPropertyPath, bCreateIDO);
-			ScopedSerializeUnknownProperties.Emplace(SerializeContext->bTrackUnknownProperties, bCreateIDO);
-			ScopedSerializeUnknownEnumNames.Emplace(SerializeContext->bTrackUnknownEnumNames, bCreateIDO);
-			ScopedSerializedObject.Emplace(SerializeContext->SerializedObject, Object);
-
-			// Enable tracking of initialized properties when loading an IDO, which is implied by impersonation being enabled.
-			const bool bLoadingIDO = bHasIDOSupport && SerializeContext->bImpersonateProperties;
-			ScopedTrackInitializedProperties.Emplace(SerializeContext->bTrackInitializedProperties, bLoadingIDO);
-			ScopedTrackSerializedProperties.Emplace(SerializeContext->bTrackSerializedProperties, bLoadingIDO);
-		}
-		else
-		{
-			ScopedImpersonateProperties.Emplace(SerializeContext->bImpersonateProperties, bHasIDOSupport);
-		}
+		bCreateIDO = bHasIDOSupport && !SerializeContext->bImpersonateProperties &&
+			!InObject->GetClass()->HasAnyClassFlags(CLASS_NewerVersionExists);
+
+		// Enable creation of a property path name tree to track any property that does not match the current class schema,
+		// except when impersonation is enabled because that implies we are deserializing an IDO.
+		SerializeContext->bTrackSerializedPropertyPath = bCreateIDO;
+		SerializeContext->bTrackUnknownProperties = bCreateIDO;
+		SerializeContext->bTrackUnknownEnumNames = bCreateIDO;
+		SerializeContext->SerializedObject = Object;
+
+		// Enable tracking of initialized properties when loading an IDO, which is implied by impersonation being enabled.
+		const bool bLoadingIDO = bHasIDOSupport && SerializeContext->bImpersonateProperties;
+		SerializeContext->bTrackInitializedProperties = bLoadingIDO;
+		SerializeContext->bTrackSerializedProperties = bLoadingIDO;
+	}
+	else
+	{
+		SerializeContext->bImpersonateProperties = bHasIDOSupport;
 	}
 }
 
 FScopedIDOSerializationContext::FScopedIDOSerializationContext(UObject* InObject, bool bImpersonate)
-	: bCreateIDO(false)
-	, Archive(nullptr)
-	, Object(InObject)
-	, PreSerializeOffset(0)
 {
 	FUObjectSerializeContext* SerializeContext = FUObjectThreadContext::Get().GetSerializeContext();
-	bool bHasIDOSupport = IsInstanceDataObjectSupportEnabled(Object);
-	if (bHasIDOSupport)
-	{
-		ScopedImpersonateProperties.Emplace(SerializeContext->bImpersonateProperties, bImpersonate);
-	}
+	SaveSerializeContext(SerializeContext);
+
+	SerializeContext->bImpersonateProperties = bImpersonate && IsInstanceDataObjectSupportEnabled(Object);
 }
 
 FScopedIDOSerializationContext::FScopedIDOSerializationContext(bool bImpersonate)
-	: bCreateIDO(false)
-	, Archive(nullptr)
-	, Object(nullptr)
-	, PreSerializeOffset(0)
 {
 	FUObjectSerializeContext* SerializeContext = FUObjectThreadContext::Get().GetSerializeContext();
-	bool bHasIDOSupport = IsInstanceDataObjectSupportEnabled();
-	if (bHasIDOSupport)
-	{
-		ScopedImpersonateProperties.Emplace(SerializeContext->bImpersonateProperties, bImpersonate);
-	}
+	SaveSerializeContext(SerializeContext);
+
+	SerializeContext->bImpersonateProperties = bImpersonate && IsInstanceDataObjectSupportEnabled();
 }
 
 FScopedIDOSerializationContext::~FScopedIDOSerializationContext()
@@ -998,6 +985,30 @@ FScopedIDOSerializationContext::~FScopedIDOSerializationContext()
 	{
 		FinishCreatingInstanceDataObject();
 	}
+
+	RestoreSerializeContext(FUObjectThreadContext::Get().GetSerializeContext());
+}
+
+void FScopedIDOSerializationContext::SaveSerializeContext(FUObjectSerializeContext* SerializeContext)
+{
+	SavedSerializedObject = SerializeContext->SerializedObject;
+	bSavedTrackSerializedPropertyPath = SerializeContext->bTrackSerializedPropertyPath;
+	bSavedTrackInitializedProperties = SerializeContext->bTrackInitializedProperties;
+	bSavedTrackSerializedProperties = SerializeContext->bTrackSerializedProperties;
+	bSavedTrackUnknownProperties = SerializeContext->bTrackUnknownProperties;
+	bSavedTrackUnknownEnumNames = SerializeContext->bTrackUnknownEnumNames;
+	bSavedImpersonateProperties = SerializeContext->bImpersonateProperties;
+}
+
+void FScopedIDOSerializationContext::RestoreSerializeContext(FUObjectSerializeContext* SerializeContext) const
+{
+	SerializeContext->SerializedObject = SavedSerializedObject;
+	SerializeContext->bTrackSerializedPropertyPath = bSavedTrackSerializedPropertyPath;
+	SerializeContext->bTrackInitializedProperties = bSavedTrackInitializedProperties;
+	SerializeContext->bTrackSerializedProperties = bSavedTrackSerializedProperties;
+	SerializeContext->bTrackUnknownProperties = bSavedTrackUnknownProperties;
+	SerializeContext->bTrackUnknownEnumNames = bSavedTrackUnknownEnumNames;
+	SerializeContext->bImpersonateProperties = bSavedImpersonateProperties;
 }
 
 void FScopedIDOSerializationContext::FinishCreatingInstanceDataObject() const
