Changed some check(false)s to unimplemented().

diff --git a/Engine/Source/Runtime/CoreUObject/Private/UObject/PropertyBaseObject.cpp b/Engine/Source/Runtime/CoreUObject/Private/UObject/PropertyBaseObject.cpp
--- a/Engine/Source/Runtime/CoreUObject/Private/UObject/PropertyBaseObject.cpp
+++ b/Engine/Source/Runtime/CoreUObject/Private/UObject/PropertyBaseObject.cpp
@@ -602,7 +602,7 @@ FName FObjectPropertyBase::GetID() const
 
 TObjectPtr<UObject> FObjectPropertyBase::GetObjectPtrPropertyValue(const void* PropertyValueAddress) const
 {
-	checkf(false, TEXT("%s is missing implementation of GetObjectPtrPropertyValue"), *GetFullName());
+	unimplemented(); // needs to be implemented by the derived class
 	return TObjectPtr<UObject>();
 }
 
@@ -613,24 +613,24 @@ void FObjectPropertyBase::SetObjectPtrPropertyValue(void* PropertyValueAddress,
 
 UObject* FObjectPropertyBase::GetObjectPropertyValue(const void* PropertyValueAddress) const
 {
-	checkf(false, TEXT("%s is missing implementation of GetObjectPropertyValue"), *GetFullName());
+	unimplemented(); // needs to be implemented by the derived class
 	return nullptr;
 }
 
 UObject* FObjectPropertyBase::GetObjectPropertyValue_InContainer(const void* ContainerAddress, int32 ArrayIndex) const
 {
-	checkf(false, TEXT("%s is missing implementation of GetObjectPropertyValue_InContainer"), *GetFullName());
+	unimplemented(); // needs to be implemented by the derived class
 	return nullptr;
 }
 
 void FObjectPropertyBase::SetObjectPropertyValue_InContainer(void* ContainerAddress, UObject* Value, int32 ArrayIndex) const
 {
-	checkf(false, TEXT("%s is missing implementation of SetObjectPropertyValue_InContainer"), *GetFullName());
+	unimplemented(); // needs to be implemented by the derived class
 }
 
 void FObjectPropertyBase::SetObjectPropertyValue(void* PropertyValueAddress, UObject* Value) const
 {
-	check(0);
+	unimplemented(); // needs to be implemented by the derived class
 }
 
 bool FObjectPropertyBase::AllowCrossLevel() const
