Added some specific overloads for TObjectPtr.

diff --git a/Engine/Source/Runtime/CoreUObject/Private/UObject/WeakObjectPtr.cpp b/Engine/Source/Runtime/CoreUObject/Private/UObject/WeakObjectPtr.cpp
--- a/Engine/Source/Runtime/CoreUObject/Private/UObject/WeakObjectPtr.cpp
+++ b/Engine/Source/Runtime/CoreUObject/Private/UObject/WeakObjectPtr.cpp
@@ -7,6 +7,7 @@
 #include "UObject/WeakObjectPtr.h"
 #include "UObject/GarbageCollection.h"
 #include "UObject/Object.h"
+#include "UObject/ObjectPtr.h"
 
 DEFINE_LOG_CATEGORY_STATIC(LogWeakObjectPtr, Log, All);
 
@@ -36,6 +37,10 @@ void FWeakObjectPtr::operator=(const class UObject *Object)
 		Reset();
 	}
 }
+void FWeakObjectPtr::operator=(TObjectPtr<UObject> Object)
+{
+	*this = Object.Get();
+}
 
 bool FWeakObjectPtr::IsValid(bool bEvenIfGarbage, bool bThreadsafeTest) const
 {
diff --git a/Engine/Source/Runtime/CoreUObject/Public/UObject/SoftObjectPtr.h b/Engine/Source/Runtime/CoreUObject/Public/UObject/SoftObjectPtr.h
--- a/Engine/Source/Runtime/CoreUObject/Public/UObject/SoftObjectPtr.h
+++ b/Engine/Source/Runtime/CoreUObject/Public/UObject/SoftObjectPtr.h
@@ -61,6 +61,11 @@ public:
 		(*this)=Object;
 	}
 
+	explicit FORCEINLINE FSoftObjectPtr(TObjectPtr<UObject> Object)
+	{
+		(*this)=Object;
+	}
+
 	/** Synchronously load (if necessary) and return the asset object represented by this asset ptr */
 	UObject* LoadSynchronous() const
 	{
@@ -114,7 +119,14 @@ public:
 	}
 #endif
 
+	// Implicit conversion from UObject* via assignment shouldn't really be allowed if the constructor is explicit
 	using TPersistentObjectPtr<FSoftObjectPath>::operator=;
+	FSoftObjectPtr& operator=(TObjectPtr<UObject> Ptr)
+	{
+		// It should be possible to do this without resolving
+		(*this) = Ptr.Get();
+		return *this;
+	}
 };
 
 template <> struct TIsPODType<FSoftObjectPtr> { enum { Value = TIsPODType<TPersistentObjectPtr<FSoftObjectPath> >::Value }; };
diff --git a/Engine/Source/Runtime/CoreUObject/Public/UObject/WeakObjectPtr.h b/Engine/Source/Runtime/CoreUObject/Public/UObject/WeakObjectPtr.h
--- a/Engine/Source/Runtime/CoreUObject/Public/UObject/WeakObjectPtr.h
+++ b/Engine/Source/Runtime/CoreUObject/Public/UObject/WeakObjectPtr.h
@@ -107,6 +107,7 @@ public:
 	 * @param Object object to create a weak pointer to
 	 */
 	COREUOBJECT_API void operator=(const class UObject* Object);
+	COREUOBJECT_API void operator=(TObjectPtr<UObject> Object);
 
 	/**  
 	 * Construct from another weak pointer
