Fixed FObjectPropertyBase::SetObjectPtr overloads, which were only partially implemented.

diff --git a/Engine/Source/Runtime/CoreUObject/Private/UObject/PropertyBaseObject.cpp b/Engine/Source/Runtime/CoreUObject/Private/UObject/PropertyBaseObject.cpp
--- a/Engine/Source/Runtime/CoreUObject/Private/UObject/PropertyBaseObject.cpp
+++ b/Engine/Source/Runtime/CoreUObject/Private/UObject/PropertyBaseObject.cpp
@@ -429,7 +429,7 @@ const TCHAR* FObjectPropertyBase::ImportText_Internal( const TCHAR* InBuffer, vo
 
 	if (PropertyPointerType == EPropertyPointerType::Container && HasSetter())
 	{
-		SetObjectPropertyValue_InContainer(ContainerOrPropertyPtr, Result); //TODO change this to not resolve TObjectPtr's
+		SetObjectPtrPropertyValue_InContainer(ContainerOrPropertyPtr, Result);
 	}
 	else
 	{
@@ -608,7 +608,12 @@ TObjectPtr<UObject> FObjectPropertyBase::GetObjectPtrPropertyValue(const void* P
 
 void FObjectPropertyBase::SetObjectPtrPropertyValue(void* PropertyValueAddress, TObjectPtr<UObject> Ptr) const
 {
-	SetObjectPropertyValue(PropertyValueAddress, Ptr.Get());
+	unimplemented(); // needs to be implemented by the derived class
+}
+
+void FObjectPropertyBase::SetObjectPtrPropertyValue_InContainer(void* ContainerAddress, TObjectPtr<UObject> Ptr, int32 ArrayIndex) const
+{
+	unimplemented(); // needs to be implemented by the derived class
 }
 
 UObject* FObjectPropertyBase::GetObjectPropertyValue(const void* PropertyValueAddress) const
diff --git a/Engine/Source/Runtime/CoreUObject/Private/UObject/PropertyLazyObjectPtr.cpp b/Engine/Source/Runtime/CoreUObject/Private/UObject/PropertyLazyObjectPtr.cpp
--- a/Engine/Source/Runtime/CoreUObject/Private/UObject/PropertyLazyObjectPtr.cpp
+++ b/Engine/Source/Runtime/CoreUObject/Private/UObject/PropertyLazyObjectPtr.cpp
@@ -146,6 +146,18 @@ void FLazyObjectProperty::SetObjectPropertyValue(void* PropertyValueAddress, UOb
 	}
 }
 
+void FLazyObjectProperty::SetObjectPtrPropertyValue(void* PropertyValueAddress, TObjectPtr<UObject> Ptr) const
+{
+	if (Ptr || !HasAnyPropertyFlags(CPF_NonNullable))
+	{
+		SetPropertyValue(PropertyValueAddress, TCppType(Ptr));
+	}
+	else
+	{
+		UE_LOG(LogProperty, Verbose /*Warning*/, TEXT("Trying to assign null object value to non-nullable \"%s\""), *GetFullName());
+	}
+}
+
 void FLazyObjectProperty::SetObjectPropertyValue_InContainer(void* ContainerAddress, UObject* Value, int32 ArrayIndex) const
 {
 	if (Value || !HasAnyPropertyFlags(CPF_NonNullable))
@@ -158,6 +170,18 @@ void FLazyObjectProperty::SetObjectPropertyValue_InContainer(void* ContainerAddr
 	}
 }
 
+void FLazyObjectProperty::SetObjectPtrPropertyValue_InContainer(void* ContainerAddress, TObjectPtr<UObject> Ptr, int32 ArrayIndex) const
+{
+	if (Ptr || !HasAnyPropertyFlags(CPF_NonNullable))
+	{
+		SetWrappedUObjectPtrValues<FLazyObjectPtr>(ContainerAddress, EPropertyMemoryAccess::InContainer, &Ptr, ArrayIndex, 1);
+	}
+	else
+	{
+		UE_LOG(LogProperty, Verbose /*Warning*/, TEXT("Trying to assign null object value to non-nullable \"%s\""), *GetFullName());
+	}
+}
+
 bool FLazyObjectProperty::AllowCrossLevel() const
 {
 	return true;
diff --git a/Engine/Source/Runtime/CoreUObject/Private/UObject/PropertyObject.cpp b/Engine/Source/Runtime/CoreUObject/Private/UObject/PropertyObject.cpp
--- a/Engine/Source/Runtime/CoreUObject/Private/UObject/PropertyObject.cpp
+++ b/Engine/Source/Runtime/CoreUObject/Private/UObject/PropertyObject.cpp
@@ -348,11 +348,11 @@ UObject* FObjectProperty::GetObjectPropertyValue_InContainer(const void* Contain
 	return Result;
 }
 
-void FObjectProperty::SetObjectPtrPropertyValue(void* PropertyValueAddress, TObjectPtr<UObject> Value) const
+void FObjectProperty::SetObjectPtrPropertyValue(void* PropertyValueAddress, TObjectPtr<UObject> Ptr) const
 {
-	if (Value || !HasAnyPropertyFlags(CPF_NonNullable))
+	if (Ptr || !HasAnyPropertyFlags(CPF_NonNullable))
 	{
-		SetPropertyValue(PropertyValueAddress, Value);
+		SetPropertyValue(PropertyValueAddress, Ptr);
 	}
 	else
 	{
@@ -372,6 +372,18 @@ void FObjectProperty::SetObjectPropertyValue(void* PropertyValueAddress, UObject
 	}
 }
 
+void FObjectProperty::SetObjectPtrPropertyValue_InContainer(void* ContainerAddress, TObjectPtr<UObject> Ptr, int32 ArrayIndex) const
+{
+	if (Ptr || !HasAnyPropertyFlags(CPF_NonNullable))
+	{
+		SetWrappedUObjectPtrValues<FObjectPtr>(ContainerAddress, EPropertyMemoryAccess::InContainer, &Ptr, ArrayIndex, 1);
+	}
+	else
+	{
+		UE_LOG(LogProperty, Verbose /*Warning*/, TEXT("Trying to assign null object value to non-nullable \"%s\""), *GetFullName());
+	}
+}
+
 void FObjectProperty::SetObjectPropertyValue_InContainer(void* ContainerAddress, UObject* Value, int32 ArrayIndex) const
 {
 	if (Value || !HasAnyPropertyFlags(CPF_NonNullable))
diff --git a/Engine/Source/Runtime/CoreUObject/Private/UObject/PropertySoftObjectPtr.cpp b/Engine/Source/Runtime/CoreUObject/Private/UObject/PropertySoftObjectPtr.cpp
--- a/Engine/Source/Runtime/CoreUObject/Private/UObject/PropertySoftObjectPtr.cpp
+++ b/Engine/Source/Runtime/CoreUObject/Private/UObject/PropertySoftObjectPtr.cpp
@@ -301,11 +301,21 @@ void FSoftObjectProperty::SetObjectPropertyValue(void* PropertyValueAddress, UOb
 	SetPropertyValue(PropertyValueAddress, TCppType(Value));
 }
 
+void FSoftObjectProperty::SetObjectPtrPropertyValue(void* PropertyValueAddress, TObjectPtr<UObject> Ptr) const
+{
+	SetPropertyValue(PropertyValueAddress, TCppType(Ptr));
+}
+
 void FSoftObjectProperty::SetObjectPropertyValue_InContainer(void* ContainerAddress, UObject* Value, int32 ArrayIndex) const
 {
 	SetWrappedUObjectPtrValues<FSoftObjectPtr>(ContainerAddress, EPropertyMemoryAccess::InContainer, &Value, ArrayIndex, 1);
 }
 
+void FSoftObjectProperty::SetObjectPtrPropertyValue_InContainer(void* ContainerAddress, TObjectPtr<UObject> Ptr, int32 ArrayIndex) const
+{
+	SetWrappedUObjectPtrValues<FSoftObjectPtr>(ContainerAddress, EPropertyMemoryAccess::InContainer, &Ptr, ArrayIndex, 1);
+}
+
 bool FSoftObjectProperty::AllowCrossLevel() const
 {
 	return true;
diff --git a/Engine/Source/Runtime/CoreUObject/Private/UObject/PropertyWeakObjectPtr.cpp b/Engine/Source/Runtime/CoreUObject/Private/UObject/PropertyWeakObjectPtr.cpp
--- a/Engine/Source/Runtime/CoreUObject/Private/UObject/PropertyWeakObjectPtr.cpp
+++ b/Engine/Source/Runtime/CoreUObject/Private/UObject/PropertyWeakObjectPtr.cpp
@@ -120,11 +120,21 @@ void FWeakObjectProperty::SetObjectPropertyValue(void* PropertyValueAddress, UOb
 	SetPropertyValue(PropertyValueAddress, TCppType(Value));
 }
 
+void FWeakObjectProperty::SetObjectPtrPropertyValue(void* PropertyValueAddress, TObjectPtr<UObject> Ptr) const
+{
+	SetPropertyValue(PropertyValueAddress, TCppType(Ptr));
+}
+
 void FWeakObjectProperty::SetObjectPropertyValue_InContainer(void* ContainerAddress, UObject* Value, int32 ArrayIndex) const
 {
 	SetWrappedUObjectPtrValues<FWeakObjectPtr>(ContainerAddress, EPropertyMemoryAccess::InContainer, &Value, ArrayIndex, 1);
 }
 
+void FWeakObjectProperty::SetObjectPtrPropertyValue_InContainer(void* ContainerAddress, TObjectPtr<UObject> Ptr, int32 ArrayIndex) const
+{
+	SetWrappedUObjectPtrValues<FWeakObjectPtr>(ContainerAddress, EPropertyMemoryAccess::InContainer, &Ptr, ArrayIndex, 1);
+}
+
 uint32 FWeakObjectProperty::GetValueTypeHashInternal(const void* Src) const
 {
 	return GetTypeHash(*(FWeakObjectPtr*)Src);
diff --git a/Engine/Source/Runtime/CoreUObject/Public/UObject/UnrealType.h b/Engine/Source/Runtime/CoreUObject/Public/UObject/UnrealType.h
--- a/Engine/Source/Runtime/CoreUObject/Public/UObject/UnrealType.h
+++ b/Engine/Source/Runtime/CoreUObject/Public/UObject/UnrealType.h
@@ -2636,6 +2636,7 @@ public:
 	virtual void SetObjectPropertyValue(void* PropertyValueAddress, UObject* Value) const;
 	virtual void SetObjectPtrPropertyValue(void* PropertyValueAddress, TObjectPtr<UObject> Value) const;
 	virtual void SetObjectPropertyValue_InContainer(void* ContainerAddress, UObject* Value, int32 ArrayIndex = 0) const;
+	virtual void SetObjectPtrPropertyValue_InContainer(void* ContainerAddress, TObjectPtr<UObject> Value, int32 ArrayIndex = 0) const;
 
 	/**
 	 * Setter function for this property's PropertyClass member. Favor this 
@@ -2725,9 +2726,12 @@ protected:
 	// Enable back buffer overrun warning
 	PRAGMA_ENABLE_BUFFER_OVERRUN_WARNING
 
-	template <typename T>
-	void SetWrappedUObjectPtrValues(void* DestAddress, EPropertyMemoryAccess DestAccess, UObject** InValues, int32 ArrayIndex, int32 ArrayCount) const
+	template <typename T, typename ValueType>
+	void SetWrappedUObjectPtrValues(void* DestAddress, EPropertyMemoryAccess DestAccess, ValueType* InValues, int32 ArrayIndex, int32 ArrayCount) const
 	{
+		// Incoming values are expected to be UObject* or TObjectPtr
+		static_assert((std::is_pointer_v<ValueType> && std::is_convertible_v<ValueType, const UObject*>) || TIsTObjectPtr_V<ValueType>);
+
 		// Ensure required range is valid
 		checkf(ArrayIndex >= 0 && ArrayCount >= 0 && ArrayIndex <= ArrayDim && ArrayCount <= ArrayDim && ArrayIndex <= ArrayDim - ArrayCount, TEXT("ArrayIndex (%d) and ArrayCount (%d) is invalid for an array of size %d"), ArrayIndex, ArrayCount, ArrayDim);
 
@@ -2905,8 +2909,9 @@ public:
 	virtual UObject* GetObjectPropertyValue(const void* PropertyValueAddress) const override;
 	virtual UObject* GetObjectPropertyValue_InContainer(const void* ContainerAddress, int32 ArrayIndex = 0) const override;
 	virtual void SetObjectPropertyValue(void* PropertyValueAddress, UObject* Value) const override;
-	virtual void SetObjectPtrPropertyValue(void* PropertyValueAddress, TObjectPtr<UObject> Value) const override;
+	virtual void SetObjectPtrPropertyValue(void* PropertyValueAddress, TObjectPtr<UObject> Ptr) const override;
 	virtual void SetObjectPropertyValue_InContainer(void* ContainerAddress, UObject* Value, int32 ArrayIndex = 0) const override;
+	virtual void SetObjectPtrPropertyValue_InContainer(void* ContainerAddress, TObjectPtr<UObject> Ptr, int32 ArrayIndex = 0) const override;
 	virtual FString GetCPPTypeCustom(FString* ExtendedTypeText, uint32 CPPExportFlags, const FString& InnerNativeTypeName)  const override;
 	// End of FObjectPropertyBase interface
 	
@@ -3007,7 +3012,9 @@ public:
 	virtual UObject* GetObjectPropertyValue(const void* PropertyValueAddress) const override;
 	virtual UObject* GetObjectPropertyValue_InContainer(const void* ContainerAddress, int32 ArrayIndex = 0) const override;
 	virtual void SetObjectPropertyValue(void* PropertyValueAddress, UObject* Value) const override;
+	virtual void SetObjectPtrPropertyValue(void* PropertyValueAddress, TObjectPtr<UObject> Ptr) const override;
 	virtual void SetObjectPropertyValue_InContainer(void* ContainerAddress, UObject* Value, int32 ArrayIndex = 0) const override;
+	virtual void SetObjectPtrPropertyValue_InContainer(void* ContainerAddress, TObjectPtr<UObject> Ptr, int32 ArrayIndex = 0) const override;
 	// End of FObjectProperty interface
 };
 
@@ -3057,7 +3064,9 @@ class COREUOBJECT_API FLazyObjectProperty : public TFObjectPropertyBase<FLazyObj
 	virtual UObject* GetObjectPropertyValue(const void* PropertyValueAddress) const override;
 	virtual UObject* GetObjectPropertyValue_InContainer(const void* ContainerAddress, int32 ArrayIndex = 0) const override;
 	virtual void SetObjectPropertyValue(void* PropertyValueAddress, UObject* Value) const override;
+	virtual void SetObjectPtrPropertyValue(void* PropertyValueAddress, TObjectPtr<UObject> Ptr) const override;
 	virtual void SetObjectPropertyValue_InContainer(void* ContainerAddress, UObject* Value, int32 ArrayIndex = 0) const override;
+	virtual void SetObjectPtrPropertyValue_InContainer(void* ContainerAddress, TObjectPtr<UObject> Ptr, int32 ArrayIndex = 0) const override;
 	virtual bool AllowCrossLevel() const override;
 private:
 	virtual uint32 GetValueTypeHashInternal(const void* Src) const override;
@@ -3118,7 +3127,9 @@ public:
 	virtual UObject* GetObjectPropertyValue(const void* PropertyValueAddress) const override;
 	virtual UObject* GetObjectPropertyValue_InContainer(const void* ContainerAddress, int32 ArrayIndex = 0) const override;
 	virtual void SetObjectPropertyValue(void* PropertyValueAddress, UObject* Value) const override;
+	virtual void SetObjectPtrPropertyValue(void* PropertyValueAddress, TObjectPtr<UObject> Ptr) const override;
 	virtual void SetObjectPropertyValue_InContainer(void* ContainerAddress, UObject* Value, int32 ArrayIndex = 0) const override;
+	virtual void SetObjectPtrPropertyValue_InContainer(void* ContainerAddress, TObjectPtr<UObject> Ptr, int32 ArrayIndex = 0) const override;
 	virtual bool AllowCrossLevel() const override;
 	virtual FString GetCPPTypeCustom(FString* ExtendedTypeText, uint32 CPPExportFlags, const FString& InnerNativeTypeName) const override;
 	virtual FString GetCPPType(FString* ExtendedTypeText, uint32 CPPExportFlags) const override;
