Added unchecked versions of FObjectPropertyBase-derived types.

diff --git a/Engine/Source/Runtime/CoreUObject/Private/UObject/PropertyBaseObject.cpp b/Engine/Source/Runtime/CoreUObject/Private/UObject/PropertyBaseObject.cpp
--- a/Engine/Source/Runtime/CoreUObject/Private/UObject/PropertyBaseObject.cpp
+++ b/Engine/Source/Runtime/CoreUObject/Private/UObject/PropertyBaseObject.cpp
@@ -608,10 +608,34 @@ TObjectPtr<UObject> FObjectPropertyBase::GetObjectPtrPropertyValue(const void* P
 
 void FObjectPropertyBase::SetObjectPtrPropertyValue(void* PropertyValueAddress, TObjectPtr<UObject> Ptr) const
 {
-	unimplemented(); // needs to be implemented by the derived class
+	if (Ptr || !HasAnyPropertyFlags(CPF_NonNullable))
+	{
+		SetObjectPtrPropertyValueUnchecked(PropertyValueAddress, Ptr);
+	}
+	else
+	{
+		UE_LOG(LogProperty, Verbose /*Warning*/, TEXT("Trying to assign null object value to non-nullable \"%s\""), *GetFullName());
+	}
 }
 
 void FObjectPropertyBase::SetObjectPtrPropertyValue_InContainer(void* ContainerAddress, TObjectPtr<UObject> Ptr, int32 ArrayIndex) const
+{
+	if (Ptr || !HasAnyPropertyFlags(CPF_NonNullable))
+	{
+		SetObjectPtrPropertyValueUnchecked_InContainer(ContainerAddress, Ptr, ArrayIndex);
+	}
+	else
+	{
+		UE_LOG(LogProperty, Verbose /*Warning*/, TEXT("Trying to assign null object value to non-nullable \"%s\""), *GetFullName());
+	}
+}
+
+void FObjectPropertyBase::SetObjectPtrPropertyValueUnchecked(void* PropertyValueAddress, TObjectPtr<UObject> Ptr) const
+{
+	unimplemented(); // needs to be implemented by the derived class
+}
+
+void FObjectPropertyBase::SetObjectPtrPropertyValueUnchecked_InContainer(void* ContainerAddress, TObjectPtr<UObject> Ptr, int32 ArrayIndex) const
 {
 	unimplemented(); // needs to be implemented by the derived class
 }
@@ -628,16 +652,40 @@ UObject* FObjectPropertyBase::GetObjectPropertyValue_InContainer(const void* Con
 	return nullptr;
 }
 
-void FObjectPropertyBase::SetObjectPropertyValue_InContainer(void* ContainerAddress, UObject* Value, int32 ArrayIndex) const
+void FObjectPropertyBase::SetObjectPropertyValueUnchecked_InContainer(void* ContainerAddress, UObject* Value, int32 ArrayIndex) const
 {
 	unimplemented(); // needs to be implemented by the derived class
 }
 
-void FObjectPropertyBase::SetObjectPropertyValue(void* PropertyValueAddress, UObject* Value) const
+void FObjectPropertyBase::SetObjectPropertyValueUnchecked(void* PropertyValueAddress, UObject* Value) const
 {
 	unimplemented(); // needs to be implemented by the derived class
 }
 
+void FObjectPropertyBase::SetObjectPropertyValue_InContainer(void* ContainerAddress, UObject* Value, int32 ArrayIndex) const
+{
+	if (Value || !HasAnyPropertyFlags(CPF_NonNullable))
+	{
+		SetObjectPropertyValueUnchecked_InContainer(ContainerAddress, Value, ArrayIndex);
+	}
+	else
+	{
+		UE_LOG(LogProperty, Verbose /*Warning*/, TEXT("Trying to assign null object value to non-nullable \"%s\""), *GetFullName());
+	}
+}
+
+void FObjectPropertyBase::SetObjectPropertyValue(void* PropertyValueAddress, UObject* Value) const
+{
+	if (Value || !HasAnyPropertyFlags(CPF_NonNullable))
+	{
+		SetObjectPropertyValueUnchecked(PropertyValueAddress, Value);
+	}
+	else
+	{
+		UE_LOG(LogProperty, Verbose /*Warning*/, TEXT("Trying to assign null object value to non-nullable \"%s\""), *GetFullName());
+	}
+}
+
 bool FObjectPropertyBase::AllowCrossLevel() const
 {
 	return false;
diff --git a/Engine/Source/Runtime/CoreUObject/Private/UObject/PropertyLazyObjectPtr.cpp b/Engine/Source/Runtime/CoreUObject/Private/UObject/PropertyLazyObjectPtr.cpp
--- a/Engine/Source/Runtime/CoreUObject/Private/UObject/PropertyLazyObjectPtr.cpp
+++ b/Engine/Source/Runtime/CoreUObject/Private/UObject/PropertyLazyObjectPtr.cpp
@@ -134,52 +134,24 @@ UObject* FLazyObjectProperty::GetObjectPropertyValue_InContainer(const void* Con
 	return Result;
 }
 
-void FLazyObjectProperty::SetObjectPropertyValue(void* PropertyValueAddress, UObject* Value) const
+void FLazyObjectProperty::SetObjectPropertyValueUnchecked(void* PropertyValueAddress, UObject* Value) const
 {
-	if (Value || !HasAnyPropertyFlags(CPF_NonNullable))
-	{
-		SetPropertyValue(PropertyValueAddress, TCppType(Value));
-	}
-	else
-	{
-		UE_LOG(LogProperty, Verbose /*Warning*/, TEXT("Trying to assign null object value to non-nullable \"%s\""), *GetFullName());
-	}
+	SetPropertyValue(PropertyValueAddress, TCppType(Value));
 }
 
-void FLazyObjectProperty::SetObjectPtrPropertyValue(void* PropertyValueAddress, TObjectPtr<UObject> Ptr) const
+void FLazyObjectProperty::SetObjectPtrPropertyValueUnchecked(void* PropertyValueAddress, TObjectPtr<UObject> Ptr) const
 {
-	if (Ptr || !HasAnyPropertyFlags(CPF_NonNullable))
-	{
-		SetPropertyValue(PropertyValueAddress, TCppType(Ptr));
-	}
-	else
-	{
-		UE_LOG(LogProperty, Verbose /*Warning*/, TEXT("Trying to assign null object value to non-nullable \"%s\""), *GetFullName());
-	}
+	SetPropertyValue(PropertyValueAddress, TCppType(Ptr));
 }
 
-void FLazyObjectProperty::SetObjectPropertyValue_InContainer(void* ContainerAddress, UObject* Value, int32 ArrayIndex) const
+void FLazyObjectProperty::SetObjectPropertyValueUnchecked_InContainer(void* ContainerAddress, UObject* Value, int32 ArrayIndex) const
 {
-	if (Value || !HasAnyPropertyFlags(CPF_NonNullable))
-	{
-		SetWrappedUObjectPtrValues<FLazyObjectPtr>(ContainerAddress, EPropertyMemoryAccess::InContainer, &Value, ArrayIndex, 1);
-	}
-	else
-	{
-		UE_LOG(LogProperty, Verbose /*Warning*/, TEXT("Trying to assign null object value to non-nullable \"%s\""), *GetFullName());
-	}
+	SetWrappedUObjectPtrValues<FLazyObjectPtr>(ContainerAddress, EPropertyMemoryAccess::InContainer, &Value, ArrayIndex, 1);
 }
 
-void FLazyObjectProperty::SetObjectPtrPropertyValue_InContainer(void* ContainerAddress, TObjectPtr<UObject> Ptr, int32 ArrayIndex) const
+void FLazyObjectProperty::SetObjectPtrPropertyValueUnchecked_InContainer(void* ContainerAddress, TObjectPtr<UObject> Ptr, int32 ArrayIndex) const
 {
-	if (Ptr || !HasAnyPropertyFlags(CPF_NonNullable))
-	{
-		SetWrappedUObjectPtrValues<FLazyObjectPtr>(ContainerAddress, EPropertyMemoryAccess::InContainer, &Ptr, ArrayIndex, 1);
-	}
-	else
-	{
-		UE_LOG(LogProperty, Verbose /*Warning*/, TEXT("Trying to assign null object value to non-nullable \"%s\""), *GetFullName());
-	}
+	SetWrappedUObjectPtrValues<FLazyObjectPtr>(ContainerAddress, EPropertyMemoryAccess::InContainer, &Ptr, ArrayIndex, 1);
 }
 
 bool FLazyObjectProperty::AllowCrossLevel() const
diff --git a/Engine/Source/Runtime/CoreUObject/Private/UObject/PropertyObject.cpp b/Engine/Source/Runtime/CoreUObject/Private/UObject/PropertyObject.cpp
--- a/Engine/Source/Runtime/CoreUObject/Private/UObject/PropertyObject.cpp
+++ b/Engine/Source/Runtime/CoreUObject/Private/UObject/PropertyObject.cpp
@@ -348,52 +348,24 @@ UObject* FObjectProperty::GetObjectPropertyValue_InContainer(const void* Contain
 	return Result;
 }
 
-void FObjectProperty::SetObjectPtrPropertyValue(void* PropertyValueAddress, TObjectPtr<UObject> Ptr) const
+void FObjectProperty::SetObjectPtrPropertyValueUnchecked(void* PropertyValueAddress, TObjectPtr<UObject> Ptr) const
 {
-	if (Ptr || !HasAnyPropertyFlags(CPF_NonNullable))
-	{
-		SetPropertyValue(PropertyValueAddress, Ptr);
-	}
-	else
-	{
-		UE_LOG(LogProperty, Verbose /*Warning*/, TEXT("Trying to assign null object value to non-nullable \"%s\""), *GetFullName());
-	}
+	SetPropertyValue(PropertyValueAddress, Ptr);
 }
 
-void FObjectProperty::SetObjectPropertyValue(void* PropertyValueAddress, UObject* Value) const
+void FObjectProperty::SetObjectPropertyValueUnchecked(void* PropertyValueAddress, UObject* Value) const
 {
-	if (Value || !HasAnyPropertyFlags(CPF_NonNullable))
-	{
-		SetPropertyValue(PropertyValueAddress, Value);
-	}
-	else
-	{
-		UE_LOG(LogProperty, Verbose /*Warning*/, TEXT("Trying to assign null object value to non-nullable \"%s\""), *GetFullName());
-	}
+	SetPropertyValue(PropertyValueAddress, Value);
 }
 
-void FObjectProperty::SetObjectPtrPropertyValue_InContainer(void* ContainerAddress, TObjectPtr<UObject> Ptr, int32 ArrayIndex) const
+void FObjectProperty::SetObjectPtrPropertyValueUnchecked_InContainer(void* ContainerAddress, TObjectPtr<UObject> Ptr, int32 ArrayIndex) const
 {
-	if (Ptr || !HasAnyPropertyFlags(CPF_NonNullable))
-	{
-		SetWrappedUObjectPtrValues<FObjectPtr>(ContainerAddress, EPropertyMemoryAccess::InContainer, &Ptr, ArrayIndex, 1);
-	}
-	else
-	{
-		UE_LOG(LogProperty, Verbose /*Warning*/, TEXT("Trying to assign null object value to non-nullable \"%s\""), *GetFullName());
-	}
+	SetWrappedUObjectPtrValues<FObjectPtr>(ContainerAddress, EPropertyMemoryAccess::InContainer, &Ptr, ArrayIndex, 1);
 }
 
-void FObjectProperty::SetObjectPropertyValue_InContainer(void* ContainerAddress, UObject* Value, int32 ArrayIndex) const
+void FObjectProperty::SetObjectPropertyValueUnchecked_InContainer(void* ContainerAddress, UObject* Value, int32 ArrayIndex) const
 {
-	if (Value || !HasAnyPropertyFlags(CPF_NonNullable))
-	{
-		SetWrappedUObjectPtrValues<FObjectPtr>(ContainerAddress, EPropertyMemoryAccess::InContainer, &Value, ArrayIndex, 1);
-	}
-	else
-	{
-		UE_LOG(LogProperty, Verbose /*Warning*/, TEXT("Trying to assign null object value to non-nullable \"%s\""), *GetFullName());
-	}
+	SetWrappedUObjectPtrValues<FObjectPtr>(ContainerAddress, EPropertyMemoryAccess::InContainer, &Value, ArrayIndex, 1);
 }
 
 void FObjectProperty::CopySingleValueToScriptVM(void* Dest, const void* Src) const
diff --git a/Engine/Source/Runtime/CoreUObject/Private/UObject/PropertySoftObjectPtr.cpp b/Engine/Source/Runtime/CoreUObject/Private/UObject/PropertySoftObjectPtr.cpp
--- a/Engine/Source/Runtime/CoreUObject/Private/UObject/PropertySoftObjectPtr.cpp
+++ b/Engine/Source/Runtime/CoreUObject/Private/UObject/PropertySoftObjectPtr.cpp
@@ -296,22 +296,22 @@ UObject* FSoftObjectProperty::GetObjectPropertyValue_InContainer(const void* Con
 	return Result;
 }
 
-void FSoftObjectProperty::SetObjectPropertyValue(void* PropertyValueAddress, UObject* Value) const
+void FSoftObjectProperty::SetObjectPropertyValueUnchecked(void* PropertyValueAddress, UObject* Value) const
 {
 	SetPropertyValue(PropertyValueAddress, TCppType(Value));
 }
 
-void FSoftObjectProperty::SetObjectPtrPropertyValue(void* PropertyValueAddress, TObjectPtr<UObject> Ptr) const
+void FSoftObjectProperty::SetObjectPtrPropertyValueUnchecked(void* PropertyValueAddress, TObjectPtr<UObject> Ptr) const
 {
 	SetPropertyValue(PropertyValueAddress, TCppType(Ptr));
 }
 
-void FSoftObjectProperty::SetObjectPropertyValue_InContainer(void* ContainerAddress, UObject* Value, int32 ArrayIndex) const
+void FSoftObjectProperty::SetObjectPropertyValueUnchecked_InContainer(void* ContainerAddress, UObject* Value, int32 ArrayIndex) const
 {
 	SetWrappedUObjectPtrValues<FSoftObjectPtr>(ContainerAddress, EPropertyMemoryAccess::InContainer, &Value, ArrayIndex, 1);
 }
 
-void FSoftObjectProperty::SetObjectPtrPropertyValue_InContainer(void* ContainerAddress, TObjectPtr<UObject> Ptr, int32 ArrayIndex) const
+void FSoftObjectProperty::SetObjectPtrPropertyValueUnchecked_InContainer(void* ContainerAddress, TObjectPtr<UObject> Ptr, int32 ArrayIndex) const
 {
 	SetWrappedUObjectPtrValues<FSoftObjectPtr>(ContainerAddress, EPropertyMemoryAccess::InContainer, &Ptr, ArrayIndex, 1);
 }
diff --git a/Engine/Source/Runtime/CoreUObject/Private/UObject/PropertyWeakObjectPtr.cpp b/Engine/Source/Runtime/CoreUObject/Private/UObject/PropertyWeakObjectPtr.cpp
--- a/Engine/Source/Runtime/CoreUObject/Private/UObject/PropertyWeakObjectPtr.cpp
+++ b/Engine/Source/Runtime/CoreUObject/Private/UObject/PropertyWeakObjectPtr.cpp
@@ -115,22 +115,22 @@ UObject* FWeakObjectProperty::GetObjectPropertyValue_InContainer(const void* Con
 	return Result;
 }
 
-void FWeakObjectProperty::SetObjectPropertyValue(void* PropertyValueAddress, UObject* Value) const
+void FWeakObjectProperty::SetObjectPropertyValueUnchecked(void* PropertyValueAddress, UObject* Value) const
 {
 	SetPropertyValue(PropertyValueAddress, TCppType(Value));
 }
 
-void FWeakObjectProperty::SetObjectPtrPropertyValue(void* PropertyValueAddress, TObjectPtr<UObject> Ptr) const
+void FWeakObjectProperty::SetObjectPtrPropertyValueUnchecked(void* PropertyValueAddress, TObjectPtr<UObject> Ptr) const
 {
 	SetPropertyValue(PropertyValueAddress, TCppType(Ptr));
 }
 
-void FWeakObjectProperty::SetObjectPropertyValue_InContainer(void* ContainerAddress, UObject* Value, int32 ArrayIndex) const
+void FWeakObjectProperty::SetObjectPropertyValueUnchecked_InContainer(void* ContainerAddress, UObject* Value, int32 ArrayIndex) const
 {
 	SetWrappedUObjectPtrValues<FWeakObjectPtr>(ContainerAddress, EPropertyMemoryAccess::InContainer, &Value, ArrayIndex, 1);
 }
 
-void FWeakObjectProperty::SetObjectPtrPropertyValue_InContainer(void* ContainerAddress, TObjectPtr<UObject> Ptr, int32 ArrayIndex) const
+void FWeakObjectProperty::SetObjectPtrPropertyValueUnchecked_InContainer(void* ContainerAddress, TObjectPtr<UObject> Ptr, int32 ArrayIndex) const
 {
 	SetWrappedUObjectPtrValues<FWeakObjectPtr>(ContainerAddress, EPropertyMemoryAccess::InContainer, &Ptr, ArrayIndex, 1);
 }
diff --git a/Engine/Source/Runtime/CoreUObject/Public/UObject/UnrealType.h b/Engine/Source/Runtime/CoreUObject/Public/UObject/UnrealType.h
--- a/Engine/Source/Runtime/CoreUObject/Public/UObject/UnrealType.h
+++ b/Engine/Source/Runtime/CoreUObject/Public/UObject/UnrealType.h
@@ -2633,10 +2633,16 @@ public:
 	virtual TObjectPtr<UObject> GetObjectPtrPropertyValue(const void* PropertyValueAddress) const;
 	virtual UObject* GetObjectPropertyValue(const void* PropertyValueAddress) const;
 	virtual UObject* GetObjectPropertyValue_InContainer(const void* ContainerAddress, int32 ArrayIndex = 0) const;
-	virtual void SetObjectPropertyValue(void* PropertyValueAddress, UObject* Value) const;
-	virtual void SetObjectPtrPropertyValue(void* PropertyValueAddress, TObjectPtr<UObject> Value) const;
-	virtual void SetObjectPropertyValue_InContainer(void* ContainerAddress, UObject* Value, int32 ArrayIndex = 0) const;
-	virtual void SetObjectPtrPropertyValue_InContainer(void* ContainerAddress, TObjectPtr<UObject> Value, int32 ArrayIndex = 0) const;
+protected:
+	virtual void SetObjectPropertyValueUnchecked(void* PropertyValueAddress, UObject* Value) const;
+	virtual void SetObjectPtrPropertyValueUnchecked(void* PropertyValueAddress, TObjectPtr<UObject> Ptr) const;
+	virtual void SetObjectPropertyValueUnchecked_InContainer(void* ContainerAddress, UObject* Value, int32 ArrayIndex = 0) const;
+	virtual void SetObjectPtrPropertyValueUnchecked_InContainer(void* ContainerAddress, TObjectPtr<UObject> Ptr, int32 ArrayIndex = 0) const;
+public:
+	void SetObjectPropertyValue(void* PropertyValueAddress, UObject* Value) const;
+	void SetObjectPtrPropertyValue(void* PropertyValueAddress, TObjectPtr<UObject> Ptr) const;
+	void SetObjectPropertyValue_InContainer(void* ContainerAddress, UObject* Value, int32 ArrayIndex = 0) const;
+	void SetObjectPtrPropertyValue_InContainer(void* ContainerAddress, TObjectPtr<UObject> Ptr, int32 ArrayIndex = 0) const;
 
 	/**
 	 * Setter function for this property's PropertyClass member. Favor this 
@@ -2908,10 +2914,12 @@ public:
 	virtual TObjectPtr<UObject> GetObjectPtrPropertyValue(const void* PropertyValueAddress) const override;
 	virtual UObject* GetObjectPropertyValue(const void* PropertyValueAddress) const override;
 	virtual UObject* GetObjectPropertyValue_InContainer(const void* ContainerAddress, int32 ArrayIndex = 0) const override;
-	virtual void SetObjectPropertyValue(void* PropertyValueAddress, UObject* Value) const override;
-	virtual void SetObjectPtrPropertyValue(void* PropertyValueAddress, TObjectPtr<UObject> Ptr) const override;
-	virtual void SetObjectPropertyValue_InContainer(void* ContainerAddress, UObject* Value, int32 ArrayIndex = 0) const override;
-	virtual void SetObjectPtrPropertyValue_InContainer(void* ContainerAddress, TObjectPtr<UObject> Ptr, int32 ArrayIndex = 0) const override;
+protected:
+	virtual void SetObjectPropertyValueUnchecked(void* PropertyValueAddress, UObject* Value) const override;
+	virtual void SetObjectPtrPropertyValueUnchecked(void* PropertyValueAddress, TObjectPtr<UObject> Ptr) const override;
+	virtual void SetObjectPropertyValueUnchecked_InContainer(void* ContainerAddress, UObject* Value, int32 ArrayIndex = 0) const override;
+	virtual void SetObjectPtrPropertyValueUnchecked_InContainer(void* ContainerAddress, TObjectPtr<UObject> Ptr, int32 ArrayIndex = 0) const override;
+public:
 	virtual FString GetCPPTypeCustom(FString* ExtendedTypeText, uint32 CPPExportFlags, const FString& InnerNativeTypeName)  const override;
 	// End of FObjectPropertyBase interface
 	
@@ -3011,10 +3019,12 @@ public:
 	virtual TObjectPtr<UObject> GetObjectPtrPropertyValue(const void* PropertyValueAddress) const override;
 	virtual UObject* GetObjectPropertyValue(const void* PropertyValueAddress) const override;
 	virtual UObject* GetObjectPropertyValue_InContainer(const void* ContainerAddress, int32 ArrayIndex = 0) const override;
-	virtual void SetObjectPropertyValue(void* PropertyValueAddress, UObject* Value) const override;
-	virtual void SetObjectPtrPropertyValue(void* PropertyValueAddress, TObjectPtr<UObject> Ptr) const override;
-	virtual void SetObjectPropertyValue_InContainer(void* ContainerAddress, UObject* Value, int32 ArrayIndex = 0) const override;
-	virtual void SetObjectPtrPropertyValue_InContainer(void* ContainerAddress, TObjectPtr<UObject> Ptr, int32 ArrayIndex = 0) const override;
+protected:
+	virtual void SetObjectPropertyValueUnchecked(void* PropertyValueAddress, UObject* Value) const override;
+	virtual void SetObjectPtrPropertyValueUnchecked(void* PropertyValueAddress, TObjectPtr<UObject> Ptr) const override;
+	virtual void SetObjectPropertyValueUnchecked_InContainer(void* ContainerAddress, UObject* Value, int32 ArrayIndex = 0) const override;
+	virtual void SetObjectPtrPropertyValueUnchecked_InContainer(void* ContainerAddress, TObjectPtr<UObject> Ptr, int32 ArrayIndex = 0) const override;
+public:
 	// End of FObjectProperty interface
 };
 
@@ -3063,10 +3073,12 @@ class COREUOBJECT_API FLazyObjectProperty : public TFObjectPropertyBase<FLazyObj
 	virtual TObjectPtr<UObject> GetObjectPtrPropertyValue(const void* PropertyValueAddress) const override;
 	virtual UObject* GetObjectPropertyValue(const void* PropertyValueAddress) const override;
 	virtual UObject* GetObjectPropertyValue_InContainer(const void* ContainerAddress, int32 ArrayIndex = 0) const override;
-	virtual void SetObjectPropertyValue(void* PropertyValueAddress, UObject* Value) const override;
-	virtual void SetObjectPtrPropertyValue(void* PropertyValueAddress, TObjectPtr<UObject> Ptr) const override;
-	virtual void SetObjectPropertyValue_InContainer(void* ContainerAddress, UObject* Value, int32 ArrayIndex = 0) const override;
-	virtual void SetObjectPtrPropertyValue_InContainer(void* ContainerAddress, TObjectPtr<UObject> Ptr, int32 ArrayIndex = 0) const override;
+protected:
+	virtual void SetObjectPropertyValueUnchecked(void* PropertyValueAddress, UObject* Value) const override;
+	virtual void SetObjectPtrPropertyValueUnchecked(void* PropertyValueAddress, TObjectPtr<UObject> Ptr) const override;
+	virtual void SetObjectPropertyValueUnchecked_InContainer(void* ContainerAddress, UObject* Value, int32 ArrayIndex = 0) const override;
+	virtual void SetObjectPtrPropertyValueUnchecked_InContainer(void* ContainerAddress, TObjectPtr<UObject> Ptr, int32 ArrayIndex = 0) const override;
+public:
 	virtual bool AllowCrossLevel() const override;
 private:
 	virtual uint32 GetValueTypeHashInternal(const void* Src) const override;
@@ -3126,10 +3138,12 @@ public:
 	virtual TObjectPtr<UObject> GetObjectPtrPropertyValue(const void* PropertyValueAddress) const override;
 	virtual UObject* GetObjectPropertyValue(const void* PropertyValueAddress) const override;
 	virtual UObject* GetObjectPropertyValue_InContainer(const void* ContainerAddress, int32 ArrayIndex = 0) const override;
-	virtual void SetObjectPropertyValue(void* PropertyValueAddress, UObject* Value) const override;
-	virtual void SetObjectPtrPropertyValue(void* PropertyValueAddress, TObjectPtr<UObject> Ptr) const override;
-	virtual void SetObjectPropertyValue_InContainer(void* ContainerAddress, UObject* Value, int32 ArrayIndex = 0) const override;
-	virtual void SetObjectPtrPropertyValue_InContainer(void* ContainerAddress, TObjectPtr<UObject> Ptr, int32 ArrayIndex = 0) const override;
+protected:
+	virtual void SetObjectPropertyValueUnchecked(void* PropertyValueAddress, UObject* Value) const override;
+	virtual void SetObjectPtrPropertyValueUnchecked(void* PropertyValueAddress, TObjectPtr<UObject> Ptr) const override;
+	virtual void SetObjectPropertyValueUnchecked_InContainer(void* ContainerAddress, UObject* Value, int32 ArrayIndex = 0) const override;
+	virtual void SetObjectPtrPropertyValueUnchecked_InContainer(void* ContainerAddress, TObjectPtr<UObject> Ptr, int32 ArrayIndex = 0) const override;
+public:
 	virtual bool AllowCrossLevel() const override;
 	virtual FString GetCPPTypeCustom(FString* ExtendedTypeText, uint32 CPPExportFlags, const FString& InnerNativeTypeName) const override;
 	virtual FString GetCPPType(FString* ExtendedTypeText, uint32 CPPExportFlags) const override;
