Reordered FObjectPropertyBase's Get and Set functions to be consistent.

diff --git a/Engine/Source/Runtime/CoreUObject/Private/UObject/PropertyBaseObject.cpp b/Engine/Source/Runtime/CoreUObject/Private/UObject/PropertyBaseObject.cpp
--- a/Engine/Source/Runtime/CoreUObject/Private/UObject/PropertyBaseObject.cpp
+++ b/Engine/Source/Runtime/CoreUObject/Private/UObject/PropertyBaseObject.cpp
@@ -600,66 +600,66 @@
 	return NAME_ObjectProperty;
 }
 
+UObject* FObjectPropertyBase::GetObjectPropertyValue(const void* PropertyValueAddress) const
+{
+	unimplemented(); // needs to be implemented by the derived class
+	return nullptr;
+}
+
 TObjectPtr<UObject> FObjectPropertyBase::GetObjectPtrPropertyValue(const void* PropertyValueAddress) const
 {
 	unimplemented(); // needs to be implemented by the derived class
 	return TObjectPtr<UObject>();
 }
 
-void FObjectPropertyBase::SetObjectPtrPropertyValue(void* PropertyValueAddress, TObjectPtr<UObject> Ptr) const
-{
-	if (Ptr || !HasAnyPropertyFlags(CPF_NonNullable))
-	{
-		SetObjectPtrPropertyValueUnchecked(PropertyValueAddress, Ptr);
-	}
-	else
-	{
-		UE_LOG(LogProperty, Verbose /*Warning*/, TEXT("Trying to assign null object value to non-nullable \"%s\""), *GetFullName());
-	}
-}
-
-void FObjectPropertyBase::SetObjectPtrPropertyValue_InContainer(void* ContainerAddress, TObjectPtr<UObject> Ptr, int32 ArrayIndex) const
-{
-	if (Ptr || !HasAnyPropertyFlags(CPF_NonNullable))
-	{
-		SetObjectPtrPropertyValueUnchecked_InContainer(ContainerAddress, Ptr, ArrayIndex);
-	}
-	else
-	{
-		UE_LOG(LogProperty, Verbose /*Warning*/, TEXT("Trying to assign null object value to non-nullable \"%s\""), *GetFullName());
-	}
-}
-
-void FObjectPropertyBase::SetObjectPtrPropertyValueUnchecked(void* PropertyValueAddress, TObjectPtr<UObject> Ptr) const
-{
-	unimplemented(); // needs to be implemented by the derived class
-}
-
-void FObjectPropertyBase::SetObjectPtrPropertyValueUnchecked_InContainer(void* ContainerAddress, TObjectPtr<UObject> Ptr, int32 ArrayIndex) const
-{
-	unimplemented(); // needs to be implemented by the derived class
-}
-
-UObject* FObjectPropertyBase::GetObjectPropertyValue(const void* PropertyValueAddress) const
-{
-	unimplemented(); // needs to be implemented by the derived class
-	return nullptr;
-}
-
 UObject* FObjectPropertyBase::GetObjectPropertyValue_InContainer(const void* ContainerAddress, int32 ArrayIndex) const
 {
 	unimplemented(); // needs to be implemented by the derived class
 	return nullptr;
 }
 
-void FObjectPropertyBase::SetObjectPropertyValueUnchecked_InContainer(void* ContainerAddress, UObject* Value, int32 ArrayIndex) const
-{
-	unimplemented(); // needs to be implemented by the derived class
-}
-
 void FObjectPropertyBase::SetObjectPropertyValueUnchecked(void* PropertyValueAddress, UObject* Value) const
 {
 	unimplemented(); // needs to be implemented by the derived class
 }
 
+void FObjectPropertyBase::SetObjectPtrPropertyValueUnchecked(void* PropertyValueAddress, TObjectPtr<UObject> Ptr) const
+{
+	unimplemented(); // needs to be implemented by the derived class
+}
+
+void FObjectPropertyBase::SetObjectPropertyValueUnchecked_InContainer(void* ContainerAddress, UObject* Value, int32 ArrayIndex) const
+{
+	unimplemented(); // needs to be implemented by the derived class
+}
+
+void FObjectPropertyBase::SetObjectPtrPropertyValueUnchecked_InContainer(void* ContainerAddress, TObjectPtr<UObject> Ptr, int32 ArrayIndex) const
+{
+	unimplemented(); // needs to be implemented by the derived class
+}
+
+void FObjectPropertyBase::SetObjectPropertyValue(void* PropertyValueAddress, UObject* Value) const
+{
+	if (Value || !HasAnyPropertyFlags(CPF_NonNullable))
+	{
+		SetObjectPropertyValueUnchecked(PropertyValueAddress, Value);
+	}
+	else
+	{
+		UE_LOG(LogProperty, Verbose /*Warning*/, TEXT("Trying to assign null object value to non-nullable \"%s\""), *GetFullName());
+	}
+}
+
+void FObjectPropertyBase::SetObjectPtrPropertyValue(void* PropertyValueAddress, TObjectPtr<UObject> Ptr) const
+{
+	if (Ptr || !HasAnyPropertyFlags(CPF_NonNullable))
+	{
+		SetObjectPtrPropertyValueUnchecked(PropertyValueAddress, Ptr);
+	}
+	else
+	{
+		UE_LOG(LogProperty, Verbose /*Warning*/, TEXT("Trying to assign null object value to non-nullable \"%s\""), *GetFullName());
+	}
+}
+
 void FObjectPropertyBase::SetObjectPropertyValue_InContainer(void* ContainerAddress, UObject* Value, int32 ArrayIndex) const
@@ -674,11 +674,11 @@
 	}
 }
 
-void FObjectPropertyBase::SetObjectPropertyValue(void* PropertyValueAddress, UObject* Value) const
+void FObjectPropertyBase::SetObjectPtrPropertyValue_InContainer(void* ContainerAddress, TObjectPtr<UObject> Ptr, int32 ArrayIndex) const
 {
-	if (Value || !HasAnyPropertyFlags(CPF_NonNullable))
+	if (Ptr || !HasAnyPropertyFlags(CPF_NonNullable))
 	{
-		SetObjectPropertyValueUnchecked(PropertyValueAddress, Value);
+		SetObjectPtrPropertyValueUnchecked_InContainer(ContainerAddress, Ptr, ArrayIndex);
 	}
 	else
 	{
diff --git a/Engine/Source/Runtime/CoreUObject/Public/UObject/UnrealType.h b/Engine/Source/Runtime/CoreUObject/Public/UObject/UnrealType.h
--- a/Engine/Source/Runtime/CoreUObject/Public/UObject/UnrealType.h
+++ b/Engine/Source/Runtime/CoreUObject/Public/UObject/UnrealType.h
@@ -2630,15 +2630,16 @@ public:
 		return LoadObjectPropertyValue(ContainerPtrToValuePtr<void>(PropertyValueAddress, ArrayIndex));
 	}
 
-	virtual TObjectPtr<UObject> GetObjectPtrPropertyValue(const void* PropertyValueAddress) const;
-	virtual UObject* GetObjectPropertyValue(const void* PropertyValueAddress) const;
-	virtual UObject* GetObjectPropertyValue_InContainer(const void* ContainerAddress, int32 ArrayIndex = 0) const;
 protected:
 	virtual void SetObjectPropertyValueUnchecked(void* PropertyValueAddress, UObject* Value) const;
 	virtual void SetObjectPtrPropertyValueUnchecked(void* PropertyValueAddress, TObjectPtr<UObject> Ptr) const;
 	virtual void SetObjectPropertyValueUnchecked_InContainer(void* ContainerAddress, UObject* Value, int32 ArrayIndex = 0) const;
 	virtual void SetObjectPtrPropertyValueUnchecked_InContainer(void* ContainerAddress, TObjectPtr<UObject> Ptr, int32 ArrayIndex = 0) const;
 public:
+	virtual UObject* GetObjectPropertyValue(const void* PropertyValueAddress) const;
+	virtual TObjectPtr<UObject> GetObjectPtrPropertyValue(const void* PropertyValueAddress) const;
+	virtual UObject* GetObjectPropertyValue_InContainer(const void* ContainerAddress, int32 ArrayIndex = 0) const;
+
 	void SetObjectPropertyValue(void* PropertyValueAddress, UObject* Value) const;
 	void SetObjectPtrPropertyValue(void* PropertyValueAddress, TObjectPtr<UObject> Ptr) const;
 	void SetObjectPropertyValue_InContainer(void* ContainerAddress, UObject* Value, int32 ArrayIndex = 0) const;
