Reordered FObjectPropertyBase-derived classes' Get functions to be consistent with the Set functions.

diff --git a/Engine/Source/Runtime/CoreUObject/Private/UObject/PropertyLazyObjectPtr.cpp b/Engine/Source/Runtime/CoreUObject/Private/UObject/PropertyLazyObjectPtr.cpp
--- a/Engine/Source/Runtime/CoreUObject/Private/UObject/PropertyLazyObjectPtr.cpp
+++ b/Engine/Source/Runtime/CoreUObject/Private/UObject/PropertyLazyObjectPtr.cpp
@@ -117,14 +117,14 @@ bool FLazyObjectProperty::Identical( const void* A, const void* B, uint32 PortFl
 	return bResult;
 }
 
-TObjectPtr<UObject> FLazyObjectProperty::GetObjectPtrPropertyValue(const void* PropertyValueAddress) const
+UObject* FLazyObjectProperty::GetObjectPropertyValue(const void* PropertyValueAddress) const
 {
-	return TObjectPtr<UObject>(GetPropertyValue(PropertyValueAddress).Get());
+	return GetPropertyValue(PropertyValueAddress).Get();
 }
 
-UObject* FLazyObjectProperty::GetObjectPropertyValue(const void* PropertyValueAddress) const
+TObjectPtr<UObject> FLazyObjectProperty::GetObjectPtrPropertyValue(const void* PropertyValueAddress) const
 {
-	return GetPropertyValue(PropertyValueAddress).Get();
+	return TObjectPtr<UObject>(GetPropertyValue(PropertyValueAddress).Get());
 }
 
 UObject* FLazyObjectProperty::GetObjectPropertyValue_InContainer(const void* ContainerAddress, int32 ArrayIndex) const
diff --git a/Engine/Source/Runtime/CoreUObject/Private/UObject/PropertyObject.cpp b/Engine/Source/Runtime/CoreUObject/Private/UObject/PropertyObject.cpp
--- a/Engine/Source/Runtime/CoreUObject/Private/UObject/PropertyObject.cpp
+++ b/Engine/Source/Runtime/CoreUObject/Private/UObject/PropertyObject.cpp
@@ -331,14 +331,14 @@ uint32 FObjectProperty::GetValueTypeHashInternal(const void* Src) const
 	return GetTypeHash(GetPropertyValue(Src));
 }
 
-TObjectPtr<UObject> FObjectProperty::GetObjectPtrPropertyValue(const void* PropertyValueAddress) const
+UObject* FObjectProperty::GetObjectPropertyValue(const void* PropertyValueAddress) const
 {
-	return GetPropertyValue(PropertyValueAddress);
+	return GetPropertyValue(PropertyValueAddress).Get();
 }
 
-UObject* FObjectProperty::GetObjectPropertyValue(const void* PropertyValueAddress) const
+TObjectPtr<UObject> FObjectProperty::GetObjectPtrPropertyValue(const void* PropertyValueAddress) const
 {
-	return GetPropertyValue(PropertyValueAddress).Get();
+	return GetPropertyValue(PropertyValueAddress);
 }
 
 UObject* FObjectProperty::GetObjectPropertyValue_InContainer(const void* ContainerAddress, int32 ArrayIndex) const
diff --git a/Engine/Source/Runtime/CoreUObject/Private/UObject/PropertySoftObjectPtr.cpp b/Engine/Source/Runtime/CoreUObject/Private/UObject/PropertySoftObjectPtr.cpp
--- a/Engine/Source/Runtime/CoreUObject/Private/UObject/PropertySoftObjectPtr.cpp
+++ b/Engine/Source/Runtime/CoreUObject/Private/UObject/PropertySoftObjectPtr.cpp
@@ -279,14 +279,14 @@ UObject* FSoftObjectProperty::LoadObjectPropertyValue(const void* PropertyValueA
 	return GetPropertyValue(PropertyValueAddress).LoadSynchronous();
 }
 
-TObjectPtr<UObject> FSoftObjectProperty::GetObjectPtrPropertyValue(const void* PropertyValueAddress) const
+UObject* FSoftObjectProperty::GetObjectPropertyValue(const void* PropertyValueAddress) const
 {
-	return TObjectPtr<UObject>(GetPropertyValue(PropertyValueAddress).Get());
+	return GetPropertyValue(PropertyValueAddress).Get();
 }
 
-UObject* FSoftObjectProperty::GetObjectPropertyValue(const void* PropertyValueAddress) const
+TObjectPtr<UObject> FSoftObjectProperty::GetObjectPtrPropertyValue(const void* PropertyValueAddress) const
 {
-	return GetPropertyValue(PropertyValueAddress).Get();
+	return TObjectPtr<UObject>(GetPropertyValue(PropertyValueAddress).Get());
 }
 
 UObject* FSoftObjectProperty::GetObjectPropertyValue_InContainer(const void* ContainerAddress, int32 ArrayIndex) const
diff --git a/Engine/Source/Runtime/CoreUObject/Private/UObject/PropertyWeakObjectPtr.cpp b/Engine/Source/Runtime/CoreUObject/Private/UObject/PropertyWeakObjectPtr.cpp
--- a/Engine/Source/Runtime/CoreUObject/Private/UObject/PropertyWeakObjectPtr.cpp
+++ b/Engine/Source/Runtime/CoreUObject/Private/UObject/PropertyWeakObjectPtr.cpp
@@ -98,14 +98,14 @@ void FWeakObjectProperty::SerializeItem( FStructuredArchive::FSlot Slot, void* V
 	}
 }
 
-TObjectPtr<UObject> FWeakObjectProperty::GetObjectPtrPropertyValue(const void* PropertyValueAddress) const
+UObject* FWeakObjectProperty::GetObjectPropertyValue(const void* PropertyValueAddress) const
 {
-	return TObjectPtr<UObject>(GetPropertyValue(PropertyValueAddress).Get());
+	return GetPropertyValue(PropertyValueAddress).Get();
 }
 
-UObject* FWeakObjectProperty::GetObjectPropertyValue(const void* PropertyValueAddress) const
+TObjectPtr<UObject> FWeakObjectProperty::GetObjectPtrPropertyValue(const void* PropertyValueAddress) const
 {
-	return GetPropertyValue(PropertyValueAddress).Get();
+	return TObjectPtr<UObject>(GetPropertyValue(PropertyValueAddress).Get());
 }
 
 UObject* FWeakObjectProperty::GetObjectPropertyValue_InContainer(const void* ContainerAddress, int32 ArrayIndex) const
