Also run CollectReferencers when searching for referencers, as it is fast and catches things that override ARO or UE::GC::DeclareIntrinsicMembers, the former is used by UObjectRedirector and as such references from UObjectRedirectors were not being reported by FindReferencers

#rnx

diff --git a/Engine/Source/Runtime/CoreUObject/Private/UObject/ReferencerFinder.cpp b/Engine/Source/Runtime/CoreUObject/Private/UObject/ReferencerFinder.cpp
--- a/Engine/Source/Runtime/CoreUObject/Private/UObject/ReferencerFinder.cpp
+++ b/Engine/Source/Runtime/CoreUObject/Private/UObject/ReferencerFinder.cpp
@@ -120,23 +120,10 @@ public:
 		checkSlow(Obj);
 		CurrentObject = Obj;
 		bDoesReferenceAnyTargets = false;
-		if(UStruct* AsStruct = Cast<UStruct>(Obj))
-		{
-			// UStruct serialize overrides (like UControlRigBlueprintGeneratedClass::Serialize)
-			// are not thread safe, but UStruct::SerializeBin will not report any
-			// references in the reflection data. This is a compromise.
-			AsStruct->UStruct::Serialize(*this);
-		}
-		else
-		{
-			// This could miss some things in a user serialize functions,
-			// but there exist many user Serialize routines that are not threadsafe
-			// and deadlocking would be disastrous.
-			Obj->GetClass()->SerializeBin(*this, Obj);
-		}
-		// Compensate for missed stream serialized reference 
-		// by calling ARO:
-		Obj->GetClass()->CallAddReferencedObjects(Obj, *this); 
+		// This could miss some things in a user serialize functions,
+		// but there exist many user Serialize routines that are not threadsafe
+		// and deadlocking would be disastrous.
+		Obj->GetClass()->SerializeBin(*this, Obj);
 
 		return bDoesReferenceAnyTargets;
 	}
@@ -188,8 +175,8 @@ static TArray<UObject*> GetAllReferencersIncludingWeak(
 	const int32 NumberOfObjectsPerThread = (MaxNumberOfObjects / NumThreads) + 1;
 	
 	// Allocate per thread results, in this case each thread will produce a list of referencers
-	TUniquePtr<TArray<UObject*>[]> ThreadResultsAlloc = MakeUnique<TArray<UObject*>[]>(NumThreads);
-	TArrayView<TArray<UObject*>> ThreadResults(ThreadResultsAlloc.Get(), NumThreads);
+	TUniquePtr<TSet<UObject*>[]> ThreadResultsAlloc = MakeUnique<TSet<UObject*>[]>(NumThreads);
+	TArrayView<TSet<UObject*>> ThreadResults(ThreadResultsAlloc.Get(), NumThreads);
 
 	ParallelFor(NumThreads, 
 		[&Referencees, 
@@ -200,8 +187,11 @@ static TArray<UObject*> GetAllReferencersIncludingWeak(
 		MaxNumberOfObjects, 
 		bSkipInnerReferences](int32 ThreadIndex)
 	{
+		const EReferencerFinderFlags Flags = bSkipInnerReferences ? EReferencerFinderFlags::SkipInnerReferences : EReferencerFinderFlags::None;
 		FReferencerFinderArchive ReferenceFinderArchive(Referencees, bSkipInnerReferences);
-		TArray<UObject*>& ThreadResult = ThreadResults[ThreadIndex];
+		TSet<UObject*>& ThreadResult = ThreadResults[ThreadIndex];
+		TArray<UObject*> ObjectsToSearch;
+		ObjectsToSearch.Reserve(NumberOfObjectsPerThread);
 		
 		// Process the block of objects assigned to this thread:
 		const int32 FirstObjectIndex = ThreadIndex * NumberOfObjectsPerThread;
@@ -227,10 +217,25 @@ static TArray<UObject*> GetAllReferencersIncludingWeak(
 				continue;
 			}
 
+			// We could skip this for objects with no reflected object references
+			// as mostly this is redundant to the faster path below (but both are 
+			// mostly going to return false so no speedup reversing the order)
 			if (ReferenceFinderArchive.SearchForReferencesToTargets(PotentialReferencer))
 			{
 				ThreadResult.Add(PotentialReferencer);
 			}
+			else
+			{
+				ObjectsToSearch.Add(PotentialReferencer);
+			}
+		}
+		
+		FAllReferencesProcessor Processor(Referencees, Flags, ThreadResult);
+		FGCArrayStruct ArrayStruct;
+		ArrayStruct.SetInitialObjectsUnpadded(ObjectsToSearch);
+		{
+			FGCScopeGuard GCGuard;
+			CollectReferences(Processor, ArrayStruct);
 		}
 	},  (GUObjectRegistrationComplete && GAllowParallelReferenceCollection) ? 
 			EParallelForFlags::None : EParallelForFlags::ForceSingleThread);
@@ -239,9 +244,9 @@ static TArray<UObject*> GetAllReferencersIncludingWeak(
 	GUObjectArray.UnlockInternalArray();
 
 	TArray<UObject*> FinalResult;
-	for(const TArray<UObject*>& ThreadResult : ThreadResults)
+	for(const TSet<UObject*>& ThreadResult : ThreadResults)
 	{
-		FinalResult.Append(ThreadResult);
+		FinalResult.Append(ThreadResult.Array());
 	}
 
 	return FinalResult;
