Added missing GetObjectPtrPropertyValue_InContainer overloads.

diff --git a/Engine/Source/Runtime/CoreUObject/Private/UObject/PropertyBaseObject.cpp b/Engine/Source/Runtime/CoreUObject/Private/UObject/PropertyBaseObject.cpp
--- a/Engine/Source/Runtime/CoreUObject/Private/UObject/PropertyBaseObject.cpp
+++ b/Engine/Source/Runtime/CoreUObject/Private/UObject/PropertyBaseObject.cpp
@@ -618,6 +618,12 @@ UObject* FObjectPropertyBase::GetObjectPropertyValue_InContainer(const void* Con
 	return nullptr;
 }
 
+TObjectPtr<UObject> FObjectPropertyBase::GetObjectPtrPropertyValue_InContainer(const void* ContainerAddress, int32 ArrayIndex) const
+{
+	unimplemented(); // needs to be implemented by the derived class
+	return TObjectPtr<UObject>();
+}
+
 void FObjectPropertyBase::SetObjectPropertyValueUnchecked(void* PropertyValueAddress, UObject* Value) const
 {
 	unimplemented(); // needs to be implemented by the derived class
diff --git a/Engine/Source/Runtime/CoreUObject/Private/UObject/PropertyLazyObjectPtr.cpp b/Engine/Source/Runtime/CoreUObject/Private/UObject/PropertyLazyObjectPtr.cpp
--- a/Engine/Source/Runtime/CoreUObject/Private/UObject/PropertyLazyObjectPtr.cpp
+++ b/Engine/Source/Runtime/CoreUObject/Private/UObject/PropertyLazyObjectPtr.cpp
@@ -134,6 +134,13 @@ UObject* FLazyObjectProperty::GetObjectPropertyValue_InContainer(const void* Con
 	return Result;
 }
 
+TObjectPtr<UObject> FLazyObjectProperty::GetObjectPtrPropertyValue_InContainer(const void* ContainerAddress, int32 ArrayIndex) const
+{
+	TObjectPtr<UObject> Result = nullptr;
+	GetWrappedUObjectPtrValues<FLazyObjectPtr>(&Result, ContainerAddress, EPropertyMemoryAccess::InContainer, ArrayIndex, 1);
+	return Result;
+}
+
 void FLazyObjectProperty::SetObjectPropertyValueUnchecked(void* PropertyValueAddress, UObject* Value) const
 {
 	SetPropertyValue(PropertyValueAddress, TCppType(Value));
diff --git a/Engine/Source/Runtime/CoreUObject/Private/UObject/PropertyObject.cpp b/Engine/Source/Runtime/CoreUObject/Private/UObject/PropertyObject.cpp
--- a/Engine/Source/Runtime/CoreUObject/Private/UObject/PropertyObject.cpp
+++ b/Engine/Source/Runtime/CoreUObject/Private/UObject/PropertyObject.cpp
@@ -348,6 +348,13 @@ UObject* FObjectProperty::GetObjectPropertyValue_InContainer(const void* Contain
 	return Result;
 }
 
+TObjectPtr<UObject> FObjectProperty::GetObjectPtrPropertyValue_InContainer(const void* ContainerAddress, int32 ArrayIndex) const
+{
+	TObjectPtr<UObject> Result = nullptr;
+	GetWrappedUObjectPtrValues<FObjectPtr>(&Result, ContainerAddress, EPropertyMemoryAccess::InContainer, ArrayIndex, 1);
+	return Result;
+}
+
 void FObjectProperty::SetObjectPtrPropertyValueUnchecked(void* PropertyValueAddress, TObjectPtr<UObject> Ptr) const
 {
 	SetPropertyValue(PropertyValueAddress, Ptr);
diff --git a/Engine/Source/Runtime/CoreUObject/Private/UObject/PropertySoftObjectPtr.cpp b/Engine/Source/Runtime/CoreUObject/Private/UObject/PropertySoftObjectPtr.cpp
--- a/Engine/Source/Runtime/CoreUObject/Private/UObject/PropertySoftObjectPtr.cpp
+++ b/Engine/Source/Runtime/CoreUObject/Private/UObject/PropertySoftObjectPtr.cpp
@@ -296,6 +296,13 @@ UObject* FSoftObjectProperty::GetObjectPropertyValue_InContainer(const void* Con
 	return Result;
 }
 
+TObjectPtr<UObject> FSoftObjectProperty::GetObjectPtrPropertyValue_InContainer(const void* ContainerAddress, int32 ArrayIndex) const
+{
+	TObjectPtr<UObject> Result = nullptr;
+	GetWrappedUObjectPtrValues<FSoftObjectPtr>(&Result, ContainerAddress, EPropertyMemoryAccess::InContainer, ArrayIndex, 1);
+	return Result;
+}
+
 void FSoftObjectProperty::SetObjectPropertyValueUnchecked(void* PropertyValueAddress, UObject* Value) const
 {
 	SetPropertyValue(PropertyValueAddress, TCppType(Value));
diff --git a/Engine/Source/Runtime/CoreUObject/Private/UObject/PropertyWeakObjectPtr.cpp b/Engine/Source/Runtime/CoreUObject/Private/UObject/PropertyWeakObjectPtr.cpp
--- a/Engine/Source/Runtime/CoreUObject/Private/UObject/PropertyWeakObjectPtr.cpp
+++ b/Engine/Source/Runtime/CoreUObject/Private/UObject/PropertyWeakObjectPtr.cpp
@@ -115,6 +115,13 @@ UObject* FWeakObjectProperty::GetObjectPropertyValue_InContainer(const void* Con
 	return Result;
 }
 
+TObjectPtr<UObject> FWeakObjectProperty::GetObjectPtrPropertyValue_InContainer(const void* ContainerAddress, int32 ArrayIndex) const
+{
+	TObjectPtr<UObject> Result = nullptr;
+	GetWrappedUObjectPtrValues<FWeakObjectPtr>(&Result, ContainerAddress, EPropertyMemoryAccess::InContainer, ArrayIndex, 1);
+	return Result;
+}
+
 void FWeakObjectProperty::SetObjectPropertyValueUnchecked(void* PropertyValueAddress, UObject* Value) const
 {
 	SetPropertyValue(PropertyValueAddress, TCppType(Value));
diff --git a/Engine/Source/Runtime/CoreUObject/Public/UObject/UnrealType.h b/Engine/Source/Runtime/CoreUObject/Public/UObject/UnrealType.h
--- a/Engine/Source/Runtime/CoreUObject/Public/UObject/UnrealType.h
+++ b/Engine/Source/Runtime/CoreUObject/Public/UObject/UnrealType.h
@@ -2639,6 +2639,7 @@ public:
 	virtual UObject* GetObjectPropertyValue(const void* PropertyValueAddress) const;
 	virtual TObjectPtr<UObject> GetObjectPtrPropertyValue(const void* PropertyValueAddress) const;
 	virtual UObject* GetObjectPropertyValue_InContainer(const void* ContainerAddress, int32 ArrayIndex = 0) const;
+	virtual TObjectPtr<UObject> GetObjectPtrPropertyValue_InContainer(const void* ContainerAddress, int32 ArrayIndex = 0) const;
 
 	void SetObjectPropertyValue(void* PropertyValueAddress, UObject* Value) const;
 	void SetObjectPtrPropertyValue(void* PropertyValueAddress, TObjectPtr<UObject> Ptr) const;
@@ -2678,9 +2679,12 @@ protected:
 	// Disable false positive buffer overrun warning during pgoprofile linking step
 	PRAGMA_DISABLE_BUFFER_OVERRUN_WARNING
 	/* Helper functions for UObject property types that wrap the object pointer in a smart pointer */
-	template <typename T>
-	void GetWrappedUObjectPtrValues(UObject** OutObjects, const void* SrcAddress, EPropertyMemoryAccess SrcAccess, int32 ArrayIndex, int32 ArrayCount) const
+	template <typename T, typename OutType>
+	void GetWrappedUObjectPtrValues(OutType* OutObjects, const void* SrcAddress, EPropertyMemoryAccess SrcAccess, int32 ArrayIndex, int32 ArrayCount) const
 	{
+		// Outgoing values are expected to be UObject* or TObjectPtr
+		static_assert((std::is_pointer_v<OutType> && std::is_convertible_v<OutType, const UObject*>) || TIsTObjectPtr_V<OutType>);
+
 		// Ensure required range is valid
 		checkf(ArrayIndex >= 0 && ArrayCount >= 0 && ArrayIndex <= ArrayDim && ArrayCount <= ArrayDim && ArrayIndex <= ArrayDim - ArrayCount, TEXT("ArrayIndex (%d) and ArrayCount (%d) is invalid for an array of size %d"), ArrayIndex, ArrayCount, ArrayDim);
 
@@ -2915,6 +2919,7 @@ public:
 	virtual UObject* GetObjectPropertyValue(const void* PropertyValueAddress) const override;
 	virtual TObjectPtr<UObject> GetObjectPtrPropertyValue(const void* PropertyValueAddress) const override;
 	virtual UObject* GetObjectPropertyValue_InContainer(const void* ContainerAddress, int32 ArrayIndex = 0) const override;
+	virtual TObjectPtr<UObject> GetObjectPtrPropertyValue_InContainer(const void* ContainerAddress, int32 ArrayIndex = 0) const override;
 protected:
 	virtual void SetObjectPropertyValueUnchecked(void* PropertyValueAddress, UObject* Value) const override;
 	virtual void SetObjectPtrPropertyValueUnchecked(void* PropertyValueAddress, TObjectPtr<UObject> Ptr) const override;
@@ -3020,6 +3025,7 @@ public:
 	virtual UObject* GetObjectPropertyValue(const void* PropertyValueAddress) const override;
 	virtual TObjectPtr<UObject> GetObjectPtrPropertyValue(const void* PropertyValueAddress) const override;
 	virtual UObject* GetObjectPropertyValue_InContainer(const void* ContainerAddress, int32 ArrayIndex = 0) const override;
+	virtual TObjectPtr<UObject> GetObjectPtrPropertyValue_InContainer(const void* ContainerAddress, int32 ArrayIndex = 0) const override;
 protected:
 	virtual void SetObjectPropertyValueUnchecked(void* PropertyValueAddress, UObject* Value) const override;
 	virtual void SetObjectPtrPropertyValueUnchecked(void* PropertyValueAddress, TObjectPtr<UObject> Ptr) const override;
@@ -3074,6 +3080,7 @@ class COREUOBJECT_API FLazyObjectProperty : public TFObjectPropertyBase<FLazyObj
 	virtual UObject* GetObjectPropertyValue(const void* PropertyValueAddress) const override;
 	virtual TObjectPtr<UObject> GetObjectPtrPropertyValue(const void* PropertyValueAddress) const override;
 	virtual UObject* GetObjectPropertyValue_InContainer(const void* ContainerAddress, int32 ArrayIndex = 0) const override;
+	virtual TObjectPtr<UObject> GetObjectPtrPropertyValue_InContainer(const void* ContainerAddress, int32 ArrayIndex = 0) const override;
 protected:
 	virtual void SetObjectPropertyValueUnchecked(void* PropertyValueAddress, UObject* Value) const override;
 	virtual void SetObjectPtrPropertyValueUnchecked(void* PropertyValueAddress, TObjectPtr<UObject> Ptr) const override;
@@ -3140,6 +3147,7 @@ public:
 	virtual UObject* GetObjectPropertyValue(const void* PropertyValueAddress) const override;
 	virtual TObjectPtr<UObject> GetObjectPtrPropertyValue(const void* PropertyValueAddress) const override;
 	virtual UObject* GetObjectPropertyValue_InContainer(const void* ContainerAddress, int32 ArrayIndex = 0) const override;
+	virtual TObjectPtr<UObject> GetObjectPtrPropertyValue_InContainer(const void* ContainerAddress, int32 ArrayIndex = 0) const override;
 protected:
 	virtual void SetObjectPropertyValueUnchecked(void* PropertyValueAddress, UObject* Value) const override;
 	virtual void SetObjectPtrPropertyValueUnchecked(void* PropertyValueAddress, TObjectPtr<UObject> Ptr) const override;
