Fixed wrong object being validated in FObjectProperty::PostSerializeObjectItem.

diff --git a/Engine/Source/Runtime/CoreUObject/Private/UObject/PropertyObject.cpp b/Engine/Source/Runtime/CoreUObject/Private/UObject/PropertyObject.cpp
--- a/Engine/Source/Runtime/CoreUObject/Private/UObject/PropertyObject.cpp
+++ b/Engine/Source/Runtime/CoreUObject/Private/UObject/PropertyObject.cpp
@@ -292,7 +292,7 @@ void FObjectProperty::PostSerializeObjectItem(FArchive& SerializingArchive, void
 		//        to accommodate this (as it depends on finding itself as the set value)
 #endif // USE_CIRCULAR_DEPENDENCY_LOAD_DEFERRING
 
-		CheckValidObject(Value, CurrentValue);
+		CheckValidObject(Value, ObjectValue);
 	}
 }
 
