[AutoRTFM] Fix memory leaking from FUObjectItem::StatIDStringStorage

If the transaction was aborted between CreateStatID() and the FUObjectItem destructor, then the memory allocated in the open for the StatID was not freed.
Allocate the memory in the open, and carefully assign / reset StatIDStringStorage atomically in the open / abort handler.

Also fix / silence some other leaks in UECoreTests.cpp

diff --git a/Engine/Source/Programs/AutoRTFMTests/Private/UECoreTests.cpp b/Engine/Source/Programs/AutoRTFMTests/Private/UECoreTests.cpp
--- a/Engine/Source/Programs/AutoRTFMTests/Private/UECoreTests.cpp
+++ b/Engine/Source/Programs/AutoRTFMTests/Private/UECoreTests.cpp
@@ -3,6 +3,7 @@
 #include "Catch2Includes.h"
 #include "AutoRTFM/AutoRTFM.h"
 #include "Delegates/IDelegateInstance.h"
+#include "HAL/MallocLeakDetection.h"
 #include "HAL/ThreadSingleton.h"
 #include "Internationalization/TextCache.h"
 #include "Internationalization/TextFormatter.h"
@@ -74,6 +75,8 @@ TEST_CASE("UECore.TThreadSingleton")
 
 	SECTION("Get")
 	{
+		MALLOCLEAK_IGNORE_SCOPE(); // TThreadSingleton will appear as a leak.
+
 		AutoRTFM::ETransactionResult Result = AutoRTFM::Transact([&]()
 			{
 				TThreadSingleton<MyStruct>::Get().I = 42;
@@ -543,6 +546,7 @@ TEST_CASE("UECore.TIntrusiveReferenceController")
 
 			REQUIRE(AutoRTFM::ETransactionResult::AbortedByRequest == Result);
 			REQUIRE(1 == Controller->WeakReferenceCount);
+			delete Controller;
 		}
 
 		SECTION("With Commit")
@@ -853,10 +857,8 @@ TEST_CASE("UECore.FUObjectItem")
 			});
 
 		REQUIRE(AutoRTFM::ETransactionResult::AbortedByRequest == Result);
-		REQUIRE(nullptr != Item.StatIDStringStorage);
-		REQUIRE(Item.StatID.IsValidStat());
-
-		PROFILER_CHAR* const StatIDStringStorage = Item.StatIDStringStorage;
+		REQUIRE(nullptr == Item.StatIDStringStorage);
+		REQUIRE(!Item.StatID.IsValidStat());
 
 		// If we commit though we'll create the stat ID.
 		Result = AutoRTFM::Transact([&]
@@ -865,7 +867,7 @@ TEST_CASE("UECore.FUObjectItem")
 			});
 
 		REQUIRE(AutoRTFM::ETransactionResult::Committed == Result);
-		REQUIRE(StatIDStringStorage == Item.StatIDStringStorage);
+		REQUIRE(nullptr != Item.StatIDStringStorage);
 		REQUIRE(Item.StatID.IsValidStat());
 	}
 
@@ -933,8 +935,8 @@ TEST_CASE("UECore.FUObjectItem")
 				});
 
 			REQUIRE(AutoRTFM::ETransactionResult::AbortedByRequest == Result);
-			REQUIRE(nullptr != Item.StatIDStringStorage);
-			REQUIRE(Item.StatID.IsValidStat());
+			REQUIRE(nullptr == Item.StatIDStringStorage);
+			REQUIRE(!Item.StatID.IsValidStat());
 		}
 
 		{
diff --git a/Engine/Source/Runtime/CoreUObject/Private/UObject/UObjectArray.cpp b/Engine/Source/Runtime/CoreUObject/Private/UObject/UObjectArray.cpp
--- a/Engine/Source/Runtime/CoreUObject/Private/UObject/UObjectArray.cpp
+++ b/Engine/Source/Runtime/CoreUObject/Private/UObject/UObjectArray.cpp
@@ -5,6 +5,7 @@
 =============================================================================*/
 
 #include "UObject/UObjectArray.h"
+#include "AutoRTFM/AutoRTFM.h"
 #include "HAL/IConsoleManager.h"
 #include "HAL/LowLevelMemStats.h"
 #include "Misc/ScopeLock.h"
@@ -68,21 +69,32 @@ void FUObjectItem::CreateStatID() const
 #if STATS
 	StatID = FDynamicStats::CreateStatId<FStatGroup_STATGROUP_UObjects>(LongName);
 #else // ENABLE_STATNAMEDEVENTS
-	UE_AUTORTFM_OPEN2
-	{
-		const auto& ConversionData = StringCast<PROFILER_CHAR>(*LongName);
-		const int32 NumStorageChars = (ConversionData.Length() + 1);	//length doesn't include null terminator
+	const auto& ConversionData = StringCast<PROFILER_CHAR>(*LongName);
+	const int32 NumStorageChars = (ConversionData.Length() + 1);	//length doesn't include null terminator
+
+	PROFILER_CHAR* const StoragePtr = new PROFILER_CHAR[NumStorageChars];
+	FMemory::Memcpy(StoragePtr, ConversionData.Get(), NumStorageChars * sizeof(PROFILER_CHAR));
 
-		PROFILER_CHAR* const StoragePtr = new PROFILER_CHAR[NumStorageChars];
-		FMemory::Memcpy(StoragePtr, ConversionData.Get(), NumStorageChars * sizeof(PROFILER_CHAR));
+	bool bExchanged = AutoRTFM::Open([&]
+	{ 
+		return FPlatformAtomics::InterlockedCompareExchangePointer((void**)&StatIDStringStorage, StoragePtr, nullptr) == nullptr;
+	});
 
-		if (FPlatformAtomics::InterlockedCompareExchangePointer((void**)&StatIDStringStorage, StoragePtr, nullptr) != nullptr)
+	if (bExchanged)
+	{
+		// If we abort, then StoragePtr will be freed so reset StatIDStringStorage.
+		// This is abort handler is popped in the destructor.
+		AutoRTFM::PushOnAbortHandler(StoragePtr, [this, StoragePtr]
 		{
-			delete[] StoragePtr;
-		}
+			FPlatformAtomics::InterlockedCompareExchangePointer((void**)&StatIDStringStorage, nullptr, StoragePtr);
+		});
+	}
+	else
+	{
+		delete[] StoragePtr;
+	}
 
-		StatID = TStatId(StatIDStringStorage);
-	};
+	StatID = TStatId(StatIDStringStorage);
 #endif
 }
 #endif
diff --git a/Engine/Source/Runtime/CoreUObject/Public/UObject/UObjectArray.h b/Engine/Source/Runtime/CoreUObject/Public/UObject/UObjectArray.h
--- a/Engine/Source/Runtime/CoreUObject/Public/UObject/UObjectArray.h
+++ b/Engine/Source/Runtime/CoreUObject/Public/UObject/UObjectArray.h
@@ -6,6 +6,7 @@
 
 #pragma once
 
+#include "AutoRTFM/AutoRTFM.h"
 #include "HAL/ThreadSafeCounter.h"
 #include "Containers/LockFreeList.h"
 #include "Misc/TransactionallySafeCriticalSection.h"
@@ -82,8 +83,11 @@ public:
 	~FUObjectItem()
 	{
 #if ENABLE_STATNAMEDEVENTS_UOBJECT
-		delete[] StatIDStringStorage;
-		StatIDStringStorage = nullptr;
+		if (PROFILER_CHAR* Storage = StatIDStringStorage)
+		{
+			AutoRTFM::PopOnAbortHandler(Storage);
+			delete[] Storage;
+		}
 #endif
 	}
 
