Improving asset leak reporting when unloading plugins by ignoring the effects of the RF_Standalone flag.

diff --git a/Engine/Source/Runtime/CoreUObject/Private/Misc/CoreUObjectPluginManager.cpp b/Engine/Source/Runtime/CoreUObject/Private/Misc/CoreUObjectPluginManager.cpp
--- a/Engine/Source/Runtime/CoreUObject/Private/Misc/CoreUObjectPluginManager.cpp
+++ b/Engine/Source/Runtime/CoreUObject/Private/Misc/CoreUObjectPluginManager.cpp
@@ -201,6 +201,13 @@ namespace UE::CoreUObject::Private
 				Options |= EPrintStaleReferencesOptions::Minimal;
 			}
 
+			// Make sure we are running on the game thread because we are going to be modifying
+			// object flags to improve the stale reference reporting later on.
+			check(IsInGameThread());
+
+			TSet<UObject*> ObjectsThatHadStandaloneCleared;
+			TArray<UObject*> ObjectsInPackage;
+
 			// Sort even if we don't limit the count, so that high priority leaks appear first 
 			int32 OmittedCount = FMath::Max(0, LeakedPackages.Num() - GLeakedAssetTrace_MaxReportCount);
 			TArray<UPackage*> PackagesToSearchFor;
@@ -210,11 +217,31 @@ namespace UE::CoreUObject::Private
 				if (i++ < GLeakedAssetTrace_MaxReportCount)
 				{
 					PackagesToSearchFor.Add(Package);
+
+					ObjectsInPackage.Reset();
+
+					// To improve the reporting for stale references we clear out the RF_Standalone flag on
+					// every object in the packages we will check. The flags will be restored after the Find.
+					GetObjectsWithPackage(Package, ObjectsInPackage, false);
+
+					for (UObject* Object : ObjectsInPackage)
+					{
+						if (Object->HasAnyFlags(RF_Standalone))
+						{
+							Object->ClearFlags(RF_Standalone);
+							ObjectsThatHadStandaloneCleared.Add(Object);
+						}
+					}
 				}
 			}
 
 			UE_LOG(PluginHandlerLog, Display, TEXT("Searching for references to %d leaked packages (%d omitted for speed) from plugin %s"), LeakedPackages.Num(), OmittedCount, *PluginName);
 			FReferenceChainSearch::FindAndPrintStaleReferencesToObjects(MakeArrayView((UObject**)PackagesToSearchFor.GetData(), PackagesToSearchFor.Num()), Options);
+
+			for (UObject* Object : ObjectsThatHadStandaloneCleared)
+			{
+				Object->SetFlags(RF_Standalone);
+			}
 		}
 
 		// Rename the packages that we are streaming out so that we can possibly reload another copy of them
