Cooker: Fix bug in MPCook that would allow the cooking of __ExternalObject__ packages (or any other package marked as PKG_EditorOnly).
Remove canskipeditoronlywhencooking; it is no longer necessary because its functionality has been replaced by skiponlyeditoronly.
MPCook was handling PKG_EditorOnly because it could not handle skipeditoronlywhencooking, and the two concepts used the same flag to toggle use of them in SavePackage. By removing canskipeditoronlywhencooking, we can turn that flag on.
#rnx

diff --git a/Engine/Source/Editor/UnrealEd/Private/CookOnTheFlyServer.cpp b/Engine/Source/Editor/UnrealEd/Private/CookOnTheFlyServer.cpp
--- a/Engine/Source/Editor/UnrealEd/Private/CookOnTheFlyServer.cpp
+++ b/Engine/Source/Editor/UnrealEd/Private/CookOnTheFlyServer.cpp
@@ -4537,14 +4537,6 @@ void UCookOnTheFlyServer::PumpSaves(UE::Cook::FTickStackData& StackData, uint32
 		UE_LOG(LogCook, Display, TEXT("Processing save for package %s"), *Package->GetName());
 #endif
 
-		if (Package->IsLoadedByEditorPropertiesOnly() && PackageTracker->UncookedEditorOnlyPackages.Contains(Package->GetFName()))
-		{
-			// We already attempted to cook this package and it's still not referenced by any non editor-only properties.
-			DemoteToIdle(PackageData, ESendFlags::QueueAdd, ESuppressCookReason::OnlyEditorOnly);
-			++OutNumPushed;
-			continue;
-		}
-
 		// This package is valid, so make sure it wasn't previously marked as being an uncooked editor only package or it would get removed from the
 		// asset registry at the end of the cook
 		PackageTracker->UncookedEditorOnlyPackages.Remove(Package->GetFName());
@@ -6339,15 +6331,6 @@ void FSaveCookedPackageContext::SetupPackage()
 		SaveFlags |= SAVE_CookSoftPackageReferences;
 	}
 
-	// removing editor only packages only works when cooking in commandlet and non iterative cooking
-	bool bCanSkipEditorOnlyPackages = COTFS.IsCookByTheBookMode() && !COTFS.IsCookingInEditor();
-	bCanSkipEditorOnlyPackages &= !COTFS.IsCookFlagSet(ECookInitializationFlags::Iterative);
-	// MPCOOKTODO: it also doesn't work in multiprocess cooking, because GetPackage()->IsLoadedByEditorPropertiesOnly()
-	// might have been set to false on a CookWorker and is not replicated. Rather than fixing this, we should delete
-	// the feature and instead use SkipOnlyEditorOnly.
-	bCanSkipEditorOnlyPackages &= !COTFS.CookDirector.IsValid() && !COTFS.CookWorkerClient.IsValid();
-	SaveFlags |= bCanSkipEditorOnlyPackages ? SAVE_None : SAVE_KeepEditorOnlyCookedPackages;
-
 	// Use SandboxFile to do path conversion to properly handle sandbox paths (outside of standard paths in particular).
 	Filename = COTFS.ConvertToFullSandboxPath(*Filename, true);
 }
diff --git a/Engine/Source/Editor/UnrealEd/Private/Cooker/CookGenerationHelper.cpp b/Engine/Source/Editor/UnrealEd/Private/Cooker/CookGenerationHelper.cpp
--- a/Engine/Source/Editor/UnrealEd/Private/Cooker/CookGenerationHelper.cpp
+++ b/Engine/Source/Editor/UnrealEd/Private/Cooker/CookGenerationHelper.cpp
@@ -1097,7 +1097,6 @@ UPackage* FGenerationHelper::TryCreateGeneratedPackage(FCookGenerationInfo& Gene
 	GeneratedPackage->SetSavedHash(GeneratedInfo.PackageHash);
 	GeneratedPackage->SetPersistentGuid(LocalOwnerPackage->GetPersistentGuid());
 	GeneratedPackage->SetPackageFlags(PKG_CookGenerated);
-	GeneratedPackage->SetLoadedByEditorPropertiesOnly(LocalOwnerPackage->IsLoadedByEditorPropertiesOnly());
 	GeneratedInfo.SetHasCreatedPackage(true);
 
 	return GeneratedPackage;
diff --git a/Engine/Source/Runtime/CoreUObject/Private/Serialization/AsyncLoading.cpp b/Engine/Source/Runtime/CoreUObject/Private/Serialization/AsyncLoading.cpp
--- a/Engine/Source/Runtime/CoreUObject/Private/Serialization/AsyncLoading.cpp
+++ b/Engine/Source/Runtime/CoreUObject/Private/Serialization/AsyncLoading.cpp
@@ -6180,10 +6180,6 @@ EAsyncPackageState::Type FAsyncPackage::CreateLinker()
 
 			// Always store package filename we loading from
 			Package->SetLoadedPath(Desc.PackagePath);
-#if WITH_EDITORONLY_DATA
-			// Assume all packages loaded through async loading are required by runtime
-			Package->SetLoadedByEditorPropertiesOnly(false);
-#endif
 		}
 
 		LastObjectWorkWasPerformedOn = Package;
diff --git a/Engine/Source/Runtime/CoreUObject/Private/UObject/LinkerLoad.cpp b/Engine/Source/Runtime/CoreUObject/Private/UObject/LinkerLoad.cpp
--- a/Engine/Source/Runtime/CoreUObject/Private/UObject/LinkerLoad.cpp
+++ b/Engine/Source/Runtime/CoreUObject/Private/UObject/LinkerLoad.cpp
@@ -4340,40 +4340,7 @@ UObject* FLinkerLoad::ResolveResource(FPackageIndex Index)
 		}
 	}
 
-	UObject* Temporary = NULL;
-	Temporary = IndexToObject(Index);
-
-#if WITH_EDITORONLY_DATA	
-	// When loading mark all packages that are accessed by non editor-only properties as being required at runtime.
-	if (Ar.IsLoading() && Temporary && !Ar.IsEditorOnlyPropertyOnTheStack())
-	{
-		const bool bReferenceFromOutsideOfThePackage = Temporary->GetOutermost() != LinkerRoot;
-		const bool bIsAClass = Temporary->IsA(UClass::StaticClass());
-		const bool bReferencingPackageIsNotEditorOnly = bReferenceFromOutsideOfThePackage && !LinkerRoot->IsLoadedByEditorPropertiesOnly();
-		if (bReferencingPackageIsNotEditorOnly || bIsAClass)
-		{
-			// The package that caused this object to be loaded is not marked as editor-only, neighter is any of the referencing properties.
-			Temporary->GetOutermost()->SetLoadedByEditorPropertiesOnly(false);
-		}
-		else if (bReferenceFromOutsideOfThePackage && !bIsAClass)
-		{
-			// In this case the object is being accessed by object property from a package that's marked as editor-only, however
-			// since we're in the middle of loading, we can't be sure that the editor-only package will still be marked as editor-only
-			// after loading has finished (this is due to the fact how objects are being processed in EndLoad).
-			// So we need to remember which packages have been kept marked as editor-only by which package so that after all
-			// objects have been serialized we can go back and make sure the LinkerRoot package is still marked as editor-only and if not,
-			// remove the flag from all packages that are marked as such because of it.
-			FUObjectThreadContext& ThreadContext = FUObjectThreadContext::Get();
-			TSet<FName>& PackagesMarkedEditorOnly = ThreadContext.PackagesMarkedEditorOnlyByOtherPackage.FindOrAdd(LinkerRoot->GetFName());
-			if (!PackagesMarkedEditorOnly.Contains(Temporary->GetOutermost()->GetFName()))
-			{
-				PackagesMarkedEditorOnly.Add(Temporary->GetOutermost()->GetFName());
-			}
-		}
-	}
-#endif
-
-	return Temporary;
+	return IndexToObject(Index);
 }
 
 // Find the index of a specified object without regard to specific package.
@@ -7057,40 +7024,6 @@ void FLinkerLoad::FixupSoftObjectPathForInstancedPackage(FSoftObjectPath& InOutS
 	InstancingContext.FixupSoftObjectPath(InOutSoftObjectPath);
 }
 
-#if WITH_EDITORONLY_DATA
-/** Performs a fixup on packages' editor-only flag */
-void FixupPackageEditorOnlyFlag(FName PackageThatGotEditorOnlyFlagCleared, bool bRecursive)
-{
-	FUObjectThreadContext& ThreadContext = FUObjectThreadContext::Get();
-	STAT(double ThisTime = 0);
-	{
-		SCOPE_SECONDS_COUNTER(ThisTime);
-
-		// Now go through all packages that were marked as editor-only at load time
-		// and if they're no longer marked as such, make sure that all packages that
-		// were marked as editor-only because of that package, are now also marked as not editor-only.
-		TSet<FName>* PackagesMarkedEditorOnlyByThisPackage = ThreadContext.PackagesMarkedEditorOnlyByOtherPackage.Find(PackageThatGotEditorOnlyFlagCleared);
-		if (PackagesMarkedEditorOnlyByThisPackage)
-		{			
-			for (FName& PackageName : *PackagesMarkedEditorOnlyByThisPackage)
-			{
-				UPackage* EditorOnlyPackage = FindObjectFast<UPackage>(nullptr, PackageName);
-				if (EditorOnlyPackage && EditorOnlyPackage->IsLoadedByEditorPropertiesOnly())
-				{
-					// Now we will recursively unset the flag on all other packages
-					EditorOnlyPackage->SetLoadedByEditorPropertiesOnly(false, true);
-				}
-			}
-			ThreadContext.PackagesMarkedEditorOnlyByOtherPackage.Remove(PackageThatGotEditorOnlyFlagCleared);
-		}
-	}
-	if (!bRecursive)
-	{
-		INC_FLOAT_STAT_BY(STAT_EditorOnlyFixupTime, ThisTime);
-	}
-}
-#endif
-
 #if WITH_EDITOR
 bool FLinkerLoad::bPreloadingEnabled = false;
 bool FLinkerLoad::GetPreloadingEnabled()
diff --git a/Engine/Source/Runtime/CoreUObject/Private/UObject/Package.cpp b/Engine/Source/Runtime/CoreUObject/Private/UObject/Package.cpp
--- a/Engine/Source/Runtime/CoreUObject/Private/UObject/Package.cpp
+++ b/Engine/Source/Runtime/CoreUObject/Private/UObject/Package.cpp
@@ -380,16 +380,6 @@ bool UPackage::IsPostLoadThreadSafe() const
 }
 
 #if WITH_EDITORONLY_DATA
-void FixupPackageEditorOnlyFlag(FName PackageThatGotEditorOnlyFlagCleared, bool bRecursive);
-
-void UPackage::SetLoadedByEditorPropertiesOnly(bool bIsEditorOnly, bool bRecursive /*= false*/)
-{
-	const bool bWasEditorOnly = bLoadedByEditorPropertiesOnly.exchange(bIsEditorOnly);
-	if (bWasEditorOnly && !bIsEditorOnly)
-	{
-		FixupPackageEditorOnlyFlag(GetFName(), bRecursive);
-	}
-}
 
 FIoHash UPackage::GetSavedHash()
 {
diff --git a/Engine/Source/Runtime/CoreUObject/Private/UObject/SavePackage/SaveContext.h b/Engine/Source/Runtime/CoreUObject/Private/UObject/SavePackage/SaveContext.h
--- a/Engine/Source/Runtime/CoreUObject/Private/UObject/SavePackage/SaveContext.h
+++ b/Engine/Source/Runtime/CoreUObject/Private/UObject/SavePackage/SaveContext.h
@@ -738,11 +738,6 @@ public:
 		return !!(SaveArgs.SaveFlags & ESaveFlags::SAVE_CompareLinker);
 	}
 
-	bool CanSkipEditorReferencedPackagesWhenCooking() const
-	{
-		return SkipEditorRefCookingSetting;
-	}
-
 	bool IsIgnoringHeaderDiff() const
 	{
 		return bIgnoreHeaderDiffs;
diff --git a/Engine/Source/Runtime/CoreUObject/Private/UObject/SavePackage/SavePackageUtilities.cpp b/Engine/Source/Runtime/CoreUObject/Private/UObject/SavePackage/SavePackageUtilities.cpp
--- a/Engine/Source/Runtime/CoreUObject/Private/UObject/SavePackage/SavePackageUtilities.cpp
+++ b/Engine/Source/Runtime/CoreUObject/Private/UObject/SavePackage/SavePackageUtilities.cpp
@@ -2032,20 +2032,25 @@ FScopedSavingFlag::~FScopedSavingFlag()
 }
 
 FCanSkipEditorReferencedPackagesWhenCooking::FCanSkipEditorReferencedPackagesWhenCooking()
-	: bCanSkipEditorReferencedPackagesWhenCooking(UE::SavePackageUtilities::CanSkipEditorReferencedPackagesWhenCooking())
 {
+	//UE_DEPRECATED(5.5, TEXT("No longer used; skiponlyeditoronly is used instead and tracks editoronly references via savepackage results."))
+	static bool bWarned = false;
+	if (!bWarned)
+	{
+		bWarned = true;
+		bool bResult = true;
+		GConfig->GetBool(TEXT("Core.System"), TEXT("CanSkipEditorReferencedPackagesWhenCooking"), bResult, GEngineIni);
+		if (bResult)
+		{
+			UE_LOG(LogSavePackage, Warning,
+				TEXT("Engine.ini:[Core.System]:CanSkipEditorReferencedPackagesWhenCooking is deprecated; it is replaced by Editor.ini:[CookSettings]:SkipOnlyEditorOnly. Remove this setting from your inis."));
+		}
+	}
 }
 
 namespace UE::SavePackageUtilities
 {
 
-bool CanSkipEditorReferencedPackagesWhenCooking()
-{
-	bool bResult = true;
-	GConfig->GetBool(TEXT("Core.System"), TEXT("CanSkipEditorReferencedPackagesWhenCooking"), bResult, GEngineIni);
-	return bResult;
-}
-
 /**
  * Static: Saves thumbnail data for the specified package outer and linker
  *
diff --git a/Engine/Source/Runtime/CoreUObject/Private/UObject/SavePackage/SavePackageUtilities.h b/Engine/Source/Runtime/CoreUObject/Private/UObject/SavePackage/SavePackageUtilities.h
--- a/Engine/Source/Runtime/CoreUObject/Private/UObject/SavePackage/SavePackageUtilities.h
+++ b/Engine/Source/Runtime/CoreUObject/Private/UObject/SavePackage/SavePackageUtilities.h
@@ -71,9 +71,7 @@ private:
 
 struct FCanSkipEditorReferencedPackagesWhenCooking
 {
-	bool bCanSkipEditorReferencedPackagesWhenCooking;
 	FCanSkipEditorReferencedPackagesWhenCooking();
-	FORCEINLINE operator bool() const { return bCanSkipEditorReferencedPackagesWhenCooking; }
 };
 
 
diff --git a/Engine/Source/Runtime/CoreUObject/Private/UObject/SavePackage2.cpp b/Engine/Source/Runtime/CoreUObject/Private/UObject/SavePackage2.cpp
--- a/Engine/Source/Runtime/CoreUObject/Private/UObject/SavePackage2.cpp
+++ b/Engine/Source/Runtime/CoreUObject/Private/UObject/SavePackage2.cpp
@@ -102,12 +102,7 @@ ESavePackageResult ValidatePackage(FSaveContext& SaveContext)
 		if (SaveContext.IsStripEditorOnly())
 		{
 			// Don't save packages marked as editor-only.
-			if (SaveContext.CanSkipEditorReferencedPackagesWhenCooking() && SaveContext.GetPackage()->IsLoadedByEditorPropertiesOnly())
-			{
-				UE_CLOG(SaveContext.IsGenerateSaveError(), LogSavePackage, Verbose, TEXT("Package loaded by editor-only properties: %s. Package will not be saved."), *SaveContext.GetPackage()->GetName());
-				return ESavePackageResult::ReferencedOnlyByEditorOnlyData;
-			}
-			else if (SaveContext.GetPackage()->HasAnyPackageFlags(PKG_EditorOnly))
+			if (SaveContext.GetPackage()->HasAnyPackageFlags(PKG_EditorOnly))
 			{
 				UE_CLOG(SaveContext.IsGenerateSaveError(), LogSavePackage, Verbose, TEXT("Package marked as editor-only: %s. Package will not be saved."), *SaveContext.GetPackage()->GetName());
 				return ESavePackageResult::ReferencedOnlyByEditorOnlyData;
diff --git a/Engine/Source/Runtime/CoreUObject/Private/UObject/UObjectGlobals.cpp b/Engine/Source/Runtime/CoreUObject/Private/UObject/UObjectGlobals.cpp
--- a/Engine/Source/Runtime/CoreUObject/Private/UObject/UObjectGlobals.cpp
+++ b/Engine/Source/Runtime/CoreUObject/Private/UObject/UObjectGlobals.cpp
@@ -1453,12 +1453,6 @@ UObject* StaticLoadObjectInternal(UClass* ObjectClass, UObject* InOuter, const T
 		StrName += FPackageName::GetShortName(InName);
 		Result = StaticLoadObjectInternal(ObjectClass, InOuter, *StrName, Filename, LoadFlags, Sandbox, bAllowObjectReconciliation, InstancingContext);
 	}
-#if WITH_EDITORONLY_DATA
-	else if (Result && !(LoadFlags & LOAD_EditorOnly))
-	{
-		Result->GetOutermost()->SetLoadedByEditorPropertiesOnly(false);
-	}
-#endif
 
 	if (Result && UE::GC::GIsIncrementalReachabilityPending)
 	{
@@ -1895,32 +1889,6 @@ UPackage* LoadPackageInternal(UPackage* InOuter, const FPackagePath& PackagePath
 			Result->ThisRequiresLocalizationGather(Linker->RequiresLocalizationGather());
 		};
 
-#if WITH_EDITORONLY_DATA
-		if (!(LoadFlags & (LOAD_IsVerifying|LOAD_EditorOnly)))
-		{
-			bool bIsEditorOnly = false;
-			FProperty* SerializingProperty = ImportLinker ? ImportLinker->GetSerializedProperty() : nullptr;
-			
-			// Check property parent chain
-			while (SerializingProperty)
-			{
-				if (SerializingProperty->IsEditorOnlyProperty())
-				{
-					bIsEditorOnly = true;
-					break;
-				}
-				SerializingProperty = SerializingProperty->GetOwner<FProperty>();
-			}
-
-			if (!bIsEditorOnly)
-			{
-				// If this package hasn't been loaded as part of import verification and there's no import linker or the
-				// currently serialized property is not editor-only mark this package as runtime.
-				Result->SetLoadedByEditorPropertiesOnly(false);
-			}
-		}
-#endif
-
 		if (Result->HasAnyFlags(RF_WasLoaded))
 		{
 			// The linker is associated with a package that has already been loaded.
diff --git a/Engine/Source/Runtime/CoreUObject/Private/UObject/UObjectThreadContext.cpp b/Engine/Source/Runtime/CoreUObject/Private/UObject/UObjectThreadContext.cpp
--- a/Engine/Source/Runtime/CoreUObject/Private/UObject/UObjectThreadContext.cpp
+++ b/Engine/Source/Runtime/CoreUObject/Private/UObject/UObjectThreadContext.cpp
@@ -22,9 +22,15 @@ FUObjectThreadContext::FUObjectThreadContext()
 , SerializeContext(new FUObjectSerializeContext())
 {}
 
+PRAGMA_DISABLE_DEPRECATION_WARNINGS;
 FUObjectThreadContext::~FUObjectThreadContext()
 {
+#if WITH_EDITORONLY_DATA
+	// Remove PRAGMA_DISABLE_DEPRECATION_WARNINGS when PackagesMarkedEditorOnlyByOtherPackage is deleted
+	(void)PackagesMarkedEditorOnlyByOtherPackage;
+#endif
 }
+PRAGMA_ENABLE_DEPRECATION_WARNINGS;
 
 FObjectInitializer& FUObjectThreadContext::ReportNull()
 {
diff --git a/Engine/Source/Runtime/CoreUObject/Public/UObject/Package.h b/Engine/Source/Runtime/CoreUObject/Public/UObject/Package.h
--- a/Engine/Source/Runtime/CoreUObject/Public/UObject/Package.h
+++ b/Engine/Source/Runtime/CoreUObject/Public/UObject/Package.h
@@ -336,10 +336,11 @@ public:
 	virtual bool IsPostLoadThreadSafe() const override;
 
 #if WITH_EDITORONLY_DATA
-	/** Sets the bLoadedByEditorPropertiesOnly flag */
-	COREUOBJECT_API void SetLoadedByEditorPropertiesOnly(bool bIsEditorOnly, bool bRecursive = false);
+	UE_DEPRECATED(5.5, "No longer used; skiponlyeditoronly is used instead and tracks editoronly references via savepackage results.")
+	void SetLoadedByEditorPropertiesOnly(bool bIsEditorOnly, bool bRecursive = false) {}
 	/** returns true when the package is only referenced by editor-only flag */
-	bool IsLoadedByEditorPropertiesOnly() const { return bLoadedByEditorPropertiesOnly; }
+	UE_DEPRECATED(5.5, "No longer used; skiponlyeditoronly is used instead and tracks editoronly references via savepackage results.")
+	bool IsLoadedByEditorPropertiesOnly() const { return false; }
 
 	/** Sets the bIsDynamicPIEPackagePending flag */
 	void SetDynamicPIEPackagePending(bool bInIsDynamicPIEPackagePending)
diff --git a/Engine/Source/Runtime/CoreUObject/Public/UObject/SavePackage.h b/Engine/Source/Runtime/CoreUObject/Public/UObject/SavePackage.h
--- a/Engine/Source/Runtime/CoreUObject/Public/UObject/SavePackage.h
+++ b/Engine/Source/Runtime/CoreUObject/Public/UObject/SavePackage.h
@@ -326,7 +326,8 @@ namespace UE::SavePackageUtilities
 	COREUOBJECT_API void EDLCookInfoMoveToCompactBinaryAndClear(FCbWriter& Writer, bool& bOutHasData);
 	COREUOBJECT_API void EDLCookInfoMoveToCompactBinaryAndClear(FCbWriter& Writer, bool& bOutHasData, FName PackageName);
 	COREUOBJECT_API bool EDLCookInfoAppendFromCompactBinary(FCbFieldView Field);
-	COREUOBJECT_API bool CanSkipEditorReferencedPackagesWhenCooking();
+	UE_DEPRECATED(5.5, "No longer used; skiponlyeditoronly is used instead and tracks editoronly references via savepackage results.")
+	inline bool CanSkipEditorReferencedPackagesWhenCooking() { return false; }
 
 
 #if WITH_EDITOR
diff --git a/Engine/Source/Runtime/CoreUObject/Public/UObject/UObjectThreadContext.h b/Engine/Source/Runtime/CoreUObject/Public/UObject/UObjectThreadContext.h
--- a/Engine/Source/Runtime/CoreUObject/Public/UObject/UObjectThreadContext.h
+++ b/Engine/Source/Runtime/CoreUObject/Public/UObject/UObjectThreadContext.h
@@ -110,7 +110,7 @@ public:
 	TArray<UObject*, TInlineAllocator<16> > DebugPostLoad;
 #endif
 #if WITH_EDITORONLY_DATA
-	/** Maps a package name to all packages marked as editor-only due to the fact it was marked as editor-only */
+	UE_DEPRECATED(5.5, "No longer used; skiponlyeditoronly is used instead and tracks editoronly references via savepackage results.")
 	TMap<FName, TSet<FName>> PackagesMarkedEditorOnlyByOtherPackage;
 #endif
 
