[Bugfix] Crash when adding to Container properties

diff --git a/Engine/Source/Runtime/CoreUObject/Private/UObject/PropertyBagRepository.cpp b/Engine/Source/Runtime/CoreUObject/Private/UObject/PropertyBagRepository.cpp
--- a/Engine/Source/Runtime/CoreUObject/Private/UObject/PropertyBagRepository.cpp
+++ b/Engine/Source/Runtime/CoreUObject/Private/UObject/PropertyBagRepository.cpp
@@ -240,6 +240,10 @@ static bool ConstructRemappedPropertyChain(const FEditPropertyChain& Chain, FEdi
 		{
 			Struct = AsStructProperty->Struct;
 		}
+		else if (FObjectProperty* AsObjectProperty = CastField<FObjectProperty>(Property))
+		{
+			Struct = AsObjectProperty->PropertyClass;
+		}
 		else
 		{
 			check(Itr->GetNextNode() == nullptr);
@@ -258,7 +262,7 @@ static bool ConstructRemappedPropertyChain(const FEditPropertyChain& Chain, FEdi
 	return true;
 }
 
-static void* ResolveChangePath(const void* StructData, FPropertyChangedChainEvent& ChangeEvent, bool bGrowContainersWhenNeeded = false)
+static void* ResolveChangePath(const void* StructData, FPropertyChangedChainEvent& ChangeEvent)
 {
 	if (ChangeEvent.PropertyChain.GetHead() == nullptr)
 	{
@@ -271,12 +275,13 @@ static void* ResolveChangePath(const void* StructData, FPropertyChangedChainEven
 	{
 		const FProperty* Property = PropertyNode->GetValue();
 		MemoryPtr = Property->ContainerPtrToValuePtr<uint8>(MemoryPtr);
-		PropertyNode = PropertyNode->GetNextNode();
 		
 		const int32 ArrayIndex = ChangeEvent.GetArrayIndex(Property->GetName());
-		if (PropertyNode && ArrayIndex != INDEX_NONE)
+		if (ArrayIndex != INDEX_NONE)
 		{
-			if (const FArrayProperty* ArrayProperty = CastField<FArrayProperty>(Property->GetOwnerProperty()))
+			PropertyNode = PropertyNode->GetNextNode();
+			FProperty* InnerProperty = PropertyNode ? PropertyNode->GetValue() : ChangeEvent.Property;
+			if (const FArrayProperty* ArrayProperty = CastField<FArrayProperty>(Property))
 			{
 				FScriptArrayHelper ArrayHelper(ArrayProperty, MemoryPtr);
 				if(!ArrayHelper.IsValidIndex(ArrayIndex))
@@ -284,11 +289,8 @@ static void* ResolveChangePath(const void* StructData, FPropertyChangedChainEven
 					return nullptr;
 				}
 				MemoryPtr = ArrayHelper.GetRawPtr(ArrayIndex);
-
-				// skip to the next property node already
-				PropertyNode = PropertyNode->GetNextNode();
 			}
-			if (const FSetProperty* SetProperty = CastField<FSetProperty>(Property->GetOwnerProperty()))
+			else if (const FSetProperty* SetProperty = CastField<FSetProperty>(Property))
 			{
 				FScriptSetHelper SetHelper(SetProperty, MemoryPtr);
 				if(!SetHelper.IsValidIndex(ArrayIndex))
@@ -296,28 +298,144 @@ static void* ResolveChangePath(const void* StructData, FPropertyChangedChainEven
 					return nullptr;
 				}
 				MemoryPtr = SetHelper.GetElementPtr(ArrayIndex);
-
-				// skip to the next property node already
-				PropertyNode = PropertyNode->GetNextNode();
 			}
-			if (const FMapProperty* MapProperty = CastField<FMapProperty>(Property->GetOwnerProperty()))
+			else if (const FMapProperty* MapProperty = CastField<FMapProperty>(Property))
 			{
 				FScriptMapHelper MapHelper(MapProperty, MemoryPtr);
 				if(!MapHelper.IsValidIndex(ArrayIndex))
 				{
 					return nullptr;
 				}
-				MemoryPtr = MapHelper.GetValuePtr(ArrayIndex);
-
-				// skip to the next property node already
-				PropertyNode = PropertyNode->GetNextNode();
+				if (InnerProperty == MapProperty->KeyProp)
+				{
+					MemoryPtr = MapHelper.GetKeyPtr(ArrayIndex);
+				}
+				else
+				{
+					MemoryPtr = MapHelper.GetValuePtr(ArrayIndex);
+				}
 			}
 		}
+		else if (const FObjectProperty* AsObjectProperty = CastField<FObjectProperty>(Property))
+		{
+			MemoryPtr = AsObjectProperty->GetObjectPropertyValue(MemoryPtr);
+		}
+
+		
+		if (PropertyNode)
+		{
+			PropertyNode = PropertyNode->GetNextNode();
+		}
 	}
 	while (PropertyNode);
 
 	return MemoryPtr;
 }
+
+static bool RemapChangeEvent(const FPropertyChangedChainEvent& InChangeEvent, FPropertyChangedChainEvent& OutRemappedChangeEvent, void*& MemoryPtr, TArray<TMap<FString, int32>>& ArrayIndices)
+{
+	
+	if (!ConstructRemappedPropertyChain(InChangeEvent.PropertyChain, OutRemappedChangeEvent.PropertyChain, static_cast<UObject*>(MemoryPtr)))
+	{
+		return false;
+	}
+	
+	if (OutRemappedChangeEvent.PropertyChain.GetHead() == nullptr)
+	{
+		return false;
+	}
+	
+	OutRemappedChangeEvent.Property = OutRemappedChangeEvent.PropertyChain.GetTail()->GetValue();
+		
+	FEditPropertyChain::TDoubleLinkedListNode* RemappedPropertyNode = OutRemappedChangeEvent.PropertyChain.GetHead();
+	FEditPropertyChain::TDoubleLinkedListNode* InPropertyNode = InChangeEvent.PropertyChain.GetHead();
+	do
+	{
+		FProperty* RemappedProperty = RemappedPropertyNode->GetValue();
+		FProperty* InProperty = InPropertyNode->GetValue();
+		MemoryPtr = RemappedProperty->ContainerPtrToValuePtr<uint8>(MemoryPtr);
+		
+		const int32 ArrayIndex = InChangeEvent.GetArrayIndex(RemappedProperty->GetName());
+		if (ArrayIndex != INDEX_NONE)
+		{
+			RemappedPropertyNode = RemappedPropertyNode->GetNextNode();
+			InPropertyNode = InPropertyNode->GetNextNode();
+			FProperty* InnerProperty = InPropertyNode ? InPropertyNode->GetValue() : InChangeEvent.Property;
+			if (const FArrayProperty* ArrayProperty = CastField<FArrayProperty>(RemappedProperty))
+			{
+				FScriptArrayHelper ArrayHelper(ArrayProperty, MemoryPtr);
+				if(!ArrayHelper.IsValidIndex(ArrayIndex))
+				{
+					check(OutRemappedChangeEvent.ChangeType & EPropertyChangeType::ArrayAdd);
+					check(ArrayHelper.Num() == ArrayIndex);
+					ArrayHelper.Resize(ArrayIndex + 1);
+				}
+				MemoryPtr = ArrayHelper.GetRawPtr(ArrayIndex);
+				if (!InPropertyNode)
+				{
+					OutRemappedChangeEvent.Property = ArrayProperty->Inner;
+				}
+			}
+			if (const FSetProperty* SetProperty = CastField<FSetProperty>(RemappedProperty))
+			{
+				FScriptSetHelper SetHelper(SetProperty, MemoryPtr);
+				if(!SetHelper.IsValidIndex(ArrayIndex))
+				{
+					int32 NewIndex = SetHelper.AddUninitializedValue();
+					ArrayIndices[InChangeEvent.ObjectIteratorIndex].Add (RemappedProperty->GetName(), NewIndex);
+				}
+				MemoryPtr = SetHelper.GetElementPtr(ArrayIndex); // TODO: do more testing with sets
+				
+				if (!InPropertyNode)
+				{
+					OutRemappedChangeEvent.Property = SetProperty->ElementProp;
+				}
+			}
+			if (const FMapProperty* MapProperty = CastField<FMapProperty>(RemappedProperty))
+			{
+				FScriptMapHelper MapHelper(MapProperty, MemoryPtr);
+				if(!MapHelper.IsValidIndex(ArrayIndex))
+				{
+					int32 NewIndex = MapHelper.AddUninitializedValue();
+					ArrayIndices[InChangeEvent.ObjectIteratorIndex].Add(RemappedProperty->GetName(), NewIndex);
+				}
+				if (InnerProperty == CastField<FMapProperty>(InProperty)->KeyProp) // TODO: do more testing with maps
+				{
+					MemoryPtr = MapHelper.GetKeyPtr(ArrayIndex);
+					
+					if (!InPropertyNode)
+					{
+						OutRemappedChangeEvent.Property = MapProperty->KeyProp;
+					}
+				}
+				else
+				{
+					MemoryPtr = MapHelper.GetValuePtr(ArrayIndex);
+					if (!InPropertyNode)
+					{
+						OutRemappedChangeEvent.Property = MapProperty->ValueProp;
+					}
+				}
+			}
+		}
+		else if (const FObjectProperty* AsObjectProperty = CastField<FObjectProperty>(RemappedProperty))
+		{
+			MemoryPtr = AsObjectProperty->GetObjectPropertyValue(MemoryPtr);
+		}
+		
+		if (RemappedPropertyNode)
+		{
+			RemappedPropertyNode = RemappedPropertyNode->GetNextNode();
+			InPropertyNode = InPropertyNode->GetNextNode();
+		}
+	}
+	while (RemappedPropertyNode);
+
+	OutRemappedChangeEvent.SetArrayIndexPerObject(ArrayIndices);
+
+	return true;
+}
+
 static void CopyProperty(const FProperty* SourceProperty, const void* SourceValue, const FProperty* DestProperty, void* DestValue)
 {
 	check(SourceProperty->GetID() == DestProperty->GetID());
@@ -408,53 +526,6 @@ static void CopyProperty(const FProperty* SourceProperty, const void* SourceValu
 	}
 }
 
-static void AddProperty(const FProperty* SourceProperty, const void* SourceValue, const FProperty* DestProperty, void* DestValue, int32 ArrayIndex)
-{
-	if (const FArrayProperty* SourcePropertyAsArray = CastField<FArrayProperty>(SourceProperty))
-	{
-		const FArrayProperty* DestPropertyAsArray = CastFieldChecked<FArrayProperty>(DestProperty);
-		FScriptArrayHelper SourceArray(SourcePropertyAsArray, SourceValue);
-        FScriptArrayHelper DestArray(DestPropertyAsArray, DestValue);
-		if (DestArray.Num() < ArrayIndex + 1)
-		{
-			DestArray.Resize(ArrayIndex + 1);
-		}
-		CopyProperty(SourcePropertyAsArray->Inner, SourceArray.GetElementPtr(ArrayIndex),
-			DestPropertyAsArray->Inner, DestArray.GetElementPtr(ArrayIndex));
-	}
-	else if (const FSetProperty* SourcePropertyAsSet = CastField<FSetProperty>(SourceProperty))
-	{
-		const FSetProperty* DestPropertyAsSet = CastFieldChecked<FSetProperty>(DestProperty);
-		FScriptSetHelper SourceSet(SourcePropertyAsSet, SourceValue);
-        FScriptSetHelper DestSet(DestPropertyAsSet, DestValue);
-		int32 DestArrayIndex = DestSet.AddUninitializedValue();
-		
-		void* DestElementPtr = DestSet.GetElementPtr(DestArrayIndex);
-		DestPropertyAsSet->ElementProp->InitializeValue(DestElementPtr);
-		CopyProperty(SourcePropertyAsSet->ElementProp, SourceSet.GetElementPtr(ArrayIndex),
-			DestPropertyAsSet->ElementProp, DestElementPtr);
-		DestSet.Rehash();
-	}
-	else if (const FMapProperty* SourcePropertyAsMap = CastField<FMapProperty>(SourceProperty))
-	{
-		const FMapProperty* DestPropertyAsMap = CastFieldChecked<FMapProperty>(DestProperty);
-		FScriptMapHelper SourceMap(SourcePropertyAsMap, SourceValue);
-		FScriptMapHelper DestMap(DestPropertyAsMap, DestValue);
-		int32 DestArrayIndex = DestMap.AddUninitializedValue();
-
-		void* DestKeyPtr = DestMap.GetKeyPtr(DestArrayIndex);
-		DestPropertyAsMap->KeyProp->InitializeValue(DestKeyPtr);
-		CopyProperty(SourcePropertyAsMap->KeyProp, SourceMap.GetKeyPtr(ArrayIndex),
-			DestPropertyAsMap->KeyProp, DestKeyPtr);
-		
-		void* DestValuePtr = DestMap.GetValuePtr(DestArrayIndex);
-		DestPropertyAsMap->ValueProp->InitializeValue(DestValuePtr);
-		CopyProperty(SourcePropertyAsMap->ValueProp, SourceMap.GetValuePtr(ArrayIndex),
-			DestPropertyAsMap->ValueProp, DestValuePtr);
-		DestMap.Rehash();
-	}
-}
-
 void FPropertyBagRepository::PostEditChangeChainProperty(const UObject* Object, FPropertyChangedChainEvent& PropertyChangedEvent)
 {
 #if WITH_EDITOR
@@ -465,29 +536,37 @@ void FPropertyBagRepository::PostEditChangeChainProperty(const UObject* Object,
 		return;
 	}
 	
-	auto CopyChanges = [&PropertyChangedEvent](const UObject* Source, UObject* Dest) 
+	auto CopyChanges = [&PropertyChangedEvent](const UObject* Source, UObject* Dest)
 	{
+		FPropertyChangedEvent BaseRemappedEvent(PropertyChangedEvent);
 		FEditPropertyChain RemappedChain;
-		if (ConstructRemappedPropertyChain(PropertyChangedEvent.PropertyChain, RemappedChain, Dest))
+		FPropertyChangedChainEvent RemappedChangeEvent(RemappedChain, BaseRemappedEvent);
+		void* DestData = Dest;
+		TArray<TMap<FString, int32>> RemappedArrayIndices;
+		RemappedArrayIndices.AddDefaulted(PropertyChangedEvent.ObjectIteratorIndex + 1);
+		if (RemapChangeEvent(PropertyChangedEvent, RemappedChangeEvent, DestData, RemappedArrayIndices))
 		{
-			Dest->PreEditChange(RemappedChain);
-            
-            FPropertyChangedChainEvent RemappedChangeEvent(RemappedChain, PropertyChangedEvent);
+			Dest->PreEditChange(RemappedChangeEvent.PropertyChain);
+			
             const void* SourceData = ResolveChangePath(Source, PropertyChangedEvent);
-            void* DestData = ResolveChangePath(Dest, RemappedChangeEvent, true);
+			
             FProperty* SourceProperty = PropertyChangedEvent.PropertyChain.GetTail()->GetValue();
-            FProperty* DestProperty = RemappedChangeEvent.PropertyChain.GetTail()->GetValue();
+			if (PropertyChangedEvent.Property->Owner == SourceProperty)
+			{
+				// sometimes the event property isn't included in the chain (example: array elements).
+				// in that case, SourceData is pointing data of type PropertyChangedEvent.Property
+				SourceProperty = PropertyChangedEvent.Property;
+			}
+			
+			FProperty* DestProperty = RemappedChangeEvent.PropertyChain.GetTail()->GetValue();
+			if (RemappedChangeEvent.Property->Owner == DestProperty)
+			{
+				// sometimes the event property isn't included in the chain (example: array elements).
+				// in that case, DestProperty is pointing data of type RemappedChangeEvent.Property
+				DestProperty = RemappedChangeEvent.Property;
+			}
 			
-            if (PropertyChangedEvent.ChangeType == EPropertyChangeType::ArrayAdd)
-            {
-            	int32 ArrayIndex = PropertyChangedEvent.GetArrayIndex(SourceProperty->GetName());
-            	check(ArrayIndex != INDEX_NONE);
-            	AddProperty(SourceProperty, SourceData, DestProperty, DestData, ArrayIndex);
-            }
-            else
-            {
-            	CopyProperty(SourceProperty, SourceData, DestProperty, DestData);
-            }
+			CopyProperty(SourceProperty, SourceData, DestProperty, DestData);
             
             Dest->PostEditChangeChainProperty(RemappedChangeEvent);
 		}
