[InstancedStruct] fixed access to nullptr when reporting a serialization error

diff --git a/Engine/Source/Runtime/CoreUObject/Private/StructUtils/InstancedStruct.cpp b/Engine/Source/Runtime/CoreUObject/Private/StructUtils/InstancedStruct.cpp
--- a/Engine/Source/Runtime/CoreUObject/Private/StructUtils/InstancedStruct.cpp
+++ b/Engine/Source/Runtime/CoreUObject/Private/StructUtils/InstancedStruct.cpp
@@ -165,7 +165,10 @@ bool FInstancedStruct::Serialize(FArchive& Ar)
 			// In this case we manually seek in the archive to skip its serialized content. 
 			// We don't want to rely on TaggedSerialization that will mark an error in the archive that
 			// may cause other serialization to fail (e.g. FArchive& operator<<(FArchive& Ar, TArray& A))
-			UE_LOG(LogCore, Warning, TEXT("Unable to find serialized UScriptStruct -> Advance %u bytes in the archive and reset to empty FInstancedStruct. SerializedProperty:%s LinkerRoot:%s"), SerialSize, *Ar.GetSerializedProperty()->GetPathName(), Ar.GetLinker() ? *GetPathNameSafe(Ar.GetLinker()->LinkerRoot) : TEXT("NoLinker"));
+			UE_LOG(LogCore, Warning, TEXT("Unable to find serialized UScriptStruct -> Advance %u bytes in the archive and reset to empty FInstancedStruct. SerializedProperty:%s LinkerRoot:%s"),
+				SerialSize,
+				*GetPathNameSafe(Ar.GetSerializedProperty()),
+				Ar.GetLinker() ? *GetPathNameSafe(Ar.GetLinker()->LinkerRoot) : TEXT("NoLinker"));
 			Ar.Seek(Ar.Tell() + SerialSize);
 		}
 		else if (NonConstStruct != nullptr && ensureMsgf(GetMutableMemory() != nullptr, TEXT("A valid script struct should always have allocated memory")))
