Fixing a bug where the GC would not be aware of package metadata in WITH_METADATA builds. (Bug found by [REDACTED_HANDLE])
Adding WITH_METADATA support for AsyncLoading2 (ZenStore).

#tests Cooked Lyra manually and tested. Tested build permutations (editor/game)

diff --git a/Engine/Source/Runtime/CoreUObject/Private/Serialization/AsyncLoading2.cpp b/Engine/Source/Runtime/CoreUObject/Private/Serialization/AsyncLoading2.cpp
--- a/Engine/Source/Runtime/CoreUObject/Private/Serialization/AsyncLoading2.cpp
+++ b/Engine/Source/Runtime/CoreUObject/Private/Serialization/AsyncLoading2.cpp
@@ -3090,9 +3090,9 @@ private:
 		int32 CreateExportIndex = 0;
 		int32 SerializeExportIndex = 0;
 		int32 PostLoadExportIndex = 0;
-#if WITH_EDITORONLY_DATA
+#if WITH_METADATA
 		int32 MetaDataIndex = -1;
-#endif
+#endif // WITH_METADATA
 		bool bIsCurrentlyResolvingImports = false;
 		bool bIsCurrentlyCreatingExports = false;
 		bool bContainsClasses = false;
@@ -3100,7 +3100,7 @@ private:
 		FAsyncPackageLinkerLoadHeaderData LinkerLoadHeaderData;
 	};
 	TOptional<FLinkerLoadState> LinkerLoadState;
-#endif
+#endif // ALT2_ENABLE_LINKERLOAD_SUPPORT
 	/** Cached async loading thread object this package was created by */
 	FAsyncLoadingThread2& AsyncLoadingThread;
 	FAsyncLoadEventGraphAllocator& GraphAllocator;
@@ -5935,7 +5935,7 @@ EEventLoadNodeExecutionResult FAsyncPackage2::ProcessLinkerLoadPackageSummary(FA
 	if (!bLoadHasFailed)
 	{
 		PackageRef.ReserveSpaceForPublicExports(LinkerLoadState->LinkerLoadHeaderData.ExportMap.Num());
-#if WITH_EDITORONLY_DATA
+#if WITH_METADATA
 		// Create metadata object, this needs to happen before any other package wants to use our exports
 		LinkerLoadState->MetaDataIndex = LinkerLoadState->Linker->LoadMetaDataFromExportMap(false);
 		if (LinkerLoadState->MetaDataIndex >= 0)
@@ -5948,7 +5948,7 @@ EEventLoadNodeExecutionResult FAsyncPackage2::ProcessLinkerLoadPackageSummary(FA
 			ExportObject.bExportLoadFailed = LinkerExport.bExportLoadFailed;
 			ExportObject.bFiltered = LinkerExport.bWasFiltered;
 		}
-#endif
+#endif // WITH_METADATA
 	}
 	else if (Desc.bCanBeImported)
 	{
@@ -6112,12 +6112,12 @@ bool FAsyncPackage2::CreateLinkerLoadExports(FAsyncLoadingThreadState2& ThreadSt
 	while (LinkerLoadState->CreateExportIndex < ExportCount)
 	{
 		const int32 ExportIndex = LinkerLoadState->CreateExportIndex++;
-#if WITH_EDITORONLY_DATA
+#if WITH_METADATA
 		if (ExportIndex == LinkerLoadState->MetaDataIndex)
 		{
 			continue;
 		}
-#endif
+#endif // WITH_METADATA
 		FObjectExport& LinkerExport = LinkerLoadState->Linker->ExportMap[ExportIndex];
 		FExportObject& ExportObject = Data.Exports[ExportIndex];
 		if (UObject* Object = LinkerLoadState->Linker->CreateExport(ExportIndex))
diff --git a/Engine/Source/Runtime/CoreUObject/Private/UObject/Package.cpp b/Engine/Source/Runtime/CoreUObject/Private/UObject/Package.cpp
--- a/Engine/Source/Runtime/CoreUObject/Private/UObject/Package.cpp
+++ b/Engine/Source/Runtime/CoreUObject/Private/UObject/Package.cpp
@@ -408,7 +408,7 @@ void UPackage::SetSavedHash(const FIoHash& InSavedHash)
 
 #endif
 
-#if WITH_EDITORONLY_DATA
+#if WITH_METADATA
 
 void DeclareIntrinsicUPackageMembers()
 {
@@ -426,4 +426,4 @@ IMPLEMENT_CORE_INTRINSIC_CLASS(UPackage, UObject,
 	{
 	}
 );
-#endif // WITH_EDITORONLY_DATA
+#endif // WITH_METADATA
