Property support for containers with intrusive unset optional states

diff --git a/Engine/Source/Runtime/Core/Public/Containers/Map.h b/Engine/Source/Runtime/Core/Public/Containers/Map.h
--- a/Engine/Source/Runtime/Core/Public/Containers/Map.h
+++ b/Engine/Source/Runtime/Core/Public/Containers/Map.h
@@ -1853,6 +1853,11 @@ public:
 		return Pairs.Num();
 	}
 
+	int32 NumUnchecked() const
+	{
+		return Pairs.NumUnchecked();
+	}
+
 	/** @return The max valid index of the elements in the sparse storage. */
 	int32 GetMaxIndex() const
 	{
diff --git a/Engine/Source/Runtime/Core/Public/Containers/ScriptArray.h b/Engine/Source/Runtime/Core/Public/Containers/ScriptArray.h
--- a/Engine/Source/Runtime/Core/Public/Containers/ScriptArray.h
+++ b/Engine/Source/Runtime/Core/Public/Containers/ScriptArray.h
@@ -43,6 +43,10 @@ public:
 		checkSlow(ArrayMax>=ArrayNum);
 		return ArrayNum;
 	}
+	FORCEINLINE int32 NumUnchecked() const
+	{
+		return ArrayNum;
+	}
 	void InsertZeroed( int32 Index, int32 Count, int32 NumBytesPerElement, uint32 AlignmentOfElement )
 	{
 		Insert( Index, Count, NumBytesPerElement, AlignmentOfElement );
diff --git a/Engine/Source/Runtime/Core/Public/Containers/Set.h b/Engine/Source/Runtime/Core/Public/Containers/Set.h
--- a/Engine/Source/Runtime/Core/Public/Containers/Set.h
+++ b/Engine/Source/Runtime/Core/Public/Containers/Set.h
@@ -1918,6 +1918,11 @@ public:
 		return Elements.Num();
 	}
 
+	int32 NumUnchecked() const
+	{
+		return Elements.NumUnchecked();
+	}
+
 	int32 GetMaxIndex() const
 	{
 		return Elements.GetMaxIndex();
diff --git a/Engine/Source/Runtime/Core/Public/Containers/SparseArray.h b/Engine/Source/Runtime/Core/Public/Containers/SparseArray.h
--- a/Engine/Source/Runtime/Core/Public/Containers/SparseArray.h
+++ b/Engine/Source/Runtime/Core/Public/Containers/SparseArray.h
@@ -91,13 +91,13 @@ public:
 	constexpr static bool bHasIntrusiveUnsetOptionalState = true;
 	using IntrusiveUnsetOptionalStateType = TSparseArray;
 
-	explicit TSparseArray(FIntrusiveUnsetOptionalState)
-		: NumFreeIndices(-1)
+	explicit TSparseArray(FIntrusiveUnsetOptionalState Tag)
+		: Data(Tag)
 	{
 	}
-	bool operator==(FIntrusiveUnsetOptionalState) const
+	bool operator==(FIntrusiveUnsetOptionalState Tag) const
 	{
-		return NumFreeIndices == -1;
+		return Data == Tag;
 	}
 	///////////////////////////////////////////////////
 	// End - intrusive TOptional<TSparseArray> state //
@@ -1176,10 +1176,10 @@ private:
 	AllocationBitArrayType AllocationFlags;
 
 	/** The index of an unallocated element in the array that currently contains the head of the linked list of free elements. */
-	int32 FirstFreeIndex;
+	int32 FirstFreeIndex = -1;
 
 	/** The number of elements in the free list. */
-	int32 NumFreeIndices;
+	int32 NumFreeIndices = 0;
 
 public:
 	void WriteMemoryImage(FMemoryImageWriter& Writer) const
@@ -1348,13 +1348,13 @@ public:
 	constexpr static bool bHasIntrusiveUnsetOptionalState = true;
 	using IntrusiveUnsetOptionalStateType = TScriptSparseArray;
 
-	explicit TScriptSparseArray(FIntrusiveUnsetOptionalState)
-		: NumFreeIndices(-1)
+	explicit TScriptSparseArray(FIntrusiveUnsetOptionalState Tag)
+		: Data(Tag)
 	{
 	}
-	bool operator==(FIntrusiveUnsetOptionalState) const
+	bool operator==(FIntrusiveUnsetOptionalState Tag) const
 	{
-		return NumFreeIndices == -1;
+		return Data == Tag;
 	}
 	///////////////////////////////////////////////////
 	// End - intrusive TOptional<TSparseArray> state //
@@ -1380,6 +1380,11 @@ public:
 		return Data.Num() - NumFreeIndices;
 	}
 
+	int32 NumUnchecked() const
+	{
+		return Data.NumUnchecked() - NumFreeIndices;
+	}
+
 	int32 GetMaxIndex() const
 	{
 		return Data.Num();
@@ -1477,8 +1482,8 @@ public:
 private:
 	TScriptArray   <typename AllocatorType::ElementAllocator>  Data;
 	TScriptBitArray<typename AllocatorType::BitArrayAllocator> AllocationFlags;
-	int32                                                      FirstFreeIndex;
-	int32                                                      NumFreeIndices;
+	int32                                                      FirstFreeIndex = -1;
+	int32                                                      NumFreeIndices = 0;
 
 	// This function isn't intended to be called, just to be compiled to validate the correctness of the type.
 	static void CheckConstraints()
diff --git a/Engine/Source/Runtime/CoreUObject/Private/UObject/GarbageCollection.cpp b/Engine/Source/Runtime/CoreUObject/Private/UObject/GarbageCollection.cpp
--- a/Engine/Source/Runtime/CoreUObject/Private/UObject/GarbageCollection.cpp
+++ b/Engine/Source/Runtime/CoreUObject/Private/UObject/GarbageCollection.cpp
@@ -6411,6 +6411,24 @@ void FObjectProperty::EmitIntrusiveOptionalReferenceInfo(UE::GC::FSchemaBuilder&
 	EmitReferenceInfo(Schema, BaseOffset, EncounteredStructProps, DebugPath);
 }
 
+void FArrayProperty::EmitIntrusiveOptionalReferenceInfo(UE::GC::FSchemaBuilder& Schema, int32 BaseOffset, TArray<const FStructProperty*>& EncounteredStructProps, UE::GC::FPropertyStack& DebugPath) 
+{
+	// Intrusive state looks like an empty container with a corrupted Max, but the GC doesn't touch the max so the default schema is sufficient
+	EmitReferenceInfo(Schema, BaseOffset, EncounteredStructProps, DebugPath);
+}
+
+void FMapProperty::EmitIntrusiveOptionalReferenceInfo(UE::GC::FSchemaBuilder& Schema, int32 BaseOffset, TArray<const FStructProperty*>& EncounteredStructProps, UE::GC::FPropertyStack& DebugPath) 
+{
+	// Intrusive state looks like an empty container with a corrupted Max, but the GC doesn't touch the max so the default schema is sufficient
+	EmitReferenceInfo(Schema, BaseOffset, EncounteredStructProps, DebugPath);
+}
+
+void FSetProperty::EmitIntrusiveOptionalReferenceInfo(UE::GC::FSchemaBuilder& Schema, int32 BaseOffset, TArray<const FStructProperty*>& EncounteredStructProps, UE::GC::FPropertyStack& DebugPath) 
+{
+	// Intrusive state looks like an empty container with a corrupted Max, but the GC doesn't touch the max so the default schema is sufficient
+	EmitReferenceInfo(Schema, BaseOffset, EncounteredStructProps, DebugPath);
+}
+
 void FArrayProperty::EmitReferenceInfo(UE::GC::FSchemaBuilder& Schema, int32 BaseOffset, TArray<const FStructProperty*>& EncounteredStructProps, UE::GC::FPropertyStack& DebugPath)
 {
 	using namespace UE::GC;
diff --git a/Engine/Source/Runtime/CoreUObject/Private/UObject/PropertyArray.cpp b/Engine/Source/Runtime/CoreUObject/Private/UObject/PropertyArray.cpp
--- a/Engine/Source/Runtime/CoreUObject/Private/UObject/PropertyArray.cpp
+++ b/Engine/Source/Runtime/CoreUObject/Private/UObject/PropertyArray.cpp
@@ -1430,3 +1430,33 @@ void* FArrayProperty::ResolveVisitedPathInfo(void* Data, const FPropertyVisitorI
 
 	return nullptr;
 }
+
+bool FArrayProperty::HasIntrusiveUnsetOptionalState() const
+{
+	return true;
+}
+
+void FArrayProperty::InitializeIntrusiveUnsetOptionalValue(void* Data) const
+{
+	// FScriptArray's unset state constructor is good enough
+	Super::InitializeIntrusiveUnsetOptionalValue(Data);
+}
+
+bool FArrayProperty::IsIntrusiveOptionalValueSet(const void* Data) const
+{
+	// FScriptArray's unset state comparison is good enough
+	return Super::IsIntrusiveOptionalValueSet(Data);
+}
+
+void FArrayProperty::ClearIntrusiveOptionalValue(void* Data) const
+{
+	// Destroy any inner elements first, because FScriptArray's destructor will only free memory
+	if (IsIntrusiveOptionalValueSet(Data))
+	{
+		FScriptArrayHelper ArrayHelper(this, Data);
+		ArrayHelper.EmptyValues();
+
+		// Call Super to actually reset the optional to the unset state, now that any elements have been destroyed
+		Super::ClearIntrusiveOptionalValue(Data);
+	}
+}
diff --git a/Engine/Source/Runtime/CoreUObject/Private/UObject/PropertyMap.cpp b/Engine/Source/Runtime/CoreUObject/Private/UObject/PropertyMap.cpp
--- a/Engine/Source/Runtime/CoreUObject/Private/UObject/PropertyMap.cpp
+++ b/Engine/Source/Runtime/CoreUObject/Private/UObject/PropertyMap.cpp
@@ -1812,3 +1812,33 @@ void* FMapProperty::ResolveVisitedPathInfo(void* Data, const FPropertyVisitorInf
 
 	return nullptr;
 }
+
+bool FMapProperty::HasIntrusiveUnsetOptionalState() const
+{
+	return true;
+}
+
+void FMapProperty::InitializeIntrusiveUnsetOptionalValue(void* Data) const
+{
+	// FScriptMap's unset state constructor is good enough
+	Super::InitializeIntrusiveUnsetOptionalValue(Data);
+}
+
+bool FMapProperty::IsIntrusiveOptionalValueSet(const void* Data) const
+{
+	// FScriptMap's unset state comparison is good enough
+	return Super::IsIntrusiveOptionalValueSet(Data);
+}
+
+void FMapProperty::ClearIntrusiveOptionalValue(void* Data) const
+{
+	// Destroy any inner elements first, because FScriptMap's destructor will only free memory
+	if (IsIntrusiveOptionalValueSet(Data))
+	{
+		FScriptMapHelper MapHelper(this, Data);
+		MapHelper.EmptyValues();
+
+		// Call Super to actually reset the optional to the unset state, now that any elements have been destroyed
+		Super::ClearIntrusiveOptionalValue(Data);
+	}
+}
diff --git a/Engine/Source/Runtime/CoreUObject/Private/UObject/PropertyOptional.cpp b/Engine/Source/Runtime/CoreUObject/Private/UObject/PropertyOptional.cpp
--- a/Engine/Source/Runtime/CoreUObject/Private/UObject/PropertyOptional.cpp
+++ b/Engine/Source/Runtime/CoreUObject/Private/UObject/PropertyOptional.cpp
@@ -456,6 +456,10 @@ void FOptionalProperty::InitializeValueInternal(void* Data) const
 void FOptionalProperty::DestroyValueInternal(void* Data) const
 {
 	MarkUnset(Data);
+	if (ValueProperty->HasIntrusiveUnsetOptionalState())
+	{
+		ValueProperty->DestroyValue(Data);
+	}
 }
 
 bool FOptionalProperty::ContainsClearOnFinishDestroyInternal(TArray<const FStructProperty*>& EncounteredStructProps) const
diff --git a/Engine/Source/Runtime/CoreUObject/Private/UObject/PropertySet.cpp b/Engine/Source/Runtime/CoreUObject/Private/UObject/PropertySet.cpp
--- a/Engine/Source/Runtime/CoreUObject/Private/UObject/PropertySet.cpp
+++ b/Engine/Source/Runtime/CoreUObject/Private/UObject/PropertySet.cpp
@@ -1170,3 +1170,33 @@ void* FSetProperty::ResolveVisitedPathInfo(void* Data, const FPropertyVisitorInf
 
 	return nullptr;
 }
+
+bool FSetProperty::HasIntrusiveUnsetOptionalState() const
+{
+	return true;
+}
+
+void FSetProperty::InitializeIntrusiveUnsetOptionalValue(void* Data) const
+{
+	// FScriptSet's unset state constructor is good enough
+	Super::InitializeIntrusiveUnsetOptionalValue(Data);
+}
+
+bool FSetProperty::IsIntrusiveOptionalValueSet(const void* Data) const
+{
+	// FScriptSet's unset state comparison is good enough
+	return Super::IsIntrusiveOptionalValueSet(Data);
+}
+
+void FSetProperty::ClearIntrusiveOptionalValue(void* Data) const
+{
+	// Destroy any inner elements first, because FScriptSet's destructor will only free memory
+	if (IsIntrusiveOptionalValueSet(Data))
+	{
+		FScriptSetHelper SetHelper(this, Data);
+		SetHelper.EmptyElements();
+
+		// Call Super to actually reset the optional to the unset state, now that any elements have been destroyed
+		Super::ClearIntrusiveOptionalValue(Data);
+	}
+}
diff --git a/Engine/Source/Runtime/CoreUObject/Public/UObject/PropertyOptional.h b/Engine/Source/Runtime/CoreUObject/Public/UObject/PropertyOptional.h
--- a/Engine/Source/Runtime/CoreUObject/Public/UObject/PropertyOptional.h
+++ b/Engine/Source/Runtime/CoreUObject/Public/UObject/PropertyOptional.h
@@ -38,7 +38,7 @@ struct COREUOBJECT_API FOptionalPropertyLayout
 			if (!IsSet(Data))
 			{
 				// Need to destroy the value in its optional unset state first 
-				ValueProperty->ClearIntrusiveOptionalValue(Data);
+				ValueProperty->DestroyValue(Data);
 				ValueProperty->InitializeValue(Data);
 			}
 		}
diff --git a/Engine/Source/Runtime/CoreUObject/Public/UObject/UnrealType.h b/Engine/Source/Runtime/CoreUObject/Public/UObject/UnrealType.h
--- a/Engine/Source/Runtime/CoreUObject/Public/UObject/UnrealType.h
+++ b/Engine/Source/Runtime/CoreUObject/Public/UObject/UnrealType.h
@@ -3564,6 +3564,12 @@ public:
 #if WITH_EDITORONLY_DATA
 	virtual void AppendSchemaHash(FBlake3& Builder, bool bSkipEditorOnly) const override;
 #endif
+
+	virtual bool HasIntrusiveUnsetOptionalState() const override;
+	virtual void InitializeIntrusiveUnsetOptionalValue(void* Data) const override;
+	virtual bool IsIntrusiveOptionalValueSet(const void* Data) const override;
+	virtual void ClearIntrusiveOptionalValue(void* Data) const override;
+	virtual void EmitIntrusiveOptionalReferenceInfo(UE::GC::FSchemaBuilder& Schema, int32 BaseOffset, TArray<const FStructProperty*>& EncounteredStructProps, UE::GC::FPropertyStack& DebugPath) override;
 };
 
 using FFreezableScriptMap = TScriptMap<FMemoryImageSetAllocator>;
@@ -3731,6 +3737,12 @@ public:
 #if WITH_EDITORONLY_DATA
 	virtual void AppendSchemaHash(FBlake3& Builder, bool bSkipEditorOnly) const override;
 #endif
+
+	virtual bool HasIntrusiveUnsetOptionalState() const override;
+	virtual void InitializeIntrusiveUnsetOptionalValue(void* Data) const override;
+	virtual bool IsIntrusiveOptionalValueSet(const void* Data) const override;
+	virtual void ClearIntrusiveOptionalValue(void* Data) const override;
+	virtual void EmitIntrusiveOptionalReferenceInfo(UE::GC::FSchemaBuilder& Schema, int32 BaseOffset, TArray<const FStructProperty*>& EncounteredStructProps, UE::GC::FPropertyStack& DebugPath) override;
 };
 
 class COREUOBJECT_API FSetProperty : public TProperty<FScriptSet, FProperty>
@@ -3858,6 +3870,12 @@ public:
 #if WITH_EDITORONLY_DATA
 	virtual void AppendSchemaHash(FBlake3& Builder, bool bSkipEditorOnly) const override;
 #endif
+
+	virtual bool HasIntrusiveUnsetOptionalState() const override;
+	virtual void InitializeIntrusiveUnsetOptionalValue(void* Data) const override;
+	virtual bool IsIntrusiveOptionalValueSet(const void* Data) const override;
+	virtual void ClearIntrusiveOptionalValue(void* Data) const override;
+	virtual void EmitIntrusiveOptionalReferenceInfo(UE::GC::FSchemaBuilder& Schema, int32 BaseOffset, TArray<const FStructProperty*>& EncounteredStructProps, UE::GC::FPropertyStack& DebugPath) override;
 };
 
 /**
@@ -3910,6 +3928,16 @@ public:
 		checkSlow(Result >= 0);
 		return Result;
 	}
+	/**
+	 *	Return the number of elements in the array without validating the state of the array.
+	 *	Needed to allow reading of the num when the array is 'invalid' during its intrusive unset state.
+	 *	@return	The number of elements in the array.
+	**/
+	FORCEINLINE int32 NumUnchecked() const
+	{
+		int32 Result = WithScriptArray([](auto* Array) { return Array->NumUnchecked(); });
+		return Result;
+	}
 	/**
 	 *	Returns a uint8 pointer to an element in the array
 	 *	@param	Index: index of the item to return a pointer to.
@@ -4053,7 +4081,7 @@ public:
 	void EmptyValues(int32 Slack = 0)
 	{
 		checkSlow(Slack>=0);
-		const int32 OldNum = Num();
+		const int32 OldNum = NumUnchecked();
 		if (OldNum)
 		{
 			DestructItems(0, OldNum);
@@ -4423,6 +4451,18 @@ public:
 		return Result;
 	}
 
+	/**
+	* Returns the number of elements in the map.
+	* Needed to allow reading of the num when the map is 'invalid' during its intrusive unset state.
+	* @return The number of elements in the map.
+	*/
+	FORCEINLINE int32 NumUnchecked() const
+	{
+		int32 Result = WithScriptMap([](auto* Map) { return Map->NumUnchecked(); });
+		checkSlow(Result >= 0); 
+		return Result;
+	}
+
 	/**
 	 * Returns the (non-inclusive) maximum index of elements in the map.
 	 *
@@ -4694,7 +4734,7 @@ public:
 	{
 		checkSlow(Slack >= 0);
 
-		int32 OldNum = Num();
+		int32 OldNum = NumUnchecked();
 		if (OldNum)
 		{
 			DestructItems(0, OldNum);
@@ -5370,6 +5410,17 @@ public:
 		return Result;
 	}
 
+	/**
+	* Returns the number of elements in the set.
+	*	Needed to allow reading of the num when the set is 'invalid' during its intrusive unset state.
+	* @return The number of elements in the set.
+	*/
+	FORCEINLINE int32 NumUnchecked() const
+	{
+		const int32 Result = Set->NumUnchecked();
+		return Result;
+	}
+
 	/**
 	* Returns the (non-inclusive) maximum index of elements in the set.
 	*
@@ -5509,7 +5560,7 @@ public:
 	{
 		checkSlow(Slack >= 0);
 
-		int32 OldNum = Num();
+		int32 OldNum = NumUnchecked();
 		if (OldNum)
 		{
 			DestructItems(0, OldNum);
