Cooker: Deprecation of PackagesMarkedEditorOnlyByOtherPackage: Suppress deprecation warnings in the autogenerated constructors.
#rnx

diff --git a/Engine/Source/Runtime/CoreUObject/Private/UObject/UObjectThreadContext.cpp b/Engine/Source/Runtime/CoreUObject/Private/UObject/UObjectThreadContext.cpp
--- a/Engine/Source/Runtime/CoreUObject/Private/UObject/UObjectThreadContext.cpp
+++ b/Engine/Source/Runtime/CoreUObject/Private/UObject/UObjectThreadContext.cpp
@@ -30,6 +30,10 @@ FUObjectThreadContext::~FUObjectThreadContext()
 	(void)PackagesMarkedEditorOnlyByOtherPackage;
 #endif
 }
+#if WITH_EDITORONLY_DATA
+FUObjectThreadContext::FUObjectThreadContext(const FUObjectThreadContext& Other) = default;
+FUObjectThreadContext::FUObjectThreadContext(FUObjectThreadContext&& Other) = default;
+#endif
 PRAGMA_ENABLE_DEPRECATION_WARNINGS;
 
 FObjectInitializer& FUObjectThreadContext::ReportNull()
diff --git a/Engine/Source/Runtime/CoreUObject/Public/UObject/UObjectThreadContext.h b/Engine/Source/Runtime/CoreUObject/Public/UObject/UObjectThreadContext.h
--- a/Engine/Source/Runtime/CoreUObject/Public/UObject/UObjectThreadContext.h
+++ b/Engine/Source/Runtime/CoreUObject/Public/UObject/UObjectThreadContext.h
@@ -39,6 +39,11 @@ class COREUOBJECT_API FUObjectThreadContext : public TThreadSingleton<FUObjectTh
 
 	FUObjectThreadContext();
 	virtual ~FUObjectThreadContext();
+#if WITH_EDITORONLY_DATA
+	// Remove declaration of manual constructors after removing deprecated PackagesMarkedEditorOnlyByOtherPackage
+	FUObjectThreadContext(const FUObjectThreadContext& Other);
+	FUObjectThreadContext(FUObjectThreadContext&& Other);
+#endif
 
 	/** Stack of currently used FObjectInitializers for this thread */
 	TArray<FObjectInitializer*> InitializerStack;
