[GC] Use MALLOCLEAK_IGNORE_SCOPE() for GScratchPages allocations

GScratchPages isn't destructed / freed until global static destruction, so allocations from this FPageAllocator will be reported as a memory leak.

diff --git a/Engine/Source/Runtime/CoreUObject/Private/UObject/GarbageCollection.cpp b/Engine/Source/Runtime/CoreUObject/Private/UObject/GarbageCollection.cpp
--- a/Engine/Source/Runtime/CoreUObject/Private/UObject/GarbageCollection.cpp
+++ b/Engine/Source/Runtime/CoreUObject/Private/UObject/GarbageCollection.cpp
@@ -8,6 +8,7 @@
 #include "Algo/Find.h"
 #include "Containers/Deque.h"
 #include "Containers/StaticBitArray.h"
+#include "HAL/MallocLeakDetection.h"
 #include "HAL/ThreadSafeBool.h"
 #include "Misc/TimeGuard.h"
 #include "HAL/IConsoleManager.h"
@@ -1831,7 +1832,8 @@ private:
 	}
 
 	void AllocateWipBlock(FStructArrayBlock* NextFull)
-	{	
+	{
+		MALLOCLEAK_IGNORE_SCOPE(); // GScratchPages allocations may leak until global static destruction
 		static_assert(sizeof(FStructArrayBlock) <= FPageAllocator::PageSize, "Block must fit in page");
 	
 		Wip = reinterpret_cast<FStructArrayBlock*>(GScratchPages.AllocatePage(WorkerIndex));
