FInstancedStructContainer: Fix crash caused by bad copypasta

diff --git a/Engine/Source/Runtime/CoreUObject/Private/StructUtils/InstancedStructContainer.cpp b/Engine/Source/Runtime/CoreUObject/Private/StructUtils/InstancedStructContainer.cpp
--- a/Engine/Source/Runtime/CoreUObject/Private/StructUtils/InstancedStructContainer.cpp
+++ b/Engine/Source/Runtime/CoreUObject/Private/StructUtils/InstancedStructContainer.cpp
@@ -49,7 +49,7 @@ FInstancedStructContainer& FInstancedStructContainer::operator=(FInstancedStruct
 			
 		Memory = InOther.Memory;
 		AllocatedSize = InOther.AllocatedSize;
-		NumItems = InOther.AllocatedSize;
+		NumItems = InOther.NumItems;
 			
 		InOther.Memory = nullptr;
 		InOther.AllocatedSize = 0;
