UE-222801 - ZenLoader can leaves BP placeholder objects unresolved

CL 28805292 added a new preload path to RegenerateBlueprintClass that's a bit more thorough than the previous path. However, it neglected to include a ResolvePlaceholders call to resolve any placeholder objects that were created for the current class. In some cases, this could result in unresolved placeholder objects leaking into Blueprint content.

diff --git a/Engine/Source/Runtime/CoreUObject/Private/Blueprint/BlueprintSupport.cpp b/Engine/Source/Runtime/CoreUObject/Private/Blueprint/BlueprintSupport.cpp
--- a/Engine/Source/Runtime/CoreUObject/Private/Blueprint/BlueprintSupport.cpp
+++ b/Engine/Source/Runtime/CoreUObject/Private/Blueprint/BlueprintSupport.cpp
@@ -775,6 +775,9 @@ bool FLinkerLoad::RegenerateBlueprintClass(UClass* LoadClass, UObject* ClassDefa
 			LoadClass->ClearFlags(RF_NeedPostLoad | RF_NeedPostLoadSubobjects);
 		}
 
+		// After regeneration, we can now notify placeholders that it's safe to create instances of our class.
+		FResolvingExportTracker::Get().ResolvePlaceholders(LoadClass);
+
 		return true;
 	}
 
