FInstancedStructContainer: Add ExportText/ImportText support

diff --git a/Engine/Source/Runtime/CoreUObject/Private/StructUtils/InstancedStructContainer.cpp b/Engine/Source/Runtime/CoreUObject/Private/StructUtils/InstancedStructContainer.cpp
--- a/Engine/Source/Runtime/CoreUObject/Private/StructUtils/InstancedStructContainer.cpp
+++ b/Engine/Source/Runtime/CoreUObject/Private/StructUtils/InstancedStructContainer.cpp
@@ -683,3 +683,91 @@ void FInstancedStructContainer::GetPreloadDependencies(TArray<UObject*>& OutDeps
 		}
 	}
 }
+
+bool FInstancedStructContainer::ExportTextItem(FString& ValueStr, FInstancedStructContainer const& DefaultValue,
+	UObject* Parent, int32 PortFlags, UObject* ExportRootScope) const
+{
+	ValueStr += TEXT('(');
+	for (int32 Index = 0; Index < NumItems; ++Index)
+	{
+		if (Index > 0)
+		{
+			ValueStr += TEXT(',');
+		}
+
+		FItem& Item = GetItem(Index);
+
+		// Serialize in the format of FInstancedStruct::ExportTextItem
+		if (Item.ScriptStruct)
+		{
+			ValueStr += Item.ScriptStruct->GetPathName();
+
+			Item.ScriptStruct->ExportText(
+				ValueStr,
+				Memory + Item.Offset,
+				DefaultValue.IsValidIndex(Index) && Item.ScriptStruct == DefaultValue.GetItem(Index).ScriptStruct
+					? DefaultValue.Memory + DefaultValue.GetItem(Index).Offset
+					: nullptr,
+				Parent,
+				PortFlags,
+				ExportRootScope
+			);
+		}
+		else
+		{
+			ValueStr += TEXT("None");
+		}
+	}
+
+	ValueStr += TEXT(')');
+	return true;
+}
+
+bool FInstancedStructContainer::ImportTextItem(const TCHAR*& Buffer, int32 PortFlags, UObject* Parent,
+	FOutputDevice* ErrorText, FArchive* InSerializingArchive)
+{
+	// We can't predict the size of the container during string read, so let's work in a temporary array
+	// instead of constantly resizing the container
+	TArray<FInstancedStruct, TInlineAllocator<64>> RawStructs;
+
+	if (*Buffer != TEXT('('))
+	{
+		ErrorText->Logf(
+			ELogVerbosity::Warning,
+			TEXT("FInstancedStructContainer: Missing opening \'(\' while importing property values.")
+		);
+
+		return false;
+	}
+	Buffer++;
+
+	bool bFirst = true;
+	while (*Buffer != TEXT(')'))
+	{
+		if (!bFirst)
+		{
+			if (*Buffer == TEXT(','))
+			{
+				Buffer++;
+			}
+			else
+			{
+				ErrorText->Logf(ELogVerbosity::Warning, TEXT("FInstancedStructContainer: Missing \',\' after item."));
+				return false;
+			}
+		}
+
+		// We can do this because we've exported exactly as FInstancedStruct::ExportTextItem does
+		FInstancedStruct& InstancedStruct = RawStructs.AddDefaulted_GetRef();
+		const bool bSuccess = InstancedStruct.ImportTextItem(Buffer, PortFlags, Parent, ErrorText, InSerializingArchive);
+		if (!bSuccess)
+		{
+			return false;
+		}
+
+		bFirst = false;
+	}
+
+	*this = RawStructs;
+	return true;
+}
diff --git a/Engine/Source/Runtime/CoreUObject/Public/StructUtils/InstancedStructContainer.h b/Engine/Source/Runtime/CoreUObject/Public/StructUtils/InstancedStructContainer.h
--- a/Engine/Source/Runtime/CoreUObject/Public/StructUtils/InstancedStructContainer.h
+++ b/Engine/Source/Runtime/CoreUObject/Public/StructUtils/InstancedStructContainer.h
@@ -227,6 +227,8 @@ public:
 	bool Identical(const FInstancedStructContainer* Other, uint32 PortFlags) const;
 	bool Serialize(FArchive& Ar);
 	void GetPreloadDependencies(TArray<UObject*>& OutDeps) const;
+	bool ExportTextItem(FString& ValueStr, FInstancedStructContainer const& DefaultValue, UObject* Parent, int32 PortFlags, UObject* ExportRootScope) const;
+	bool ImportTextItem(const TCHAR*& Buffer, int32 PortFlags, UObject* Parent, FOutputDevice* ErrorText, FArchive* InSerializingArchive = nullptr);
 
 private:
 
@@ -273,5 +275,7 @@ struct TStructOpsTypeTraits<FInstancedStructContainer> : public TStructOpsTypeTr
 		WithIdentical = true,
 		WithAddStructReferencedObjects = true,
 		WithGetPreloadDependencies = true,
+		WithExportTextItem = true,
+		WithImportTextItem = true,
 	};
 };
