Prevent marking as replaced the properties that are always overridden which fixes the fact that all components are displaying overrides upon reload since the Entity back pointer is always overridden.

diff --git a/Engine/Source/Runtime/CoreUObject/Private/UObject/Class.cpp b/Engine/Source/Runtime/CoreUObject/Private/UObject/Class.cpp
--- a/Engine/Source/Runtime/CoreUObject/Private/UObject/Class.cpp
+++ b/Engine/Source/Runtime/CoreUObject/Private/UObject/Class.cpp
@@ -1762,7 +1762,11 @@ void UStruct::SerializeVersionedTaggedProperties(FStructuredArchive::FSlot Slot,
 								// No need to restore none operations
 								if (Tag.OverrideOperation != EOverriddenPropertyOperation::None)
 								{
-									OverriddenProperties->SetOverriddenPropertyOperation(Tag.OverrideOperation, UnderlyingArchive.GetSerializedPropertyChain(), Property);
+									// Prevent marking as replaced the properties that are always overridden
+									if (Tag.OverrideOperation != EOverriddenPropertyOperation::Replace || !Property->HasAnyPropertyFlags(CPF_ExperimentalAlwaysOverriden))
+									{
+										OverriddenProperties->SetOverriddenPropertyOperation(Tag.OverrideOperation, UnderlyingArchive.GetSerializedPropertyChain(), Property);
+									}
 								}
 							}
 						}
