Cooker: ICookInfo: Add sessionplatforms, outputfolder, and IsIterative for use by plugins that need to work with the cooker.
#rnx

diff --git a/Engine/Source/Editor/UnrealEd/Classes/CookOnTheSide/CookOnTheFlyServer.h b/Engine/Source/Editor/UnrealEd/Classes/CookOnTheSide/CookOnTheFlyServer.h
--- a/Engine/Source/Editor/UnrealEd/Classes/CookOnTheSide/CookOnTheFlyServer.h
+++ b/Engine/Source/Editor/UnrealEd/Classes/CookOnTheSide/CookOnTheFlyServer.h
@@ -644,6 +644,10 @@ public:
 	UNREALED_API virtual UE::Cook::ECookType GetCookType() override;
 	UNREALED_API virtual UE::Cook::ECookingDLC GetCookingDLC() override;
 	UNREALED_API virtual UE::Cook::EProcessType GetProcessType() override;
+	UNREALED_API virtual bool IsIterative() override;
+	UNREALED_API virtual TArray<const ITargetPlatform*> GetSessionPlatforms() override;
+	UNREALED_API virtual FString GetCookOutputFolder(const ITargetPlatform* TargetPlatform) override;
+
 	UNREALED_API virtual void RegisterCollector(UE::Cook::IMPCollector* Collector,
 		UE::Cook::EProcessType ProcessType = UE::Cook::EProcessType::AllMPCook) override;
 	UNREALED_API virtual void UnregisterCollector(UE::Cook::IMPCollector* Collector) override;
diff --git a/Engine/Source/Editor/UnrealEd/Private/CookOnTheFlyServer.cpp b/Engine/Source/Editor/UnrealEd/Private/CookOnTheFlyServer.cpp
--- a/Engine/Source/Editor/UnrealEd/Private/CookOnTheFlyServer.cpp
+++ b/Engine/Source/Editor/UnrealEd/Private/CookOnTheFlyServer.cpp
@@ -7705,6 +7705,40 @@ UE::Cook::EProcessType UCookOnTheFlyServer::GetProcessType()
 	}
 }
 
+bool UCookOnTheFlyServer::IsIterative()
+{
+	// TODO: For simplicity, we provide a single bool for all platforms in multiprocess cooks
+	// But it is not currently guaranteed that they all have the same value; add enforcement of
+	// that in CookByTheBookStarted.
+	if (!PlatformManager || PlatformManager->GetNumSessionPlatforms() == 0)
+	{
+		return false;
+	}
+	const ITargetPlatform* TargetPlatform = PlatformManager->GetSessionPlatforms()[0];
+	UE::Cook::FPlatformData* PlatformData = PlatformManager->GetPlatformData(TargetPlatform);
+	return !PlatformData->bFullBuild;
+}
+
+TArray<const ITargetPlatform*> UCookOnTheFlyServer::GetSessionPlatforms()
+{
+	return PlatformManager ? PlatformManager->GetSessionPlatforms() : TArray<const ITargetPlatform*>();
+}
+
+FString UCookOnTheFlyServer::GetCookOutputFolder(const ITargetPlatform* TargetPlatform)
+{
+	if (!SandboxFile || !PlatformManager || !TargetPlatform)
+	{
+		return FString();
+	}
+	if (!PlatformManager->GetSessionPlatforms().Contains(TargetPlatform))
+	{
+		return FString();
+	}
+	FString Result = SandboxFile->GetSandboxDirectory(TargetPlatform->PlatformName());
+	FPaths::MakeStandardFilename(Result);
+	return Result;
+}
+
 void UCookOnTheFlyServer::RegisterCollector(UE::Cook::IMPCollector* Collector, UE::Cook::EProcessType ProcessType)
 {
 	using namespace UE::Cook;
diff --git a/Engine/Source/Runtime/CoreUObject/Public/UObject/ICookInfo.h b/Engine/Source/Runtime/CoreUObject/Public/UObject/ICookInfo.h
--- a/Engine/Source/Runtime/CoreUObject/Public/UObject/ICookInfo.h
+++ b/Engine/Source/Runtime/CoreUObject/Public/UObject/ICookInfo.h
@@ -15,6 +15,7 @@
 
 enum class EDataValidationResult : uint8;
 class FDataValidationContext;
+class ITargetPlatform;
 class UPackage;
 
 namespace UE::Cook { class IMPCollector; }
@@ -114,6 +115,31 @@ public:
 	virtual UE::Cook::ECookingDLC GetCookingDLC() = 0;
 	/** The role the current process plays in its MPCook session, or EProcessType::SingleProcess if it is running standalone. */
 	virtual UE::Cook::EProcessType GetProcessType() = 0;
+	/**
+	 * Returns true if the cooker is cooking after a previous cook session and is cooking only the changed files.
+	 * Returns false if the cooker is doing a recook of all packages discovered in the session.
+	 * Returns false if not yet initialized, but it will be initialized whenever a session is in progress (GetSessionPlatforms
+	 * is non-empty).
+	 * When IsIterative is true, systems that write artifacts to the cook output should load/update/resave their
+	 * artifacts.
+	 */
+	virtual bool IsIterative() = 0;
+
+	/**
+	 * Returns the list of platforms that will be/are being/have been cooked for the current cook session. Returns
+	 * empty array when outside of a cooksession, including in the case that GetCookType() == ECookType::OnTheFly and
+	 * no platforms have been requested yet. During cook by the book, this list will not change throughout the cook,
+	 * during CookOnTheFly it can be added to or removed from when platforms are requested or go idle and are dropped.
+	 */
+	virtual TArray<const ITargetPlatform*> GetSessionPlatforms() = 0;
+	/**
+	 * Returns the output folder being used by the cooker for the given platform in the given session.
+	 * Returns empty string if not in a session or the given platform is not in GetSessionPlatforms().
+	 * Returns the path to the root folder of the output, so e.g. GetCookOutputFolder()/<ProjectName>/Metadata
+	 * is the path to the output metadata for the cook.
+	 * Returns the path in FPaths::MakeStandardFilename format ("../../../<ProjectName>/Saved/Cooked/<PlatformName>")
+	 */
+	virtual FString GetCookOutputFolder(const ITargetPlatform* TargetPlatform) = 0;
 
 	/**
 	 * MPCook: register in the current process a collector that replicates system-specific and package-specific
