Removed auto-creation of non-nullable properties.

diff --git a/Engine/Source/Editor/PropertyEditor/Private/PropertyHandleImpl.cpp b/Engine/Source/Editor/PropertyEditor/Private/PropertyHandleImpl.cpp
--- a/Engine/Source/Editor/PropertyEditor/Private/PropertyHandleImpl.cpp
+++ b/Engine/Source/Editor/PropertyEditor/Private/PropertyHandleImpl.cpp
@@ -1164,14 +1164,6 @@ int32 FPropertyValueImpl::AddChild()
 								FScriptArrayHelper	ArrayHelper(Array, DirectAddress);
 								Index = ArrayHelper.AddValue();
 								ReturnLogicalIndex = Index;
-
-								// check whether the inner type is flagged as a non-nullable. if so, create it.
-								FObjectProperty* InnerObjectProperty = CastField<FObjectProperty>(Array->Inner);
-								if (InnerObjectProperty && InnerObjectProperty->HasAnyPropertyFlags(CPF_NonNullable))
-								{
-									UObject* NewItem = NewObject<UObject>(Obj, InnerObjectProperty->PropertyClass);
-									InnerObjectProperty->SetObjectPropertyValue(ArrayHelper.GetRawPtr(Index), NewItem);
-								}
 							});
 						}
 						else if (Set)
@@ -1182,14 +1174,6 @@ int32 FPropertyValueImpl::AddChild()
 								Index = SetHelper.AddDefaultValue_Invalid_NeedsRehash();
 								ReturnLogicalIndex = SetHelper.FindLogicalIndex(Index);
 
-								// check whether the element type is flagged as a non-nullable. if so, create it.
-								FObjectProperty* ElementObjectProperty = CastField<FObjectProperty>(Set->ElementProp);
-								if (ElementObjectProperty && ElementObjectProperty->HasAnyPropertyFlags(CPF_NonNullable))
-								{
-									UObject* NewItem = NewObject<UObject>(Obj, ElementObjectProperty->PropertyClass);
-									ElementObjectProperty->SetObjectPropertyValue(SetHelper.GetElementPtr(Index), NewItem);
-								}
-
 								SetHelper.Rehash();
 							});
 						}
@@ -1201,25 +1185,6 @@ int32 FPropertyValueImpl::AddChild()
 								Index = MapHelper.AddDefaultValue_Invalid_NeedsRehash();
 								ReturnLogicalIndex = MapHelper.FindLogicalIndex(Index);
 
-								// check whether the key or value type is flagged as a non-nullable. if so, create it.
-								{
-									FObjectProperty* KeyObjectProperty = CastField<FObjectProperty>(Map->KeyProp);
-									if (KeyObjectProperty && KeyObjectProperty->HasAnyPropertyFlags(CPF_NonNullable))
-									{
-										UObject* NewItem = NewObject<UObject>(Obj, KeyObjectProperty->PropertyClass);
-										KeyObjectProperty->SetObjectPropertyValue(MapHelper.GetKeyPtr(Index), NewItem);
-									}
-								}
-
-								{
-									FObjectProperty* ValueObjectProperty = CastField<FObjectProperty>(Map->ValueProp);
-									if (ValueObjectProperty && ValueObjectProperty->HasAnyPropertyFlags(CPF_NonNullable))
-									{
-										UObject* NewItem = NewObject<UObject>(Obj, ValueObjectProperty->PropertyClass);
-										ValueObjectProperty->SetObjectPropertyValue(MapHelper.GetValuePtr(Index), NewItem);
-									}
-								}
-
 								MapHelper.Rehash();
 								bAddedMapEntry = true;
 							});
@@ -1491,14 +1456,6 @@ void FPropertyValueImpl::InsertChild( TSharedPtr<FPropertyNode> ChildNodeToInser
 
 		ArrayHelper.InsertValues(Index, 1 );
 
-		// check whether the inner type is flagged as a non-nullable. if so, create it.
-		FObjectProperty* InnerObjectProperty = CastField<FObjectProperty>(ArrayProperty->Inner);
-		if (InnerObjectProperty && InnerObjectProperty->HasAnyPropertyFlags(CPF_NonNullable))
-		{
-			UObject* NewItem = NewObject<UObject>(Obj, InnerObjectProperty->PropertyClass);
-			InnerObjectProperty->SetObjectPropertyValue(ArrayHelper.GetRawPtr(Index), NewItem);
-		}
-
 		//set up indices for the coming events
 		TArray< TMap<FString,int32> > ArrayIndicesPerObject;
 		for (int32 ObjectIndex = 0; ObjectIndex < ReadAddresses.Num(); ++ObjectIndex)
diff --git a/Engine/Source/Runtime/CoreUObject/Private/UObject/PropertyBaseObject.cpp b/Engine/Source/Runtime/CoreUObject/Private/UObject/PropertyBaseObject.cpp
--- a/Engine/Source/Runtime/CoreUObject/Private/UObject/PropertyBaseObject.cpp
+++ b/Engine/Source/Runtime/CoreUObject/Private/UObject/PropertyBaseObject.cpp
@@ -702,56 +702,6 @@ bool FObjectPropertyBase::AllowObjectTypeReinterpretationTo(const FObjectPropert
 	return false;
 }
 
-UObject* FObjectPropertyBase::ConstructDefaultObjectValueIfNecessary(UObject* ExistingValue) const
-{
-	UObject* NewDefaultObjectValue = nullptr;
-
-	FUObjectSerializeContext* SerializeContext = FUObjectThreadContext::Get().GetSerializeContext();
-	UObject* Outer = SerializeContext ? SerializeContext->SerializedObject : nullptr;
-	if (!Outer)
-	{
-		Outer = GetTransientPackage();
-	}
-
-	if (ExistingValue)
-	{
-		UClass* ExistingValueClass = ExistingValue->GetClass();
-		// Sanity check to make sure the existing value class matches the property class
-		if (ExistingValueClass && (ExistingValueClass->IsChildOf(PropertyClass) || ExistingValueClass->GetAuthoritativeClass()->IsChildOf(PropertyClass)))
-		{
-			if (ExistingValue->IsTemplate() && 	// Existing value is a template so we can construct a new value with it as the archetype
-				ExistingValue->GetOuter() != Outer) // Unless the template's Outer is the same as the new Outer in which case the template (ExistingValue) IS the object we can reuse
-			{
-				// We probably got here because an object value failed to load (missing import class) and the property is left with a template of default subobject
-				NewDefaultObjectValue = NewObject<UObject>(Outer, ExistingValue->GetClass(), NAME_None, RF_NoFlags, ExistingValue);
-			}
-			else
-			{
-				// Existing value is not a template or a template is what this property was pointing to so we can use it directly
-				// Similar to the above condition but the property was not referencing an instanced value in which case it's ok to leave the CDO default here
-				NewDefaultObjectValue = ExistingValue;
-			}
-		}
-	}
-
-	if (!NewDefaultObjectValue && !PropertyClass->HasAnyClassFlags(CLASS_Abstract))
-	{
-		// Existing value did not exist or it could not be used as a template
-		// Existing value may be null in case we were serializing an array of UObjects that failed to load (missing import class). Since the array is first pre-allocated with null values
-		// it will not have any existing objects to instantiate
-		NewDefaultObjectValue = NewObject<UObject>(Outer, PropertyClass);
-	}
-
-	// Final sanity check. We still may end up with a null object if the property class is abstract and the previous object value was missing or was not compatible
-	UE_CLOG(!NewDefaultObjectValue, LogProperty, Fatal, TEXT("Failed to create default object value for property %s. Previous value: %s, property class: %s, outer: %s"), 
-		*GetFullName(), 
-		ExistingValue ? *ExistingValue->GetFullName() : TEXT("None"),
-		PropertyClass ? *PropertyClass->GetFullName() : TEXT("None"),
-		Outer ? *Outer->GetFullName() : TEXT("None"));
-
-	return NewDefaultObjectValue;
-}
-
 void FObjectPropertyBase::CheckValidObject(void* ValueAddress, TObjectPtr<UObject> OldValue) const
 {
 	const TObjectPtr<UObject> Object = GetObjectPtrPropertyValue(ValueAddress);
@@ -806,24 +756,20 @@ void FObjectPropertyBase::CheckValidObject(void* ValueAddress, TObjectPtr<UObjec
 					*GetFullName(),
 					*Object.GetFullName()
 				);
-				SetObjectPropertyValue(ValueAddress, nullptr);
 			}
 			else
 			{
-				UObject* DefaultValue = ConstructDefaultObjectValueIfNecessary(OldValue);
-
 				UE_LOG(LogProperty, Warning,
-					TEXT("Serialized %s for a non-nullable property of %s. Reference will be defaulted to %s.\n    ReferencingObject = %s\n    Property = %s\n    Item = %s"),
+					TEXT("Serialized %s for a non-nullable property of %s. Reference will be nulled - will cause a runtime error if accessed.\n    ReferencingObject = %s\n    Property = %s\n    Item = %s"),
 					*ObjectClass->GetFullName(),
 					*PropertyClass->GetFullName(),
-					DefaultValue ? *DefaultValue->GetFullName() : TEXT("None"),
 					*GetFullNameSafe(SerializeContext ? SerializeContext->SerializedObject : nullptr),
 					*GetFullName(),
 					*Object.GetFullName()
 				);
-
-				SetObjectPropertyValue(ValueAddress, DefaultValue);
 			}
+
+			SetObjectPropertyValueUnchecked(ValueAddress, nullptr);
 		}
 	}
 }
diff --git a/Engine/Source/Runtime/CoreUObject/Private/UObject/PropertyObject.cpp b/Engine/Source/Runtime/CoreUObject/Private/UObject/PropertyObject.cpp
--- a/Engine/Source/Runtime/CoreUObject/Private/UObject/PropertyObject.cpp
+++ b/Engine/Source/Runtime/CoreUObject/Private/UObject/PropertyObject.cpp
@@ -254,16 +254,12 @@ void FObjectProperty::PostSerializeObjectItem(FArchive& SerializingArchive, void
 		!SerializingArchive.IsTransacting() && // Don't create new objects when loading from the transaction buffer
 		!SerializingArchive.IsCountingMemory())
 	{
-		UObject* DefaultValue = ConstructDefaultObjectValueIfNecessary(CurrentValue);
-
 		UE_LOG(LogProperty, Warning,
-			TEXT("Failed to serialize value for non-nullable property %s. Reference will be defaulted to %s."),
-			*GetFullName(),
-			DefaultValue ? *DefaultValue->GetFullName() : TEXT("None")
+			TEXT("Failed to serialize value for non-nullable property %s. Reference will be nulled - will cause a runtime error if accessed."),
+			*GetFullName()
 		);
 
-		SetObjectPropertyValue(Value, DefaultValue);
-		ObjectValue = DefaultValue;
+		SetObjectPropertyValueUnchecked(Value, nullptr);
 	}
 
 	if (ObjectValue != CurrentValue)
diff --git a/Engine/Source/Runtime/CoreUObject/Public/UObject/UnrealType.h b/Engine/Source/Runtime/CoreUObject/Public/UObject/UnrealType.h
--- a/Engine/Source/Runtime/CoreUObject/Public/UObject/UnrealType.h
+++ b/Engine/Source/Runtime/CoreUObject/Public/UObject/UnrealType.h
@@ -2668,14 +2668,6 @@ protected:
 	virtual bool AllowCrossLevel() const;
 	// End of FObjectPropertyBase interface
 
-	/**
-	 * Constructs a new object if the existing one is missing or is not compatible with the property class
-	 * Used for making sure non-nullable properties have valid values.
-	 * @param ExistingValue Previous object value (can be null)
-	 * @return Non-null object that was assigned to this property value address
-	 */
-	UObject* ConstructDefaultObjectValueIfNecessary(UObject* ExistingValue) const;
-
 	// Disable false positive buffer overrun warning during pgoprofile linking step
 	PRAGMA_DISABLE_BUFFER_OVERRUN_WARNING
 	/* Helper functions for UObject property types that wrap the object pointer in a smart pointer */
diff --git a/Engine/Source/Runtime/CoreUObject/Tests/ObjectPropertyTest.cpp b/Engine/Source/Runtime/CoreUObject/Tests/ObjectPropertyTest.cpp
--- a/Engine/Source/Runtime/CoreUObject/Tests/ObjectPropertyTest.cpp
+++ b/Engine/Source/Runtime/CoreUObject/Tests/ObjectPropertyTest.cpp
@@ -6,11 +6,13 @@
 #include "TestMacros/Assertions.h"
 #include "ObjectPtrTestClass.h"
 #include "UObject/Package.h"
+#include "Logging/LogScopedVerbosityOverride.h"
 #include "Misc/ScopeExit.h"
 #include "UObject/UnrealType.h"
 #include "UObject/LinkerPlaceholderExportObject.h"
 #include "UObject/LinkerPlaceholderClass.h"
 #include "UObject/ObjectHandleTracking.h"
+#include "UObject/PropertyHelper.h"
 #include "Tests/WarnFilterScope.h"
 #include "Misc/AssetRegistryInterface.h"
 #include "AssetRegistry/AssetData.h"
@@ -121,6 +123,9 @@ TEST_CASE("UE::CoreUObject::FObjectProperty::CheckValidAddressNonNullable")
 	Property->CheckValidObject(&Obj->ObjectPtrNonNullable, nullptr);
 	CHECK(Obj->ObjectPtrNonNullable == ObjectPtr);
 
+	// Disable property warnings that will fire because we're deliberately setting non-nullable properties to null
+	LOG_SCOPE_VERBOSITY_OVERRIDE(LogProperty, ELogVerbosity::NoLogging);
+
 	bAllowRead = true; //has resolve the old value to construct a new default value for the property
 	//assign a bad value to the pointer
 	Obj->ObjectPtrNonNullable = reinterpret_cast<UObjectPtrTestClass*>(OtherTestPackage);
@@ -135,16 +140,15 @@ TEST_CASE("UE::CoreUObject::FObjectProperty::CheckValidAddressNonNullable")
 			return false;
 		});
 	Property->CheckValidObject(&Obj->ObjectPtrNonNullable, ObjectPtr);
-	CHECK(Obj->ObjectPtrNonNullable == ObjectPtr); //non nullable properties should be assigned the old value
+	CHECK(Obj->ObjectPtrNonNullable == nullptr); //non nullable properties should be nulled if invalid
 
 	//assign a bad value to the pointer
 	Obj->ObjectPtrNonNullable = reinterpret_cast<UObjectPtrTestClass*>(Obj);
 	CHECK(Obj->ObjectPtrNonNullable != nullptr);
 
-	//new value is required for non nullable properties
+	//null is required for invalid non nullable properties
 	Property->CheckValidObject(&Obj->ObjectPtrNonNullable, nullptr);
-	CHECK(Obj->ObjectPtrNonNullable != nullptr);
-	CHECK(Obj->ObjectPtrNonNullable->IsA(UObjectPtrTestClass::StaticClass()));
+	CHECK(Obj->ObjectPtrNonNullable == nullptr);
 }
 
 class FMockArchive : public FArchive
