Cooker: SavePackage: Content is Missing From Cook Warning: Change the log message to a structured log record so that farm jobs that use structured logging can interpret the message.
#rnx

diff --git a/Engine/Source/Editor/UnrealEd/Private/Commandlets/CookCommandlet.cpp b/Engine/Source/Editor/UnrealEd/Private/Commandlets/CookCommandlet.cpp
--- a/Engine/Source/Editor/UnrealEd/Private/Commandlets/CookCommandlet.cpp
+++ b/Engine/Source/Editor/UnrealEd/Private/Commandlets/CookCommandlet.cpp
@@ -25,6 +25,7 @@
 #include "IPlatformFileSandboxWrapper.h"
 #include "Interfaces/ITargetPlatform.h"
 #include "Interfaces/ITargetPlatformManagerModule.h"
+#include "Logging/StructuredLog.h"
 #include "Misc/App.h"
 #include "Misc/CommandLine.h"
 #include "Misc/ConfigCacheIni.h"
@@ -601,11 +602,11 @@ void UCookCommandlet::RunCookByTheBookCook(UCookOnTheFlyServer* CookOnTheFlyServ
 	if (bShouldVerifyEDLCookInfo)
 	{
 		bool bFullReferencesExpected = !(CookOptions & ECookByTheBookOptions::SkipHardReferences);
-		UE::SavePackageUtilities::VerifyEDLCookInfo([](ELogVerbosity::Type Verbosity, FStringView Message)
+		UE::SavePackageUtilities::VerifyEDLCookInfo([](UE::FLogRecord&& Record)
 			{
 #if !NO_LOGGING
-				FMsg::Logf(__FILE__, __LINE__, LogCook.GetCategoryName(), Verbosity, TEXT("%.*s"),
-				Message.Len(), Message.GetData());
+				Record.SetCategory(LogCook.GetCategoryName());
+				UE::DispatchDynamicLogRecord(Record);
 #endif
 			}, bFullReferencesExpected);
 	}
diff --git a/Engine/Source/Runtime/CoreUObject/Private/UObject/SavePackage/SavePackageUtilities.cpp b/Engine/Source/Runtime/CoreUObject/Private/UObject/SavePackage/SavePackageUtilities.cpp
--- a/Engine/Source/Runtime/CoreUObject/Private/UObject/SavePackage/SavePackageUtilities.cpp
+++ b/Engine/Source/Runtime/CoreUObject/Private/UObject/SavePackage/SavePackageUtilities.cpp
@@ -11,6 +11,7 @@
 #include "CoreMinimal.h"
 #include "HAL/FileManager.h"
 #include "Interfaces/ITargetPlatform.h"
+#include "Logging/StructuredLog.h"
 #include "Misc/AssetRegistryInterface.h"
 #include "Misc/CommandLine.h"
 #include "Misc/ConfigCacheIni.h"
@@ -1668,7 +1669,7 @@ FEDLCookChecker FEDLCookChecker::AccumulateAndClear()
 	return Accumulator;
 }
 
-void FEDLCookChecker::Verify(const UE::SavePackageUtilities::FEDLMessageCallback& MessageCallback,
+void FEDLCookChecker::Verify(const UE::SavePackageUtilities::FEDLLogRecordCallback& MessageCallback,
 	bool bFullReferencesExpected)
 {
 	check(!GIsSavingPackage);
@@ -1726,14 +1727,33 @@ void FEDLCookChecker::Verify(const UE::SavePackageUtilities::FEDLMessageCallback
 
 				for (FName PackageName : NodeData.ImportingPackagesSorted)
 				{
-					TStringBuilder<512> Message;
-					Message << TEXTVIEW("Content is missing from cook. Source package referenced an object in target package but ");
-					Message << ReasonExportIsMissing << TEXT(".\n");
-					Message << TEXT("\tSource package: ") << PackageName << TEXT("\n");
-					Message << TEXT("\tTarget package: ") << NodeDataOfExportPackage->Name << TEXT("\n");
-					Message << TEXT("\tReferenced object: ");
-					NodeData.AppendPathName(Accumulator, Message);
-					MessageCallback(MissingContentSeverity, Message);
+					UE::FLogRecord Record;
+#if !NO_LOGGING
+					Record.SetCategory(LogSavePackage.GetCategoryName());
+#endif
+					Record.SetVerbosity(MissingContentSeverity);
+					Record.SetTime(UE::FLogTime::Now());
+					Record.SetFormat(TEXT("Content is missing from cook. Source package referenced an object in target package but {Reason}.")
+						TEXT("\n\tSource package: {Source}")
+						TEXT("\n\tTarget package: {Target}")
+						TEXT("\n\tReferenced object: {ReferencedObject}"));
+					{
+						FCbWriter Writer;
+						Writer.BeginObject();
+						Writer << "Reason" << ReasonExportIsMissing;
+						Writer << "Source" << WriteToUtf8String<256>(PackageName);
+						Writer << "Target" << WriteToUtf8String<256>(NodeDataOfExportPackage->Name);
+						{
+							TStringBuilder<256> ReferencedObjectStr;
+							NodeData.AppendPathName(Accumulator, ReferencedObjectStr);
+							Writer << "ReferencedObject" << ReferencedObjectStr;
+						}
+						Writer.EndObject();
+						Record.SetFields(Writer.Save().AsObject());
+					}
+					Record.SetFile(__FILE__);
+					Record.SetLine(__LINE__);
+					MessageCallback(MoveTemp(Record));
 				}
 			}
 		}
@@ -1937,17 +1957,27 @@ void StartSavingEDLCookInfoForVerification()
 
 void VerifyEDLCookInfo(bool bFullReferencesExpected)
 {
-	VerifyEDLCookInfo([](ELogVerbosity::Type Verbosity, FStringView Message)
+	VerifyEDLCookInfo([](UE::FLogRecord&& Record)
 		{
 #if !NO_LOGGING
-			FMsg::Logf(__FILE__, __LINE__, LogSavePackage.GetCategoryName(), Verbosity, TEXT("%.*s"),
-				Message.Len(), Message.GetData());
+			DispatchDynamicLogRecord(Record);
 #endif
 		}, bFullReferencesExpected);
 }
 
 void VerifyEDLCookInfo(const UE::SavePackageUtilities::FEDLMessageCallback& MessageCallback,
 	bool bFullReferencesExpected)
+{
+	VerifyEDLCookInfo([&MessageCallback](FLogRecord&& Record)
+		{
+			TStringBuilder<256> Message;
+			Record.FormatMessageTo(Message);
+			MessageCallback(Record.GetVerbosity(), Message.ToView());
+		});
+}
+
+void VerifyEDLCookInfo(const UE::SavePackageUtilities::FEDLLogRecordCallback& MessageCallback,
+	bool bFullReferencesExpected)
 {
 	LLM_SCOPE_BYTAG(EDLCookChecker);
 	FEDLCookChecker::Verify(MessageCallback, bFullReferencesExpected);
diff --git a/Engine/Source/Runtime/CoreUObject/Private/UObject/SavePackage/SavePackageUtilities.h b/Engine/Source/Runtime/CoreUObject/Private/UObject/SavePackage/SavePackageUtilities.h
--- a/Engine/Source/Runtime/CoreUObject/Private/UObject/SavePackage/SavePackageUtilities.h
+++ b/Engine/Source/Runtime/CoreUObject/Private/UObject/SavePackage/SavePackageUtilities.h
@@ -160,7 +160,7 @@ struct FEDLCookChecker
 	void AddPackageWithUnknownExports(FName LongPackageName);
 
 	static void StartSavingEDLCookInfoForVerification();
-	static void Verify(const UE::SavePackageUtilities::FEDLMessageCallback& MessageCallback,
+	static void Verify(const UE::SavePackageUtilities::FEDLLogRecordCallback& MessageCallback,
 		bool bFullReferencesExpected);
 	static void MoveToCompactBinaryAndClear(FCbWriter& Writer, bool& bOutHasData);
 	static bool AppendFromCompactBinary(FCbFieldView Field);
diff --git a/Engine/Source/Runtime/CoreUObject/Public/UObject/SavePackage.h b/Engine/Source/Runtime/CoreUObject/Public/UObject/SavePackage.h
--- a/Engine/Source/Runtime/CoreUObject/Public/UObject/SavePackage.h
+++ b/Engine/Source/Runtime/CoreUObject/Public/UObject/SavePackage.h
@@ -40,6 +40,7 @@ class FOutputDevice;
 class FPackagePath;
 class FSavePackageContext;
 class IPackageWriter;
+namespace UE { class FLogRecord; }
 struct FObjectSaveContextData;
 
 /**
@@ -320,8 +321,11 @@ namespace UE::SavePackageUtilities
 
 	COREUOBJECT_API void StartSavingEDLCookInfoForVerification();
 	using FEDLMessageCallback = TFunction<void(ELogVerbosity::Type, FStringView)>;
+	using FEDLLogRecordCallback = TFunction<void(UE::FLogRecord&& Record)>;
 	COREUOBJECT_API void VerifyEDLCookInfo(bool bFullReferencesExpected = true);
+	UE_DEPRECATED(5.5, "Use version that takes FEDLLogRecordCallback")
 	COREUOBJECT_API void VerifyEDLCookInfo(const FEDLMessageCallback& MessageCallback, bool bFullReferencesExpected = true);
+	COREUOBJECT_API void VerifyEDLCookInfo(const FEDLLogRecordCallback& MessageCallback, bool bFullReferencesExpected = true);
 	COREUOBJECT_API void EDLCookInfoAddIterativelySkippedPackage(FName LongPackageName);
 	COREUOBJECT_API void EDLCookInfoMoveToCompactBinaryAndClear(FCbWriter& Writer, bool& bOutHasData);
 	COREUOBJECT_API void EDLCookInfoMoveToCompactBinaryAndClear(FCbWriter& Writer, bool& bOutHasData, FName PackageName);
