Cooker: SavePackage: Content is Missing From Cook Warning: Fix compile error "Call to 'VerifyEDLCookInfo' is ambiguous; it looks like some compilers support casting a lambda to a bool.
#rnx

diff --git a/Engine/Source/Runtime/CoreUObject/Private/UObject/SavePackage/SavePackageUtilities.cpp b/Engine/Source/Runtime/CoreUObject/Private/UObject/SavePackage/SavePackageUtilities.cpp
--- a/Engine/Source/Runtime/CoreUObject/Private/UObject/SavePackage/SavePackageUtilities.cpp
+++ b/Engine/Source/Runtime/CoreUObject/Private/UObject/SavePackage/SavePackageUtilities.cpp
@@ -1973,7 +1973,7 @@ void VerifyEDLCookInfo(const UE::SavePackageUtilities::FEDLMessageCallback& Mess
 			TStringBuilder<256> Message;
 			Record.FormatMessageTo(Message);
 			MessageCallback(Record.GetVerbosity(), Message.ToView());
-		});
+		}, bFullReferencesExpected);
 }
 
 void VerifyEDLCookInfo(const UE::SavePackageUtilities::FEDLLogRecordCallback& MessageCallback,
