#UE Add more info to find slow structs in uninitialized members test

diff --git a/Engine/Source/Runtime/CoreUObject/Private/UObject/Class.cpp b/Engine/Source/Runtime/CoreUObject/Private/UObject/Class.cpp
--- a/Engine/Source/Runtime/CoreUObject/Private/UObject/Class.cpp
+++ b/Engine/Source/Runtime/CoreUObject/Private/UObject/Class.cpp
@@ -4017,6 +4017,28 @@ static FString GetFieldLocation(const UField* Field)
 
 int32 FStructUtils::AttemptToFindUninitializedScriptStructMembers()
 {
+	struct FExecuteIfExceedsTimeLimit
+	{
+		FExecuteIfExceedsTimeLimit(double InTimeLimit, TFunction<void(double)> InFunc)
+			: StartTime(FPlatformTime::Seconds()), TimeLimit(InTimeLimit), Func(InFunc) {}
+
+		~FExecuteIfExceedsTimeLimit()
+		{
+			double Time = FPlatformTime::Seconds() - StartTime;
+			if (Time > TimeLimit)
+			{
+				Func(Time);
+			}
+		}
+		double StartTime;
+		double TimeLimit;
+		TFunction<void(double)> Func;
+	};
+
+	FExecuteIfExceedsTimeLimit TotalTimeLimit(2.0, [](double Time) {
+		UE_LOG(LogClass, Display, TEXT("AttemptToFindUninitializedScriptStructMembers took more than 2s to complete. Time: %.2f"), Time);
+	});
+
 	auto DetermineIfModuleIsEngine = [](const UScriptStruct* ScriptStruct) -> bool
 	{
 		UPackage* ScriptPackage = ScriptStruct->GetOutermost();
@@ -4043,6 +4065,9 @@ int32 FStructUtils::AttemptToFindUninitializedScriptStructMembers()
 
 		FScriptStructSettings()
 		{
+			FExecuteIfExceedsTimeLimit SettingsTimeLimit(0.05, [](double Time) {
+				UE_LOG(LogClass, Display, TEXT("AttemptToFindUninitializedScriptStructMembers took more than 50ms to construct the FScriptStructSettings. Time: %.1f"), Time * 1000);
+				});
 			{
 				FString ProjectSettingString;
 				if (GConfig->GetString(TEXT("CoreUObject.UninitializedScriptStructMembersCheck"), TEXT("ProjectModuleReflectedUninitializedPropertyVerbosity"), ProjectSettingString, GEngineIni))
@@ -4102,11 +4127,20 @@ int32 FStructUtils::AttemptToFindUninitializedScriptStructMembers()
 		}
 	}
 
+	double PreScriptStructTime = FPlatformTime::Seconds() - TotalTimeLimit.StartTime;
+	if (PreScriptStructTime > 0.05)
+	{
+		UE_LOG(LogClass, Display, TEXT("AttemptToFindUninitializedScriptStructMembers took more than 50ms before starting to check ScriptStructs. Time(ms):%.1f"), PreScriptStructTime * 1000);
+	}
+
 	TSet<const FProperty*> UninitializedPropertiesNoInit;
 	TSet<const FProperty*> UninitializedPropertiesZeroed;
 	for (TObjectIterator<UScriptStruct> ScriptIt; ScriptIt; ++ScriptIt)
 	{
 		UScriptStruct* ScriptStruct = *ScriptIt;
+		FExecuteIfExceedsTimeLimit StructTimeLimit(0.001, [ScriptStruct](double Time) {
+			UE_LOG(LogClass, Display, TEXT("AttemptToFindUninitializedScriptStructMembers took more than 1ms to process ScriptStruct %s. Time(ms): %.1f"), ScriptStruct ? *ScriptStruct->GetPathName() : TEXT("None"), Time * 1000);
+			});
 
 		if (!FScriptStructTestWrapper::CanRunTests(ScriptStruct) || ScriptStruct == TestUninitializedScriptStructMembersTestStruct)
 		{
