[Bugfix] Crash when adding to Container properties

diff --git a/Engine/Source/Runtime/CoreUObject/Private/UObject/PropertyBagRepository.cpp b/Engine/Source/Runtime/CoreUObject/Private/UObject/PropertyBagRepository.cpp
--- a/Engine/Source/Runtime/CoreUObject/Private/UObject/PropertyBagRepository.cpp
+++ b/Engine/Source/Runtime/CoreUObject/Private/UObject/PropertyBagRepository.cpp
@@ -240,6 +240,10 @@ static bool ConstructRemappedPropertyChain(const FEditPropertyChain& Chain, FEdi
 		{
 			Struct = AsStructProperty->Struct;
 		}
+		else if (FObjectProperty* AsObjectProperty = CastField<FObjectProperty>(Property))
+		{
+			Struct = AsObjectProperty->PropertyClass;
+		}
 		else
 		{
 			check(Itr->GetNextNode() == nullptr);
@@ -258,7 +262,8 @@ static bool ConstructRemappedPropertyChain(const FEditPropertyChain& Chain, FEdi
 	return true;
 }
 
-static void* ResolveChangePath(const void* StructData, FPropertyChangedChainEvent& ChangeEvent, bool bGrowContainersWhenNeeded = false)
+// note: this method is very similar to RemapChangeEvent so there's a good chance bugs found here will also need to be fixed there
+static void* ResolveChangePath(const void* StructData, FProperty*& ResolvedProperty, FPropertyChangedChainEvent& ChangeEvent)
 {
 	if (ChangeEvent.PropertyChain.GetHead() == nullptr)
 	{
@@ -269,14 +274,16 @@ static void* ResolveChangePath(const void* StructData, FPropertyChangedChainEven
 	void* MemoryPtr = const_cast<void*>(StructData);
 	do
 	{
-		const FProperty* Property = PropertyNode->GetValue();
+		FProperty* Property = PropertyNode->GetValue();
 		MemoryPtr = Property->ContainerPtrToValuePtr<uint8>(MemoryPtr);
-		PropertyNode = PropertyNode->GetNextNode();
+		ResolvedProperty = Property;
 		
 		const int32 ArrayIndex = ChangeEvent.GetArrayIndex(Property->GetName());
-		if (PropertyNode && ArrayIndex != INDEX_NONE)
+		if (ArrayIndex != INDEX_NONE && !(ChangeEvent.ChangeType & EPropertyChangeType::ArrayRemove))
 		{
-			if (const FArrayProperty* ArrayProperty = CastField<FArrayProperty>(Property->GetOwnerProperty()))
+			PropertyNode = PropertyNode->GetNextNode();
+			FProperty* InnerProperty = PropertyNode ? PropertyNode->GetValue() : ChangeEvent.Property;
+			if (const FArrayProperty* ArrayProperty = CastField<FArrayProperty>(Property))
 			{
 				FScriptArrayHelper ArrayHelper(ArrayProperty, MemoryPtr);
 				if(!ArrayHelper.IsValidIndex(ArrayIndex))
@@ -284,11 +291,9 @@ static void* ResolveChangePath(const void* StructData, FPropertyChangedChainEven
 					return nullptr;
 				}
 				MemoryPtr = ArrayHelper.GetRawPtr(ArrayIndex);
-
-				// skip to the next property node already
-				PropertyNode = PropertyNode->GetNextNode();
+				ResolvedProperty = ArrayProperty->Inner;
 			}
-			if (const FSetProperty* SetProperty = CastField<FSetProperty>(Property->GetOwnerProperty()))
+			else if (const FSetProperty* SetProperty = CastField<FSetProperty>(Property))
 			{
 				FScriptSetHelper SetHelper(SetProperty, MemoryPtr);
 				if(!SetHelper.IsValidIndex(ArrayIndex))
@@ -296,28 +301,190 @@ static void* ResolveChangePath(const void* StructData, FPropertyChangedChainEven
 					return nullptr;
 				}
 				MemoryPtr = SetHelper.GetElementPtr(ArrayIndex);
-
-				// skip to the next property node already
-				PropertyNode = PropertyNode->GetNextNode();
+				ResolvedProperty = SetProperty->ElementProp;
 			}
-			if (const FMapProperty* MapProperty = CastField<FMapProperty>(Property->GetOwnerProperty()))
+			else if (const FMapProperty* MapProperty = CastField<FMapProperty>(Property))
 			{
 				FScriptMapHelper MapHelper(MapProperty, MemoryPtr);
 				if(!MapHelper.IsValidIndex(ArrayIndex))
 				{
 					return nullptr;
 				}
-				MemoryPtr = MapHelper.GetValuePtr(ArrayIndex);
-
-				// skip to the next property node already
-				PropertyNode = PropertyNode->GetNextNode();
+				if (InnerProperty == MapProperty->KeyProp)
+				{
+					MemoryPtr = MapHelper.GetKeyPtr(ArrayIndex);
+					ResolvedProperty = MapProperty->KeyProp;
+				}
+				else
+				{
+					MemoryPtr = MapHelper.GetValuePtr(ArrayIndex);
+					ResolvedProperty = MapProperty->ValueProp;
+				}
+			}
+		}
+		else if (PropertyNode->GetNextNode())
+		{
+			if (const FObjectProperty* AsObjectProperty = CastField<FObjectProperty>(Property))
+			{
+				MemoryPtr = AsObjectProperty->GetObjectPropertyValue(MemoryPtr);
 			}
 		}
+
+		
+		if (PropertyNode)
+		{
+			PropertyNode = PropertyNode->GetNextNode();
+		}
 	}
 	while (PropertyNode);
 
 	return MemoryPtr;
 }
+
+
+// note: this method is very similar to ResolveChangePath so there's a good chance bugs found here will also need to be fixed there
+static bool RemapChangeEvent(const FPropertyChangedChainEvent& InChangeEvent, FPropertyChangedChainEvent& OutRemappedChangeEvent, void*& MemoryPtr, FProperty*& ResolvedProperty, TArray<TMap<FString, int32>>& ArrayIndices)
+{
+	
+	if (!ConstructRemappedPropertyChain(InChangeEvent.PropertyChain, OutRemappedChangeEvent.PropertyChain, static_cast<UObject*>(MemoryPtr)))
+	{
+		return false;
+	}
+	
+	if (OutRemappedChangeEvent.PropertyChain.GetHead() == nullptr)
+	{
+		return false;
+	}
+		
+	FEditPropertyChain::TDoubleLinkedListNode* RemappedPropertyNode = OutRemappedChangeEvent.PropertyChain.GetHead();
+	FEditPropertyChain::TDoubleLinkedListNode* InPropertyNode = InChangeEvent.PropertyChain.GetHead();
+	do
+	{
+		FProperty* RemappedProperty = RemappedPropertyNode->GetValue();
+		FProperty* InProperty = InPropertyNode->GetValue();
+		MemoryPtr = RemappedProperty->ContainerPtrToValuePtr<uint8>(MemoryPtr);
+		ResolvedProperty = RemappedProperty;
+
+		if (InChangeEvent.Property == InProperty)
+		{
+			OutRemappedChangeEvent.Property = RemappedProperty;
+		}
+		if (InChangeEvent.MemberProperty == InProperty)
+		{
+			OutRemappedChangeEvent.MemberProperty = RemappedProperty;
+		}
+		
+		int32 ArrayIndex = InChangeEvent.GetArrayIndex(RemappedProperty->GetName());
+		if (ArrayIndex != INDEX_NONE && !(InChangeEvent.ChangeType & EPropertyChangeType::ArrayRemove))
+		{
+			RemappedPropertyNode = RemappedPropertyNode->GetNextNode();
+			
+			InPropertyNode = InPropertyNode->GetNextNode();
+			FProperty* InnerProperty = InPropertyNode ? InPropertyNode->GetValue() : InChangeEvent.Property;
+			if (const FArrayProperty* ArrayProperty = CastField<FArrayProperty>(RemappedProperty))
+			{
+				ResolvedProperty = ArrayProperty->Inner;
+				if (InChangeEvent.Property->Owner == InProperty)
+				{
+					OutRemappedChangeEvent.Property = ArrayProperty->Inner;
+				}
+				if (InChangeEvent.MemberProperty->Owner == InProperty)
+				{
+					OutRemappedChangeEvent.MemberProperty = ArrayProperty->Inner;
+				}
+				
+				FScriptArrayHelper ArrayHelper(ArrayProperty, MemoryPtr);
+				if(!ArrayHelper.IsValidIndex(ArrayIndex))
+				{
+					check(OutRemappedChangeEvent.ChangeType & EPropertyChangeType::ArrayAdd);
+					check(ArrayHelper.Num() == ArrayIndex);
+					ArrayHelper.Resize(ArrayIndex + 1);
+				}
+				ArrayIndices[InChangeEvent.ObjectIteratorIndex].Add(RemappedProperty->GetName(), ArrayIndex);
+				MemoryPtr = ArrayHelper.GetRawPtr(ArrayIndex);
+			}
+			if (const FSetProperty* SetProperty = CastField<FSetProperty>(RemappedProperty))
+			{
+				ResolvedProperty = SetProperty->ElementProp;
+				if (InChangeEvent.Property->Owner == InProperty)
+				{
+					OutRemappedChangeEvent.Property = SetProperty->ElementProp;
+				}
+				if (InChangeEvent.MemberProperty->Owner == InProperty)
+				{
+					OutRemappedChangeEvent.MemberProperty = SetProperty->ElementProp;
+				}
+				
+				FScriptSetHelper SetHelper(SetProperty, MemoryPtr);
+				if(!SetHelper.IsValidIndex(ArrayIndex))
+				{
+					check(OutRemappedChangeEvent.ChangeType & EPropertyChangeType::ArrayAdd);
+					ArrayIndex = SetHelper.AddUninitializedValue();
+				}
+				ArrayIndices[InChangeEvent.ObjectIteratorIndex].Add(RemappedProperty->GetName(), ArrayIndex);
+				MemoryPtr = SetHelper.GetElementPtr(ArrayIndex); // TODO: do more testing with sets
+			}
+			if (const FMapProperty* MapProperty = CastField<FMapProperty>(RemappedProperty))
+			{
+				if (InnerProperty == CastField<FMapProperty>(InProperty)->KeyProp) // TODO: do more testing with maps
+				{
+					ResolvedProperty = MapProperty->KeyProp;
+					if (InChangeEvent.Property->Owner == InProperty)
+					{
+						OutRemappedChangeEvent.Property = MapProperty->KeyProp;
+					}
+					if (InChangeEvent.MemberProperty->Owner == InProperty)
+					{
+						OutRemappedChangeEvent.MemberProperty = MapProperty->KeyProp;
+					}
+				}
+				else
+				{
+					ResolvedProperty = MapProperty->ValueProp;
+					if (InChangeEvent.Property->Owner == InProperty)
+					{
+						OutRemappedChangeEvent.Property = MapProperty->ValueProp;
+					}
+					if (InChangeEvent.MemberProperty->Owner == InProperty)
+					{
+						OutRemappedChangeEvent.MemberProperty = MapProperty->ValueProp;
+					}
+				}
+				
+				FScriptMapHelper MapHelper(MapProperty, MemoryPtr);
+				if(!MapHelper.IsValidIndex(ArrayIndex))
+				{
+					check(OutRemappedChangeEvent.ChangeType & EPropertyChangeType::ArrayAdd);
+					ArrayIndex = MapHelper.AddUninitializedValue();
+				}
+				ArrayIndices[InChangeEvent.ObjectIteratorIndex].Add(RemappedProperty->GetName(), ArrayIndex);
+				MemoryPtr = InnerProperty == CastField<FMapProperty>(InProperty)->KeyProp ? MapHelper.GetKeyPtr(ArrayIndex) : MapHelper.GetValuePtr(ArrayIndex);
+			}
+		}
+		else if (RemappedPropertyNode->GetNextNode())
+		{
+			if (const FObjectProperty* AsObjectProperty = CastField<FObjectProperty>(RemappedProperty))
+         	{
+         		MemoryPtr = AsObjectProperty->GetObjectPropertyValue(MemoryPtr);
+         	}
+		}
+		
+		if (RemappedPropertyNode)
+		{
+			RemappedPropertyNode = RemappedPropertyNode->GetNextNode();
+			InPropertyNode = InPropertyNode->GetNextNode();
+		}
+	}
+	while (RemappedPropertyNode);
+
+	OutRemappedChangeEvent.SetArrayIndexPerObject(ArrayIndices);
+
+	check(InChangeEvent.Property != OutRemappedChangeEvent.Property);
+	check(InChangeEvent.MemberProperty != OutRemappedChangeEvent.MemberProperty);
+
+	return true;
+}
+
 static void CopyProperty(const FProperty* SourceProperty, const void* SourceValue, const FProperty* DestProperty, void* DestValue)
 {
 	check(SourceProperty->GetID() == DestProperty->GetID());
@@ -408,53 +575,6 @@ static void CopyProperty(const FProperty* SourceProperty, const void* SourceValu
 	}
 }
 
-static void AddProperty(const FProperty* SourceProperty, const void* SourceValue, const FProperty* DestProperty, void* DestValue, int32 ArrayIndex)
-{
-	if (const FArrayProperty* SourcePropertyAsArray = CastField<FArrayProperty>(SourceProperty))
-	{
-		const FArrayProperty* DestPropertyAsArray = CastFieldChecked<FArrayProperty>(DestProperty);
-		FScriptArrayHelper SourceArray(SourcePropertyAsArray, SourceValue);
-        FScriptArrayHelper DestArray(DestPropertyAsArray, DestValue);
-		if (DestArray.Num() < ArrayIndex + 1)
-		{
-			DestArray.Resize(ArrayIndex + 1);
-		}
-		CopyProperty(SourcePropertyAsArray->Inner, SourceArray.GetElementPtr(ArrayIndex),
-			DestPropertyAsArray->Inner, DestArray.GetElementPtr(ArrayIndex));
-	}
-	else if (const FSetProperty* SourcePropertyAsSet = CastField<FSetProperty>(SourceProperty))
-	{
-		const FSetProperty* DestPropertyAsSet = CastFieldChecked<FSetProperty>(DestProperty);
-		FScriptSetHelper SourceSet(SourcePropertyAsSet, SourceValue);
-        FScriptSetHelper DestSet(DestPropertyAsSet, DestValue);
-		int32 DestArrayIndex = DestSet.AddUninitializedValue();
-		
-		void* DestElementPtr = DestSet.GetElementPtr(DestArrayIndex);
-		DestPropertyAsSet->ElementProp->InitializeValue(DestElementPtr);
-		CopyProperty(SourcePropertyAsSet->ElementProp, SourceSet.GetElementPtr(ArrayIndex),
-			DestPropertyAsSet->ElementProp, DestElementPtr);
-		DestSet.Rehash();
-	}
-	else if (const FMapProperty* SourcePropertyAsMap = CastField<FMapProperty>(SourceProperty))
-	{
-		const FMapProperty* DestPropertyAsMap = CastFieldChecked<FMapProperty>(DestProperty);
-		FScriptMapHelper SourceMap(SourcePropertyAsMap, SourceValue);
-		FScriptMapHelper DestMap(DestPropertyAsMap, DestValue);
-		int32 DestArrayIndex = DestMap.AddUninitializedValue();
-
-		void* DestKeyPtr = DestMap.GetKeyPtr(DestArrayIndex);
-		DestPropertyAsMap->KeyProp->InitializeValue(DestKeyPtr);
-		CopyProperty(SourcePropertyAsMap->KeyProp, SourceMap.GetKeyPtr(ArrayIndex),
-			DestPropertyAsMap->KeyProp, DestKeyPtr);
-		
-		void* DestValuePtr = DestMap.GetValuePtr(DestArrayIndex);
-		DestPropertyAsMap->ValueProp->InitializeValue(DestValuePtr);
-		CopyProperty(SourcePropertyAsMap->ValueProp, SourceMap.GetValuePtr(ArrayIndex),
-			DestPropertyAsMap->ValueProp, DestValuePtr);
-		DestMap.Rehash();
-	}
-}
-
 void FPropertyBagRepository::PostEditChangeChainProperty(const UObject* Object, FPropertyChangedChainEvent& PropertyChangedEvent)
 {
 #if WITH_EDITOR
@@ -465,29 +585,23 @@ void FPropertyBagRepository::PostEditChangeChainProperty(const UObject* Object,
 		return;
 	}
 	
-	auto CopyChanges = [&PropertyChangedEvent](const UObject* Source, UObject* Dest) 
+	auto CopyChanges = [&PropertyChangedEvent](const UObject* Source, UObject* Dest)
 	{
+		FPropertyChangedEvent BaseRemappedEvent(PropertyChangedEvent);
 		FEditPropertyChain RemappedChain;
-		if (ConstructRemappedPropertyChain(PropertyChangedEvent.PropertyChain, RemappedChain, Dest))
+		FPropertyChangedChainEvent RemappedChangeEvent(RemappedChain, BaseRemappedEvent);
+		void* DestData = Dest;
+		TArray<TMap<FString, int32>> RemappedArrayIndices;
+		RemappedArrayIndices.AddDefaulted(PropertyChangedEvent.ObjectIteratorIndex + 1);
+		FProperty* DestProperty = nullptr;
+		if (RemapChangeEvent(PropertyChangedEvent, RemappedChangeEvent, DestData, DestProperty, RemappedArrayIndices))
 		{
-			Dest->PreEditChange(RemappedChain);
-            
-            FPropertyChangedChainEvent RemappedChangeEvent(RemappedChain, PropertyChangedEvent);
-            const void* SourceData = ResolveChangePath(Source, PropertyChangedEvent);
-            void* DestData = ResolveChangePath(Dest, RemappedChangeEvent, true);
-            FProperty* SourceProperty = PropertyChangedEvent.PropertyChain.GetTail()->GetValue();
-            FProperty* DestProperty = RemappedChangeEvent.PropertyChain.GetTail()->GetValue();
+			Dest->PreEditChange(RemappedChangeEvent.PropertyChain);
+			
+            FProperty* SourceProperty = nullptr;
+            const void* SourceData = ResolveChangePath(Source, SourceProperty, PropertyChangedEvent);
 			
-            if (PropertyChangedEvent.ChangeType == EPropertyChangeType::ArrayAdd)
-            {
-            	int32 ArrayIndex = PropertyChangedEvent.GetArrayIndex(SourceProperty->GetName());
-            	check(ArrayIndex != INDEX_NONE);
-            	AddProperty(SourceProperty, SourceData, DestProperty, DestData, ArrayIndex);
-            }
-            else
-            {
-            	CopyProperty(SourceProperty, SourceData, DestProperty, DestData);
-            }
+			CopyProperty(SourceProperty, SourceData, DestProperty, DestData);
             
             Dest->PostEditChangeChainProperty(RemappedChangeEvent);
 		}
