[AutoRTFM] Update asset data tag map to use transactionally-safe locks.

This is necessary (but not sufficient) to allow FAssetData objects to be copy-constructed within a transaction.

diff --git a/Engine/Source/Runtime/CoreUObject/Private/AssetRegistry/AssetDataTagMap.cpp b/Engine/Source/Runtime/CoreUObject/Private/AssetRegistry/AssetDataTagMap.cpp
--- a/Engine/Source/Runtime/CoreUObject/Private/AssetRegistry/AssetDataTagMap.cpp
+++ b/Engine/Source/Runtime/CoreUObject/Private/AssetRegistry/AssetDataTagMap.cpp
@@ -6,6 +6,8 @@
 #include "AssetRegistry/AssetDataTagMapSerializationDetails.h"
 #include "Misc/PackageName.h"
 #include "Misc/Optional.h"
+#include "Misc/TransactionallySafeCriticalSection.h"
+#include "Misc/TransactionallySafeScopeLock.h"
 #include "Misc/ScopeLock.h"
 #include "Serialization/ArrayWriter.h"
 #include "Serialization/MemoryReader.h"
@@ -362,7 +364,6 @@ namespace FixedTagPrivate
 
 	//////////////////////////////////////////////////////////////////////////
 
-
 	class FStoreManager
 	{
 	public:
@@ -376,7 +377,7 @@ namespace FixedTagPrivate
 
 		FStore* CreateAndRegister()
 		{
-			FScopeLock Lock(&Mutex);
+			FTransactionallySafeScopeLock Lock(&Mutex);
 
 			static constexpr uint32 Mask = Capacity - 1; 
 			for (uint32 RawIndex = ProbeIndex, RawEnd = ProbeIndex + Capacity; RawIndex < RawEnd + Capacity; ++RawIndex)
@@ -398,7 +399,7 @@ namespace FixedTagPrivate
 
 		void Unregister(FStore& Store)
 		{
-			FScopeLock Lock(&Mutex);
+			FTransactionallySafeScopeLock Lock(&Mutex);
 
 			check(Store.Index < Capacity);
 			check(Data[Store.Index] == &Store);
@@ -407,7 +408,7 @@ namespace FixedTagPrivate
 		}
 
 	private:		
-		FCriticalSection Mutex;
+		FTransactionallySafeCriticalSection Mutex;
 		FStore* Data[Capacity] = {};
 		uint32 ProbeIndex = 0;
 	}
@@ -1468,7 +1469,7 @@ FAssetDataTagMapSharedView::FAssetDataTagMapSharedView(FAssetDataTagMap&& InLoos
 	if (InLoose.Num())
 	{
 		Loose = new FAssetDataTagMap(MoveTemp(InLoose));
-		Loose->RefCount.Set(1);
+		Loose->RefCount.Increment();
 	}
 }
 
