[LLM] Remove FName block allocations and TextureManager construction from contributing to the MemQuery/Memtrace asset memory cost readings

#tests local run

diff --git a/Engine/Source/Runtime/Core/Private/HAL/LowLevelMemTracker.cpp b/Engine/Source/Runtime/Core/Private/HAL/LowLevelMemTracker.cpp
--- a/Engine/Source/Runtime/Core/Private/HAL/LowLevelMemTracker.cpp
+++ b/Engine/Source/Runtime/Core/Private/HAL/LowLevelMemTracker.cpp
@@ -3639,7 +3639,9 @@ const UE::LLMPrivate::FTagData* FLowLevelMemTracker::FindTagData(FName TagName,
 
 void FLLMScope::Init(ELLMTag TagEnum, bool bInIsStatTag, ELLMTagSet InTagSet, ELLMTracker InTracker, bool bOverride)
 {
-	LLMCheck(!bInIsStatTag && InTagSet == ELLMTagSet::None);
+	// This scope does not support TagSets (use FLLMScopeDynamic instead), except for the special case of
+	// anonymized allocations, which are indicated with a TagSet and ELLMTag::EngineMisc
+	LLMCheck((!bInIsStatTag && InTagSet == ELLMTagSet::None) || (TagEnum == ELLMTag::EngineMisc));
 	// ELLMTag::FMalloc is a special tag expected to only be used by the Platform Tracker (see header where defined).
 	LLMCheck((TagEnum != ELLMTag::FMalloc) || (ELLMTracker::Platform == InTracker));
 
diff --git a/Engine/Source/Runtime/Core/Private/UObject/UnrealNames.cpp b/Engine/Source/Runtime/Core/Private/UObject/UnrealNames.cpp
--- a/Engine/Source/Runtime/Core/Private/UObject/UnrealNames.cpp
+++ b/Engine/Source/Runtime/Core/Private/UObject/UnrealNames.cpp
@@ -23,6 +23,7 @@
 #include "HAL/ThreadSafeCounter.h"
 #include "Misc/ScopeRWLock.h"
 #include "ProfilingDebugging/StringsTrace.h"
+#include "ProfilingDebugging/MetadataTrace.h"
 #include "Containers/Set.h"
 #include "Internationalization/Text.h"
 #include "Internationalization/Internationalization.h"
@@ -382,6 +383,9 @@ private:
 		if (!Buffer)
 		{
 			LLM_SCOPE(ELLMTag::FName);
+			LLM_TAGSET_SCOPE(ELLMTag::EngineMisc, ELLMTagSet::Assets);
+			LLM_TAGSET_SCOPE(ELLMTag::EngineMisc, ELLMTagSet::AssetClasses);
+			UE_TRACE_METADATA_CLEAR_SCOPE();
 			constexpr uint32 BufferSizeBytes = NumWords * sizeof(WordType);
 			Buffer = (WordType*) FMemory::MallocZeroed(BufferSizeBytes, alignof(WordType));
 			WordType* Expected = nullptr;
@@ -657,6 +661,9 @@ private:
 	static uint8* AllocBlock()
 	{
 		LLM_SCOPE(ELLMTag::FName);
+		LLM_TAGSET_SCOPE(ELLMTag::EngineMisc, ELLMTagSet::Assets);
+		LLM_TAGSET_SCOPE(ELLMTag::EngineMisc, ELLMTagSet::AssetClasses);
+		UE_TRACE_METADATA_CLEAR_SCOPE();
 		return (uint8*)FMemory::Malloc(BlockSizeBytes, alignof(FNameEntry));
 	}
 	
@@ -996,6 +1003,9 @@ public:
 	void Initialize(FNameEntryAllocator& InEntries)
 	{
 		LLM_SCOPE(ELLMTag::FName);
+		LLM_TAGSET_SCOPE(ELLMTag::EngineMisc, ELLMTagSet::Assets);
+		LLM_TAGSET_SCOPE(ELLMTag::EngineMisc, ELLMTagSet::AssetClasses);
+		UE_TRACE_METADATA_CLEAR_SCOPE();
 		Entries = &InEntries;
 
 		Slots = (FNameSlot*)FMemory::Malloc(FNamePoolInitialSlotsPerShard * sizeof(FNameSlot), alignof(FNameSlot));
@@ -1398,6 +1408,9 @@ private:
 	void Grow(const uint32 NewCapacity)
 	{
 		LLM_SCOPE(ELLMTag::FName);
+		LLM_TAGSET_SCOPE(ELLMTag::EngineMisc, ELLMTagSet::Assets);
+		LLM_TAGSET_SCOPE(ELLMTag::EngineMisc, ELLMTagSet::AssetClasses);
+		UE_TRACE_METADATA_CLEAR_SCOPE();
 		TArrayView<FNameSlot> OldSlots(Slots, Capacity());
 		const uint32 OldUsedSlots = UsedSlots;
 
diff --git a/Engine/Source/Runtime/CoreUObject/Private/Serialization/AsyncLoading2.cpp b/Engine/Source/Runtime/CoreUObject/Private/Serialization/AsyncLoading2.cpp
--- a/Engine/Source/Runtime/CoreUObject/Private/Serialization/AsyncLoading2.cpp
+++ b/Engine/Source/Runtime/CoreUObject/Private/Serialization/AsyncLoading2.cpp
@@ -9761,6 +9761,10 @@ void FAsyncLoadingThread2::FlushLoading(TConstArrayView<int32> RequestIDs)
 {
 	if (IsAsyncLoadingPackages())
 	{
+		LLM_SCOPE(ELLMTag::AsyncLoading);
+		LLM_TAGSET_SCOPE(ELLMTag::EngineMisc, ELLMTagSet::Assets);
+		LLM_TAGSET_SCOPE(ELLMTag::EngineMisc, ELLMTagSet::AssetClasses);
+		UE_TRACE_METADATA_CLEAR_SCOPE();
 		// Prevent objects still being loaded from being accessible from the game thread.
 		TGuardValue GuardVisibilityFilter(FUObjectThreadContext::Get().AsyncVisibilityFilter, EInternalObjectFlags::AsyncLoadingPhase1);
 
diff --git a/Engine/Source/Runtime/CoreUObject/Private/UObject/UObjectArray.cpp b/Engine/Source/Runtime/CoreUObject/Private/UObject/UObjectArray.cpp
--- a/Engine/Source/Runtime/CoreUObject/Private/UObject/UObjectArray.cpp
+++ b/Engine/Source/Runtime/CoreUObject/Private/UObject/UObjectArray.cpp
@@ -9,8 +9,8 @@
 #include "HAL/IConsoleManager.h"
 #include "HAL/LowLevelMemStats.h"
 #include "Misc/ScopeLock.h"
+#include "ProfilingDebugging/MetadataTrace.h"
 #include "Misc/TransactionallySafeScopeLock.h"
-#include "ProfilingDebugging/AssetMetadataTrace.h"
 #include "UObject/UObjectAllocator.h"
 #include "UObject/Class.h"
 #include "UObject/GarbageCollectionInternalFlags.h"
@@ -229,9 +229,9 @@ void FUObjectArray::AllocateUObjectIndex(UObjectBase* Object, EInternalObjectFla
 {
 	LLM_SCOPE(ELLMTag::UObject);
 	// Clear asset scopes
-	LLM_SCOPE_DYNAMIC_STAT_OBJECTPATH_FNAME(FName{NAME_Default}, ELLMTagSet::Assets);
-	LLM_SCOPE_DYNAMIC_STAT_OBJECTPATH_FNAME(FName{NAME_Default}, ELLMTagSet::AssetClasses);
-	UE_TRACE_METADATA_SCOPE_ASSET_FNAME(NAME_None, NAME_None, NAME_None);
+	LLM_TAGSET_SCOPE(ELLMTag::EngineMisc, ELLMTagSet::Assets);
+	LLM_TAGSET_SCOPE(ELLMTag::EngineMisc, ELLMTagSet::AssetClasses);
+	UE_TRACE_METADATA_CLEAR_SCOPE();
 
 	int32 Index = INDEX_NONE;
 	check(Object->InternalIndex == INDEX_NONE);
diff --git a/Engine/Source/Runtime/CoreUObject/Private/UObject/UObjectHash.cpp b/Engine/Source/Runtime/CoreUObject/Private/UObject/UObjectHash.cpp
--- a/Engine/Source/Runtime/CoreUObject/Private/UObject/UObjectHash.cpp
+++ b/Engine/Source/Runtime/CoreUObject/Private/UObject/UObjectHash.cpp
@@ -17,7 +17,7 @@
 #include "HAL/IConsoleManager.h"
 #include "HAL/LowLevelMemStats.h"
 #include "UObject/AnyPackagePrivate.h"
-#include "ProfilingDebugging/AssetMetadataTrace.h"
+#include "ProfilingDebugging/MetadataTrace.h"
 #include "AutoRTFM/AutoRTFM.h"
 
 #include <atomic>
@@ -1143,9 +1143,9 @@ FORCEINLINE static UPackage* UnassignExternalPackageFromObject(FUObjectHashTable
 void ShrinkUObjectHashTables()
 {
 	LLM_SCOPE_BYTAG(UObjectHash);
-	LLM_SCOPE_DYNAMIC_STAT_OBJECTPATH_FNAME(FName{NAME_Default}, ELLMTagSet::Assets);
-	LLM_SCOPE_DYNAMIC_STAT_OBJECTPATH_FNAME(FName{NAME_Default}, ELLMTagSet::AssetClasses);
-	UE_TRACE_METADATA_SCOPE_ASSET_FNAME(NAME_None, NAME_None, NAME_None);
+	LLM_TAGSET_SCOPE(ELLMTag::EngineMisc, ELLMTagSet::Assets);
+	LLM_TAGSET_SCOPE(ELLMTag::EngineMisc, ELLMTagSet::AssetClasses);
+	UE_TRACE_METADATA_CLEAR_SCOPE();
 	TRACE_CPUPROFILER_EVENT_SCOPE(ShrinkUObjectHashTables);
 	FUObjectHashTables& ThreadHash = FUObjectHashTables::Get();
 	FHashTableLock HashLock(ThreadHash);
@@ -1578,9 +1578,9 @@ void HashObject(UObjectBase* Object)
 	if (Name != NAME_None)
 	{
 		LLM_SCOPE_BYTAG(UObjectHash);
-		LLM_SCOPE_DYNAMIC_STAT_OBJECTPATH_FNAME(FName{NAME_Default}, ELLMTagSet::Assets);
-		LLM_SCOPE_DYNAMIC_STAT_OBJECTPATH_FNAME(FName{NAME_Default}, ELLMTagSet::AssetClasses);
-		UE_TRACE_METADATA_SCOPE_ASSET_FNAME(NAME_None, NAME_None, NAME_None);
+		LLM_TAGSET_SCOPE(ELLMTag::EngineMisc, ELLMTagSet::Assets);
+		LLM_TAGSET_SCOPE(ELLMTag::EngineMisc, ELLMTagSet::AssetClasses);
+		UE_TRACE_METADATA_CLEAR_SCOPE();
 #if !UE_BUILD_TEST && !UE_BUILD_SHIPPING
 		SCOPE_CYCLE_COUNTER(STAT_Hash_HashObject);
 #endif
@@ -1628,9 +1628,9 @@ void UnhashObject(UObjectBase* Object)
 	if (Name != NAME_None)
 	{
 		LLM_SCOPE_BYTAG(UObjectHash);
-		LLM_SCOPE_DYNAMIC_STAT_OBJECTPATH_FNAME(FName{NAME_Default}, ELLMTagSet::Assets);
-		LLM_SCOPE_DYNAMIC_STAT_OBJECTPATH_FNAME(FName{NAME_Default}, ELLMTagSet::AssetClasses);
-		UE_TRACE_METADATA_SCOPE_ASSET_FNAME(NAME_None, NAME_None, NAME_None);
+		LLM_TAGSET_SCOPE(ELLMTag::EngineMisc, ELLMTagSet::Assets);
+		LLM_TAGSET_SCOPE(ELLMTag::EngineMisc, ELLMTagSet::AssetClasses);
+		UE_TRACE_METADATA_CLEAR_SCOPE();
 #if !UE_BUILD_TEST && !UE_BUILD_SHIPPING
 		SCOPE_CYCLE_COUNTER(STAT_Hash_UnhashObject);
 #endif
@@ -1668,9 +1668,9 @@ void UnhashObject(UObjectBase* Object)
 void HashObjectExternalPackage(UObjectBase* Object, UPackage* Package)
 {
 	LLM_SCOPE_BYTAG(UObjectHash);
-	LLM_SCOPE_DYNAMIC_STAT_OBJECTPATH_FNAME(FName{NAME_Default}, ELLMTagSet::Assets);
-	LLM_SCOPE_DYNAMIC_STAT_OBJECTPATH_FNAME(FName{NAME_Default}, ELLMTagSet::AssetClasses);
-	UE_TRACE_METADATA_SCOPE_ASSET_FNAME(NAME_None, NAME_None, NAME_None);
+	LLM_TAGSET_SCOPE(ELLMTag::EngineMisc, ELLMTagSet::Assets);
+	LLM_TAGSET_SCOPE(ELLMTag::EngineMisc, ELLMTagSet::AssetClasses);
+	UE_TRACE_METADATA_CLEAR_SCOPE();
 	if (Package)
 	{
 		FUObjectHashTables& ThreadHash = FUObjectHashTables::Get();
@@ -1695,9 +1695,9 @@ void HashObjectExternalPackage(UObjectBase* Object, UPackage* Package)
 void UnhashObjectExternalPackage(class UObjectBase* Object)
 {
 	LLM_SCOPE_BYTAG(UObjectHash);
-	LLM_SCOPE_DYNAMIC_STAT_OBJECTPATH_FNAME(FName{NAME_Default}, ELLMTagSet::Assets);
-	LLM_SCOPE_DYNAMIC_STAT_OBJECTPATH_FNAME(FName{NAME_Default}, ELLMTagSet::AssetClasses);
-	UE_TRACE_METADATA_SCOPE_ASSET_FNAME(NAME_None, NAME_None, NAME_None);
+	LLM_TAGSET_SCOPE(ELLMTag::EngineMisc, ELLMTagSet::Assets);
+	LLM_TAGSET_SCOPE(ELLMTag::EngineMisc, ELLMTagSet::AssetClasses);
+	UE_TRACE_METADATA_CLEAR_SCOPE();
 	FUObjectHashTables& ThreadHash = FUObjectHashTables::Get();
 	FHashTableLock LockHash(ThreadHash);
 	Object->AtomicallyClearFlags(RF_HasExternalPackage);
diff --git a/Engine/Source/Runtime/Engine/Private/ContentStreaming.cpp b/Engine/Source/Runtime/Engine/Private/ContentStreaming.cpp
--- a/Engine/Source/Runtime/Engine/Private/ContentStreaming.cpp
+++ b/Engine/Source/Runtime/Engine/Private/ContentStreaming.cpp
@@ -20,6 +20,8 @@
 #include "VT/VirtualTextureChunkManager.h"
 #include "Rendering/NaniteCoarseMeshStreamingManager.h"
 #include "AutoRTFM/AutoRTFM.h"
+#include "HAL/LowLevelMemStats.h"
+#include "ProfilingDebugging/MetadataTrace.h"
 
 #if WITH_EDITOR
 #include "AudioDevice.h"
@@ -370,6 +372,9 @@ FStreamingManagerCollection& IStreamingManager::Get()
 {
 	if (StreamingManagerCollection == nullptr)
 	{
+		LLM_TAGSET_SCOPE(ELLMTag::EngineMisc, ELLMTagSet::Assets);
+		LLM_TAGSET_SCOPE(ELLMTag::EngineMisc, ELLMTagSet::AssetClasses);
+		UE_TRACE_METADATA_CLEAR_SCOPE();
 		// Since this is a lazily created static global variable we create it
 		// in the open.
 		UE_AUTORTFM_OPEN
