Improve the ability to have UHT generate code that uses classes other than UClass and UScriptStruct.

#rnx

diff --git a/Engine/Source/Runtime/CoreUObject/Private/UObject/Class.cpp b/Engine/Source/Runtime/CoreUObject/Private/UObject/Class.cpp
--- a/Engine/Source/Runtime/CoreUObject/Private/UObject/Class.cpp
+++ b/Engine/Source/Runtime/CoreUObject/Private/UObject/Class.cpp
@@ -106,6 +106,7 @@ FThreadSafeBool& InternalSafeGetTokenStreamDirtyFlag()
  * Shared function called from the various InitializePrivateStaticClass functions generated my the IMPLEMENT_CLASS macro.
  */
 COREUOBJECT_API void InitializePrivateStaticClass(
+	class UClass* (*TClass_StaticClassFn)(),
 	class UClass* TClass_Super_StaticClass,
 	class UClass* TClass_PrivateStaticClass,
 	class UClass* TClass_WithinClass_StaticClass,
@@ -130,10 +131,27 @@ COREUOBJECT_API void InitializePrivateStaticClass(
 	TClass_PrivateStaticClass->RegisterDependencies();
 	{
 		// Defer
-		TClass_PrivateStaticClass->Register(PackageName, Name);
+		TClass_PrivateStaticClass->Register(TClass_StaticClassFn, PackageName, Name);
 	}
 }
 
+COREUOBJECT_API void InitializePrivateStaticClass(
+	class UClass* TClass_Super_StaticClass,
+	class UClass* TClass_PrivateStaticClass,
+	class UClass* TClass_WithinClass_StaticClass,
+	const TCHAR* PackageName,
+	const TCHAR* Name
+)
+{
+	InitializePrivateStaticClass(
+		UClass::StaticClass,
+		TClass_Super_StaticClass,
+		TClass_PrivateStaticClass,
+		TClass_WithinClass_StaticClass,
+		PackageName,
+		Name);
+}
+
 void FNativeFunctionRegistrar::RegisterFunction(class UClass* Class, const ANSICHAR* InName, FNativeFuncPtr InPointer)
 {
 	Class->AddNativeFunction(InName, InPointer);
@@ -7100,6 +7118,7 @@ void GetPrivateStaticClassBody(
 	check(ReturnClass);
 	
 	InitializePrivateStaticClass(
+		UClass::StaticClass,
 		InSuperClassFn(),
 		ReturnClass,
 		InWithinClassFn(),
diff --git a/Engine/Source/Runtime/CoreUObject/Private/UObject/UObjectBase.cpp b/Engine/Source/Runtime/CoreUObject/Private/UObject/UObjectBase.cpp
--- a/Engine/Source/Runtime/CoreUObject/Private/UObject/UObjectBase.cpp
+++ b/Engine/Source/Runtime/CoreUObject/Private/UObject/UObjectBase.cpp
@@ -59,9 +59,12 @@ struct FPendingRegistrantInfo
 {
 	const TCHAR*	Name;
 	const TCHAR*	PackageName;
-	FPendingRegistrantInfo(const TCHAR* InName,const TCHAR* InPackageName)
-		:	Name(InName)
-		,	PackageName(InPackageName)
+	class UClass*  (*StaticClassFn)();
+
+	FPendingRegistrantInfo(class UClass* (*InStaticClassFn)(), const TCHAR* InName, const TCHAR* InPackageName)
+		: Name(InName)
+		, PackageName(InPackageName)
+		, StaticClassFn(InStaticClassFn)
 	{}
 	static TMap<UObjectBase*, FPendingRegistrantInfo>& GetMap()
 	{
@@ -165,7 +168,7 @@ UObjectBase::~UObjectBase()
 /**
  * Convert a boot-strap registered class into a real one, add to uobject array, etc
  *
- * @param UClassStaticClass Now that it is known, fill in UClass::StaticClass() as the class
+ * @param UClassStaticClass Now that it is known, fill in ClassPrivate.  This is usually UClass::StaticClass.
  */
 void UObjectBase::DeferredRegister(UClass *UClassStaticClass,const TCHAR* PackageName,const TCHAR* InName)
 {
@@ -462,13 +465,18 @@ static FAutoConsoleCommand DumpPendingUObjectModulesCmd(
 #endif
 
 /** Enqueue the registration for this object. */
-void UObjectBase::Register(const TCHAR* PackageName,const TCHAR* InName)
+void UObjectBase::Register(const TCHAR* PackageName, const TCHAR* InName)
+{
+	Register(UClass::StaticClass, PackageName, InName);
+}
+
+void UObjectBase::Register(class UClass* (*StaticClassFn)(), const TCHAR* PackageName, const TCHAR* InName)
 {
 	LLM_SCOPE(ELLMTag::UObject);
 	TMap<UObjectBase*, FPendingRegistrantInfo>& PendingRegistrants = FPendingRegistrantInfo::GetMap();
 
 	FPendingRegistrant* PendingRegistration = new FPendingRegistrant(this);
-	PendingRegistrants.Add(this, FPendingRegistrantInfo(InName, PackageName));
+	PendingRegistrants.Add(this, FPendingRegistrantInfo(StaticClassFn, InName, PackageName));
 
 #if USE_PER_MODULE_UOBJECT_BOOTSTRAP
 	if (FName(PackageName) != FName("/Script/CoreUObject"))
@@ -555,8 +563,9 @@ void UObjectForceRegistration(UObjectBase* Object, bool bCheckForModuleRelease)
 		}
 #endif
 		const TCHAR* Name = Info->Name;
+		UClass* StaticClass = Info->StaticClassFn();
 		PendingRegistrants.Remove(Object);  // delete this first so that it doesn't try to do it twice
-		Object->DeferredRegister(UClass::StaticClass(),PackageName,Name);
+		Object->DeferredRegister(StaticClass, PackageName, Name);
 	}
 }
 
diff --git a/Engine/Source/Runtime/CoreUObject/Private/UObject/UObjectGlobals.cpp b/Engine/Source/Runtime/CoreUObject/Private/UObject/UObjectGlobals.cpp
--- a/Engine/Source/Runtime/CoreUObject/Private/UObject/UObjectGlobals.cpp
+++ b/Engine/Source/Runtime/CoreUObject/Private/UObject/UObjectGlobals.cpp
@@ -6278,6 +6278,15 @@ namespace UECodeGen_Private
 #endif
 	}
 
+	void FinalizeConstructedUScriptStruct(UScriptStruct* Struct, const FStructParams& Params)
+	{
+		ConstructFProperties(Struct, Params.PropertyArray, Params.NumProperties);
+		Struct->StaticLink();
+#if WITH_METADATA
+		AddMetaData(Struct, Params.MetaDataArray, Params.NumMetaData);
+#endif
+	}
+
 	void ConstructUScriptStruct(UScriptStruct*& OutStruct, const FStructParams& Params)
 	{
 		UObject*                      (*OuterFunc)()     = Params.OuterFunc;
@@ -6296,15 +6305,8 @@ namespace UECodeGen_Private
 		UScriptStruct* NewStruct = new(EC_InternalUseOnlyConstructor, Outer, UTF8_TO_TCHAR(Params.NameUTF8), Params.ObjectFlags) UScriptStruct(FObjectInitializer(), Super, StructOps, (EStructFlags)Params.StructFlags, Params.SizeOf, Params.AlignOf);
 		OutStruct = NewStruct;
 
-		ConstructFProperties(NewStruct, Params.PropertyArray, Params.NumProperties);
-
-		NewStruct->StaticLink();
-
-#if WITH_METADATA
-		AddMetaData(NewStruct, Params.MetaDataArray, Params.NumMetaData);
-#endif
+		FinalizeConstructedUScriptStruct(NewStruct, Params);
 	}
-
 	void ConstructUPackage(UPackage*& OutPackage, const FPackageParams& Params)
 	{
 		if (OutPackage)
@@ -6439,7 +6441,7 @@ namespace UECodeGen_Private
 
 		NewClass->SetSparseClassDataStruct(NewClass->GetSparseClassDataArchetypeStruct());
 	}
-}
+} // namespace UECodeGen_Private
 
 void FReferenceCollector::AddStableReferenceSetFwd(TSet<FObjectPtr>* Objects)
 {
diff --git a/Engine/Source/Runtime/CoreUObject/Public/UObject/Class.h b/Engine/Source/Runtime/CoreUObject/Public/UObject/Class.h
--- a/Engine/Source/Runtime/CoreUObject/Public/UObject/Class.h
+++ b/Engine/Source/Runtime/CoreUObject/Public/UObject/Class.h
@@ -3759,6 +3759,16 @@ UObject* InternalVTableHelperCtorCaller(FVTableHelper& Helper)
 	return T::__VTableCtorCaller(Helper);
 }
 
+COREUOBJECT_API void InitializePrivateStaticClass(
+	class UClass* (*TClass_StaticClassFn)(),
+	class UClass* TClass_Super_StaticClass,
+	class UClass* TClass_PrivateStaticClass,
+	class UClass* TClass_WithinClass_StaticClass,
+	const TCHAR* PackageName,
+	const TCHAR* Name
+);
+
+UE_DEPRECATED(5.6, "Use the new method above and specify UClass::StaticClass as the first argument")
 COREUOBJECT_API void InitializePrivateStaticClass(
 	class UClass* TClass_Super_StaticClass,
 	class UClass* TClass_PrivateStaticClass,
diff --git a/Engine/Source/Runtime/CoreUObject/Public/UObject/UObjectBase.h b/Engine/Source/Runtime/CoreUObject/Public/UObject/UObjectBase.h
--- a/Engine/Source/Runtime/CoreUObject/Public/UObject/UObjectBase.h
+++ b/Engine/Source/Runtime/CoreUObject/Public/UObject/UObjectBase.h
@@ -47,6 +47,7 @@ class UObjectBase
 	friend struct FInternalUObjectBaseUtilityIsValidFlagsChecker; // for access to InternalIndex
 	friend COREUOBJECT_API void UObjectForceRegistration(UObjectBase* Object, bool bCheckForModuleRelease);
 	friend COREUOBJECT_API void InitializePrivateStaticClass(
+		class UClass* (*TClass_StaticClassFn)(),
 		class UClass* TClass_Super_StaticClass,
 		class UClass* TClass_PrivateStaticClass,
 		class UClass* TClass_WithinClass_StaticClass,
@@ -104,8 +105,10 @@ protected:
 	virtual void RegisterDependencies() {}
 
 	/** Enqueue the registration for this object. */
-	COREUOBJECT_API void Register(const TCHAR* PackageName,const TCHAR* Name);
-	
+	COREUOBJECT_API void Register(class UClass* (*StaticClassFn)(), const TCHAR* PackageName, const TCHAR* Name);
+	UE_DEPRECATED(5.6, "Use the new method above and specify UClass::StaticClass as the first argument")
+	COREUOBJECT_API void Register(const TCHAR* PackageName, const TCHAR* Name);
+
 	/**
 	 * Convert a boot-strap registered class into a real one, add to uobject array, etc
 	 *
diff --git a/Engine/Source/Runtime/CoreUObject/Public/UObject/UObjectGlobals.h b/Engine/Source/Runtime/CoreUObject/Public/UObject/UObjectGlobals.h
--- a/Engine/Source/Runtime/CoreUObject/Public/UObject/UObjectGlobals.h
+++ b/Engine/Source/Runtime/CoreUObject/Public/UObject/UObjectGlobals.h
@@ -3954,6 +3954,9 @@ namespace UECodeGen_Private
 	COREUOBJECT_API void ConstructUScriptStruct(UScriptStruct*& OutStruct, const FStructParams& Params);
 	COREUOBJECT_API void ConstructUPackage(UPackage*& OutPackage, const FPackageParams& Params);
 	COREUOBJECT_API void ConstructUClass(UClass*& OutClass, const FClassParams& Params);
+
+	// Helper methods to finish the construction process for different types
+	COREUOBJECT_API void FinalizeConstructedUScriptStruct(UScriptStruct* Struct, const FStructParams& Params);
 }
 /// @endcond
 
