[ZenLoader]
 - Fix crash when running in cook on the fly mode
 - Restore zenloader by default for loose cooked files

#rnx

diff --git a/Engine/Source/Runtime/CoreUObject/Private/Serialization/AsyncLoading2.cpp b/Engine/Source/Runtime/CoreUObject/Private/Serialization/AsyncLoading2.cpp
--- a/Engine/Source/Runtime/CoreUObject/Private/Serialization/AsyncLoading2.cpp
+++ b/Engine/Source/Runtime/CoreUObject/Private/Serialization/AsyncLoading2.cpp
@@ -5931,28 +5931,32 @@ EEventLoadNodeExecutionResult FAsyncPackage2::ProcessLinkerLoadPackageSummary(FA
 
 	AsyncLoadingThread.FinishInitializeAsyncPackage(ThreadState, this);
 
-	FLoadedPackageRef& PackageRef = AsyncLoadingThread.GlobalImportStore.FindPackageRefChecked(Desc.UPackageId, Desc.UPackageName);
-	if (!bLoadHasFailed)
+	// Package that can't be imported are never registered in the global import store, don't try to search for them
+	if (Desc.bCanBeImported)
 	{
-		PackageRef.ReserveSpaceForPublicExports(LinkerLoadState->LinkerLoadHeaderData.ExportMap.Num());
-#if WITH_METADATA
-		// Create metadata object, this needs to happen before any other package wants to use our exports
-		LinkerLoadState->MetaDataIndex = LinkerLoadState->Linker->LoadMetaDataFromExportMap(false);
-		if (LinkerLoadState->MetaDataIndex >= 0)
+		FLoadedPackageRef& PackageRef = AsyncLoadingThread.GlobalImportStore.FindPackageRefChecked(Desc.UPackageId, Desc.UPackageName);
+		if (!bLoadHasFailed)
 		{
-			FObjectExport& LinkerExport = LinkerLoadState->Linker->ExportMap[LinkerLoadState->MetaDataIndex];
-			FExportObject& ExportObject = Data.Exports[LinkerLoadState->MetaDataIndex];
-			ExportObject.Object = LinkerExport.Object;
+			PackageRef.ReserveSpaceForPublicExports(LinkerLoadState->LinkerLoadHeaderData.ExportMap.Num());
+#if WITH_METADATA
+			// Create metadata object, this needs to happen before any other package wants to use our exports
+			LinkerLoadState->MetaDataIndex = LinkerLoadState->Linker->LoadMetaDataFromExportMap(false);
+			if (LinkerLoadState->MetaDataIndex >= 0)
+			{
+				FObjectExport& LinkerExport = LinkerLoadState->Linker->ExportMap[LinkerLoadState->MetaDataIndex];
+				FExportObject& ExportObject = Data.Exports[LinkerLoadState->MetaDataIndex];
+				ExportObject.Object = LinkerExport.Object;
 
-			ExportObject.bWasFoundInMemory = !!LinkerExport.Object; // Make sure that the async flags are cleared in ClearConstructedObjects
-			ExportObject.bExportLoadFailed = LinkerExport.bExportLoadFailed;
-			ExportObject.bFiltered = LinkerExport.bWasFiltered;
-		}
+				ExportObject.bWasFoundInMemory = !!LinkerExport.Object; // Make sure that the async flags are cleared in ClearConstructedObjects
+				ExportObject.bExportLoadFailed = LinkerExport.bExportLoadFailed;
+				ExportObject.bFiltered = LinkerExport.bWasFiltered;
+			}
 #endif // WITH_METADATA
-	}
-	else if (Desc.bCanBeImported)
-	{
-		PackageRef.SetHasFailed();
+		}
+		else
+		{
+			PackageRef.SetHasFailed();
+		}
 	}
 
 	AsyncPackageLoadingState = EAsyncPackageLoadingState2::WaitingForDependencies;
diff --git a/Engine/Source/Runtime/CoreUObject/Private/Serialization/AsyncPackageLoader.cpp b/Engine/Source/Runtime/CoreUObject/Private/Serialization/AsyncPackageLoader.cpp
--- a/Engine/Source/Runtime/CoreUObject/Private/Serialization/AsyncPackageLoader.cpp
+++ b/Engine/Source/Runtime/CoreUObject/Private/Serialization/AsyncPackageLoader.cpp
@@ -15,7 +15,7 @@
 #include "IO/IoDispatcherInternal.h"
 #include "HAL/IConsoleManager.h"
 
-static bool GUseOldLoaderAsFallback = true;
+static bool GUseOldLoaderAsFallback = false;
 static FAutoConsoleVariableRef CVarGUseOldLoaderAsFallback(
 	TEXT("s.UseOldLoaderAsFallback"),
 	GUseOldLoaderAsFallback,
