Add PackageLocation to FAssetPackageData to allow the AssetRegistry to inform where AssetData is coming from. This is used to accelerate DoesPackageExist calls when the asset registry has been initialized to alleviate the need to read from disk if the registry already knows about the data. In particular, cooked editor cases are optimized since previously the asset registry cache was ignored if the IOStore contained any data at all.

Opening a large project with a cooked editor now sees DoesPackageExist time decrease by ~8x, from ~17s to ~2s

#rnx

diff --git a/Engine/Source/Editor/UnrealEd/Private/Commandlets/AssetRegistryGenerator.cpp b/Engine/Source/Editor/UnrealEd/Private/Commandlets/AssetRegistryGenerator.cpp
--- a/Engine/Source/Editor/UnrealEd/Private/Commandlets/AssetRegistryGenerator.cpp
+++ b/Engine/Source/Editor/UnrealEd/Private/Commandlets/AssetRegistryGenerator.cpp
@@ -1441,6 +1441,9 @@ FAssetPackageData FAssetRegistryGenerator::CopyAssetPackageDataForIncrementalCoo
 	Result.SetIsLicenseeVersion(Source.IsLicenseeVersion());
 	Result.SetHasVirtualizedPayloads(Source.HasVirtualizedPayloads());
 
+	// Currently assumes IoDispatcher however should we cook to loose files in the future this would need updating
+	Result.SetPackageLocation(FPackageName::EPackageLocationFilter::IoDispatcher);
+
 	// CustomVersions are not read by the cook
 	//Result.SetCustomVersions(Source.GetCustomVersions());
 
diff --git a/Engine/Source/Editor/UnrealEd/Private/CookOnTheFlyServer.cpp b/Engine/Source/Editor/UnrealEd/Private/CookOnTheFlyServer.cpp
--- a/Engine/Source/Editor/UnrealEd/Private/CookOnTheFlyServer.cpp
+++ b/Engine/Source/Editor/UnrealEd/Private/CookOnTheFlyServer.cpp
@@ -10687,6 +10687,7 @@ void UCookOnTheFlyServer::CookByTheBookFinishedInternal()
 					FAssetPackageData* AssetPackageData = Generator.GetAssetPackageData(HashSet.Key);
 					TRefCountPtr<FPackageHashes>& PackageHashes = HashSet.Value;
 
+					AssetPackageData->SetPackageLocation(FPackageName::EPackageLocationFilter::IoDispatcher);
 					AssetPackageData->CookedHash = PackageHashes->PackageHash;
 					Move(AssetPackageData->ChunkHashes, PackageHashes->ChunkHashes);
 				}
diff --git a/Engine/Source/Editor/UnrealEd/Private/Cooker/CookGenerationHelper.cpp b/Engine/Source/Editor/UnrealEd/Private/Cooker/CookGenerationHelper.cpp
--- a/Engine/Source/Editor/UnrealEd/Private/Cooker/CookGenerationHelper.cpp
+++ b/Engine/Source/Editor/UnrealEd/Private/Cooker/CookGenerationHelper.cpp
@@ -1231,6 +1231,8 @@ void FGenerationHelper::FinishGeneratedPlatformSave(FPackageData& PackageData,
 	OutGeneratedResult.AssetPackageData.FileVersionUE = GPackageFileUEVersion;
 	OutGeneratedResult.AssetPackageData.FileVersionLicenseeUE = GPackageFileLicenseeUEVersion;
 	OutGeneratedResult.AssetPackageData.SetIsLicenseeVersion(FEngineVersion::Current().IsLicenseeVersion());
+	// Currently assumes IoDispatcher however should we cook to loose files in the future this would need updating
+	OutGeneratedResult.AssetPackageData.SetPackageLocation(FPackageName::EPackageLocationFilter::IoDispatcher);
 	OutGeneratedResult.AssetPackageData.Extension = FPackagePath::ParseExtension(
 		WriteToString<256>(PackageData.GetFileName()));
 	OutGeneratedResult.AssetPackageData.SetPackageSavedHash(Info->PackageHash);
diff --git a/Engine/Source/Runtime/AssetRegistry/Private/PackageReader.cpp b/Engine/Source/Runtime/AssetRegistry/Private/PackageReader.cpp
--- a/Engine/Source/Runtime/AssetRegistry/Private/PackageReader.cpp
+++ b/Engine/Source/Runtime/AssetRegistry/Private/PackageReader.cpp
@@ -754,6 +754,7 @@ bool FPackageReader::ReadDependencyData(FPackageDependencyData& OutDependencyDat
 		PackageData.FileVersionUE = PackageFileSummary.GetFileVersionUE();
 		PackageData.FileVersionLicenseeUE = PackageFileSummary.GetFileVersionLicenseeUE();
 		PackageData.SetIsLicenseeVersion(PackageFileSummary.SavedByEngineVersion.IsLicenseeVersion());
+		PackageData.SetPackageLocation(FPackageName::EPackageLocationFilter::FileSystem);
 		PackageData.Extension = FPackagePath::ParseExtension(PackageFilename);
 
 		if (!SerializeImportedClasses(ImportMap, PackageData.ImportedClasses))
diff --git a/Engine/Source/Runtime/CoreUObject/Private/Misc/PackageName.cpp b/Engine/Source/Runtime/CoreUObject/Private/Misc/PackageName.cpp
--- a/Engine/Source/Runtime/CoreUObject/Private/Misc/PackageName.cpp
+++ b/Engine/Source/Runtime/CoreUObject/Private/Misc/PackageName.cpp
@@ -1968,34 +1968,45 @@ FPackageName::EPackageLocationFilter FPackageName::DoesPackageExistEx(const FPac
 		return EPackageLocationFilter::None;
 	}
 
-
 #if WITH_EDITOR
 	IAssetRegistryInterface* AssetRegistry = IAssetRegistryInterface::GetPtr();
-
-	// Todo: The AssetRegistry currently cannot determine if a package comes from the Filesystem 
-	// or cooked content so we avoid registry lookups since we can't provide a reliable Location
-	if (AssetRegistry && ((uint8)Filter & (uint8)EPackageLocationFilter::FileSystem) && (!FIoDispatcher::IsInitialized() || !FIoDispatcherInternal::HasPackageData()))
+	if (AssetRegistry)
 	{
-		FName PackageName = PackagePath.GetPackageFName();
 		FName CorrectCasePackageName;
 		FAssetPackageData AssetPackageData;
-		if (AssetRegistry->TryGetAssetPackageData(PackageName, AssetPackageData, CorrectCasePackageName) == UE::AssetRegistry::EExists::Exists)
+		FName PackageName = PackagePath.GetPackageFName();
+		UE::AssetRegistry::EExists ExistsStatus = AssetRegistry->TryGetAssetPackageData(PackageName, AssetPackageData, CorrectCasePackageName);
+		if (ExistsStatus == UE::AssetRegistry::EExists::Exists)
 		{
-			if (OutPackagePath)
+			EPackageLocationFilter PackageLocation = AssetPackageData.GetPackageLocation();
+
+			// If the registry doesn't have this information at all, we can't trust it (e.g. reading older cache version)
+			if (PackageLocation != EPackageLocationFilter::None)
 			{
-				if (AssetPackageData.Extension != EPackageExtension::Unspecified && AssetPackageData.Extension != EPackageExtension::Custom)
+				PackageLocation = EPackageLocationFilter((uint8)Filter & (uint8)PackageLocation);
+				if (!OutPackagePath)
 				{
-					*OutPackagePath = FPackagePath::FromPackageNameUnchecked(MoveTemp(CorrectCasePackageName));
-					OutPackagePath->SetHeaderExtension(AssetPackageData.Extension);
-					return EPackageLocationFilter::FileSystem;
+					return PackageLocation;
+				}
+
+				*OutPackagePath = FPackagePath::FromPackageNameUnchecked(MoveTemp(CorrectCasePackageName));
+				if ((uint8)PackageLocation & (uint8)EPackageLocationFilter::FileSystem)
+				{
+					// If an OutPackagePath was provided but a custom or unspecified extension was used we must use the disk
+					// scanning code futher down to allow the PackageResourceManager to handle the unspecified and custom extension cases
+					if (AssetPackageData.Extension != EPackageExtension::Unspecified && AssetPackageData.Extension != EPackageExtension::Custom)
+					{
+						OutPackagePath->SetHeaderExtension(AssetPackageData.Extension);
+						return PackageLocation;
+					}
+					// Intentionally do not return a EPackageLocationFilter if an OutPackagePath was provided but a custom or 
+					// unspecified extension was used. In such cases we must use the disk scanning code futher down to allow the 
+					// PackageResourceManager to handle the unspecified and custom extension cases
+				}
+				else
+				{
+					return PackageLocation;
 				}
-				// Intentionally do not return a EPackageLocationFilter if an OutPackagePath was provided but a custom or 
-				// unspecified extension was used. In such cases we must use the disk scanning code futher down to allow the 
-				// PackageResourceManager to handle the unspecified and custom extension cases
-			}
-			else
-			{
-				return EPackageLocationFilter::FileSystem;
 			}
 		}
 	}
diff --git a/Engine/Source/Runtime/CoreUObject/Public/AssetRegistry/AssetData.h b/Engine/Source/Runtime/CoreUObject/Public/AssetRegistry/AssetData.h
--- a/Engine/Source/Runtime/CoreUObject/Public/AssetRegistry/AssetData.h
+++ b/Engine/Source/Runtime/CoreUObject/Public/AssetRegistry/AssetData.h
@@ -101,6 +101,7 @@ struct FAssetRegistryVersion
 		RemoveAssetPathFNames,				// Asset bundles are serialized as FTopLevelAssetPath instead of FSoftObjectPath, deprecated FAssetData::ObjectPath	
 		AddedHeader,						// Added header with bFilterEditorOnlyData flag 
 		AssetPackageDataHasExtension,		// Added Extension to AssetPackageData.
+		AssetPackageDataHasPackageLocation,	// Added PackageLocation to AssetPackageData.
 
 		// -----<new versions can be added above this line>-------------------------------------------------
 		VersionPlusOne,
@@ -1014,13 +1015,33 @@ public:
 	COREUOBJECT_API void NetworkWrite(FCbWriter& Writer) const;
 	COREUOBJECT_API bool TryNetworkRead(FCbFieldView Field);
 
+	COREUOBJECT_API void SetPackageLocation(FPackageName::EPackageLocationFilter PackageLocation)
+	{	
+		// Clear existing location flags
+		Flags = Flags & ~FLAG_PACKAGE_LOCATION_MASK;
+
+		check((uint8)PackageLocation < 4); // We only support 2 bits for location currently
+		Flags |= uint32(PackageLocation) << FLAG_PACKAGE_LOCATION_TO_ENUM_SHIFT_SIZE;
+	}
+
+	COREUOBJECT_API FPackageName::EPackageLocationFilter GetPackageLocation() const
+	{
+		return FPackageName::EPackageLocationFilter((Flags & FLAG_PACKAGE_LOCATION_MASK) >> FLAG_PACKAGE_LOCATION_TO_ENUM_SHIFT_SIZE);
+	}
+
 private:
 	FORCEINLINE void SerializeForCacheInternal(FArchive& Ar, FAssetPackageData& PackageData, FAssetRegistryVersion::Type Version);
 
-	enum
+	enum : uint32
 	{
-		FLAG_LICENSEE_VERSION			= 1 << 0,
-		FLAG_HAS_VIRTUALIZED_PAYLOADS	= 1 << 1
+		FLAG_LICENSEE_VERSION						= 1u << 0,
+		FLAG_HAS_VIRTUALIZED_PAYLOADS				= 1u << 1,
+
+		FLAG_PACKAGE_LOCATION_IODISPATCHER			= 1u << 30,
+		FLAG_PACKAGE_LOCATION_FILESYSTEM			= 1u << 31,
+
+		FLAG_PACKAGE_LOCATION_TO_ENUM_SHIFT_SIZE	= 30u,
+		FLAG_PACKAGE_LOCATION_MASK					= 0xc0000000
 	};
 };
 PRAGMA_ENABLE_DEPRECATION_WARNINGS
