[AutoRTFM] Wrap thread-safe counter in asset-tag store so that it is transactionally safe.

This relies on the CL at https://p4-swarm.epicgames.net/reviews/36271853 , which implements TTransactionallySafeCounter<>.

diff --git a/Engine/Source/Runtime/CoreUObject/Public/AssetRegistry/AssetDataTagMapSerializationDetails.h b/Engine/Source/Runtime/CoreUObject/Public/AssetRegistry/AssetDataTagMapSerializationDetails.h
--- a/Engine/Source/Runtime/CoreUObject/Public/AssetRegistry/AssetDataTagMapSerializationDetails.h
+++ b/Engine/Source/Runtime/CoreUObject/Public/AssetRegistry/AssetDataTagMapSerializationDetails.h
@@ -5,6 +5,8 @@
 #include "AssetRegistry/AssetData.h"
 #include "AssetRegistry/AssetDataTagMap.h"
 #include "Async/Async.h"
+#include "Misc/TransactionallySafeCounter.h"
+#include "HAL/ThreadSafeCounter.h"
 
 struct FAssetRegistrySerializationOptions;
 
@@ -72,7 +74,7 @@ namespace FixedTagPrivate
 		explicit FStore(uint32 InIndex) : Index(InIndex) {}
 		~FStore();
 
-		mutable FThreadSafeCounter RefCount;
+		mutable TTransactionallySafeCounter<FThreadSafeCounter> RefCount;
 	};
 
 	struct FOptions
