Make unique names for IDO instances to avoid tracking namespaces for them

#rnx

diff --git a/Engine/Source/Runtime/CoreUObject/Internal/UObject/PropertyBagRepository.h b/Engine/Source/Runtime/CoreUObject/Internal/UObject/PropertyBagRepository.h
--- a/Engine/Source/Runtime/CoreUObject/Internal/UObject/PropertyBagRepository.h
+++ b/Engine/Source/Runtime/CoreUObject/Internal/UObject/PropertyBagRepository.h
@@ -50,9 +50,6 @@ private:
 	
 	TMap<const UObject*, const UObject*> InstanceDataObjectToOwner;
 
-	// used to make sure IDOs don't have name overlap
-	TMap<const UObject*, TObjectPtr<UObject>> Namespaces;
-
 	FPropertyBagRepository() = default;
 
 public:
diff --git a/Engine/Source/Runtime/CoreUObject/Private/UObject/PropertyBagRepository.cpp b/Engine/Source/Runtime/CoreUObject/Private/UObject/PropertyBagRepository.cpp
--- a/Engine/Source/Runtime/CoreUObject/Private/UObject/PropertyBagRepository.cpp
+++ b/Engine/Source/Runtime/CoreUObject/Private/UObject/PropertyBagRepository.cpp
@@ -187,7 +187,6 @@ void FPropertyBagRepository::ReassociateObjects(const TMap<UObject*, UObject*>&
 				FPropertyBagPlaceholderTypeRegistry::Get().Remove(TypeObject);
 			}
 		}
-		Namespaces.Remove(Pair.Key);
 	}
 }
 
@@ -762,32 +761,29 @@ UObject* FPropertyBagRepository::CreateInstanceDataObject(UObject* Owner, FArchi
 
 UObject* FPropertyBagRepository::DuplicateInstanceDataObject(UObject* SourceOwner, UObject* DestOwner)
 {
+	FPropertyBagRepositoryLock LockRepo(this);
 	if (FPropertyBagAssociationData* SourceData = AssociatedData.Find(SourceOwner))
 	{
 		check(SourceData->InstanceDataObject);
 		FPropertyBagAssociationData& DestData = AssociatedData.FindOrAdd(DestOwner);
 		ensure(!DestData.InstanceDataObject);
 
-		// get outer pointer for new IDO
-		TObjectPtr<UObject>* OuterPtr = nullptr;
-		if (FPropertyBagAssociationData* OuterData = AssociatedData.Find(DestOwner->GetOuter()))
+		// use the IDO of the outer as the new outer if there is one, otherwise use the transient package
+		UObject* Outer = nullptr;
+		if (const FPropertyBagAssociationData* OuterData = AssociatedData.Find(DestOwner->GetOuter()))
 		{
-			OuterPtr = &OuterData->InstanceDataObject;
+			Outer = OuterData->InstanceDataObject;
 		}
-		if (!OuterPtr || !*OuterPtr)
+		if (!Outer)
 		{
-			OuterPtr = &Namespaces.FindOrAdd(DestOwner->GetOuter());
-			if (*OuterPtr == nullptr)
-			{
-				*OuterPtr = CreatePackage(nullptr); // TODO: replace with dummy object
-			}
+			Outer = GetTransientPackage();
 		}
 
 		// construct InstanceDataObject
 		FStaticConstructObjectParameters Params(SourceData->InstanceDataObject->GetClass());
 		Params.SetFlags |= EObjectFlags::RF_Transactional;
-		Params.Name = DestOwner->GetFName();
-		Params.Outer = *OuterPtr;
+		Params.Name = MakeUniqueObjectName(Outer, Params.Class, DestOwner->GetFName());
+		Params.Outer = Outer;
 		DestData.InstanceDataObject = StaticConstructObject_Internal(Params);
 		InstanceDataObjectToOwner.Add(DestData.InstanceDataObject, DestOwner);
 
@@ -868,12 +864,8 @@ void FPropertyBagRepository::MarkAsFixedUp(const UObject* Object)
 
 bool FPropertyBagRepository::RemoveAssociationUnsafe(const UObject* Owner)
 {
-	// note: RemoveAssociationUnsafe is called on every object regardless of whether it has a property bag.
-	// in that scenario, there's a chance we have a namespace associated with it. Remove that namespace.
-	Namespaces.Remove(Owner);
-
 	FPropertyBagAssociationData OldData;
-	if(AssociatedData.RemoveAndCopyValue(Owner, OldData))
+	if (AssociatedData.RemoveAndCopyValue(Owner, OldData))
 	{
 		InstanceDataObjectToOwner.Remove(OldData.InstanceDataObject);
 		OldData.Destroy();
@@ -951,11 +943,6 @@ bool FPropertyBagRepository::WasPropertyValueSerialized(const UStruct* Struct, c
 
 void FPropertyBagRepository::AddReferencedObjects(FReferenceCollector& Collector)
 {
-	for (TPair<const UObject*, TObjectPtr<UObject>>& Element : Namespaces)
-	{
-		Collector.AddReferencedObject(Element.Value);
-	}
-
 	FPropertyBagPlaceholderTypeRegistry::Get().AddReferencedObjects(Collector);
 }
 
@@ -970,36 +957,27 @@ void FPropertyBagRepository::CreateInstanceDataObjectUnsafe(UObject* Owner, FPro
 	const FPropertyPathNameTree* PropertyTree = BagData.Tree;
 	const FUnknownEnumNames* EnumNames = BagData.EnumNames;
 	// construct InstanceDataObject class
-	// TODO: should we put the InstanceDataObject or it's class in a package?
+	// TODO: should we put the InstanceDataObject or its class in a package?
 	const UClass* InstanceDataObjectClass = CreateInstanceDataObjectClass(PropertyTree, EnumNames, Owner->GetClass(), GetTransientPackage());
 
 	BagData.bNeedsFixup = StructContainsLooseProperties(InstanceDataObjectClass);
 
-	TObjectPtr<UObject>* OuterPtr = nullptr;
+	// use the IDO of the outer as the new outer if there is one, otherwise use the transient package
+	UObject* Outer = nullptr;
 	if (FPropertyBagAssociationData* OuterData = AssociatedData.Find(Owner->GetOuter()))
 	{
-		OuterPtr = &OuterData->InstanceDataObject;
-	}
-	if (!OuterPtr || !*OuterPtr)
-	{
-		OuterPtr = &Namespaces.FindOrAdd(Owner->GetOuter());
-		if (*OuterPtr == nullptr)
-		{
-			*OuterPtr = CreatePackage(nullptr); // TODO: replace with dummy object
-		}
+		Outer = OuterData->InstanceDataObject;
 	}
-
-	// if an old IDO still exists with the same name, rename it out of the way so StaticConstructObject_Internal doesn't have conflicts
-	if (UObject* OldIDO = StaticFindObjectFastInternal( /*Class=*/ nullptr, *OuterPtr, Owner->GetFName() ))
+	if (!Outer)
 	{
-		OldIDO->Rename(nullptr, GetTransientPackage(), REN_DontCreateRedirectors | REN_NonTransactional | REN_DoNotDirty);
+		Outer = GetTransientPackage();
 	}
 
 	// construct InstanceDataObject object
 	FStaticConstructObjectParameters Params(InstanceDataObjectClass);
 	Params.SetFlags |= EObjectFlags::RF_Transactional;
-	Params.Name = Owner->GetFName();
-	Params.Outer = *OuterPtr;
+	Params.Name = MakeUniqueObjectName(Outer, Params.Class, Owner->GetFName());
+	Params.Outer = Outer;
 	UObject* InstanceDataObjectObject = StaticConstructObject_Internal(Params);
 	BagData.InstanceDataObject = InstanceDataObjectObject;
 	InstanceDataObjectToOwner.Add(InstanceDataObjectObject, Owner);
