Push array inner property to current archiver state when serializing "replacement" array item value overrides on load.

diff --git a/Engine/Source/Runtime/CoreUObject/Private/UObject/PropertyArray.cpp b/Engine/Source/Runtime/CoreUObject/Private/UObject/PropertyArray.cpp
--- a/Engine/Source/Runtime/CoreUObject/Private/UObject/PropertyArray.cpp
+++ b/Engine/Source/Runtime/CoreUObject/Private/UObject/PropertyArray.cpp
@@ -250,6 +250,8 @@ void FArrayProperty::SerializeItem(FStructuredArchive::FSlot Slot, void* Value,
 			FStructuredArchive::FArray ReplacedArray = Record.EnterArray(TEXT("Replaced"), NumReplaced);
 			if (NumReplaced != INDEX_NONE)
 			{
+				FSerializedPropertyScope SerializedProperty(UnderlyingArchive, Inner, this);
+
 				ArrayHelper.EmptyAndAddValues(NumReplaced);
 				for (int32 i = 0; i < NumReplaced; i++)
 				{
