[ZenLoader]
 - Fix missing ref counting causing use-after-free when ALT is active

diff --git a/Engine/Source/Runtime/CoreUObject/Private/Serialization/AsyncLoading2.cpp b/Engine/Source/Runtime/CoreUObject/Private/Serialization/AsyncLoading2.cpp
--- a/Engine/Source/Runtime/CoreUObject/Private/Serialization/AsyncLoading2.cpp
+++ b/Engine/Source/Runtime/CoreUObject/Private/Serialization/AsyncLoading2.cpp
@@ -2533,9 +2533,7 @@ public:
 		}
 	}
 
-	~FAsyncLoadingSyncLoadContext()
-	{
-	}
+	~FAsyncLoadingSyncLoadContext();
 
 	void AddRef()
 	{
@@ -7960,6 +7958,13 @@ void FAsyncLoadingThread2::UpdateSyncLoadContext(FAsyncLoadingThreadState2& Thre
 			}
 			else if (FAsyncPackage2* RequestedPackage = RequestIdToPackageMap.FindRef(RequestID))
 			{
+				// Make sure the package we found is not queued for deletion by trying to add a ref to it. 
+				// RequestIdToPackageMap is only manipulated from ALT so we can't race between FindRef, TryAddRef and the actual deletion of the package.
+				if (!RequestedPackage->TryAddRef())
+				{
+					continue;
+				}
+
 				// Set RequestedPackage before setting bHasFoundRequestedPackage so that another thread looking at RequestedPackage
 				// after validating that bHasFoundRequestedPackage is true would see the proper value.
 				SyncLoadContext->RequestedPackages[i] = RequestedPackage;
@@ -10089,6 +10094,17 @@ void FGlobalImportStore::FlushDeferredDeletePackagesQueue()
 	AsyncLoadingThread.ProcessDeferredDeletePackagesQueue();
 }
 
+FAsyncLoadingSyncLoadContext::~FAsyncLoadingSyncLoadContext()
+{
+	for (FAsyncPackage2* AsyncPackage : RequestedPackages)
+	{
+		if (AsyncPackage)
+		{
+			AsyncPackage->ReleaseRef();
+		}
+	}
+}
+
 #if !UE_BUILD_SHIPPING
 
 void FAsyncPackage2::DumpState()
