Initial Remote Object support (disabled by default)

diff --git a/Engine/Plugins/Runtime/MultiServerReplication/Source/MultiServerReplication/Private/MultiServerNode.cpp b/Engine/Plugins/Runtime/MultiServerReplication/Source/MultiServerReplication/Private/MultiServerNode.cpp
--- a/Engine/Plugins/Runtime/MultiServerReplication/Source/MultiServerReplication/Private/MultiServerNode.cpp
+++ b/Engine/Plugins/Runtime/MultiServerReplication/Source/MultiServerReplication/Private/MultiServerNode.cpp
@@ -222,3 +222,27 @@ AMultiServerBeaconClient* UMultiServerNode::GetBeaconClientForURL(const FString&
 
 	return nullptr;
 }
+
+void UMultiServerNode::ForEachBeaconClient(TFunctionRef<void(AMultiServerBeaconClient*)> Operation)
+{
+	UNetDriver* HostNetDriver = BeaconHost ? BeaconHost->GetNetDriver() : nullptr;
+	if (HostNetDriver)
+	{
+		for (UNetConnection* ClientConnection : HostNetDriver->ClientConnections)
+		{
+			AMultiServerBeaconClient* BeaconClient = Cast<AMultiServerBeaconClient>(BeaconHost->GetClientActor(ClientConnection));
+			if (BeaconClient)
+			{
+				Operation(BeaconClient);
+			}
+		}
+	}
+	for (UMultiServerPeerConnection* Peer : PeerConnections)
+	{
+		AMultiServerBeaconClient* BeaconClient = Peer->BeaconClient;
+		if (BeaconClient)
+		{
+			Operation(BeaconClient);
+		}
+	}
+}
\ No newline at end of file
diff --git a/Engine/Plugins/Runtime/MultiServerReplication/Source/MultiServerReplication/Public/MultiServerNode.h b/Engine/Plugins/Runtime/MultiServerReplication/Source/MultiServerReplication/Public/MultiServerNode.h
--- a/Engine/Plugins/Runtime/MultiServerReplication/Source/MultiServerReplication/Public/MultiServerNode.h
+++ b/Engine/Plugins/Runtime/MultiServerReplication/Source/MultiServerReplication/Public/MultiServerNode.h
@@ -84,6 +84,8 @@ public:
 
 	TSubclassOf<AMultiServerBeaconClient> GetUserBeaconClass() const { return UserBeaconClass; }
 
+	void ForEachBeaconClient(TFunctionRef<void(AMultiServerBeaconClient*)> Operation);
+
 private:
 	friend AMultiServerBeaconClient;
 
diff --git a/Engine/Source/Developer/ProfileVisualizer/ProfileVisualizer.Build.cs b/Engine/Source/Developer/ProfileVisualizer/ProfileVisualizer.Build.cs
--- a/Engine/Source/Developer/ProfileVisualizer/ProfileVisualizer.Build.cs
+++ b/Engine/Source/Developer/ProfileVisualizer/ProfileVisualizer.Build.cs
@@ -9,6 +9,7 @@ public class ProfileVisualizer : ModuleRules
 		PrivateDependencyModuleNames.AddRange(
 			new string[] {
 				"Core",
+				"CoreUObject",
 				"Slate",
 				"SlateCore",
 				
diff --git a/Engine/Source/Editor/KismetCompiler/Private/KismetCompilerVMBackend.cpp b/Engine/Source/Editor/KismetCompiler/Private/KismetCompilerVMBackend.cpp
--- a/Engine/Source/Editor/KismetCompiler/Private/KismetCompilerVMBackend.cpp
+++ b/Engine/Source/Editor/KismetCompiler/Private/KismetCompilerVMBackend.cpp
@@ -97,6 +97,8 @@ public:
 	{
 #if UE_WITH_OBJECT_HANDLE_LATE_RESOLVE
 		ScriptPointerType D = (ScriptPointerType)Res.GetHandle().PointerOrRef;
+#elif UE_WITH_REMOTE_OBJECT_HANDLE
+		ScriptPointerType D = (ScriptPointerType)Res.GetHandle().PointerOrHandle;
 #else
 		ScriptPointerType D = (ScriptPointerType)Res.GetHandle();
 #endif
diff --git a/Engine/Source/Runtime/CoreUObject/Private/Serialization/AsyncLoading2.cpp b/Engine/Source/Runtime/CoreUObject/Private/Serialization/AsyncLoading2.cpp
--- a/Engine/Source/Runtime/CoreUObject/Private/Serialization/AsyncLoading2.cpp
+++ b/Engine/Source/Runtime/CoreUObject/Private/Serialization/AsyncLoading2.cpp
@@ -64,6 +64,7 @@
 #include "UObject/UObjectClusters.h"
 #include "UObject/LinkerInstancingContext.h"
 #include "UObject/LinkerLoadImportBehavior.h"
+#include "UObject/ObjectHandlePrivate.h"
 #include "UObject/ObjectSerializeAccessScope.h"
 #include "ProfilingDebugging/CountersTrace.h"
 #include "ProfilingDebugging/AssetMetadataTrace.h"
@@ -8905,7 +8906,8 @@ void FAsyncLoadingThread2::CollectUnreachableObjects(
 			Item.ObjectIndex = GUObjectArray.ObjectToIndex(Object);
 			Item.ObjectName = Object->GetFName();
 
-			if (!Object->GetOuter())
+			UObject* Outer = UE::CoreUObject::Private::FObjectHandleUtils::GetNonAccessTrackedOuterNoResolve(Object);
+			if (!Outer)
 			{
 				UPackage* Package = static_cast<UPackage*>(Object);
 				if (Package->bCanBeImported)
@@ -8933,7 +8935,8 @@ void FAsyncLoadingThread2::CollectUnreachableObjects(
 			[this, &UnreachableObjectItems](int32 Index)
 		{
 			UObject* Object = static_cast<UObject*>(UnreachableObjectItems[Index]->Object);
-			if (!Object->GetOuter())
+			UObject* Outer = UE::CoreUObject::Private::FObjectHandleUtils::GetNonAccessTrackedOuterNoResolve(Object);
+			if (!Outer)
 			{
 				UPackage* Package = static_cast<UPackage*>(Object);
 				if (Package->bCanBeImported)
diff --git a/Engine/Source/Runtime/CoreUObject/Private/UObject/Class.cpp b/Engine/Source/Runtime/CoreUObject/Private/UObject/Class.cpp
--- a/Engine/Source/Runtime/CoreUObject/Private/UObject/Class.cpp
+++ b/Engine/Source/Runtime/CoreUObject/Private/UObject/Class.cpp
@@ -7730,7 +7730,23 @@ UScriptStruct* TBaseStructure<FTopLevelAssetPath>::Get()
 	return ScriptStruct;
 }
 
+UScriptStruct* TBaseStructure<FRemoteObjectId>::Get()
+{
+	static UScriptStruct* ScriptStruct = StaticGetBaseStructureInternal(TEXT("RemoteObjectId"));
+	return ScriptStruct;
+}
+
+UScriptStruct* TBaseStructure<FRemoteObjectData>::Get()
+{
+	static UScriptStruct* ScriptStruct = StaticGetBaseStructureInternal(TEXT("RemoteObjectData"));
+	return ScriptStruct;
+}
 
+UScriptStruct* TBaseStructure<FRemoteServerId>::Get()
+{
+	static UScriptStruct* ScriptStruct = StaticGetBaseStructureInternal(TEXT("RemoteServerId"));
+	return ScriptStruct;
+}
 
 #define UE_DEFINE_CORE_VARIANT_TYPE(VARIANT, CORE)												\
 UScriptStruct* TBaseStructure<F##CORE>::Get()												\
diff --git a/Engine/Source/Runtime/CoreUObject/Private/UObject/GarbageCollection.cpp b/Engine/Source/Runtime/CoreUObject/Private/UObject/GarbageCollection.cpp
--- a/Engine/Source/Runtime/CoreUObject/Private/UObject/GarbageCollection.cpp
+++ b/Engine/Source/Runtime/CoreUObject/Private/UObject/GarbageCollection.cpp
@@ -1128,7 +1128,7 @@ struct FResolvedMutableReference : FImmutableReference
 
 FORCEINLINE_DEBUGGABLE static bool ValidateReference(UObject* Object, FPermanentObjectPoolExtents PermanentPool, FReferenceToken Referencer, FMemberId MemberId)
 {
-	bool bOk = (!PermanentPool.Contains(Object)) & (!!Object) & IsObjectHandleResolved(reinterpret_cast<FObjectHandle&>(Object)); //-V792
+	bool bOk = (!PermanentPool.Contains(Object)) & (!!Object) & IsObjectHandleResolved_ForGC(reinterpret_cast<FObjectHandle&>(Object)); //-V792
 
 #if ENABLE_GC_OBJECT_CHECKS
 	if (bOk)
@@ -1580,7 +1580,7 @@ private:
 		for (uint32 Idx = 0; Idx < Num; ++Idx)
 		{
 			UObject* Object = GetObject(UnvalidatedReferences[Idx]);
-			ValidsB.Set(Idx, (!!Object) & IsObjectHandleResolved(reinterpret_cast<FObjectHandle&>(Object))); //-V792
+			ValidsB.Set(Idx, (!!Object) & IsObjectHandleResolved_ForGC(reinterpret_cast<FObjectHandle&>(Object))); //-V792
 		}
 		
 		FValidatedBitmask Validations = FValidatedBitmask::And(ValidsA, ValidsB);
@@ -2816,6 +2816,13 @@ int32 GetNumSlowAROs()
 // Helps identify code locations that kill references
 FORCEINLINE_DEBUGGABLE void KillReference(UObject*& Object) { Object = nullptr; }
 
+#if UE_WITH_REMOTE_OBJECT_HANDLE
+FORCEINLINE_DEBUGGABLE void MakeRemoteReference(UObject*& Object) 
+{ 
+	Object = reinterpret_cast<UObject*>(FObjectHandle::ConvertToRemoteHandle(Object).PointerOrHandle);
+}
+#endif // UE_WITH_REMOTE_OBJECT_HANDLE
+
 template <EGCOptions Options>
 constexpr FORCEINLINE EKillable MayKill(EOrigin Origin, bool bAllowKill)
 {
@@ -2894,7 +2901,13 @@ public:
 				KillReference(Object);
 				return;
 			}
-			
+#if UE_WITH_REMOTE_OBJECT_HANDLE
+			else if (Metadata.Has(EInternalObjectFlags::Remote))
+			{
+				MakeRemoteReference(Object);
+				return;
+			}
+#endif
 			DetectGarbageReference(Context, Metadata);
 			HandleValidReference(Context, /*ReferencingObject, */Reference, Metadata);
 		}
@@ -2908,6 +2921,12 @@ public:
 			checkSlow(Metadata.ObjectItem->GetOwnerIndex() <= 0);
 			KillReference(*Reference.Mutable);
 		}
+#if UE_WITH_REMOTE_OBJECT_HANDLE
+		else if (Metadata.Has(EInternalObjectFlags::Remote))
+		{
+			MakeRemoteReference(*Reference.Mutable);
+		}
+#endif
 		else
 		{
 			HandleValidReference(Context, Reference, Metadata);
@@ -3422,13 +3441,23 @@ public:
 					KillReference(Object);
 					return;
 				}
-				else if (!IsEliminatingGarbage(Options) && bTrackGarbage)
+				else if (!IsEliminatingGarbage(Options) && bTrackGarbage
+#if UE_WITH_REMOTE_OBJECT_HANDLE
+					&& !Metadata.Has(EInternalObjectFlags::Remote)
+#endif
+					)
 				{
 					check(Origin == EOrigin::Other);
 					HandleGarbageReference(Context, ReferencingObject, Object, MemberId);
 				}
 			}
-			
+#if UE_WITH_REMOTE_OBJECT_HANDLE
+			if (Metadata.Has(EInternalObjectFlags::Remote))
+			{
+				MakeRemoteReference(Object);
+				return;
+			}
+#endif			
 			bool bReachedFirst = TReachabilityProcessor<Options>::HandleValidReference(Context, FImmutableReference{Object}, Metadata);
 
 			if (bReachedFirst && bTrackHistory)
diff --git a/Engine/Source/Runtime/CoreUObject/Private/UObject/Obj.cpp b/Engine/Source/Runtime/CoreUObject/Private/UObject/Obj.cpp
--- a/Engine/Source/Runtime/CoreUObject/Private/UObject/Obj.cpp
+++ b/Engine/Source/Runtime/CoreUObject/Private/UObject/Obj.cpp
@@ -42,6 +42,7 @@
 #include "UObject/ObjectSaveContext.h"
 #include "UObject/UObjectAnnotation.h"
 #include "UObject/ReferenceChainSearch.h"
+#include "UObject/RemoteObjectPrivate.h"
 #include "UObject/OverridableManager.h"
 #include "Serialization/ArchiveCountMem.h"
 #include "Serialization/ArchiveShowReferences.h"
@@ -5468,6 +5469,10 @@ void StaticUObjectInit()
 		}
 	}
 
+#if UE_WITH_REMOTE_OBJECT_HANDLE
+	UE::RemoteObject::Private::InitRemoteObjects();
+#endif // UE_WITH_REMOTE_OBJECT_HANDLE
+
 	UE_LOG(LogInit, Log, TEXT("Object subsystem initialized") );
 }
 
diff --git a/Engine/Source/Runtime/CoreUObject/Private/UObject/ObjectHandlePrivate.h b/Engine/Source/Runtime/CoreUObject/Private/UObject/ObjectHandlePrivate.h
new file mode 100644
--- /dev/null
+++ b/Engine/Source/Runtime/CoreUObject/Private/UObject/ObjectHandlePrivate.h
@@ -0,0 +1,38 @@
+// Copyright Epic Games, Inc. All Rights Reserved.
+
+#pragma once
+#include "UObject/Object.h"
+
+namespace UE::CoreUObject::Private
+{
+	struct FObjectHandleUtils
+	{
+		FORCEINLINE static const ObjectPtr_Private::TNonAccessTrackedObjectPtr<UObject>& GetNonAccessTrackedOuter(UObjectBase* Object)
+		{
+			return Object->OuterPrivate;
+		}
+
+		FORCEINLINE static UObject* GetNonAccessTrackedOuterNoResolve(UObjectBase* Object)
+		{
+#if UE_WITH_REMOTE_OBJECT_HANDLE
+			return Object->OuterPrivate.GetNoResolve();
+#else
+			return Object->OuterPrivate;
+#endif 
+		}
+
+#if UE_WITH_REMOTE_OBJECT_HANDLE
+		static void ChangeRemoteId(UObjectBase* Object, FRemoteObjectId Id);
+#endif
+
+		FORCEINLINE static FRemoteObjectId GetRemoteId(const UObjectBase* Object)
+		{
+#if UE_WITH_REMOTE_OBJECT_HANDLE
+			return Object->RemoteId;
+#else
+			return FRemoteObjectId();
+#endif // UE_WITH_REMOTE_OBJECT_HANDLE
+		}
+	};
+}
+
diff --git a/Engine/Source/Runtime/CoreUObject/Private/UObject/PropertyObject.cpp b/Engine/Source/Runtime/CoreUObject/Private/UObject/PropertyObject.cpp
--- a/Engine/Source/Runtime/CoreUObject/Private/UObject/PropertyObject.cpp
+++ b/Engine/Source/Runtime/CoreUObject/Private/UObject/PropertyObject.cpp
@@ -167,7 +167,7 @@ void FObjectProperty::SerializeItem(FStructuredArchive::FSlot Slot, void* Value,
 	if (UnderlyingArchive.IsObjectReferenceCollector())
 	{
 		TObjectPtr<UObject> CurrentValue = *ObjectPtr;
-#if UE_WITH_OBJECT_HANDLE_LATE_RESOLVE
+#if UE_WITH_OBJECT_HANDLE_LATE_RESOLVE || UE_WITH_REMOTE_OBJECT_HANDLE
 		if (HasAnyPropertyFlags(CPF_TObjectPtr))
 		{
 			FObjectPtr* Ptr = (FObjectPtr*)ObjectPtr;
@@ -188,7 +188,7 @@ void FObjectProperty::SerializeItem(FStructuredArchive::FSlot Slot, void* Value,
 	}
 	else
 	{
-#if UE_WITH_OBJECT_HANDLE_LATE_RESOLVE
+#if UE_WITH_OBJECT_HANDLE_LATE_RESOLVE || UE_WITH_REMOTE_OBJECT_HANDLE
 		if (HasAnyPropertyFlags(CPF_TObjectPtr))
 		{
 			FObjectHandle OriginalHandle = ObjectPtr->GetHandle();
@@ -403,7 +403,7 @@ void FObjectProperty::CopyCompleteValueFromScriptVM_InContainer(void* OutContain
 
 void FObjectProperty::CopyValuesInternal(void* Dest, void const* Src, int32 Count) const
 {
-#if UE_WITH_OBJECT_HANDLE_LATE_RESOLVE
+#if UE_WITH_OBJECT_HANDLE_LATE_RESOLVE || UE_WITH_REMOTE_OBJECT_HANDLE
 	if (HasAllPropertyFlags(CPF_TObjectPtr))
 	{
 		for (int32 Index = 0; Index < Count; Index++)
diff --git a/Engine/Source/Runtime/CoreUObject/Private/UObject/RemoteObject.cpp b/Engine/Source/Runtime/CoreUObject/Private/UObject/RemoteObject.cpp
new file mode 100644
--- /dev/null
+++ b/Engine/Source/Runtime/CoreUObject/Private/UObject/RemoteObject.cpp
@@ -0,0 +1,322 @@
+// Copyright Epic Games, Inc. All Rights Reserved.
+
+#include "UObject/RemoteObject.h"
+#include "UObject/ObjectPtr.h"
+#include "UObject/Class.h"
+#include "UObject/GCObject.h"
+#include "Serialization/ArchiveUObject.h"
+#include "Serialization/MemoryWriter.h"
+#include "Serialization/MemoryReader.h"
+#include "UObject/RemoteObjectTransfer.h"
+#include "UObject/RemoteObjectSerialization.h"
+#include "UObject/RemoteObjectPrivate.h"
+#include "UObject/ObjectHandlePrivate.h"
+#include "Misc/CommandLine.h"
+#include "Misc/CoreDelegates.h"
+
+namespace UE::RemoteObject::Private
+{
+
+FRemoteServerId RemoteServerId;
+std::atomic<uint64> RemoteObjectSerialNumber(1);
+
+class FRemoteObjectMaps
+{
+	mutable FCriticalSection ObjectMapCritical;
+	TMap<FRemoteObjectId, UE::RemoteObject::Handle::FRemoteObjectStub*> RemoteObjects;
+	TMap<FRemoteObjectId, UObject*> IdToObjectMap;
+
+public:
+
+	bool IsRemote(FRemoteObjectId ObjectId) const
+	{
+		FScopeLock MapLock(&ObjectMapCritical);
+		bool bResult = !IdToObjectMap.Contains(ObjectId);
+		return bResult;
+	}
+
+	void MapObject(UObject* Object)
+	{
+		FRemoteObjectId Id = UE::CoreUObject::Private::FObjectHandleUtils::GetRemoteId(Object);
+		{
+			FScopeLock MapLock(&ObjectMapCritical);
+			UObject*& MappedObject = IdToObjectMap.FindOrAdd(Id);
+			MappedObject = Object;
+		}
+	}
+
+	void UnmapObject(UObject* Object)
+	{
+		FRemoteObjectId Id = UE::CoreUObject::Private::FObjectHandleUtils::GetRemoteId(Object);
+		{
+			FScopeLock MapLock(&ObjectMapCritical);
+			IdToObjectMap.Remove(Id);
+		}
+	}
+
+	UObject* FindObject(FRemoteObjectId Id)
+	{
+		FScopeLock MapLock(&ObjectMapCritical);
+		UObject** MappedObject = IdToObjectMap.Find(Id);
+		return MappedObject ? *MappedObject : nullptr;
+	}
+
+	UE::RemoteObject::Handle::FRemoteObjectStub* FindRemoteObjectStub(FRemoteObjectId Id)
+	{
+		using namespace UE::RemoteObject::Handle;
+
+		FScopeLock MapLock(&ObjectMapCritical);
+		FRemoteObjectStub* ExistingStub = RemoteObjects.FindChecked(Id);
+		return ExistingStub;
+	}
+
+	UE::RemoteObject::Handle::FRemoteObjectStub* FindOrAddRemoteObjectStub(FRemoteObjectId Id)
+	{
+		using namespace UE::RemoteObject::Handle;
+
+		FScopeLock MapLock(&ObjectMapCritical);
+		FRemoteObjectStub*& ExistingStub = RemoteObjects.FindOrAdd(Id);
+		if (!ExistingStub)
+		{
+			ExistingStub = new FRemoteObjectStub();
+			ExistingStub->Id = Id;
+			ExistingStub->RemoteServerId = Id.GetServerId();
+		}
+		return ExistingStub;
+	}
+
+	UE::RemoteObject::Handle::FRemoteObjectStub* FindOrAddRemoteObjectStub(UObject* Object, FRemoteServerId DestinationServerId)
+	{
+		using namespace UE::RemoteObject::Handle;
+
+		FScopeLock MapLock(&ObjectMapCritical);
+		FRemoteObjectStub*& ExistingStub = RemoteObjects.FindOrAdd(UE::CoreUObject::Private::FObjectHandleUtils::GetRemoteId(Object));
+		if (!ExistingStub)
+		{
+			ExistingStub = new FRemoteObjectStub(Object);
+		}
+		ExistingStub->RemoteServerId = DestinationServerId;
+		return ExistingStub;
+	}
+};
+FRemoteObjectMaps* ObjectMaps = nullptr;
+
+void InitRemoteObjects()
+{
+	ObjectMaps = new FRemoteObjectMaps();
+
+	uint32 ServerIdNumber = 0;
+	if (!FParse::Value(FCommandLine::Get(), TEXT("MultiServerLocalId="), ServerIdNumber))
+	{
+		if (!FParse::Value(FCommandLine::Get(), TEXT("LocalPeerId="), ServerIdNumber))
+		{
+			ServerIdNumber = 1;
+		}
+	}
+	checkf(ServerIdNumber != 0, TEXT("Remote Server Id=%u is not valid"), ServerIdNumber);
+	RemoteServerId = FRemoteServerId(ServerIdNumber);
+
+	if (!FCoreUObjectDelegates::RemoteObjectTransferDelegate.IsBound())
+	{
+		FCoreUObjectDelegates::RemoteObjectTransferDelegate.BindStatic(&UE::RemoteObject::Serialization::Disk::SaveObjectToDisk);
+	}
+	if (!FCoreUObjectDelegates::RequestRemoteObjectDelegate.IsBound())
+	{
+		FCoreUObjectDelegates::RequestRemoteObjectDelegate.BindStatic(&UE::RemoteObject::Serialization::Disk::LoadObjectFromDisk);
+	}
+
+	FCoreDelegates::OnEndFrame.AddStatic(&TickRemoteObjectRequests);
+}
+
+void MarkAsRemote(UObject* Object, FRemoteServerId DestinationServerId)
+{
+	static_assert(sizeof(FObjectHandle) == sizeof(UObject*));
+	FUObjectItem* ObjectItem = GUObjectArray.ObjectToObjectItem(Object);
+	ObjectItem->SetFlags(EInternalObjectFlags::Remote);
+	ObjectItem->ClearFlags(EInternalObjectFlags_RootFlags);
+	ObjectMaps->FindOrAddRemoteObjectStub(Object, DestinationServerId);
+}
+
+void MarkAsLocal(UObject* Object)
+{
+	FUObjectItem* ObjectItem = GUObjectArray.ObjectToObjectItem(Object);
+	ObjectItem->ClearFlags(EInternalObjectFlags::Remote);
+}
+
+void MapObject(UObject* Object)
+{
+	ObjectMaps->MapObject(Object);
+}
+
+void UnmapObject(UObject* Object)
+{
+	ObjectMaps->UnmapObject(Object);
+}
+
+UObject* FindObject(FRemoteObjectId ObjectId)
+{
+	UObject* Object = ObjectMaps->FindObject(ObjectId);
+	return Object;
+}
+
+UE::RemoteObject::Handle::FRemoteObjectStub* FindRemoteObjectStub(FRemoteObjectId ObjectId)
+{
+	return ObjectMaps->FindRemoteObjectStub(ObjectId);
+}
+
+} // namespace UE::RemoteObject::Private
+
+namespace UE::RemoteObject
+{
+	FRemoteServerId GetGlobalServerId()
+	{
+		return UE::RemoteObject::Private::RemoteServerId;
+	}
+}
+
+namespace UE::RemoteObject::Handle
+{
+
+FRemoteObjectStub::FRemoteObjectStub(UObject* Object)
+{
+	using namespace UE::CoreUObject::Private;
+
+	Id = FObjectHandleUtils::GetRemoteId(Object);
+	if (UObject* Outer = FObjectHandleUtils::GetNonAccessTrackedOuterNoResolve(Object))
+	{
+		OuterId = FObjectHandleUtils::GetRemoteId(Outer);
+	}
+}
+
+bool IsRemote(FRemoteObjectId ObjectId)
+{
+	return UE::RemoteObject::Private::ObjectMaps->IsRemote(ObjectId);
+}
+
+bool IsRemote(UObject* Object)
+{
+	const int32 InternalIndex = GUObjectArray.ObjectToIndex(Object);
+	bool bResult = InternalIndex >= 0 && GUObjectArray.IndexToObject(InternalIndex)->HasAnyFlags(EInternalObjectFlags::Remote);
+	return bResult;
+}
+
+UObject* ResolveObject(const FRemoteObjectStub* Stub)
+{
+	using namespace UE::RemoteObject::Transfer;
+	using namespace UE::RemoteObject::Private;
+
+	UObject* Object = ObjectMaps->FindObject(Stub->Id);
+	bool bRemoteObject = false;
+	if (Object)
+	{
+		bRemoteObject = IsRemote(Object);
+	}
+	else
+	{
+		bRemoteObject = IsRemote(Stub->Id);
+	}
+	if (bRemoteObject)
+	{
+		UObject* Outer = Object ? Object->GetOuter() : ObjectMaps->FindObject(Stub->OuterId);
+		Object = MigrateObjectFromRemoteServer(Stub->Id, Stub->RemoteServerId, Outer);
+	}
+	if (Object)
+	{
+		MarkAsLocal(Object);
+	}
+	else
+	{
+		Object = nullptr;
+	}
+	return Object;
+}
+
+COREUOBJECT_API UObject* ResolveObject(UObject* Object)
+{
+	using namespace UE::RemoteObject::Private;
+
+	FRemoteObjectStub* Stub = ObjectMaps->FindRemoteObjectStub(UE::CoreUObject::Private::FObjectHandleUtils::GetRemoteId(Object));
+	return ResolveObject(Stub);
+}
+
+} // namespace UE::RemoteObject::Handle
+
+namespace UE::CoreUObject::Private
+{
+#if UE_WITH_REMOTE_OBJECT_HANDLE
+	void FObjectHandleUtils::ChangeRemoteId(UObjectBase* Object, FRemoteObjectId Id)
+	{
+		using namespace UE::RemoteObject::Private;
+		UnmapObject((UObject*)Object);
+		Object->RemoteId = Id;
+		MapObject((UObject*)Object);
+	}
+
+	FRemoteObjectHandlePrivate::FRemoteObjectHandlePrivate(UObject* Object)
+	{
+		if (Object)
+		{
+			PointerOrHandle = UPTRINT(Object);
+		}
+		else
+		{
+			PointerOrHandle = 0;
+		}
+	}
+
+	FRemoteObjectHandlePrivate::FRemoteObjectHandlePrivate(UE::RemoteObject::Handle::FRemoteObjectStub* RemoteInfo)
+	{
+		PointerOrHandle = UPTRINT(RemoteInfo) | 1;
+	}
+
+	FRemoteObjectHandlePrivate FRemoteObjectHandlePrivate::ConvertToRemoteHandle(UObject* Object)
+	{
+		using namespace UE::RemoteObject::Handle;
+		using namespace UE::RemoteObject::Private;		
+
+		FRemoteObjectStub* Stub = ObjectMaps->FindRemoteObjectStub(UE::CoreUObject::Private::FObjectHandleUtils::GetRemoteId(Object));
+		return FRemoteObjectHandlePrivate(Stub);
+	}
+
+	FRemoteObjectHandlePrivate FRemoteObjectHandlePrivate::FromIdNoResolve(FRemoteObjectId ObjectId)
+	{
+		using namespace UE::RemoteObject::Handle;
+		using namespace UE::RemoteObject::Private;
+
+		UObject* Obj = nullptr;
+		if (ObjectId.IsValid())
+		{
+			Obj = ObjectMaps->FindObject(ObjectId);
+			if (Obj && !Obj->HasAnyInternalFlags(EInternalObjectFlags::Remote))
+			{
+				return FRemoteObjectHandlePrivate(Obj);
+			}
+			else if (FRemoteObjectStub* Stub = ObjectMaps->FindOrAddRemoteObjectStub(ObjectId))
+			{
+				return FRemoteObjectHandlePrivate(Stub);
+			}
+			check(false);
+		}
+		return FRemoteObjectHandlePrivate(Obj);
+	}
+#endif // UE_WITH_REMOTE_OBJECT_HANDLE
+
+} // namespace UE::CoreUObject::Private
+
+FRemoteObjectId FRemoteObjectId::Generate()
+{
+	using namespace UE::RemoteObject::Private;
+
+	return FRemoteObjectId(RemoteServerId, RemoteObjectSerialNumber.fetch_add(1));
+}
+
+FRemoteObjectId FRemoteObjectId::FromObject(const UObject* Object)
+{
+	return UE::CoreUObject::Private::FObjectHandleUtils::GetRemoteId(Object);
+}
+
+FArchive& operator<<(FArchive& Ar, FRemoteObjectId& Id)
+{
+	Ar << Id.Id;
+	return Ar;
+}
diff --git a/Engine/Source/Runtime/CoreUObject/Private/UObject/RemoteObjectPrivate.h b/Engine/Source/Runtime/CoreUObject/Private/UObject/RemoteObjectPrivate.h
new file mode 100644
--- /dev/null
+++ b/Engine/Source/Runtime/CoreUObject/Private/UObject/RemoteObjectPrivate.h
@@ -0,0 +1,53 @@
+// Copyright Epic Games, Inc. All Rights Reserved.
+
+#pragma once
+
+#include "CoreTypes.h"
+#include "Misc/AssertionMacros.h"
+#include "UObject/NameTypes.h"
+#include "UObject/ObjectMacros.h"
+#include "UObject/Object.h"
+
+namespace UE::RemoteObject::Private
+{
+	/**
+	* Initializes remote objects subsystems
+	*/
+	void InitRemoteObjects();
+
+	/**
+	* Marks object memory as remote and creates its stub
+	*/
+	void MarkAsRemote(UObject* Object, FRemoteServerId DestinationServerId);
+
+	/**
+	* Marks object memory as local
+	*/
+	void MarkAsLocal(UObject* Object);
+
+	/**
+	* Maps an object with its id
+	*/
+	void MapObject(UObject* Object);
+
+	/**
+	* Unmaps object memory from its id
+	*/
+	void UnmapObject(UObject* Object);
+
+	/**
+	* Finds a local object by its unique id
+	* @returns Object if it exists in the current processes (and has been mapped with its id)
+	*/
+	UObject* FindObject(FRemoteObjectId ObjectId);
+
+	/**
+	* Processes remote object requests (send and receive requests)
+	*/
+	void TickRemoteObjectRequests();
+
+	/**
+	* Finds remote object stub
+	*/
+	UE::RemoteObject::Handle::FRemoteObjectStub* FindRemoteObjectStub(FRemoteObjectId ObjectId);
+}
diff --git a/Engine/Source/Runtime/CoreUObject/Private/UObject/RemoteObjectSerialization.cpp b/Engine/Source/Runtime/CoreUObject/Private/UObject/RemoteObjectSerialization.cpp
new file mode 100644
--- /dev/null
+++ b/Engine/Source/Runtime/CoreUObject/Private/UObject/RemoteObjectSerialization.cpp
@@ -0,0 +1,429 @@
+// Copyright Epic Games, Inc. All Rights Reserved.
+
+#include "UObject/RemoteObjectSerialization.h"
+#include "UObject/RemoteObject.h"
+#include "UObject/RemoteObjectTransfer.h"
+#include "UObject/RemoteObjectPrivate.h"
+#include "UObject/ObjectPtr.h"
+#include "UObject/Class.h"
+#include "UObject/PropertyPortFlags.h"
+#include "Serialization/ArchiveUObject.h"
+#include "Serialization/MemoryWriter.h"
+#include "Serialization/MemoryReader.h"
+#include "Templates/Casts.h"
+#include "Async/Async.h"
+#include "UObject/ObjectHandlePrivate.h"
+#include "Misc/Paths.h"
+#include "HAL/FileManager.h"
+
+namespace UE::RemoteObject::Serialization::Disk
+{
+
+FString GenerateRemoteObjectFilename(FRemoteObjectId ObjectId)
+{
+	return FPaths::Combine(*FPaths::ProjectSavedDir(), *(ObjectId.ToString() + TEXT(".tmp")));
+}
+
+void LoadObjectFromDisk(FRemoteObjectId ObjectId, FRemoteServerId CurrentOwnerServerId, FRemoteServerId DestinationServerId)
+{
+	TUniquePtr<FArchive> FileReader{ IFileManager::Get().CreateFileReader(*GenerateRemoteObjectFilename(ObjectId)) };
+	checkf(FileReader.IsValid(), TEXT("Unable to create file reader for remote object %s"), *ObjectId.ToString());
+	FRemoteObjectData* ObjectData = new FRemoteObjectData();
+	
+	*FileReader << ObjectData->Bytes;
+	FileReader->Close();
+
+	UE::RemoteObject::Transfer::OnObjectDataReceived(ObjectId, ObjectData);
+}
+
+void SaveObjectToDisk(FRemoteObjectId ObjectId, FRemoteServerId ServerId, FRemoteObjectData& ObjectData)
+{
+	TUniquePtr<FArchive> FileWriter{ IFileManager::Get().CreateFileWriter(*GenerateRemoteObjectFilename(ObjectId)) };
+	checkf(FileWriter.IsValid(), TEXT("Unable to create file writer for remote object %s"), *ObjectId.ToString());
+	
+	*FileWriter << ObjectData.Bytes;
+	FileWriter->Close();
+}
+
+} // namespace UE::RemoteObject::Serialization::Disk
+
+namespace UE::RemoteObject::Serialization
+{
+
+enum class ERemoteReferenceType : uint8
+{
+	None = 0,
+	IdOnly = 1,
+	PathName = 2
+};
+
+FArchive& operator<<(FArchive& Ar, ERemoteReferenceType& RefType)
+{
+	uint8 SerialziedType = (uint8)RefType;
+	Ar << SerialziedType;
+	RefType = (ERemoteReferenceType)SerialziedType;
+	return Ar;
+}
+
+class FArchiveRemoteObjectWriter : public FMemoryWriter
+{
+	UObject* RootObject = nullptr;	
+	TArray<UObject*>& ObjectsToSerialize;
+	const TSet<UObject*>& SerializedObjects;
+	TArray<FName>& Names;
+	TMap<FName, int32>& NameMap;
+
+	ERemoteReferenceType WriteObjectReference(UObject*& Obj)
+	{
+		using namespace UE::CoreUObject::Private;
+
+		ERemoteReferenceType Type = ERemoteReferenceType::None;
+		FRemoteObjectId Id;
+
+		if (Obj)
+		{
+			Id = FObjectHandleUtils::GetRemoteId(Obj);
+
+			if (Obj == RootObject || UE::RemoteObject::Handle::IsRemote(Obj) || Obj->IsIn(RootObject) || (Id.GetServerId().IsValid() && Id.GetServerId() != UE::RemoteObject::GetGlobalServerId()))
+			{
+				Type = ERemoteReferenceType::IdOnly;
+			}
+			else
+			{
+				Type = ERemoteReferenceType::PathName;
+			}
+		}
+
+		*this << Type;
+		// Always serialize unique id as objects may not exist on the other server and then we may need to pull them from this server
+		if (Type != ERemoteReferenceType::None)
+		{
+			*this << Id;
+		}
+
+		if (Type == ERemoteReferenceType::PathName)
+		{
+			FString PathName = GetPathNameSafe(Obj);
+			*this << PathName;
+		}
+
+		return Type;
+	}
+
+public:
+
+	FArchiveRemoteObjectWriter(UObject* InRootObject, FRemoteObjectData& OutBytes, TArray<UObject*>& OutObjectsToSerialize, const TSet<UObject*>& InSerializedObjects, TArray<FName>& InOutNames, TMap<FName, int32>& InOutNameMap)
+		: FMemoryWriter(OutBytes.Bytes)
+		, RootObject(InRootObject)
+		, ObjectsToSerialize(OutObjectsToSerialize)
+		, SerializedObjects(InSerializedObjects)
+		, Names(InOutNames)
+		, NameMap(InOutNameMap)
+	{
+		SetPortFlags(PPF_IncludeTransient);
+	}
+
+	virtual FArchive& operator<<(UObject*& Obj) override
+	{
+		FRemoteObjectId ObjId;
+		ERemoteReferenceType RefType = WriteObjectReference(Obj);
+		if (RefType == ERemoteReferenceType::IdOnly)
+		{
+			if (Obj->IsIn(RootObject))
+			{
+				ObjectsToSerialize.Add(Obj);
+			}			
+		}
+		if (Obj)
+		{
+			// Map all references we encounter during serialization, especially if they're objects that are not going to be serialized
+			// because we may end up needing to resolve them later if a remote server requests them with their unique id.
+			// Note: we could just hash all objects with their FRemoteObjectId in UObjectHash.cpp but for now we can limit the memory impact by hashing 
+			// only those we encounter when serializing objects that are actually going to be migrated
+			UE::RemoteObject::Private::MapObject(Obj);
+		}
+
+		return *this;
+	}
+
+	virtual FArchive& operator<<(FName& Name) override
+	{
+		int32 NameIndex = -1;
+		if (int32* ExistingNameIndex = NameMap.Find(Name))
+		{
+			NameIndex = *ExistingNameIndex;
+		}
+		else
+		{
+			NameIndex = Names.Add(Name);
+			NameMap.Add(Name, NameIndex);
+		}
+		*this << NameIndex;
+		return *this;
+	}
+
+	virtual FArchive& operator<<(FLazyObjectPtr& Value) override { return FArchiveUObject::SerializeLazyObjectPtr(*this, Value); }
+	virtual FArchive& operator<<(FObjectPtr& Value) override 
+	{ 
+		return FArchiveUObject::SerializeObjectPtr(*this, Value); 
+	}
+	virtual FArchive& operator<<(FSoftObjectPtr& Value) override { return FArchiveUObject::SerializeSoftObjectPtr(*this, Value); }
+	virtual FArchive& operator<<(FSoftObjectPath& Value) override { return FArchiveUObject::SerializeSoftObjectPath(*this, Value); }
+	virtual FArchive& operator<<(FWeakObjectPtr& Value) override { return FArchiveUObject::SerializeWeakObjectPtr(*this, Value); }
+};
+
+void SerializeRemoteObject(UObject* Object, FArchive& Ar)
+{
+	using namespace UE::RemoteObject;
+	using namespace UE::RemoteObject::Private;
+	using namespace UE::CoreUObject::Private;
+
+	FName Name = Object->GetFName();
+	FRemoteObjectId Id = FObjectHandleUtils::GetRemoteId(Object);
+	UClass* Class = Object->GetClass();
+	UObject* Outer = FObjectHandleUtils::GetNonAccessTrackedOuterNoResolve(Object);
+
+	MapObject(Class);
+	MapObject(Outer);
+
+	Ar << Name;
+	Ar << Id;
+	Ar << Class;
+	Ar << Outer;
+
+	Object->Serialize(Ar);
+}
+
+class FArchiveRemoteObjectReader : public FMemoryReader
+{
+	TArray<FName>& Names;
+	TMap<FRemoteObjectId, TArray<FObjectPtr*>> IdToPtrMap;
+
+	void StoreUnresolvedHandle(FRemoteObjectId Id, FObjectPtr* Ptr)
+	{
+		TArray<FObjectPtr*>& HandleList = IdToPtrMap.FindOrAdd(Id);
+		HandleList.Add(Ptr);
+	}
+
+	ERemoteReferenceType ReadObjectReference(FObjectPtr& Value)
+	{
+		ERemoteReferenceType Type;		
+		*this << Type;
+
+		if (Type != ERemoteReferenceType::None)
+		{
+			FRemoteObjectId ObjId;
+			*this << ObjId;
+
+			bool bNeedsResolvingWithId = true;
+			if (Type == ERemoteReferenceType::PathName)
+			{
+				FString PathName;
+				*this << PathName;
+				checkf(!PathName.IsEmpty(), TEXT("Serialized remote object with empty path name but it should've been ERemoteReferenceType::None"));
+
+				// Try to resolve path name immediately as we expect the object to exist in memory.
+				UObject* Obj = ::FindObject<UObject>(nullptr, *PathName);
+				if (Obj)
+				{
+					Value = FObjectPtr(Obj);
+					bNeedsResolvingWithId = false;
+				}
+			}
+			if (bNeedsResolvingWithId)
+			{
+				// If the serialized reference was not found in memory or if the reference was serialized as id-only 
+				// keep it as an unresolved ObjectPtr and store a pointer to it so that we can try to resolve it after all objects have been deserialize
+#if UE_WITH_REMOTE_OBJECT_HANDLE
+				FObjectHandle Handle = FObjectHandle::FromIdNoResolve(ObjId);
+				Value = FObjectPtr(Handle);
+				if (!Value.IsResolved())
+				{
+					StoreUnresolvedHandle(ObjId, &Value);
+				}
+#endif // UE_WITH_REMOTE_OBJECT_HANDLE
+			}
+		}
+		else
+		{
+			Value = FObjectPtr();
+		}
+		return Type;
+	}
+
+public:
+
+	FArchiveRemoteObjectReader(const FRemoteObjectData& InBytes, TArray<FName>& InNames)
+		: FMemoryReader(InBytes.Bytes)
+		, Names(InNames)
+	{
+		SetPortFlags(PPF_IncludeTransient);
+	}
+
+	virtual FArchive& operator<<(UObject*& Obj) override
+	{
+		FObjectPtr Value;
+		ReadObjectReference(Value);
+		check(Value.IsResolved());
+		Obj = Value.Get();
+		return *this;
+	}
+
+	virtual FArchive& operator<<(FName& Name) override
+	{
+		int32 NameIndex = -1;
+		*this << NameIndex;
+		Name = Names[NameIndex];
+		return *this;
+	}
+
+	virtual FArchive& operator<<(FLazyObjectPtr& Value) override { return FArchiveUObject::SerializeLazyObjectPtr(*this, Value); }
+	virtual FArchive& operator<<(FObjectPtr& Value) override 
+	{
+		ReadObjectReference(Value);
+		return *this;
+	}
+
+	virtual FArchive& operator<<(FSoftObjectPtr& Value) override { return FArchiveUObject::SerializeSoftObjectPtr(*this, Value); }
+	virtual FArchive& operator<<(FSoftObjectPath& Value) override { return FArchiveUObject::SerializeSoftObjectPath(*this, Value); }
+	virtual FArchive& operator<<(FWeakObjectPtr& Value) override { return FArchiveUObject::SerializeWeakObjectPtr(*this, Value); }
+
+	/**
+	* Fixup (resolve) all deserialized FObjectPtrs
+	*/
+	void ResolveObjectPtrs(const TArray<UObject*>& SerializedObjects)
+	{
+		for (UObject* Obj : SerializedObjects)
+		{
+			if (TArray<FObjectPtr*>* UnresolvedPtrs = IdToPtrMap.Find(UE::CoreUObject::Private::FObjectHandleUtils::GetRemoteId(Obj)))
+			{
+				for (FObjectPtr* Ptr : *UnresolvedPtrs)
+				{
+					*Ptr = FObjectPtr(Obj);
+				}
+			}
+		}
+	}
+};
+
+UObject* DeserializeRemoteObject(FArchive& Ar, UObject* DestinationOuter)
+{
+	using namespace UE::RemoteObject;
+	using namespace UE::RemoteObject::Private;
+
+	FName Name;
+	FRemoteObjectId Id;
+	FObjectPtr ClassPtr;
+	FObjectPtr OuterPtr;
+
+	Ar << Name;
+	Ar << Id;
+	Ar << ClassPtr;
+	Ar << OuterPtr;
+
+	UClass* Class = CastChecked<UClass>(ClassPtr.Get());
+	UObject* Outer = OuterPtr.Get();
+
+	UObject* Object = NewObject<UObject>(Outer, Class, Name, RF_NoFlags, nullptr, false, nullptr, nullptr, Id);
+	MapObject(Object);
+	MarkAsLocal(Object);
+
+	Object->Serialize(Ar);
+
+	return Object;
+}
+
+void SerializeNames(FArchive& Ar, TArray<FName>& InNames)
+{
+	int32 NumNames = InNames.Num();
+	Ar << NumNames;
+	for (FName& Name : InNames)
+	{
+		FString NameString = Name.ToString();
+		Ar << NameString;
+	}
+}
+
+void DeserializeNames(FArchive& Ar, TArray<FName>& OutNames)
+{
+	int32 NumNames = 0;
+	Ar << NumNames;
+	for (int32 NameIndex = 0; NameIndex < NumNames; ++NameIndex)
+	{
+		FString NameString;
+		Ar << NameString;
+		OutNames.Add(FName(NameString, FNAME_Add));
+	}
+}
+
+FRemoteObjectData* SerializeObjectData(UObject* Object, TSet<UObject*>& OutObjects)
+{
+	using namespace UE::RemoteObject;
+
+	FRemoteObjectData* ObjectData = new FRemoteObjectData();
+	TArray<UObject*> ObjectsToSerialize;
+	TArray<FName> Names;
+	TMap<FName, int32> NameMap;
+	int32 Version = 0;
+	int64 HeaderOffset = 0;
+
+	{
+		FArchiveRemoteObjectWriter Ar(Object, *ObjectData, ObjectsToSerialize, OutObjects, Names, NameMap);
+		ObjectsToSerialize.Add(Object);
+
+		Ar << Version;
+		const int64 OffsetOfHeaderOffset = Ar.Tell();
+		Ar << HeaderOffset;
+
+		for (int32 ObjectIndex = 0; ObjectIndex < ObjectsToSerialize.Num(); ++ObjectIndex)
+		{
+			UObject* ObjectToSerialize = ObjectsToSerialize[ObjectIndex];
+			if (!OutObjects.Contains(ObjectToSerialize))
+			{
+				OutObjects.Add(ObjectToSerialize);
+				SerializeRemoteObject(ObjectToSerialize, Ar);
+			}
+		}
+
+		HeaderOffset = Ar.Tell();
+		Ar.Seek(OffsetOfHeaderOffset);
+		Ar << HeaderOffset;
+		Ar.Seek(HeaderOffset);
+
+		int32 NumObjects = OutObjects.Num();
+		Ar << NumObjects;
+		SerializeNames(Ar, Names);
+	}
+
+	return ObjectData;
+}
+
+void DeserializeObjectData(FRemoteObjectData& ObjectData, UObject* DestinationOuter, TArray<UObject*>& OutObjects)
+{
+	int32 Version = 0;
+	int64 HeaderOffset = 0;
+	int32 NumObjects = 0;
+	TArray<FName> Names;
+
+	FArchiveRemoteObjectReader Ar(ObjectData, Names);
+
+	Ar << Version;
+	Ar << HeaderOffset;
+
+	const int64 ObjectDataOffset = Ar.Tell();
+	Ar.Seek(HeaderOffset);
+
+	Ar << NumObjects;
+	DeserializeNames(Ar, Names);
+
+	Ar.Seek(ObjectDataOffset);
+
+	for (int32 ObjectIndex = 0; ObjectIndex < NumObjects; ++ObjectIndex)
+	{
+		OutObjects.Add(DeserializeRemoteObject(Ar, ObjectIndex == 0 ? DestinationOuter : nullptr));
+	}
+
+	Ar.ResolveObjectPtrs(OutObjects);
+}
+
+} // namespace UE::RemoveObject::Serialization
\ No newline at end of file
diff --git a/Engine/Source/Runtime/CoreUObject/Private/UObject/RemoteObjectSerialization.h b/Engine/Source/Runtime/CoreUObject/Private/UObject/RemoteObjectSerialization.h
new file mode 100644
--- /dev/null
+++ b/Engine/Source/Runtime/CoreUObject/Private/UObject/RemoteObjectSerialization.h
@@ -0,0 +1,22 @@
+// Copyright Epic Games, Inc. All Rights Reserved.
+
+#pragma once
+
+#include "UObject/RemoteObjectTransfer.h"
+
+struct FRemoteObjectStub;
+
+namespace UE::RemoteObject::Serialization
+{
+
+	FRemoteObjectData* SerializeObjectData(UObject* Object, TSet<UObject*>& OutObjects);
+	void DeserializeObjectData(FRemoteObjectData& ObjectData, UObject* DestinationOuter, TArray<UObject*>& OutReceivedObjects);
+
+} // namespace UE::RemoveObject::Serialization
+
+
+namespace UE::RemoteObject::Serialization::Disk
+{
+	void LoadObjectFromDisk(FRemoteObjectId ObjectId, FRemoteServerId CurrentOwnerServerId, FRemoteServerId DestinationServerId);
+	void SaveObjectToDisk(FRemoteObjectId ObjectId, FRemoteServerId ServerId, FRemoteObjectData& ObjectData);
+}
\ No newline at end of file
diff --git a/Engine/Source/Runtime/CoreUObject/Private/UObject/RemoteObjectTransfer.cpp b/Engine/Source/Runtime/CoreUObject/Private/UObject/RemoteObjectTransfer.cpp
new file mode 100644
--- /dev/null
+++ b/Engine/Source/Runtime/CoreUObject/Private/UObject/RemoteObjectTransfer.cpp
@@ -0,0 +1,451 @@
+// Copyright Epic Games, Inc. All Rights Reserved.
+
+#include "UObject/RemoteObjectTransfer.h"
+#include "UObject/RemoteObjectPrivate.h"
+#include "UObject/RemoteObjectSerialization.h"
+#include "UObject/ObjectPtr.h"
+#include "UObject/Class.h"
+#include "Serialization/ArchiveUObject.h"
+#include "Serialization/MemoryWriter.h"
+#include "Serialization/MemoryReader.h"
+#include "UObject/RemoteObject.h"
+#include "UObject/ObjectHandlePrivate.h"
+#include "HAL/CriticalSection.h"
+
+
+namespace UE::RemoteObject::Transfer
+{
+
+/**
+ * Remote object transfer queue. 
+ * Queued requests (both send and receive) are processed on the game thread since some systems (like RPCs) that are used for transfering object data are GT-only
+ */
+class FRemoteObjectTransferQueue
+{
+	enum class ERequestState : uint8
+	{
+		None,
+		InProgress,
+		Complete
+	};
+
+	enum class ERequestType : uint8
+	{
+		Unknown,
+		Send,
+		Receive,
+		Forward
+	};
+
+	/** Base class for remote object requests */
+	struct FRemoteObjectRequestBase
+	{
+		/** Unique object id this request is related to */
+		FRemoteObjectId ObjectId;
+		/** Server Id this request will be sent to */
+		FRemoteServerId ServerId;
+		/** Type of this request */
+		ERequestType Type = ERequestType::Unknown;
+		/** Current state of this request */
+		ERequestState State = ERequestState::None;
+
+		/** Next request on the list */
+		FRemoteObjectRequestBase* Next = nullptr;
+
+		virtual void Process() = 0;
+		virtual ~FRemoteObjectRequestBase() = default;
+	};
+
+	/**
+	* Receive request - requests an object from a remote server
+	* Receive request is not complete until a remote server fullfils the request with a 'send' request
+	*/
+	struct FReceiveRemoteObjectRequest : public FRemoteObjectRequestBase
+	{		
+		TPromise<TUniquePtr<FRemoteObjectData>> ReceivePromise;
+		FRemoteObjectDataFuture ReceiveFuture;
+
+		FReceiveRemoteObjectRequest()
+		{
+			Type = ERequestType::Receive;
+		}
+		virtual void Process() override
+		{
+			// Request an object from remote server
+			FCoreUObjectDelegates::RequestRemoteObjectDelegate.Execute(ObjectId, /*CurrentOwner*/ ServerId, GetGlobalServerId());
+		}
+	};
+
+	/**
+	* Forward request - forwards receive request for an object to another remote server
+	*/
+	struct FForwardRemoteObjectRequest : public FRemoteObjectRequestBase
+	{
+		/** Server Id that wants to receive object data */
+		FRemoteServerId DestinationServerId;
+
+		FForwardRemoteObjectRequest()
+		{
+			Type = ERequestType::Forward;
+		}
+		virtual void Process() override
+		{
+			// Request an object from remote server to be sent to another remote server
+			FCoreUObjectDelegates::RequestRemoteObjectDelegate.Execute(ObjectId, /*CurrentOwner*/ ServerId, DestinationServerId);
+			State = ERequestState::Complete;
+		}
+	};
+
+	/**
+	* Send request - sends object data to a remote server
+	* Send requests can be insigated either by the current server or be a response to a receive request sent by a remote server
+	*/
+	struct FSendRemoteObjectRequest : public FRemoteObjectRequestBase
+	{
+		TUniquePtr<FRemoteObjectData> ObjectData;
+
+		FSendRemoteObjectRequest()
+		{
+			Type = ERequestType::Send;
+		}
+		virtual void Process() override
+		{
+			// Send the associated object data to a remote server
+			FCoreUObjectDelegates::RemoteObjectTransferDelegate.Execute(ObjectId, ServerId, *ObjectData);
+			State = ERequestState::Complete;
+		}
+	};
+	
+	FCriticalSection RequestCritical;
+	// Request queue. Order of requests needs to be preserved
+	FRemoteObjectRequestBase* Head = nullptr;
+	FRemoteObjectRequestBase* Tail = nullptr;
+
+	// For quick lookups, receive requests are stored in ID -> Request map
+	TMap<FRemoteObjectId, FReceiveRemoteObjectRequest*> IdToReceiveRequestMap;
+
+	/**
+	* Find an existing receive request by object id
+	* @param ObjectId Remote object id
+	* @return Existing receive request or null if not found
+	*/
+	FORCEINLINE FReceiveRemoteObjectRequest* FindReceiveRequestUnsafe(FRemoteObjectId ObjectId)
+	{
+		FReceiveRemoteObjectRequest** FoundRequest = IdToReceiveRequestMap.Find(ObjectId);
+		if (FoundRequest)
+		{
+			return *FoundRequest;
+		}
+		return nullptr;
+	}
+
+	/**
+	* Appends (links to the list tail) new request
+	* @param Request request to link
+	*/
+	void AddRequest(FRemoteObjectRequestBase* Request)
+	{
+		if (Tail)
+		{
+			Tail->Next = Request;
+			Tail = Request;
+		}
+		else
+		{
+			Tail = Request;
+			Head = Request;
+		}
+		if (Request->Type == ERequestType::Receive)
+		{
+			IdToReceiveRequestMap.Add(Request->ObjectId, static_cast<FReceiveRemoteObjectRequest*>(Request));
+		}
+	}
+
+	/**
+	* Removes (unlinks) and destroys an existing request
+	* @param Request request to unlink and destroy
+	* @param Prev Request that was previous to the one being destroyed
+	* @return Request following the request being destroyed
+	*/
+	FRemoteObjectRequestBase* RemoveRequest(FRemoteObjectRequestBase* Request, FRemoteObjectRequestBase* Prev)
+	{
+		FRemoteObjectRequestBase* Next = Request->Next;
+		if (Prev)
+		{
+			Prev->Next = Next;
+		}
+		if (Request == Head)
+		{
+			Head = Next;
+		}
+		if (Request == Tail)
+		{
+			Tail = Prev;
+		}
+		if (Request->Type == ERequestType::Receive)
+		{
+			IdToReceiveRequestMap.Remove(Request->ObjectId);
+		}
+		delete Request;
+		return Next;
+	}
+	
+public:
+
+	/**
+	* Requests an object from a remote server (adds a receive request)
+	* @param ObjectId Id of the object being requested
+	* @param SourceServerId Remote server id
+	* @return Future that's going to receive the object data
+	*/
+	FRemoteObjectDataFuture RequestRemoteObject(FRemoteObjectId ObjectId, FRemoteServerId CurrentOwnerServerId)
+	{
+		FReceiveRemoteObjectRequest* Request = nullptr;
+		{
+			FScopeLock RequestLock(&RequestCritical);
+			Request = FindReceiveRequestUnsafe(ObjectId);
+			if (!Request)
+			{
+				Request = new FReceiveRemoteObjectRequest();
+				Request->ObjectId = ObjectId;
+				Request->ServerId = CurrentOwnerServerId;
+				Request->ReceiveFuture = Request->ReceivePromise.GetFuture().Share();
+				AddRequest(Request);
+			}
+		}
+		return Request->ReceiveFuture;
+	}
+
+	/**
+	* Forwards receive request to a remote server
+	* @param ObjectId Id of the object being requested
+	* @param CurrentOwnerServerId Remote server id that currently owns the reqested object
+	* @param DestinationServerId Remote server id that requests the object ownership
+	*/
+	void ForwardRemoteObjectRequest(FRemoteObjectId ObjectId, FRemoteServerId CurrentOwnerServerId, FRemoteServerId DestinationServerId)
+	{
+		FForwardRemoteObjectRequest* Request = new FForwardRemoteObjectRequest();
+		Request->ObjectId = ObjectId;
+		Request->ServerId = CurrentOwnerServerId;
+		Request->DestinationServerId = DestinationServerId;
+
+		FScopeLock RequestLock(&RequestCritical);
+		AddRequest(Request);
+	}
+
+	/**
+	* Sends an object to a remote server (adds a send request)
+	* @param ObjectId Id of the object being requested
+	* @param DestinationServerId Remote server id
+	* @param ObjectData Object data to send. Send request takes ownership of the data
+	* @return Future that's going to receive the object data
+	*/
+	void SendRemoteObject(FRemoteObjectId ObjectId, FRemoteServerId DestinationServerId, FRemoteObjectData* ObjectData)
+	{		
+		FSendRemoteObjectRequest* Request = new FSendRemoteObjectRequest();
+		Request->ObjectId = ObjectId;
+		Request->ServerId = DestinationServerId;
+		Request->ObjectData = MakeUnique<FRemoteObjectData>(MoveTemp(*ObjectData));
+
+		FScopeLock RequestLock(&RequestCritical);
+		AddRequest(Request);
+	}
+
+	/**
+	* Fullfils receive request
+	* @param ObjectId Id of the object that was requested
+	* @param Data Object data. Data ownership is transferred to the receive request (if it exists)
+	* @return True if a receive request for the specified object exists (and object data ownership was transferred), false otherwise
+	*/
+	bool FulfillReceiveRequest(FRemoteObjectId ObjectId, FRemoteObjectData* Data)
+	{
+		FScopeLock RequestLock(&RequestCritical);
+		FReceiveRemoteObjectRequest* Request = FindReceiveRequestUnsafe(ObjectId);
+		if (Request)
+		{
+			check(Request->State == ERequestState::InProgress);
+			Request->State = ERequestState::Complete;
+			TUniquePtr<FRemoteObjectData> UniqueData(Data);
+			Request->ReceivePromise.SetValue(MoveTemp(UniqueData));
+			return true;
+		}
+		return false;
+	}
+
+	/**
+	* Finds an existing receive request
+	* @param ObjectId Id of the object the receive request is to be found
+	* @return Object data future associated with the existing request (invalid if a receive request for the specified object id was not found)
+	*/
+	FRemoteObjectDataFuture FindReceiveRequest(FRemoteObjectId ObjectId)
+	{
+		FScopeLock RequestLock(&RequestCritical);
+		FReceiveRemoteObjectRequest* Request = FindReceiveRequestUnsafe(ObjectId);
+		if (Request)
+		{
+			return Request->ReceiveFuture;
+		}
+		return FRemoteObjectDataFuture();
+	}
+
+	/**
+	* Processes exiting requests
+	*/
+	FORCENOINLINE void Tick()
+	{
+		FScopeLock RequestLock(&RequestCritical);
+		FRemoteObjectRequestBase* Prev = nullptr;
+		for (FRemoteObjectRequestBase* Request = Head; Request;)
+		{
+			if (Request->State == ERequestState::None)
+			{
+				Request->State = ERequestState::InProgress;
+				Request->Process();
+			}
+
+			if (Request->State == ERequestState::Complete)
+			{
+				Request = RemoveRequest(Request, Prev);
+			}
+			else
+			{
+				Prev = Request;
+				Request = Request->Next;
+			}
+		}
+	}
+
+	/**
+	* Processes exiting requests until the receive request associated with the specified object id is fullfilled
+	*/
+	FORCENOINLINE void FlushReceiveRequest(FRemoteObjectId ObjectId)
+	{
+		FRemoteObjectDataFuture RequestFuture = FindReceiveRequest(ObjectId);
+		while (RequestFuture.IsValid() && !RequestFuture.IsReady())
+		{
+			FCoreUObjectDelegates::FlushRemoteObjectRequestsDelegate.ExecuteIfBound();
+			Tick();
+			FPlatformProcess::Sleep(0.0f);
+		}
+	}
+};
+
+FRemoteObjectTransferQueue GTransferQueue;
+
+
+FRemoteObjectDataFuture RequestObjectData(FRemoteObjectId ObjectId, FRemoteServerId CurrentOwnerServerId, FRemoteServerId DestinationServerId)
+{
+	if (DestinationServerId == GetGlobalServerId())
+	{
+		return GTransferQueue.RequestRemoteObject(ObjectId, CurrentOwnerServerId);
+	}
+	else
+	{
+		GTransferQueue.ForwardRemoteObjectRequest(ObjectId, CurrentOwnerServerId, DestinationServerId);
+		// We're not expecting to receive anything on this server
+		return FRemoteObjectDataFuture();
+	}
+}
+
+void SendObjectData(UObject* Object, FRemoteServerId DestinationServerId, FRemoteObjectData* ObjectData)
+{
+	GTransferQueue.SendRemoteObject(UE::CoreUObject::Private::FObjectHandleUtils::GetRemoteId(Object), DestinationServerId, ObjectData);
+}
+
+UObject* MigrateObjectFromObjectData(FRemoteObjectData& ObjectData, UObject* DestinationOuter)
+{
+	using namespace UE::RemoteObject::Serialization;
+
+	TArray<UObject*> ReceivedObjects;
+	DeserializeObjectData(ObjectData, DestinationOuter, ReceivedObjects);
+
+	checkf(ReceivedObjects.Num() > 0, TEXT("PeerId:%s unable to deserialize object data (%d bytes)"), *GetGlobalServerId().ToString(), ObjectData.Bytes.Num());
+
+#if UE_WITH_REMOTE_OBJECT_HANDLE
+	for (UObject* ReceivedObject : ReceivedObjects)
+	{
+		ReceivedObject->PostMigrate(GetGlobalServerId());
+	}
+#endif
+
+	return ReceivedObjects[0];
+} 
+
+UObject* MigrateObjectFromRemoteServer(FRemoteObjectId ObjectId, FRemoteServerId CurrentOwnerServerId, UObject* DestinationOuter)
+{
+	FRemoteObjectDataFuture ObjectFuture = RequestObjectData(ObjectId, CurrentOwnerServerId, /*DestinationServerId*/ GetGlobalServerId());
+	if (!ObjectFuture.IsReady())
+	{
+		// RobM todo: this function should not expect to receive data immediately. Instead we should abort the current transaction after RequestObjectData is called.
+		GTransferQueue.FlushReceiveRequest(ObjectId);
+	}
+
+	// The below code should probably be moved to FulfillReceiveRequest
+	FRemoteObjectData& ObjectData = *ObjectFuture.Get();
+	UObject* Result = MigrateObjectFromObjectData(ObjectData, DestinationOuter);
+	checkf(UE::CoreUObject::Private::FObjectHandleUtils::GetRemoteId(Result) == ObjectId, TEXT("PeerId:%s requested object %s migration from PeerId:%s but received %s"),
+		*GetGlobalServerId().ToString(), *ObjectId.ToString(), *CurrentOwnerServerId.ToString(), *UE::CoreUObject::Private::FObjectHandleUtils::GetRemoteId(Result).ToString());
+
+	return Result;
+}
+
+void MigrateObjectToRemoteServer(UObject* Object, FRemoteServerId DestinationServerId)
+{
+	using namespace UE::RemoteObject::Serialization;
+	using namespace UE::RemoteObject::Private;
+
+	TSet<UObject*> SentObjects;	
+	FRemoteObjectData* ObjectData = SerializeObjectData(Object, SentObjects);
+	SendObjectData(Object, DestinationServerId, ObjectData);
+
+	checkf(SentObjects.Num() > 0, TEXT("PeerId:%s unable to migrate object %s to PeerId:%s"), *GetGlobalServerId().ToString(), *UE::CoreUObject::Private::FObjectHandleUtils::GetRemoteId(Object).ToString(), *DestinationServerId.ToString());
+
+	for (UObject* SentObject : SentObjects)
+	{
+#if UE_WITH_REMOTE_OBJECT_HANDLE
+		SentObject->PostMigrate(DestinationServerId);
+#endif
+		MarkAsRemote(SentObject, DestinationServerId);
+	}
+}
+
+void MigrateObjectToRemoteServer(FRemoteObjectId ObjectId, FRemoteServerId DestinationServerId)
+{
+	using namespace UE::RemoteObject::Handle;
+	using namespace UE::RemoteObject::Serialization;
+	using namespace UE::RemoteObject::Private;
+
+	UObject* Object = FindObject(ObjectId);
+	if (Object)
+	{
+		MigrateObjectToRemoteServer(Object, DestinationServerId);
+	}
+	else
+	{		
+		FRemoteObjectStub* RemoteObjectStub = FindRemoteObjectStub(ObjectId);
+		FRemoteServerId CurrentOwnerServerId = RemoteObjectStub ? RemoteObjectStub->RemoteServerId : ObjectId.GetServerId();
+		checkf(CurrentOwnerServerId != GetGlobalServerId(), TEXT("PeerId:%s attempting to forward request for object %s to PeerId:%s but object location is unknown"), 
+			*GetGlobalServerId().ToString(), *ObjectId.ToString(), *DestinationServerId.ToString());
+
+		// Forward the request to CurrentOwnerServerId
+		RequestObjectData(ObjectId, CurrentOwnerServerId, DestinationServerId);
+	}
+}
+
+void OnObjectDataReceived(FRemoteObjectId ObjectId, FRemoteObjectData* Data)
+{
+	if (!GTransferQueue.FulfillReceiveRequest(ObjectId, Data))
+	{
+		MigrateObjectFromObjectData(*Data, nullptr);
+	}
+}
+
+} // namespace UE::RemoteObject::Transfer
+
+namespace UE::RemoteObject::Private
+{
+	void TickRemoteObjectRequests()
+	{
+		UE::RemoteObject::Transfer::GTransferQueue.Tick();
+	}
+}
\ No newline at end of file
diff --git a/Engine/Source/Runtime/CoreUObject/Private/UObject/UObjectArray.cpp b/Engine/Source/Runtime/CoreUObject/Private/UObject/UObjectArray.cpp
--- a/Engine/Source/Runtime/CoreUObject/Private/UObject/UObjectArray.cpp
+++ b/Engine/Source/Runtime/CoreUObject/Private/UObject/UObjectArray.cpp
@@ -16,6 +16,7 @@
 #include "UObject/GarbageCollectionInternalFlags.h"
 #include "UObject/UObjectIterator.h"
 #include "UObject/ReachabilityAnalysisState.h"
+#include "UObject/RemoteObjectPrivate.h"
 
 DEFINE_LOG_CATEGORY_STATIC(LogUObjectArray, Log, All);
 
@@ -353,7 +354,10 @@ void FUObjectArray::FreeUObjectIndex(UObjectBase* Object)
 
 #if UE_WITH_OBJECT_HANDLE_LATE_RESOLVE
 	UE::CoreUObject::Private::FreeObjectHandle(Object);
-#endif 	
+#elif UE_WITH_REMOTE_OBJECT_HANDLE
+	UE::RemoteObject::Private::UnmapObject((UObject*)Object);
+#endif
+
 
 	int32 Index = Object->InternalIndex;
 	FUObjectItem* ObjectItem = IndexToObject(Index);
diff --git a/Engine/Source/Runtime/CoreUObject/Private/UObject/UObjectBase.cpp b/Engine/Source/Runtime/CoreUObject/Private/UObject/UObjectBase.cpp
--- a/Engine/Source/Runtime/CoreUObject/Private/UObject/UObjectBase.cpp
+++ b/Engine/Source/Runtime/CoreUObject/Private/UObject/UObjectBase.cpp
@@ -15,6 +15,7 @@
 #include "UObject/UObjectHash.h"
 #include "UObject/Class.h"
 #include "UObject/DeferredRegistry.h"
+#include "UObject/ObjectHandlePrivate.h"
 #include "UObject/UObjectIterator.h"
 #include "UObject/UObjectStats.h"
 #include "UObject/Package.h"
@@ -25,6 +26,7 @@
 #include "Misc/CommandLine.h"
 #include "Interfaces/IPluginManager.h"
 #include "Serialization/LoadTimeTrace.h"
+#include <Atomic>
 
 DEFINE_LOG_CATEGORY_STATIC(LogUObjectBase, Log, All);
 DEFINE_STAT(STAT_UObjectsStatGroupTester);
@@ -105,6 +107,9 @@ static TMap<FName, TArray<FPendingRegistrant*>>& GetPerModuleBootstrapMap()
 UObjectBase::UObjectBase( EObjectFlags InFlags )
 :	ObjectFlags			(InFlags)
 ,	InternalIndex		(INDEX_NONE)
+#if UE_WITH_REMOTE_OBJECT_HANDLE
+,	RemoteId(FRemoteObjectId::Generate())
+#endif // UE_WITH_REMOTE_OBJECT_HANDLE
 ,	ClassPrivate		(nullptr)
 ,	OuterPrivate		(nullptr)
 {
@@ -127,9 +132,13 @@ UObjectBase::UObjectBase(UClass* InClass,
 	UObject *InOuter,
 	FName InName,
 	int32 InInternalIndex,
-	int32 InSerialNumber)
+	int32 InSerialNumber,
+	FRemoteObjectId InRemoteId)
 :	ObjectFlags			(InFlags)
 ,	InternalIndex		(INDEX_NONE)
+#if UE_WITH_REMOTE_OBJECT_HANDLE
+,	RemoteId			(InRemoteId.IsValid() ? InRemoteId : FRemoteObjectId::Generate())
+#endif // UE_WITH_REMOTE_OBJECT_HANDLE
 ,	ClassPrivate		(InClass)
 ,	OuterPrivate		(InOuter)
 {
@@ -273,7 +282,7 @@ UPackage* UObjectBase::GetExternalPackageInternal() const
 void UObjectBase::SetExternalPackage(UPackage* InPackage)
 {
 	// if we have no outer, consider this a package, packages have themselves as their external package and that shouldn't be added to the object hash
-	if (OuterPrivate == nullptr)
+	if (UE::CoreUObject::Private::FObjectHandleUtils::GetNonAccessTrackedOuterNoResolve(this) == nullptr)
 	{
 		// Just validate that we tried to set ourselves or nothing as our external package which is a no-op. anything else is illegal for package
 		check(GetClass()->IsChildOf(UPackage::StaticClass()) && (InPackage == this || InPackage == nullptr));
@@ -353,7 +362,8 @@ bool UObjectBase::IsValidLowLevelFast(bool bRecursive /*= true*/) const
 	}
 
 	// These should all be 0.
-	const UPTRINT CheckZero = (GetFlagsInternal() & ~RF_AllFlags) | ((UPTRINT)ClassPrivate & AlignmentCheck) | ((UPTRINT)OuterPrivate & AlignmentCheck);
+	UObject* Outer = UE::CoreUObject::Private::FObjectHandleUtils::GetNonAccessTrackedOuterNoResolve((UObject*)this);
+	const UPTRINT CheckZero = (GetFlagsInternal() & ~RF_AllFlags) | ((UPTRINT)ClassPrivate & AlignmentCheck) | ((UPTRINT)Outer & AlignmentCheck);
 	if (!!CheckZero)
 	{
 		UE_LOG(LogUObjectBase, Error, TEXT("Object flags are invalid or either Class or Outer is misaligned"));
diff --git a/Engine/Source/Runtime/CoreUObject/Private/UObject/UObjectGlobals.cpp b/Engine/Source/Runtime/CoreUObject/Private/UObject/UObjectGlobals.cpp
--- a/Engine/Source/Runtime/CoreUObject/Private/UObject/UObjectGlobals.cpp
+++ b/Engine/Source/Runtime/CoreUObject/Private/UObject/UObjectGlobals.cpp
@@ -85,6 +85,7 @@
 #include "AutoRTFM/AutoRTFM.h"
 #include "UObject/PropertyOptional.h"
 #include "UObject/VerseValueProperty.h"
+#include "UObject/ObjectHandlePrivate.h"
 
 #include "Interfaces/IPluginManager.h"
 
@@ -244,6 +245,10 @@ FCoreUObjectDelegates::FShouldCookPackageForPlatform FCoreUObjectDelegates::Shou
 FCoreUObjectDelegates::FPackageCreatedForLoad FCoreUObjectDelegates::PackageCreatedForLoad;
 FCoreUObjectDelegates::FGetPrimaryAssetIdForObject FCoreUObjectDelegates::GetPrimaryAssetIdForObject;
 
+TDelegate<void(FRemoteObjectId ObjectId, FRemoteServerId CurrentOwnerServerId, FRemoteServerId DestinationServerId)> FCoreUObjectDelegates::RequestRemoteObjectDelegate;
+TDelegate<void(FRemoteObjectId ObjectId, FRemoteServerId ServerId, FRemoteObjectData& ObjectData)> FCoreUObjectDelegates::RemoteObjectTransferDelegate;
+TDelegate<void()> FCoreUObjectDelegates::FlushRemoteObjectRequestsDelegate;
+
 /** Check whether we should report progress or not */
 bool ShouldReportProgress()
 {
@@ -3386,8 +3391,8 @@ UObject* StaticAllocateObject
 	EInternalObjectFlags InternalSetFlags,
 	bool bCanRecycleSubobjects,
 	bool* bOutRecycledSubobject,
-	UPackage* ExternalPackage
-)
+	UPackage* ExternalPackage,
+	FRemoteObjectId RemoteId)
 {
 	LLM_SCOPE(ELLMTag::UObject);
 
@@ -3594,6 +3599,12 @@ UObject* StaticAllocateObject
 		}
 		else
 		{
+#if UE_WITH_REMOTE_OBJECT_HANDLE
+			if (RemoteId.IsValid() && RemoteId != UE::CoreUObject::Private::FObjectHandleUtils::GetRemoteId(Obj))
+			{
+				UE::CoreUObject::Private::FObjectHandleUtils::ChangeRemoteId(Obj, RemoteId);
+			}
+#endif // UE_WITH_REMOTE_OBJECT_HANDLE
 			bSubObject = true;
 		}
 	}
@@ -3611,7 +3622,7 @@ UObject* StaticAllocateObject
 		UE_AUTORTFM_OPEN
 		{
 			FMemory::Memzero((void *)Obj, TotalSize);
-			new ((void *)Obj) UObjectBase(const_cast<UClass*>(InClass), InFlags|RF_NeedInitialization, InternalSetFlags, InOuter, InName, OldIndex, OldSerialNumber);
+			new ((void *)Obj) UObjectBase(const_cast<UClass*>(InClass), InFlags|RF_NeedInitialization, InternalSetFlags, InOuter, InName, OldIndex, OldSerialNumber, RemoteId);
 		};
 	}
 	else
@@ -4524,7 +4535,7 @@ UObject* StaticConstructObject_Internal(const FStaticConstructObjectParameters&
 	const bool bCanRecycleSubobjects = bIsNativeFromCDO && (!(InFlags & RF_DefaultSubObject) || !FUObjectThreadContext::Get().IsInConstructor);
 
 	bool bRecycledSubobject = false;	
-	Result = StaticAllocateObject(InClass, InOuter, InName, InFlags, Params.InternalSetFlags, bCanRecycleSubobjects, &bRecycledSubobject, Params.ExternalPackage);
+	Result = StaticAllocateObject(InClass, InOuter, InName, InFlags, Params.InternalSetFlags, bCanRecycleSubobjects, &bRecycledSubobject, Params.ExternalPackage, Params.RemoteId);
 	check(Result != nullptr);
 	// Don't call the constructor on recycled subobjects, they haven't been destroyed.
 	if (!bRecycledSubobject)
@@ -4944,7 +4955,7 @@ void CollectArrayReferences(FReferenceCollector& Collector, FArrayProperty& Prop
 		{
 			if (bIsReferenceArray)
 			{
-#if UE_WITH_OBJECT_HANDLE_LATE_RESOLVE
+#if UE_WITH_OBJECT_HANDLE_LATE_RESOLVE || UE_WITH_REMOTE_OBJECT_HANDLE
 				if (InnerProperty.HasAnyPropertyFlags(CPF_TObjectPtr))
 				{
 					Collector.AddReferencedObjects(*reinterpret_cast<TArray<TObjectPtr<UObject>>*>(Instance), Referencer, &Property);
diff --git a/Engine/Source/Runtime/CoreUObject/Private/UObject/UObjectHash.cpp b/Engine/Source/Runtime/CoreUObject/Private/UObject/UObjectHash.cpp
--- a/Engine/Source/Runtime/CoreUObject/Private/UObject/UObjectHash.cpp
+++ b/Engine/Source/Runtime/CoreUObject/Private/UObject/UObjectHash.cpp
@@ -8,6 +8,7 @@
 #include "UObject/Class.h"
 #include "UObject/GarbageCollectionGlobals.h"
 #include "UObject/Package.h"
+#include "UObject/ObjectHandlePrivate.h"
 #include "UObject/UObjectIterator.h"
 #include "UObject/ObjectVisibility.h"
 #include "Templates/Casts.h"
@@ -1046,7 +1047,8 @@ FORCEINLINE static UPackage* AssignExternalPackageToObject(FUObjectHashTables& T
 // Assumes that ThreadHash's critical is already locked
 FORCEINLINE static void RemoveFromOuterMap(FUObjectHashTables& ThreadHash, UObjectBase* Object)
 {
-	FHashBucket& Bucket = ThreadHash.ObjectOuterMap.FindOrAdd(Object->GetOuter());
+	UObject* Outer = UE::CoreUObject::Private::FObjectHandleUtils::GetNonAccessTrackedOuterNoResolve(Object);
+	FHashBucket& Bucket = ThreadHash.ObjectOuterMap.FindOrAdd(Outer);
 	int32 NumRemoved = Bucket.Remove(Object);
 
 	if (!LIKELY(NumRemoved == 1))
@@ -1056,7 +1058,7 @@ FORCEINLINE static void RemoveFromOuterMap(FUObjectHashTables& ThreadHash, UObje
 
 	if (!Bucket.Num())
 	{
-		ThreadHash.ObjectOuterMap.Remove(Object->GetOuter());
+		ThreadHash.ObjectOuterMap.Remove(Outer);
 	}
 }
 
@@ -1648,7 +1650,7 @@ void UnhashObject(UObjectBase* Object)
 			OnHashFailure((UObjectBaseUtility*)Object, TEXT("Hash"), TEXT("remove miscount"));
 		}
 
-		if (PTRINT Outer = (PTRINT)Object->GetOuter())
+		if (PTRINT Outer = (PTRINT)UE::CoreUObject::Private::FObjectHandleUtils::GetNonAccessTrackedOuterNoResolve(Object))
 		{
 			Hash = GetObjectOuterHash(Name, Outer);
 			NumRemoved = ThreadHash.HashOuter.RemoveSingle(Hash, Object->GetUniqueID());
diff --git a/Engine/Source/Runtime/CoreUObject/Public/UObject/Class.h b/Engine/Source/Runtime/CoreUObject/Public/UObject/Class.h
--- a/Engine/Source/Runtime/CoreUObject/Public/UObject/Class.h
+++ b/Engine/Source/Runtime/CoreUObject/Public/UObject/Class.h
@@ -4430,6 +4430,21 @@ template<> struct TBaseStructure<FTopLevelAssetPath>
 	static COREUOBJECT_API UScriptStruct* Get();
 };
 
+template<> struct TBaseStructure<FRemoteObjectId>
+{
+	static COREUOBJECT_API UScriptStruct* Get();
+};
+
+template<> struct TBaseStructure<FRemoteObjectData>
+{
+	static COREUOBJECT_API UScriptStruct* Get();
+};
+
+template<> struct TBaseStructure<FRemoteServerId>
+{
+	static COREUOBJECT_API UScriptStruct* Get();
+};
+
 // TBaseStructure for explicit core variant types only. e.g. FVector3d returns "Vector3d" struct. 
 template< class T > struct TVariantStructure
 {
diff --git a/Engine/Source/Runtime/CoreUObject/Public/UObject/FastReferenceCollector.h b/Engine/Source/Runtime/CoreUObject/Public/UObject/FastReferenceCollector.h
--- a/Engine/Source/Runtime/CoreUObject/Public/UObject/FastReferenceCollector.h
+++ b/Engine/Source/Runtime/CoreUObject/Public/UObject/FastReferenceCollector.h
@@ -744,7 +744,7 @@ struct TDirectDispatcher
 
 	FORCEINLINE void HandleReferenceDirectly(UObject* ReferencingObject, UObject*& Object, FMemberId MemberId, EOrigin Origin, bool bAllowReferenceElimination) const
 	{
-		if (IsObjectHandleResolved(*reinterpret_cast<FObjectHandle*>(&Object)))
+		if (IsObjectHandleResolved_ForGC(*reinterpret_cast<FObjectHandle*>(&Object)))
 		{
 			Processor.HandleTokenStreamObjectReference(Context, ReferencingObject, Object, MemberId, Origin, bAllowReferenceElimination);
 		}
diff --git a/Engine/Source/Runtime/CoreUObject/Public/UObject/NoExportTypes.h b/Engine/Source/Runtime/CoreUObject/Public/UObject/NoExportTypes.h
--- a/Engine/Source/Runtime/CoreUObject/Public/UObject/NoExportTypes.h
+++ b/Engine/Source/Runtime/CoreUObject/Public/UObject/NoExportTypes.h
@@ -3306,6 +3306,33 @@ struct FTemplateString
 	FString Template;
 };
 
+USTRUCT(noexport, BlueprintType)
+struct FRemoteObjectId
+{
+	GENERATED_BODY()
+
+	UPROPERTY()
+	uint64 Id;
+};
+
+USTRUCT(noexport, BlueprintType)
+struct FRemoteObjectData
+{
+	GENERATED_BODY()
+
+	UPROPERTY()
+	TArray<uint8> Bytes;
+};
+
+USTRUCT(noexport, BlueprintType)
+struct FRemoteServerId
+{
+	GENERATED_BODY()
+
+	UPROPERTY()
+	uint32 Id;
+};
+
 /// @endcond
 
 #endif
diff --git a/Engine/Source/Runtime/CoreUObject/Public/UObject/Object.h b/Engine/Source/Runtime/CoreUObject/Public/UObject/Object.h
--- a/Engine/Source/Runtime/CoreUObject/Public/UObject/Object.h
+++ b/Engine/Source/Runtime/CoreUObject/Public/UObject/Object.h
@@ -542,6 +542,13 @@ public:
 		PostDuplicate(DuplicateMode == EDuplicateMode::PIE);
 	}
 
+#if UE_WITH_REMOTE_OBJECT_HANDLE
+	/**
+	 * Called when an object is about to be migrated and after it's been migrated to another process.
+	 */
+	virtual void PostMigrate(FRemoteServerId DestinationServerId) {}
+#endif
+
 	/**
 	 * Called during saving to determine the load flags to save with the object.
 	 * If false, this object will be discarded on clients
diff --git a/Engine/Source/Runtime/CoreUObject/Public/UObject/ObjectHandle.h b/Engine/Source/Runtime/CoreUObject/Public/UObject/ObjectHandle.h
--- a/Engine/Source/Runtime/CoreUObject/Public/UObject/ObjectHandle.h
+++ b/Engine/Source/Runtime/CoreUObject/Public/UObject/ObjectHandle.h
@@ -14,7 +14,7 @@
 #include "UObject/ObjectMacros.h"
 #include "UObject/ObjectRef.h"
 #include "UObject/PackedObjectRef.h"
-
+#include "UObject/RemoteObject.h"
 
 class UClass;
 class UPackage;
@@ -32,6 +32,15 @@ namespace UE::CoreUObject::Private
 
 using FObjectHandle = UE::CoreUObject::Private::FObjectHandlePrivate;
 
+#elif UE_WITH_REMOTE_OBJECT_HANDLE
+
+namespace UE::CoreUObject::Private
+{
+	struct FRemoteObjectHandlePrivate;
+}
+
+using FObjectHandle = UE::CoreUObject::Private::FRemoteObjectHandlePrivate;
+
 #else
 
 using FObjectHandle = UObject*;
@@ -65,6 +74,25 @@ namespace UE::CoreUObject::Private
 	/* Creates and ObjectRef from a packed object ref*/
 	COREUOBJECT_API FObjectRef MakeObjectRef(FPackedObjectRef Handle);
 
+#elif UE_WITH_REMOTE_OBJECT_HANDLE
+
+	struct FRemoteObjectHandlePrivate
+	{
+		UPTRINT PointerOrHandle;
+
+		FRemoteObjectHandlePrivate() = default;
+		COREUOBJECT_API explicit FRemoteObjectHandlePrivate(UObject* Object);
+		COREUOBJECT_API explicit FRemoteObjectHandlePrivate(UE::RemoteObject::Handle::FRemoteObjectStub* RemoteInfo);
+
+		FORCEINLINE const UE::RemoteObject::Handle::FRemoteObjectStub* ToStub() const
+		{
+			return reinterpret_cast<UE::RemoteObject::Handle::FRemoteObjectStub*>(PointerOrHandle & ~UPTRINT(1));
+		}
+
+		COREUOBJECT_API static FRemoteObjectHandlePrivate ConvertToRemoteHandle(UObject* Object);
+		COREUOBJECT_API static FRemoteObjectHandlePrivate FromIdNoResolve(FRemoteObjectId ObjectId);
+	};
+
 #endif
 
 	///these functions are always defined regardless of UE_WITH_OBJECT_HANDLE_LATE_RESOLVE value
@@ -97,6 +125,8 @@ inline bool IsObjectHandleNull(FObjectHandle Handle)
 {
 #if UE_WITH_OBJECT_HANDLE_LATE_RESOLVE
 	return !Handle.PointerOrRef;
+#elif UE_WITH_REMOTE_OBJECT_HANDLE
+	return !Handle.PointerOrHandle;
 #else
 	return !Handle;
 #endif
@@ -110,6 +140,31 @@ inline bool IsObjectHandleResolved(FObjectHandle Handle)
 {
 #if UE_WITH_OBJECT_HANDLE_LATE_RESOLVE
 	return !(Handle.PointerOrRef & 1);
+#elif UE_WITH_REMOTE_OBJECT_HANDLE
+	if (!!(Handle.PointerOrHandle & 1))
+	{
+		return false;
+	}
+	else if (!IsObjectHandleNull(Handle))
+	{
+		return !UE::RemoteObject::Handle::IsRemote(ReadObjectHandlePointerNoCheck(Handle));
+	}
+	return true;
+#else
+	return true;
+#endif
+}
+
+/* checks if a handle is resolved.
+ * nullptr counts as resolved
+ * all handles are resolved when late resolved is off
+ */
+inline bool IsObjectHandleResolved_ForGC(FObjectHandle Handle)
+{
+#if UE_WITH_OBJECT_HANDLE_LATE_RESOLVE
+	return !(Handle.PointerOrRef & 1);
+#elif UE_WITH_REMOTE_OBJECT_HANDLE
+	return !(Handle.PointerOrHandle & 1);
 #else
 	return true;
 #endif
@@ -122,6 +177,8 @@ inline bool IsObjectHandleTypeSafe(FObjectHandle Handle)
 {
 #if UE_WITH_OBJECT_HANDLE_LATE_RESOLVE && UE_WITH_OBJECT_HANDLE_TYPE_SAFETY
 	return !((Handle.PointerOrRef & 3) == 3);
+#elif UE_WITH_REMOTE_OBJECT_HANDLE
+	return true;
 #else
 	return true;
 #endif
@@ -208,6 +265,75 @@ inline uint32 GetTypeHash(UE::CoreUObject::Private::FObjectHandlePrivate Handle)
 	}
 	return GetTypeHash(Handle.PointerOrRef);
 }
+#elif UE_WITH_REMOTE_OBJECT_HANDLE
+inline bool operator==(UE::CoreUObject::Private::FRemoteObjectHandlePrivate LHS, UE::CoreUObject::Private::FRemoteObjectHandlePrivate RHS)
+{
+	using namespace UE::CoreUObject::Private;
+	using namespace UE::RemoteObject::Handle;
+
+	bool bLhsResolved = IsObjectHandleResolved(LHS);
+	bool bRhsResolved = IsObjectHandleResolved(RHS);
+
+	//if both resolved or both unresolved compare the uintptr
+	if (bLhsResolved == bRhsResolved)
+	{
+		return LHS.PointerOrHandle == RHS.PointerOrHandle;
+	}
+
+	//only one side can be resolved
+	if (bLhsResolved)
+	{
+		//both sides can't be null as resolved status would have be true for both
+		const UObject* Obj = ReadObjectHandlePointerNoCheck(LHS);
+		if (!Obj)
+		{
+			return false;
+		}
+
+		const FRemoteObjectStub* RhsStub = RHS.ToStub();
+		return RhsStub->Id == FRemoteObjectId::FromObject(Obj);
+	}
+	else
+	{
+		//both sides can't be null as resolved status would have be true for both
+		const UObject* Obj = ReadObjectHandlePointerNoCheck(RHS);
+		if (!Obj)
+		{
+			return false;
+		}
+
+		const FRemoteObjectStub* LhsStub = LHS.ToStub();
+		return LhsStub->Id == FRemoteObjectId::FromObject(Obj);
+	}
+
+}
+
+inline bool operator!=(UE::CoreUObject::Private::FRemoteObjectHandlePrivate LHS, UE::CoreUObject::Private::FRemoteObjectHandlePrivate RHS)
+{
+	return !(LHS == RHS);
+}
+
+/**
+* GetTypeHash for FRemoteObjectHandlePrivate is guaranteed to return the same value whether a non-null handle is resolved or not (GetTypeHash of FRemoteObjectId)
+*/
+inline uint32 GetTypeHash(UE::CoreUObject::Private::FRemoteObjectHandlePrivate Handle)
+{
+	using namespace UE::CoreUObject::Private;
+	using namespace UE::RemoteObject::Handle;
+
+	if (Handle.PointerOrHandle == 0)
+	{
+		return 0;
+	}
+
+	if (IsObjectHandleResolved(Handle))
+	{
+		const UObject* Obj = ReadObjectHandlePointerNoCheck(Handle);
+		return GetTypeHash(FRemoteObjectId::FromObject(Obj));
+	}
+	const FRemoteObjectStub* Stub = Handle.ToStub();
+	return GetTypeHash(Stub->Id);
+}
 #endif
 
 /////////////////////////////////////////////////////////////////////////////////////////////////////////////////////////
@@ -218,6 +344,8 @@ namespace UE::CoreUObject::Private
 	{
 #if UE_WITH_OBJECT_HANDLE_LATE_RESOLVE
 		{ return { UPTRINT(Object) }; }
+#elif UE_WITH_REMOTE_OBJECT_HANDLE
+		return FObjectHandle(Object);
 #else
 		return Object;
 #endif
@@ -229,6 +357,9 @@ namespace UE::CoreUObject::Private
 		UObject* ResolvedObject = ResolveObjectHandleNoRead(Handle);
 		UE::CoreUObject::Private::OnHandleRead(ResolvedObject);
 		return ResolvedObject;
+#elif UE_WITH_REMOTE_OBJECT_HANDLE
+		UObject* ResolvedObject = ResolveObjectHandleNoRead(Handle);
+		return ResolvedObject;
 #else
 		return ReadObjectHandlePointerNoCheck(Handle);
 #endif
@@ -256,6 +387,17 @@ namespace UE::CoreUObject::Private
 			FObjectRef ObjectRef = MakeObjectRef(PackedObjectRef);
 			return ObjectRef.ResolveObjectRefClass();
 		}
+#elif UE_WITH_REMOTE_OBJECT_HANDLE
+		UObject* Obj = nullptr;
+		if (IsObjectHandleResolved(Handle))
+		{
+			Obj = ReadObjectHandlePointerNoCheck(Handle);
+		}
+		else
+		{
+			Obj = ResolveObjectHandle(Handle);
+		}
+		return Obj != nullptr ? UE::CoreUObject::Private::GetClass(Obj) : nullptr;
 #else
 		UObject* Obj = ReadObjectHandlePointerNoCheck(Handle);
 		return Obj != nullptr ? UE::CoreUObject::Private::GetClass(Obj) : nullptr;
@@ -284,12 +426,44 @@ namespace UE::CoreUObject::Private
 			}
 			return ResolvedObject;
 		}
+#elif UE_WITH_REMOTE_OBJECT_HANDLE
+		FObjectHandle LocalHandle = Handle;
+		UObject* ResolvedObject = nullptr;
+		if (!!(LocalHandle.PointerOrHandle & 1))
+		{
+			ResolvedObject = UE::RemoteObject::Handle::ResolveObject(LocalHandle.ToStub());
+			
+		}
+		else if (!IsObjectHandleNull(Handle) && UE::RemoteObject::Handle::IsRemote(ReadObjectHandlePointerNoCheck(LocalHandle)))
+		{
+			ResolvedObject = UE::RemoteObject::Handle::ResolveObject(ReadObjectHandlePointerNoCheck(LocalHandle));
+		}
+		else
+		{
+			return ReadObjectHandlePointerNoCheck(LocalHandle);
+		}
+		Handle = MakeObjectHandle(ResolvedObject);
+		return ResolvedObject;
 #else
 		return ReadObjectHandlePointerNoCheck(Handle);
 #endif
 	}
 
-	
+
+	inline UObject* NoResolveObjectHandleNoRead(const FObjectHandle& Handle)
+	{
+		FObjectHandle LocalHandle = Handle;
+		if (IsObjectHandleResolved_ForGC(LocalHandle))
+		{
+			UObject* ResolvedObject = ReadObjectHandlePointerNoCheck(LocalHandle);
+			return ResolvedObject;
+		}
+		else
+		{
+			return nullptr;
+		}
+	}
+
 	inline UObject* ResolveObjectHandleNoReadNoCheck(FObjectHandle& Handle)
 	{
 #if UE_WITH_OBJECT_HANDLE_LATE_RESOLVE
@@ -305,6 +479,13 @@ namespace UE::CoreUObject::Private
 			Handle = LocalHandle;
 		}
 		return ResolvedObject;
+#elif UE_WITH_REMOTE_OBJECT_HANDLE
+		FObjectHandle LocalHandle = Handle;
+		check(!!(LocalHandle.PointerOrHandle & 1) || !LocalHandle.PointerOrHandle);
+		UObject* ResolvedObject = UE::RemoteObject::Handle::ResolveObject(LocalHandle.ToStub());
+		LocalHandle = MakeObjectHandle(ResolvedObject);
+		Handle = LocalHandle;
+		return ResolvedObject;	
 #else
 		return ReadObjectHandlePointerNoCheck(Handle);
 #endif
@@ -314,6 +495,8 @@ namespace UE::CoreUObject::Private
 	{
 #if UE_WITH_OBJECT_HANDLE_LATE_RESOLVE
 		return reinterpret_cast<UObject*>(Handle.PointerOrRef);
+#elif UE_WITH_REMOTE_OBJECT_HANDLE
+		return reinterpret_cast<UObject*>(Handle.PointerOrHandle);
 #else
 		return Handle;
 #endif
diff --git a/Engine/Source/Runtime/CoreUObject/Public/UObject/ObjectHandleDefines.h b/Engine/Source/Runtime/CoreUObject/Public/UObject/ObjectHandleDefines.h
--- a/Engine/Source/Runtime/CoreUObject/Public/UObject/ObjectHandleDefines.h
+++ b/Engine/Source/Runtime/CoreUObject/Public/UObject/ObjectHandleDefines.h
@@ -6,16 +6,21 @@
 
 #include "Misc/CoreMiscDefines.h"
 
+#ifndef UE_WITH_REMOTE_OBJECT_HANDLE
+	#define UE_WITH_REMOTE_OBJECT_HANDLE 0
+#endif
+
+
 #ifndef UE_WITH_OBJECT_HANDLE_LATE_RESOLVE
-	#define UE_WITH_OBJECT_HANDLE_LATE_RESOLVE WITH_EDITORONLY_DATA
+	#define UE_WITH_OBJECT_HANDLE_LATE_RESOLVE (!UE_WITH_REMOTE_OBJECT_HANDLE && WITH_EDITORONLY_DATA)
 #endif
 
 
 #ifndef UE_WITH_OBJECT_HANDLE_TRACKING
-	#define UE_WITH_OBJECT_HANDLE_TRACKING WITH_EDITORONLY_DATA
+	#define UE_WITH_OBJECT_HANDLE_TRACKING (!UE_WITH_REMOTE_OBJECT_HANDLE && WITH_EDITORONLY_DATA)
 #endif
 
 
 #ifndef UE_WITH_OBJECT_HANDLE_TYPE_SAFETY
-	#define UE_WITH_OBJECT_HANDLE_TYPE_SAFETY WITH_EDITORONLY_DATA
+	#define UE_WITH_OBJECT_HANDLE_TYPE_SAFETY (!UE_WITH_REMOTE_OBJECT_HANDLE && WITH_EDITORONLY_DATA)
 #endif
\ No newline at end of file
diff --git a/Engine/Source/Runtime/CoreUObject/Public/UObject/ObjectMacros.h b/Engine/Source/Runtime/CoreUObject/Public/UObject/ObjectMacros.h
--- a/Engine/Source/Runtime/CoreUObject/Public/UObject/ObjectMacros.h
+++ b/Engine/Source/Runtime/CoreUObject/Public/UObject/ObjectMacros.h
@@ -618,6 +618,7 @@ enum class EInternalObjectFlags : int32
 	ReachabilityFlag1 = 1 << 1, ///< One of the flags used by Garbage Collector to determine UObject's reachability state
 	ReachabilityFlag2 = 1 << 2, ///< One of the flags used by Garbage Collector to determine UObject's reachability state
 
+	Remote = 1 << 19, //< Object is no longer owned by this process
 	LoaderImport = 1 << 20, ///< Object is ready to be imported by another package during loading
 	Garbage = 1 << 21, ///< Garbage from logical point of view and should not be referenced. This flag is mirrored in EObjectFlags as RF_Garbage for performance
 	AsyncLoadingPhase1 = 1 << 22, ///< Object is being asynchronously loaded.
@@ -640,7 +641,7 @@ ENUM_CLASS_FLAGS(EInternalObjectFlags);
 
 //~ Make sure these macros are up to date!
 #define EInternalObjectFlags_GarbageCollectionKeepFlags (EInternalObjectFlags::Native | EInternalObjectFlags::Async | EInternalObjectFlags::AsyncLoadingPhase1 | EInternalObjectFlags::AsyncLoadingPhase2 | EInternalObjectFlags::LoaderImport | EInternalObjectFlags::RefCounted)
-#define EInternalObjectFlags_AllFlags (EInternalObjectFlags::ReachabilityFlag0 | EInternalObjectFlags::ReachabilityFlag1 | EInternalObjectFlags::ReachabilityFlag2 | EInternalObjectFlags::LoaderImport | EInternalObjectFlags::Garbage | EInternalObjectFlags::ReachableInCluster | EInternalObjectFlags::ClusterRoot | EInternalObjectFlags::Native | EInternalObjectFlags::RefCounted | EInternalObjectFlags::Async | EInternalObjectFlags::AsyncLoadingPhase1 | EInternalObjectFlags::AsyncLoadingPhase2 | EInternalObjectFlags::RootSet | EInternalObjectFlags::PendingConstruction | (EInternalObjectFlags::Unreachable))
+#define EInternalObjectFlags_AllFlags (EInternalObjectFlags::ReachabilityFlag0 | EInternalObjectFlags::ReachabilityFlag1 | EInternalObjectFlags::ReachabilityFlag2 | EInternalObjectFlags::Remote | EInternalObjectFlags::LoaderImport | EInternalObjectFlags::Garbage | EInternalObjectFlags::ReachableInCluster | EInternalObjectFlags::ClusterRoot | EInternalObjectFlags::Native | EInternalObjectFlags::RefCounted | EInternalObjectFlags::Async | EInternalObjectFlags::AsyncLoadingPhase1 | EInternalObjectFlags::AsyncLoadingPhase2 | EInternalObjectFlags::RootSet | EInternalObjectFlags::PendingConstruction | (EInternalObjectFlags::Unreachable))
 #define EInternalObjectFlags_RootFlags (EInternalObjectFlags::RootSet | EInternalObjectFlags_GarbageCollectionKeepFlags)
 #define EInternalObjectFlags_ReachabilityFlags (EInternalObjectFlags::ReachabilityFlag0 | EInternalObjectFlags::ReachabilityFlag1 | EInternalObjectFlags::ReachabilityFlag2 | EInternalObjectFlags::Unreachable)
 #define EInternalObjectFlags_AsyncLoading (EInternalObjectFlags::AsyncLoadingPhase1 | EInternalObjectFlags::AsyncLoadingPhase2)
diff --git a/Engine/Source/Runtime/CoreUObject/Public/UObject/ObjectPtr.h b/Engine/Source/Runtime/CoreUObject/Public/UObject/ObjectPtr.h
--- a/Engine/Source/Runtime/CoreUObject/Public/UObject/ObjectPtr.h
+++ b/Engine/Source/Runtime/CoreUObject/Public/UObject/ObjectPtr.h
@@ -89,7 +89,7 @@ public:
 	{
 	}
 
-#if UE_WITH_OBJECT_HANDLE_LATE_RESOLVE
+#if UE_WITH_OBJECT_HANDLE_LATE_RESOLVE || UE_WITH_REMOTE_OBJECT_HANDLE
 	explicit FORCEINLINE FObjectPtr(FObjectHandle Handle)
 		: Handle(Handle)
 	{
@@ -358,7 +358,7 @@ namespace ObjectPtr_Private
 			return !Other;
 		}
 
-#if UE_WITH_OBJECT_HANDLE_LATE_RESOLVE
+#if UE_WITH_OBJECT_HANDLE_LATE_RESOLVE || UE_WITH_REMOTE_OBJECT_HANDLE
 		if (Ptr.IsResolved())
 		{
 			return Ptr.GetNoResolveNoCheck() == Other;
@@ -773,6 +773,12 @@ namespace ObjectPtr_Private
 		{
 			return reinterpret_cast<T*>(UE::CoreUObject::Private::ResolveObjectHandleNoRead(Ptr.ObjectPtr.GetHandleRef()));
 		}
+
+		template <typename T>
+		FORCEINLINE static T* NoAccessTrackingGetNoResolve(const TObjectPtr<T>& Ptr)
+		{
+			return reinterpret_cast<T*>(UE::CoreUObject::Private::NoResolveObjectHandleNoRead(Ptr.ObjectPtr.GetHandleRef()));
+		}
 	};
 	
 	template <typename T>
@@ -802,6 +808,11 @@ namespace ObjectPtr_Private
 			return ObjectPtr_Private::Friend::NoAccessTrackingGet(ObjectPtr);
 		}
 
+		UE_INTERNAL T* GetNoResolve() const
+		{
+			return ObjectPtr_Private::Friend::NoAccessTrackingGetNoResolve(ObjectPtr);
+		}
+
 		explicit operator UPTRINT() const
 		{
 			return BitCast<UPTRINT>(Get());
diff --git a/Engine/Source/Runtime/CoreUObject/Public/UObject/RemoteObject.h b/Engine/Source/Runtime/CoreUObject/Public/UObject/RemoteObject.h
new file mode 100644
--- /dev/null
+++ b/Engine/Source/Runtime/CoreUObject/Public/UObject/RemoteObject.h
@@ -0,0 +1,59 @@
+// Copyright Epic Games, Inc. All Rights Reserved.
+
+#pragma once
+
+#include "CoreTypes.h"
+#include "UObject/RemoteObjectTypes.h"
+#include "UObject/ObjectMacros.h"
+
+class UObject;
+
+namespace UE::RemoteObject
+{
+	/**
+	* Returns a unique id associated with this process
+	*/
+	COREUOBJECT_API FRemoteServerId GetGlobalServerId();
+
+}
+
+namespace UE::RemoteObject::Handle
+{
+	// Structure that holds basic information about a remote object
+	struct FRemoteObjectStub
+	{
+		FRemoteObjectId Id;
+		FRemoteObjectId OuterId;
+		/** Server id that currently owns the object */
+		FRemoteServerId RemoteServerId;
+
+		FRemoteObjectStub() = default;
+		COREUOBJECT_API explicit FRemoteObjectStub(UObject* Object);
+	};
+
+	/**
+	* Resolves a remote object given its stub
+	* @param Stub Basic data required to migrate a remote object
+	* @return Resolved object
+	*/
+	COREUOBJECT_API UObject* ResolveObject(const FRemoteObjectStub* Stub);
+
+	/**
+	* Resolves a remote object
+	* @param Object Object to resolve (remote object memory that has not yet been GC'd)
+	* @return Resolved object
+	*/
+	COREUOBJECT_API UObject* ResolveObject(UObject* Object);
+
+	/**
+	* Checks if an object associated with the specified unique id is remote
+	* @return True if an object associated with the specified unique id is remote
+	*/
+	COREUOBJECT_API bool IsRemote(FRemoteObjectId ObjectId);
+
+	/**
+	* Checks if an object (memory that has not yet been GC'd) is remote
+	* @return True if the object is remote
+	*/
+	COREUOBJECT_API bool IsRemote(UObject* Object);
+}
diff --git a/Engine/Source/Runtime/CoreUObject/Public/UObject/RemoteObjectTransfer.h b/Engine/Source/Runtime/CoreUObject/Public/UObject/RemoteObjectTransfer.h
new file mode 100644
--- /dev/null
+++ b/Engine/Source/Runtime/CoreUObject/Public/UObject/RemoteObjectTransfer.h
@@ -0,0 +1,55 @@
+// Copyright Epic Games, Inc. All Rights Reserved.
+
+#pragma once
+
+#include "CoreTypes.h"
+#include "Misc/AssertionMacros.h"
+#include "UObject/RemoteObjectTypes.h"
+#include "UObject/Object.h"
+#include "Async/Async.h"
+
+namespace UE::RemoteObject::Transfer
+{
+	using FRemoteObjectDataFuture = TSharedFuture<TUniquePtr<FRemoteObjectData>>;
+
+	/**
+	* Called when remote object data has been received from a remote server
+	* @param ObjectId Remote object id
+	* @param Data Remote object data. Data ownership is transferred to the remote object queue
+	*/
+	COREUOBJECT_API void OnObjectDataReceived(FRemoteObjectId ObjectId, FRemoteObjectData* Data);
+
+	/**
+	* Requests remote object data to be transferred from Current Owner Server to Destination Server
+	* @param ObjectId Remote object id
+	* @param CurrentOwnerServerId Id of a remote server that is known to currently own the object
+	* @param DestinationServerId Id of a remote server that wants to own the object (may be different than the the server that sent the request)
+	* @return Future that will receive the object data
+	*/
+	COREUOBJECT_API FRemoteObjectDataFuture RequestObjectData(FRemoteObjectId ObjectId, FRemoteServerId CurrentOwnerServerId, FRemoteServerId DestinationServerId);
+
+	/**
+	* Migrates an object to a remote server
+	* @param Object Local object to migrate
+	* @param DestinationServerId Id of a remote server to receive the object
+	*/
+	COREUOBJECT_API void MigrateObjectToRemoteServer(UObject* Object, FRemoteServerId DestinationServerId);
+
+	/**
+	* Migrates an object to a remote server
+	* @param ObjectId Remote object id
+	* @param DestinationServerId Id of a remote server to receive the object
+	*/
+	COREUOBJECT_API void MigrateObjectToRemoteServer(FRemoteObjectId Id, FRemoteServerId DestinationServerId);
+
+	/**
+	* Migrates an object from a remote server (temp function)
+	* @param ObjectId Remote object id
+	* @param CurrentOwnerServerId Id of a remote server that currently owns the object
+	* @param DestinationOuter Optional outer the migrated object should be reparented to
+	* @return Migrated object
+	*/
+	COREUOBJECT_API UObject* MigrateObjectFromRemoteServer(FRemoteObjectId ObjectId, FRemoteServerId CurrentOwnerServerId, UObject* DestinationOuter = nullptr);
+
+} // namespace UE::RemoveObject::Transfer
+
diff --git a/Engine/Source/Runtime/CoreUObject/Public/UObject/RemoteObjectTypes.h b/Engine/Source/Runtime/CoreUObject/Public/UObject/RemoteObjectTypes.h
new file mode 100644
--- /dev/null
+++ b/Engine/Source/Runtime/CoreUObject/Public/UObject/RemoteObjectTypes.h
@@ -0,0 +1,125 @@
+// Copyright Epic Games, Inc. All Rights Reserved.
+
+#pragma once
+
+#include "CoreTypes.h"
+#include "Templates/TypeHash.h"
+#include "Containers/Array.h"
+#include "Containers/UnrealString.h"
+
+class UObject;
+class FArchive;
+
+struct FRemoteServerId
+{
+	friend struct FRemoteObjectId;
+
+	FRemoteServerId() = default;
+	explicit FRemoteServerId(uint32 InId)
+		: Id(InId)
+	{
+		checkf(Id < (1 << 10), TEXT("Remote server id can only be a 10-bit value, got: %u"), Id);
+	}
+	FString ToString() const
+	{
+		return FString::FromInt(Id);
+	}
+	bool IsValid() const
+	{
+		return Id != 0;
+	}
+	bool operator == (FRemoteServerId Other) const
+	{
+		return Id == Other.Id;
+	}
+	bool operator != (FRemoteServerId Other) const
+	{
+		return Id != Other.Id;
+	}
+	bool operator < (FRemoteServerId Other) const
+	{
+		return Id < Other.Id;
+	}
+
+private:
+	uint32 Id = 0;
+};
+
+struct FRemoteObjectId
+{
+private:
+	union
+	{
+		struct
+		{
+			uint64 ServerId : 10;
+			uint64 SerialNumber : 54;
+		};
+		uint64 Id;
+	};
+
+public:
+
+	FRemoteObjectId()
+		: ServerId(0)
+		, SerialNumber(0)
+	{
+	}
+
+	explicit FRemoteObjectId(uint64 InId)
+		: Id(InId)
+	{
+
+	}
+
+	FRemoteObjectId(FRemoteServerId InServerId, uint64 InSerialNumber)
+		: ServerId(InServerId.Id)
+		, SerialNumber(InSerialNumber)
+	{
+	}
+	explicit FRemoteObjectId(ENoInit)
+	{
+	}
+	FORCEINLINE friend uint32 GetTypeHash(const FRemoteObjectId& ObjectId)
+	{
+		return GetTypeHash(ObjectId.Id);
+	}
+
+	FORCEINLINE bool operator==(const FRemoteObjectId& Other) const
+	{
+		return Id == Other.Id;
+	}
+	FORCEINLINE bool operator!=(const FRemoteObjectId& Other) const
+	{
+		return Id != Other.Id;
+	}
+	FORCEINLINE bool IsValid() const
+	{
+		return *this != FRemoteObjectId();
+	}
+
+	FORCEINLINE FRemoteServerId GetServerId() const
+	{
+		return FRemoteServerId(ServerId);
+	}
+
+	FORCEINLINE FString ToString() const
+	{
+		return FString::Printf(TEXT("%llu-%llu"), ServerId, SerialNumber);
+	}
+
+	COREUOBJECT_API static FRemoteObjectId Generate();
+
+	COREUOBJECT_API static FRemoteObjectId FromObject(const UObject* Object);
+
+	friend COREUOBJECT_API FArchive& operator<<(FArchive& Ar, FRemoteObjectId& Id);
+};
+
+
+/**
+* Structure that holds remote object memory
+*/
+struct FRemoteObjectData
+{
+	TArray<uint8> Bytes;
+};
\ No newline at end of file
diff --git a/Engine/Source/Runtime/CoreUObject/Public/UObject/Stack.h b/Engine/Source/Runtime/CoreUObject/Public/UObject/Stack.h
--- a/Engine/Source/Runtime/CoreUObject/Public/UObject/Stack.h
+++ b/Engine/Source/Runtime/CoreUObject/Public/UObject/Stack.h
@@ -400,7 +400,7 @@ inline UObject* FFrame::ReadObject()
 {
 	UObject* Result = (UObject*) ReadPointer();
 
-#if UE_WITH_OBJECT_HANDLE_LATE_RESOLVE
+#if UE_WITH_OBJECT_HANDLE_LATE_RESOLVE || UE_WITH_REMOTE_OBJECT_HANDLE
 	TObjectPtr<UObject> ObjPtr(Result);
 	return ObjPtr.Get();
 #else
diff --git a/Engine/Source/Runtime/CoreUObject/Public/UObject/UObjectArray.h b/Engine/Source/Runtime/CoreUObject/Public/UObject/UObjectArray.h
--- a/Engine/Source/Runtime/CoreUObject/Public/UObject/UObjectArray.h
+++ b/Engine/Source/Runtime/CoreUObject/Public/UObject/UObjectArray.h
@@ -791,7 +791,7 @@ public:
 class FUObjectArray
 {
 	friend class UObject;
-	friend COREUOBJECT_API UObject* StaticAllocateObject(const UClass*, UObject*, FName, EObjectFlags, EInternalObjectFlags, bool, bool*, UPackage*);
+	friend COREUOBJECT_API UObject* StaticAllocateObject(const UClass*, UObject*, FName, EObjectFlags, EInternalObjectFlags, bool, bool*, UPackage*, FRemoteObjectId);
 
 private:
 	/**
diff --git a/Engine/Source/Runtime/CoreUObject/Public/UObject/UObjectBase.h b/Engine/Source/Runtime/CoreUObject/Public/UObject/UObjectBase.h
--- a/Engine/Source/Runtime/CoreUObject/Public/UObject/UObjectBase.h
+++ b/Engine/Source/Runtime/CoreUObject/Public/UObject/UObjectBase.h
@@ -15,6 +15,7 @@
 #include "Stats/Stats.h"
 #include "Stats/Stats2.h"
 #include "UObject/NameTypes.h"
+#include "UObject/RemoteObjectTypes.h"
 #include "UObject/ObjectMacros.h"
 #include "UObject/UObjectGlobals.h"
 #include "UObject/UnrealNames.h"
@@ -27,6 +28,11 @@ class UObject;
 class UPackage;
 class UScriptStruct;
 
+namespace UE::CoreUObject::Private
+{
+	struct FObjectHandleUtils;
+}
+
 // If FName is 4 bytes than we can use padding after it to store internal object list index and use array instead of a hash map for lookup in UObjectHash.cpp
 // This might change the each UClass' object list iteration order
 #if !defined(UE_STORE_OBJECT_LIST_INTERNAL_INDEX)
@@ -45,6 +51,7 @@ class UObjectBase
 	friend class FUObjectArray; // for access to InternalIndex without revealing it to anyone else
 	friend class FUObjectAllocator; // for access to destructor without revealing it to anyone else
 	friend struct FInternalUObjectBaseUtilityIsValidFlagsChecker; // for access to InternalIndex
+	friend struct UE::CoreUObject::Private::FObjectHandleUtils; // for access to OuterPrivate
 	friend COREUOBJECT_API void UObjectForceRegistration(UObjectBase* Object, bool bCheckForModuleRelease);
 	friend COREUOBJECT_API void InitializePrivateStaticClass(
 		class UClass* (*TClass_StaticClassFn)(),
@@ -56,6 +63,9 @@ class UObjectBase
 		);
 protected:
 	UObjectBase() :
+#if UE_WITH_REMOTE_OBJECT_HANDLE
+		RemoteId(NoInit),
+#endif // UE_WITH_REMOTE_OBJECT_HANDLE
 		ClassPrivate(NoInit),
 		NamePrivate(NoInit),  // screwy, but the name was already set and we don't want to set it again
 		OuterPrivate(NoInit)
@@ -85,7 +95,8 @@ public:
 			UObject *InOuter,
 			FName InName,
 			int32 InInternalIndex = -1,
-			int32 InSerialNumber = 0);
+			int32 InSerialNumber = 0,
+			FRemoteObjectId InRemoteId = FRemoteObjectId());
 
 	/**
 	 * Final destructor, removes the object from the object array, and indirectly, from any annotations
@@ -316,6 +327,11 @@ private:
 	/** Index into GObjectArray...very private. */
 	int32							InternalIndex;
 
+#if UE_WITH_REMOTE_OBJECT_HANDLE
+	/** Unique Id associated with this object */
+	FRemoteObjectId					RemoteId;
+#endif // UE_WITH_REMOTE_OBJECT_HANDLE
+
 	/** Class the object belongs to. */
 	ObjectPtr_Private::TNonAccessTrackedObjectPtr<UClass>							ClassPrivate;
 
diff --git a/Engine/Source/Runtime/CoreUObject/Public/UObject/UObjectGlobals.h b/Engine/Source/Runtime/CoreUObject/Public/UObject/UObjectGlobals.h
--- a/Engine/Source/Runtime/CoreUObject/Public/UObject/UObjectGlobals.h
+++ b/Engine/Source/Runtime/CoreUObject/Public/UObject/UObjectGlobals.h
@@ -40,9 +40,11 @@
 #include "Traits/IsCharEncodingCompatibleWith.h"
 #include "UObject/NameTypes.h"
 #include "UObject/ObjectMacros.h"
+#include "UObject/ObjectHandle.h"
 #include "UObject/PrimaryAssetId.h"
 #include "UObject/Script.h"
 #include "UObject/TopLevelAssetPath.h"
+#include "UObject/RemoteObjectTypes.h"
 #include "UObject/UnrealNames.h"
 #include "VerseTypesFwd.h"
 
@@ -1167,9 +1169,10 @@ bool StaticAllocateObjectErrorTests( const UClass* Class, UObject* InOuter, FNam
  * @param bCanReuseSubobjects	if set to true, SAO will not attempt to destroy a subobject if it already exists in memory.
  * @param bOutReusedSubobject	flag indicating if the object is a subobject that has already been created (in which case further initialization is not necessary).
  * @param ExternalPackage	External Package assigned to the allocated object, if any	
+ * @param RemoteId		Globally unique id for this object
  * @return	a pointer to a fully initialized object of the specified class.
  */
-COREUOBJECT_API UObject* StaticAllocateObject(const UClass* Class, UObject* InOuter, FName Name, EObjectFlags SetFlags, EInternalObjectFlags InternalSetFlags = EInternalObjectFlags::None, bool bCanReuseSubobjects = false, bool* bOutReusedSubobject = nullptr, UPackage* ExternalPackage = nullptr);
+COREUOBJECT_API UObject* StaticAllocateObject(const UClass* Class, UObject* InOuter, FName Name, EObjectFlags SetFlags, EInternalObjectFlags InternalSetFlags = EInternalObjectFlags::None, bool bCanReuseSubobjects = false, bool* bOutReusedSubobject = nullptr, UPackage* ExternalPackage = nullptr, FRemoteObjectId RemoteId = FRemoteObjectId());
 
 /** FObjectInitializer options */
 enum class EObjectInitializerOptions
@@ -1683,6 +1686,8 @@ struct FStaticConstructObjectParameters
 	/** Callback for custom code to initialize properties before PostInitProperties runs */
 	TFunction<void()> PropertyInitCallback;
 
+	/** Existing remote id for the object */
+	FRemoteObjectId RemoteId;
 private:
 	FObjectInitializer::FOverrides* SubobjectOverrides = nullptr;
 
@@ -1762,12 +1767,13 @@ COREUOBJECT_API void CheckIsClassChildOf_Internal(const UClass* Parent, const UC
  * @param	bCopyTransientsFromClassDefaults	if true, copy transient from the class defaults instead of the pass in archetype ptr (often these are the same)
  * @param	InInstanceGraph						contains the mappings of instanced objects and components to their templates
  * @param	ExternalPackage						Assign an external Package to the created object if non-null
+ * @param 	RemoteId	Globally unique id for this object
  *
  * @return	a pointer of type T to a new object of the specified class
  */
 template< class T >
 FUNCTION_NON_NULL_RETURN_START
-	T* NewObject(UObject* Outer, const UClass* Class, FName Name = NAME_None, EObjectFlags Flags = RF_NoFlags, UObject* Template = nullptr, bool bCopyTransientsFromClassDefaults = false, FObjectInstancingGraph* InInstanceGraph = nullptr, UPackage* ExternalPackage = nullptr)
+	T* NewObject(UObject* Outer, const UClass* Class, FName Name = NAME_None, EObjectFlags Flags = RF_NoFlags, UObject* Template = nullptr, bool bCopyTransientsFromClassDefaults = false, FObjectInstancingGraph* InInstanceGraph = nullptr, UPackage* ExternalPackage = nullptr, FRemoteObjectId RemoteId = FRemoteObjectId())
 FUNCTION_NON_NULL_RETURN_END
 {
 	if (Name == NAME_None)
@@ -1788,6 +1794,7 @@ FUNCTION_NON_NULL_RETURN_END
 	Params.bCopyTransientsFromClassDefaults = bCopyTransientsFromClassDefaults;
 	Params.InstanceGraph = InInstanceGraph;
 	Params.ExternalPackage = ExternalPackage;
+	Params.RemoteId = RemoteId;
 
 	T* Result = static_cast<T*>(StaticConstructObject_Internal(Params));
 	return Result;
@@ -2539,7 +2546,7 @@ public:
 	template<class UObjectType>
 	void AddReferencedObject(TObjectPtr<UObjectType>& Object, const UObject* ReferencingObject = nullptr, const FProperty* ReferencingProperty = nullptr)
 	{
-		if (Object.IsResolved())
+		if (IsObjectHandleResolved_ForGC(*reinterpret_cast<FObjectHandle*>(&Object)))
 		{
 			HandleObjectReference(*reinterpret_cast<UObject**>(&Object), ReferencingObject, ReferencingProperty);
 		}
@@ -3283,6 +3290,15 @@ struct FCoreUObjectDelegates
 	/** Called during cooking to see if a specific package should be cooked for a given target platform */
 	DECLARE_DELEGATE_RetVal_TwoParams(bool, FShouldCookPackageForPlatform, const UPackage*, const ITargetPlatform*);
 	static COREUOBJECT_API FShouldCookPackageForPlatform ShouldCookPackageForPlatform;
+
+	/** Delegate that transfers object data to another server */
+	static COREUOBJECT_API TDelegate<void(FRemoteObjectId /*ObjectId*/, FRemoteServerId /*DestinationServerId*/, FRemoteObjectData& /*ObjectData*/)> RemoteObjectTransferDelegate;
+
+	/** Delegate that requests remote object data from CurrentOwnerServerId to be transferred to DestinationServerId */
+	static COREUOBJECT_API TDelegate<void(FRemoteObjectId /*ObjectId*/, FRemoteServerId /*CurrentOwnerServerId*/, FRemoteServerId /*DestinationServerId*/)> RequestRemoteObjectDelegate;
+
+	/** Delegate that flushes remote object requests */
+	static COREUOBJECT_API TDelegate<void()> FlushRemoteObjectRequestsDelegate;
 };
 
 /** Allows release builds to override not verifying GC assumptions. Useful for profiling as it's hitchy. */
diff --git a/Engine/Source/Runtime/CoreUObject/Public/UObject/UObjectIterator.h b/Engine/Source/Runtime/CoreUObject/Public/UObject/UObjectIterator.h
--- a/Engine/Source/Runtime/CoreUObject/Public/UObject/UObjectIterator.h
+++ b/Engine/Source/Runtime/CoreUObject/Public/UObject/UObjectIterator.h
@@ -58,7 +58,7 @@ enum class EObjectIteratorThreadSafetyOptions : uint8
 
 inline EInternalObjectFlags GetObjectIteratorDefaultInternalExclusionFlags(EInternalObjectFlags InternalExclusionFlags)
 {
-	return InternalExclusionFlags | EInternalObjectFlags::Unreachable | EInternalObjectFlags::PendingConstruction | UE::GetAsyncLoadingInternalFlagsExclusion();
+	return InternalExclusionFlags | EInternalObjectFlags::Unreachable | EInternalObjectFlags::PendingConstruction | EInternalObjectFlags::Remote | UE::GetAsyncLoadingInternalFlagsExclusion();
 }
 
 /**
diff --git a/Engine/Source/Runtime/Engine/Classes/GameFramework/Actor.h b/Engine/Source/Runtime/Engine/Classes/GameFramework/Actor.h
--- a/Engine/Source/Runtime/Engine/Classes/GameFramework/Actor.h
+++ b/Engine/Source/Runtime/Engine/Classes/GameFramework/Actor.h
@@ -2284,6 +2284,9 @@ public:
 	ENGINE_API virtual void Serialize(FArchive& Ar) override;
 	ENGINE_API virtual void PostLoad() override;
 	ENGINE_API virtual void PostLoadSubobjects( FObjectInstancingGraph* OuterInstanceGraph ) override;
+#if UE_WITH_REMOTE_OBJECT_HANDLE
+	ENGINE_API virtual void PostMigrate(FRemoteServerId DestinationServerId) override;
+#endif
 	ENGINE_API virtual void BeginDestroy() override;
 #if !UE_STRIP_DEPRECATED_PROPERTIES
 	ENGINE_API virtual bool IsReadyForFinishDestroy() override;
diff --git a/Engine/Source/Runtime/Engine/Private/Actor.cpp b/Engine/Source/Runtime/Engine/Private/Actor.cpp
--- a/Engine/Source/Runtime/Engine/Private/Actor.cpp
+++ b/Engine/Source/Runtime/Engine/Private/Actor.cpp
@@ -796,12 +796,18 @@ void AActor::SetTickableWhenPaused(bool bTickableWhenPaused)
 
 void AActor::BeginDestroy()
 {
-	UnregisterAllComponents();
-	ULevel* Level = GetLevel();
-	if (Level && !Level->IsUnreachable())
+#if UE_WITH_REMOTE_OBJECT_HANDLE
+	// Remote actors have already gone through the pcoesses of removing from world and by repeating it we'd migrate them back
+	if (!HasAnyInternalFlags(EInternalObjectFlags::Remote))
 	{
-		Level->Actors.RemoveSingleSwap(this, EAllowShrinking::No);
+		UnregisterAllComponents();
+		ULevel* Level = GetLevel();
+		if (Level && !Level->IsUnreachable())
+		{
+			Level->Actors.RemoveSingleSwap(this, EAllowShrinking::No);
+		}
 	}
+#endif // UE_WITH_REMOTE_OBJECT_HANDLE
 
 #if (CSV_PROFILER_STATS && !UE_BUILD_SHIPPING)
 	// Decrement actor class count
@@ -1130,6 +1136,28 @@ void AActor::PostLoadSubobjects(FObjectInstancingGraph* OuterInstanceGraph)
 	}
 }
 
+#if UE_WITH_REMOTE_OBJECT_HANDLE
+void AActor::PostMigrate(FRemoteServerId DestinationServerId)
+{
+	Super::PostMigrate(DestinationServerId);
+
+	TGuardValue<ITransaction*> SuppressTransaction(GUndo, nullptr);
+	if (DestinationServerId != UE::RemoteObject::GetGlobalServerId())
+	{
+		// Remove the actor from the world on this server as it is about to be migrated to another server.
+		// Migrated actors are unique to the server instance that currently owns their memory and are not replicated to other servers.
+		Destroy(false, false);
+	}
+	else
+	{
+		// Add the actor that was migrated from another server to this server's world
+		GetLevel()->TryAddActorToList(this, /*bAddUnique*/false);
+		PostSpawnInitialize(GetTransform(), GetOwner(), nullptr, false, true, false);
+		GetWorld()->AddNetworkActor(this);
+	}
+}
+#endif // UE_WITH_REMOTE_OBJECT_HANDLE
+
 void AActor::ProcessEvent(UFunction* Function, void* Parameters)
 {
 	LLM_SCOPE(ELLMTag::EngineMisc);
