Fixing CIS errors

diff --git a/Engine/Source/Runtime/CoreUObject/Private/UObject/RemoteObject.cpp b/Engine/Source/Runtime/CoreUObject/Private/UObject/RemoteObject.cpp
--- a/Engine/Source/Runtime/CoreUObject/Private/UObject/RemoteObject.cpp
+++ b/Engine/Source/Runtime/CoreUObject/Private/UObject/RemoteObject.cpp
@@ -13,6 +13,7 @@
 #include "UObject/ObjectHandlePrivate.h"
 #include "Misc/CommandLine.h"
 #include "Misc/CoreDelegates.h"
+#include <atomic>
 
 namespace UE::RemoteObject::Private
 {
diff --git a/Engine/Source/Runtime/CoreUObject/Private/UObject/UObjectBase.cpp b/Engine/Source/Runtime/CoreUObject/Private/UObject/UObjectBase.cpp
--- a/Engine/Source/Runtime/CoreUObject/Private/UObject/UObjectBase.cpp
+++ b/Engine/Source/Runtime/CoreUObject/Private/UObject/UObjectBase.cpp
@@ -26,7 +26,6 @@
 #include "Misc/CommandLine.h"
 #include "Interfaces/IPluginManager.h"
 #include "Serialization/LoadTimeTrace.h"
-#include <Atomic>
 
 DEFINE_LOG_CATEGORY_STATIC(LogUObjectBase, Log, All);
 DEFINE_STAT(STAT_UObjectsStatGroupTester);
