Fixing SA error

diff --git a/Engine/Source/Runtime/CoreUObject/Private/UObject/RemoteObjectSerialization.cpp b/Engine/Source/Runtime/CoreUObject/Private/UObject/RemoteObjectSerialization.cpp
--- a/Engine/Source/Runtime/CoreUObject/Private/UObject/RemoteObjectSerialization.cpp
+++ b/Engine/Source/Runtime/CoreUObject/Private/UObject/RemoteObjectSerialization.cpp
@@ -129,7 +129,7 @@ public:
 		ERemoteReferenceType RefType = WriteObjectReference(Obj);
 		if (RefType == ERemoteReferenceType::IdOnly)
 		{
-			if (Obj->IsIn(RootObject))
+			if (Obj && Obj->IsIn(RootObject))
 			{
 				ObjectsToSerialize.Add(Obj);
 			}			
