Minimal support for staging/cooking/loading localized asset variants using both "en-US" and "en_US" style culture names

This allows both styles to be used for L10N asset paths, ie, /MyPlugin/L10N/en-US/Foo and /MyPlugin/L10N/en_US/Foo are now equivalent, with the latter using a valid Verse identifier.

The rest of the tooling around creating/managing localized assets still needs updating, but this unblocks being unable to create localized assets in environments that enforce Verse identifier restrictions on folder names.

#rnx

diff --git a/Engine/Source/Editor/UnrealEd/Private/CookOnTheFlyServer.cpp b/Engine/Source/Editor/UnrealEd/Private/CookOnTheFlyServer.cpp
--- a/Engine/Source/Editor/UnrealEd/Private/CookOnTheFlyServer.cpp
+++ b/Engine/Source/Editor/UnrealEd/Private/CookOnTheFlyServer.cpp
@@ -13257,8 +13257,13 @@ void UCookOnTheFlyServer::GenerateLocalizationReferences()
 	{
 		for (const FString& CultureName : CookByTheBookOptions->AllCulturesToCook)
 		{
-			FString LocalizedPackagePath = RootPath / TEXT("L10N") / CultureName;
-			Filter.PackagePaths.Add(*LocalizedPackagePath);
+			// Cook both UE style (eg, "en-US") and Verse style (eg, "en_US") localized assets
+			const FString VerseIdentifier = FCulture::CultureNameToVerseIdentifier(CultureName);
+			if (CultureName != VerseIdentifier)
+			{
+				Filter.PackagePaths.Add(*(RootPath / TEXT("L10N") / VerseIdentifier));
+			}
+			Filter.PackagePaths.Add(*(RootPath / TEXT("L10N") / CultureName));
 		}
 	}
 
diff --git a/Engine/Source/Runtime/CoreUObject/Private/Internationalization/PackageLocalizationCache.cpp b/Engine/Source/Runtime/CoreUObject/Private/Internationalization/PackageLocalizationCache.cpp
--- a/Engine/Source/Runtime/CoreUObject/Private/Internationalization/PackageLocalizationCache.cpp
+++ b/Engine/Source/Runtime/CoreUObject/Private/Internationalization/PackageLocalizationCache.cpp
@@ -13,6 +13,17 @@ FPackageLocalizationCultureCache::FPackageLocalizationCultureCache(FPackageLocal
 	: OwnerCache(InOwnerCache)
 {
 	PrioritizedCultureNames = FInternationalization::Get().GetPrioritizedCultureNames(InCultureName);
+
+	// Query both UE style (eg, "en-US") and Verse style (eg, "en_US") localized assets
+	for (const FString& PrioritizedCultureName : PrioritizedCultureNames)
+	{
+		FString VerseIdentifier = FCulture::CultureNameToVerseIdentifier(PrioritizedCultureName);
+		if (PrioritizedCultureName != VerseIdentifier)
+		{
+			PrioritizedCultureNamesAndVerseIdentifiers.Add(MoveTemp(VerseIdentifier));
+		}
+		PrioritizedCultureNamesAndVerseIdentifiers.Add(PrioritizedCultureName);
+	}
 }
 
 void FPackageLocalizationCultureCache::ConditionalUpdateCache()
@@ -35,7 +46,7 @@ void FPackageLocalizationCultureCache::ConditionalUpdateCache_NoLock()
 	for (const FString& SourceRootPath : PendingSourceRootPathsToSearch)
 	{
 		TArray<FString>& LocalizedRootPaths = SourcePathsToLocalizedPaths.FindOrAdd(SourceRootPath);
-		for (const FString& PrioritizedCultureName : PrioritizedCultureNames)
+		for (const FString& PrioritizedCultureName : PrioritizedCultureNamesAndVerseIdentifiers)
 		{
 			const FString LocalizedRootPath = SourceRootPath / TEXT("L10N") / PrioritizedCultureName;
 			if (!LocalizedRootPaths.Contains(LocalizedRootPath))
diff --git a/Engine/Source/Runtime/CoreUObject/Private/Internationalization/PackageLocalizationManager.cpp b/Engine/Source/Runtime/CoreUObject/Private/Internationalization/PackageLocalizationManager.cpp
--- a/Engine/Source/Runtime/CoreUObject/Private/Internationalization/PackageLocalizationManager.cpp
+++ b/Engine/Source/Runtime/CoreUObject/Private/Internationalization/PackageLocalizationManager.cpp
@@ -149,10 +149,21 @@ FName FPackageLocalizationManager::FindLocalizedPackageNameNoCache(const FName I
 	const TArray<FString> PrioritizedCultureNames = FInternationalization::Get().GetPrioritizedCultureNames(InCultureName);
 	for (const FString& PrioritizedCultureName : PrioritizedCultureNames)
 	{
-		const FString LocalizedPackageName = PackageNameRoot / TEXT("L10N") / PrioritizedCultureName / PackageNameSubPath;
-		if (FPackageName::DoesPackageExist(LocalizedPackageName))
+		// Query both UE style (eg, "en-US") and Verse style (eg, "en_US") localized assets
+		const FString VerseIdentifier = FCulture::CultureNameToVerseIdentifier(PrioritizedCultureName);
+		if (PrioritizedCultureName != VerseIdentifier)
 		{
-			return *LocalizedPackageName;
+			const FString LocalizedPackageNameForVerseIdentifier = PackageNameRoot / TEXT("L10N") / VerseIdentifier / PackageNameSubPath;
+			if (FPackageName::DoesPackageExist(LocalizedPackageNameForVerseIdentifier))
+			{
+				return *LocalizedPackageNameForVerseIdentifier;
+			}
+		}
+
+		const FString LocalizedPackageNameForPrioritizedCulture = PackageNameRoot / TEXT("L10N") / PrioritizedCultureName / PackageNameSubPath;
+		if (FPackageName::DoesPackageExist(LocalizedPackageNameForPrioritizedCulture))
+		{
+			return *LocalizedPackageNameForPrioritizedCulture;
 		}
 	}
 
diff --git a/Engine/Source/Runtime/CoreUObject/Public/Internationalization/PackageLocalizationCache.h b/Engine/Source/Runtime/CoreUObject/Public/Internationalization/PackageLocalizationCache.h
--- a/Engine/Source/Runtime/CoreUObject/Public/Internationalization/PackageLocalizationCache.h
+++ b/Engine/Source/Runtime/CoreUObject/Public/Internationalization/PackageLocalizationCache.h
@@ -94,6 +94,9 @@ private:
 	/** An array of culture names that should be scanned, sorted in priority order. */
 	TArray<FString> PrioritizedCultureNames;
 
+	/** An array of culture names and Verse identifiers that should be scanned, sorted in priority order. */
+	TArray<FString> PrioritizedCultureNamesAndVerseIdentifiers;
+
 	/** An array of source paths we should scan on the next call to ConditionalUpdateCache. */
 	TArray<FString> PendingSourceRootPathsToSearch;
 
