[ZenLoader]
 - Fix ref counting issue causing use-after-free when ALT is active

#rnx

diff --git a/Engine/Source/Runtime/CoreUObject/Private/Serialization/AsyncLoading2.cpp b/Engine/Source/Runtime/CoreUObject/Private/Serialization/AsyncLoading2.cpp
--- a/Engine/Source/Runtime/CoreUObject/Private/Serialization/AsyncLoading2.cpp
+++ b/Engine/Source/Runtime/CoreUObject/Private/Serialization/AsyncLoading2.cpp
@@ -5524,13 +5524,24 @@ void FAsyncPackage2::ImportPackagesRecursiveInner(FAsyncLoadingThreadState2& Thr
 			TEXT("Start loading imported package with id '%s'"), *FormatPackageId(ImportedPackageId));
 			++AsyncLoadingThread.PackagesWithRemainingWorkCounter;
 			TRACE_COUNTER_SET(AsyncLoadingPackagesWithRemainingWork, AsyncLoadingThread.PackagesWithRemainingWorkCounter);
+
+			// This should never fail since we just did the insert.
+			ImportedPackage->AddRef();
 		}
 		else
 		{
 			UE_ASYNC_PACKAGE_LOG_VERBOSE(VeryVerbose, Desc, TEXT("ImportPackages: UpdatePackage"),
 				TEXT("Imported package with id '%s' is already being loaded."), *FormatPackageId(ImportedPackageId));
+
+			// When using ALT, we could end up with a package that is finishing its loading on the GT and the last ref released.
+			// In such a case, we need to avoid taking a ref otherwise we would end up with a use-after-free.
+			if (!ImportedPackage->TryAddRef())
+			{
+				ensureMsgf(bIsFullyLoaded, TEXT("Found a package being destructed that is not marked as fully loaded"));
+				continue;
+			}
 		}
-		ImportedPackage->AddRef();
+
 		Header.ImportedAsyncPackagesView[LocalImportedPackageIndex] = ImportedPackage;
 
 		if (bIsZenPackage != bIsZenPackageImport)
