Remote Object cleanup:

FRemoteObjectId is now a literal type
FRemoteObjectId::FromObject is now explicit FRemoteObjectid(const UObjectBase*)
Added FRemoteObjectId::operator<
Removed unused variables / unnecessary code from a few functions

diff --git a/Engine/Source/Runtime/CoreUObject/Private/UObject/RemoteObject.cpp b/Engine/Source/Runtime/CoreUObject/Private/UObject/RemoteObject.cpp
--- a/Engine/Source/Runtime/CoreUObject/Private/UObject/RemoteObject.cpp
+++ b/Engine/Source/Runtime/CoreUObject/Private/UObject/RemoteObject.cpp
@@ -221,14 +221,6 @@ UObject* ResolveObject(const FRemoteObjectStub* Stub)
 		UObject* Outer = Object ? Object->GetOuter() : ObjectMaps->FindObject(Stub->OuterId);
 		Object = MigrateObjectFromRemoteServer(Stub->Id, Stub->RemoteServerId, Outer);
 	}
-	if (Object)
-	{
-		MarkAsLocal(Object);
-	}
-	else
-	{
-		Object = nullptr;
-	}
 	return Object;
 }
 
@@ -253,23 +245,6 @@ namespace UE::CoreUObject::Private
 		MapObject((UObject*)Object);
 	}
 
-	FRemoteObjectHandlePrivate::FRemoteObjectHandlePrivate(UObject* Object)
-	{
-		if (Object)
-		{
-			PointerOrHandle = UPTRINT(Object);
-		}
-		else
-		{
-			PointerOrHandle = 0;
-		}
-	}
-
-	FRemoteObjectHandlePrivate::FRemoteObjectHandlePrivate(UE::RemoteObject::Handle::FRemoteObjectStub* RemoteInfo)
-	{
-		PointerOrHandle = UPTRINT(RemoteInfo) | 1;
-	}
-
 	FRemoteObjectHandlePrivate FRemoteObjectHandlePrivate::ConvertToRemoteHandle(UObject* Object)
 	{
 		using namespace UE::RemoteObject::Handle;
@@ -311,9 +286,9 @@ FRemoteObjectId FRemoteObjectId::Generate()
 	return FRemoteObjectId(RemoteServerId, RemoteObjectSerialNumber.fetch_add(1));
 }
 
-FRemoteObjectId FRemoteObjectId::FromObject(const UObject* Object)
+FRemoteObjectId::FRemoteObjectId(const UObjectBase* Object)
+	: FRemoteObjectId(UE::CoreUObject::Private::FObjectHandleUtils::GetRemoteId(Object))
 {
-	return UE::CoreUObject::Private::FObjectHandleUtils::GetRemoteId(Object);
 }
 
 FArchive& operator<<(FArchive& Ar, FRemoteObjectId& Id)
diff --git a/Engine/Source/Runtime/CoreUObject/Private/UObject/RemoteObjectSerialization.cpp b/Engine/Source/Runtime/CoreUObject/Private/UObject/RemoteObjectSerialization.cpp
--- a/Engine/Source/Runtime/CoreUObject/Private/UObject/RemoteObjectSerialization.cpp
+++ b/Engine/Source/Runtime/CoreUObject/Private/UObject/RemoteObjectSerialization.cpp
@@ -125,7 +125,6 @@ public:
 
 	virtual FArchive& operator<<(UObject*& Obj) override
 	{
-		FRemoteObjectId ObjId;
 		ERemoteReferenceType RefType = WriteObjectReference(Obj);
 		if (RefType == ERemoteReferenceType::IdOnly)
 		{
diff --git a/Engine/Source/Runtime/CoreUObject/Private/UObject/UObjectGlobals.cpp b/Engine/Source/Runtime/CoreUObject/Private/UObject/UObjectGlobals.cpp
--- a/Engine/Source/Runtime/CoreUObject/Private/UObject/UObjectGlobals.cpp
+++ b/Engine/Source/Runtime/CoreUObject/Private/UObject/UObjectGlobals.cpp
@@ -3699,6 +3699,10 @@ void UObject::PostInitProperties()
 UE_AUTORTFM_ALWAYS_OPEN
 UObject::UObject()
 {
+#if UE_WITH_REMOTE_OBJECT_HANDLE
+	checkf(UE::CoreUObject::Private::FObjectHandleUtils::GetRemoteId(this).IsValid(), TEXT("Constructing an object %s with an invalid RemoteId. Did UObjectBase() constructor reset it?"), *GetPathName());
+#endif // UE_WITH_REMOTE_OBJECT_HANDLE
+
 	EnsureNotRetrievingVTablePtr();
 
 	FObjectInitializer* ObjectInitializerPtr = FUObjectThreadContext::Get().TopInitializer();
@@ -3712,6 +3716,10 @@ UObject::UObject()
 UE_AUTORTFM_ALWAYS_OPEN
 UObject::UObject(const FObjectInitializer& ObjectInitializer)
 {
+#if UE_WITH_REMOTE_OBJECT_HANDLE
+	checkf(UE::CoreUObject::Private::FObjectHandleUtils::GetRemoteId(this).IsValid(), TEXT("Constructing an object %s with an invalid RemoteId. Did UObjectBase() constructor reset it?"), *GetPathName());
+#endif // UE_WITH_REMOTE_OBJECT_HANDLE
+
 	EnsureNotRetrievingVTablePtr();
 
 	UE_CLOG(ObjectInitializer.Obj != nullptr && ObjectInitializer.Obj != this, LogUObjectGlobals, Fatal, TEXT("UObject(const FObjectInitializer&) constructor called but it's not the object that's currently being constructed with NewObject. Maybe you are trying to construct it on the stack, which is not supported."));
diff --git a/Engine/Source/Runtime/CoreUObject/Public/UObject/ObjectHandle.h b/Engine/Source/Runtime/CoreUObject/Public/UObject/ObjectHandle.h
--- a/Engine/Source/Runtime/CoreUObject/Public/UObject/ObjectHandle.h
+++ b/Engine/Source/Runtime/CoreUObject/Public/UObject/ObjectHandle.h
@@ -81,8 +81,14 @@ namespace UE::CoreUObject::Private
 		UPTRINT PointerOrHandle;
 
 		FRemoteObjectHandlePrivate() = default;
-		COREUOBJECT_API explicit FRemoteObjectHandlePrivate(UObject* Object);
-		COREUOBJECT_API explicit FRemoteObjectHandlePrivate(UE::RemoteObject::Handle::FRemoteObjectStub* RemoteInfo);
+		explicit FRemoteObjectHandlePrivate(UObject* Object)
+		{
+			PointerOrHandle = UPTRINT(Object);
+		}
+		explicit FRemoteObjectHandlePrivate(UE::RemoteObject::Handle::FRemoteObjectStub* RemoteInfo)
+		{
+			PointerOrHandle = UPTRINT(RemoteInfo) | 1;
+		}
 
 		FORCEINLINE const UE::RemoteObject::Handle::FRemoteObjectStub* ToStub() const
 		{
@@ -291,7 +297,7 @@ inline bool operator==(UE::CoreUObject::Private::FRemoteObjectHandlePrivate LHS,
 		}
 
 		const FRemoteObjectStub* RhsStub = RHS.ToStub();
-		return RhsStub->Id == FRemoteObjectId::FromObject(Obj);
+		return RhsStub->Id == FRemoteObjectId((const UObjectBase*)Obj);
 	}
 	else
 	{
@@ -303,7 +309,7 @@ inline bool operator==(UE::CoreUObject::Private::FRemoteObjectHandlePrivate LHS,
 		}
 
 		const FRemoteObjectStub* LhsStub = LHS.ToStub();
-		return LhsStub->Id == FRemoteObjectId::FromObject(Obj);
+		return LhsStub->Id == FRemoteObjectId((const UObjectBase*)Obj);
 	}
 
 }
@@ -329,7 +335,7 @@ inline uint32 GetTypeHash(UE::CoreUObject::Private::FRemoteObjectHandlePrivate H
 	if (IsObjectHandleResolved(Handle))
 	{
 		const UObject* Obj = ReadObjectHandlePointerNoCheck(Handle);
-		return GetTypeHash(FRemoteObjectId::FromObject(Obj));
+		return GetTypeHash(FRemoteObjectId((const UObjectBase*)Obj));
 	}
 	const FRemoteObjectStub* Stub = Handle.ToStub();
 	return GetTypeHash(Stub->Id);
diff --git a/Engine/Source/Runtime/CoreUObject/Public/UObject/RemoteObjectTypes.h b/Engine/Source/Runtime/CoreUObject/Public/UObject/RemoteObjectTypes.h
--- a/Engine/Source/Runtime/CoreUObject/Public/UObject/RemoteObjectTypes.h
+++ b/Engine/Source/Runtime/CoreUObject/Public/UObject/RemoteObjectTypes.h
@@ -55,21 +55,16 @@ private:
 			uint64 ServerId : 10;
 			uint64 SerialNumber : 54;
 		};
-		uint64 Id;
+		uint64 Id = 0;
 	};
 
 public:
 
-	FRemoteObjectId()
-		: ServerId(0)
-		, SerialNumber(0)
-	{
-	}
+	FRemoteObjectId() = default;
 
 	explicit FRemoteObjectId(uint64 InId)
 		: Id(InId)
 	{
-
 	}
 
 	FRemoteObjectId(FRemoteServerId InServerId, uint64 InSerialNumber)
@@ -77,12 +72,12 @@ public:
 		, SerialNumber(InSerialNumber)
 	{
 	}
-	explicit FRemoteObjectId(ENoInit)
-	{
-	}
-	FORCEINLINE friend uint32 GetTypeHash(const FRemoteObjectId& ObjectId)
+
+	explicit COREUOBJECT_API FRemoteObjectId(const UObjectBase* InObject);
+
+	FORCEINLINE uint32 GetTypeHash() const
 	{
-		return GetTypeHash(ObjectId.Id);
+		return ::GetTypeHash(Id);
 	}
 
 	FORCEINLINE bool operator==(const FRemoteObjectId& Other) const
@@ -93,6 +88,18 @@ public:
 	{
 		return Id != Other.Id;
 	}
+	FORCEINLINE bool operator<(const FRemoteObjectId& Other) const
+	{
+		if (ServerId < Other.ServerId)
+		{
+			return true;
+		}
+		else if (ServerId == Other.ServerId)
+		{
+			return SerialNumber < Other.SerialNumber;
+		}
+		return false;
+	}
 	FORCEINLINE bool IsValid() const
 	{
 		return *this != FRemoteObjectId();
@@ -110,11 +117,13 @@ public:
 
 	COREUOBJECT_API static FRemoteObjectId Generate();
 
-	COREUOBJECT_API static FRemoteObjectId FromObject(const UObject* Object);
-
 	friend COREUOBJECT_API FArchive& operator<<(FArchive& Ar, FRemoteObjectId& Id);
 };
 
+FORCEINLINE uint32 GetTypeHash(const FRemoteObjectId& ObjectId)
+{
+	return ObjectId.GetTypeHash();
+}
 
 /**
 * Structure that holds remote object memory
diff --git a/Engine/Source/Runtime/CoreUObject/Public/UObject/UObjectBase.h b/Engine/Source/Runtime/CoreUObject/Public/UObject/UObjectBase.h
--- a/Engine/Source/Runtime/CoreUObject/Public/UObject/UObjectBase.h
+++ b/Engine/Source/Runtime/CoreUObject/Public/UObject/UObjectBase.h
@@ -64,7 +64,7 @@ class UObjectBase
 protected:
 	UObjectBase() :
 #if UE_WITH_REMOTE_OBJECT_HANDLE
-		RemoteId(NoInit),
+		RemoteId(this), // FRemoteObjectId is a literal type and can't have a NoInit constructor so copy the id from whatever is already set in (this) memory
 #endif // UE_WITH_REMOTE_OBJECT_HANDLE
 		ClassPrivate(NoInit),
 		NamePrivate(NoInit),  // screwy, but the name was already set and we don't want to set it again
