Added cvar to control if non-nullable properties produce an error, a warning or a default object when failing to resolve an object reference.
Default objects will not be created if the class is abstract.

#tests Ran editor, new unit tests

diff --git a/Engine/Source/Editor/PropertyEditor/Private/PropertyHandleImpl.cpp b/Engine/Source/Editor/PropertyEditor/Private/PropertyHandleImpl.cpp
--- a/Engine/Source/Editor/PropertyEditor/Private/PropertyHandleImpl.cpp
+++ b/Engine/Source/Editor/PropertyEditor/Private/PropertyHandleImpl.cpp
@@ -1161,10 +1161,18 @@ int32 FPropertyValueImpl::AddChild()
 						if (Array)
 						{
 							Array->PerformOperationWithSetter(Obj, Addr, [Obj, Array, &Index, &ReturnLogicalIndex](void* DirectAddress)
-							{
+							{								
 								FScriptArrayHelper	ArrayHelper(Array, DirectAddress);
 								Index = ArrayHelper.AddValue();
 								ReturnLogicalIndex = Index;
+
+								// check whether the inner type is flagged as a non-nullable. if so, create it.
+								FObjectProperty* InnerObjectProperty = CastField<FObjectProperty>(Array->Inner);
+								if (InnerObjectProperty && InnerObjectProperty->HasAnyPropertyFlags(CPF_NonNullable) && !InnerObjectProperty->PropertyClass->HasAnyClassFlags(CLASS_Abstract))
+								{
+									UObject* NewItem = NewObject<UObject>(Obj, InnerObjectProperty->PropertyClass);
+									InnerObjectProperty->SetObjectPropertyValue(ArrayHelper.GetRawPtr(Index), NewItem);
+								}
 							});
 						}
 						else if (Set)
@@ -1175,6 +1183,14 @@ int32 FPropertyValueImpl::AddChild()
 								Index = SetHelper.AddDefaultValue_Invalid_NeedsRehash();
 								ReturnLogicalIndex = SetHelper.FindLogicalIndex(Index);
 
+								// check whether the element type is flagged as a non-nullable. if so, create it.
+								FObjectProperty* ElementObjectProperty = CastField<FObjectProperty>(Set->ElementProp);
+								if (ElementObjectProperty && ElementObjectProperty->HasAnyPropertyFlags(CPF_NonNullable) && !ElementObjectProperty->PropertyClass->HasAnyClassFlags(CLASS_Abstract))
+								{
+									UObject* NewItem = NewObject<UObject>(Obj, ElementObjectProperty->PropertyClass);
+									ElementObjectProperty->SetObjectPropertyValue(SetHelper.GetElementPtr(Index), NewItem);
+								}
+
 								SetHelper.Rehash();
 							});
 						}
@@ -1186,6 +1202,25 @@ int32 FPropertyValueImpl::AddChild()
 								Index = MapHelper.AddDefaultValue_Invalid_NeedsRehash();
 								ReturnLogicalIndex = MapHelper.FindLogicalIndex(Index);
 
+								// check whether the key or value type is flagged as a non-nullable. if so, create it.
+								{
+									FObjectProperty* KeyObjectProperty = CastField<FObjectProperty>(Map->KeyProp);
+									if (KeyObjectProperty && KeyObjectProperty->HasAnyPropertyFlags(CPF_NonNullable) && !KeyObjectProperty->PropertyClass->HasAnyClassFlags(CLASS_Abstract))
+									{
+										UObject* NewItem = NewObject<UObject>(Obj, KeyObjectProperty->PropertyClass);
+										KeyObjectProperty->SetObjectPropertyValue(MapHelper.GetKeyPtr(Index), NewItem);
+									}
+								}
+
+								{
+									FObjectProperty* ValueObjectProperty = CastField<FObjectProperty>(Map->ValueProp);
+									if (ValueObjectProperty && ValueObjectProperty->HasAnyPropertyFlags(CPF_NonNullable) && !ValueObjectProperty->PropertyClass->HasAnyClassFlags(CLASS_Abstract))
+									{
+										UObject* NewItem = NewObject<UObject>(Obj, ValueObjectProperty->PropertyClass);
+										ValueObjectProperty->SetObjectPropertyValue(MapHelper.GetValuePtr(Index), NewItem);
+									}
+								}
+
 								MapHelper.Rehash();
 								bAddedMapEntry = true;
 							});
@@ -4608,17 +4643,17 @@ FPropertyHandleVector::FPropertyHandleVector( TSharedRef<class FPropertyNode> Pr
 	}
 
 	if( TSharedPtr<FPropertyNode> YComponentPropertyNode = Implementation->GetChildNode("Y", bRecurse) )
-	{
+		{
 		VectorComponents.Add( MakeShareable( new FPropertyHandleMixed( YComponentPropertyNode.ToSharedRef(), NotifyHook, PropertyUtilities ) ) );
-	}
+		}
 		
 	if( TSharedPtr<FPropertyNode> ZComponentPropertyNode = Implementation->GetChildNode("Z", bRecurse) )
-	{
+		{
 		VectorComponents.Add( MakeShareable( new FPropertyHandleMixed( ZComponentPropertyNode.ToSharedRef(), NotifyHook, PropertyUtilities ) ) );
-	}
+		}
 		
 	if( TSharedPtr<FPropertyNode> WComponentPropertyNode = Implementation->GetChildNode("W", bRecurse) )
-	{
+		{
 		VectorComponents.Add( MakeShareable( new FPropertyHandleMixed( WComponentPropertyNode.ToSharedRef(), NotifyHook, PropertyUtilities ) ) );
 	}
 }
diff --git a/Engine/Source/Runtime/CoreUObject/Private/UObject/PropertyBaseObject.cpp b/Engine/Source/Runtime/CoreUObject/Private/UObject/PropertyBaseObject.cpp
--- a/Engine/Source/Runtime/CoreUObject/Private/UObject/PropertyBaseObject.cpp
+++ b/Engine/Source/Runtime/CoreUObject/Private/UObject/PropertyBaseObject.cpp
@@ -1,6 +1,7 @@
 // Copyright Epic Games, Inc. All Rights Reserved.
 
 #include "CoreMinimal.h"
+#include "HAL/IConsoleManager.h"
 #include "UObject/ObjectMacros.h"
 #include "Templates/Casts.h"
 #include "UObject/Package.h"
@@ -16,6 +17,25 @@
 #include "UObject/UObjectThreadContext.h"
 #include "Misc/StringBuilder.h"
 
+namespace UE::CoreUObject::Private
+{
+	TAutoConsoleVariable<int32> CVarNonNullableBehavior(
+		TEXT("CoreUObject.NonNullableBehavior"),
+		(int32)ENonNullableBehavior::CreateDefaultObjectIfPossible,
+		TEXT("Sets the behavior when a non-null property cannot be resolved into an object reference - 0=Leave property null and log a warning, 1=Leave property null and log an error, 2=Create a default object and log a warning if successful, or leave it null and log an error if unsuccessful")
+	);
+
+	ENonNullableBehavior GetNonNullableBehavior()
+	{
+		int32 Value = CVarNonNullableBehavior.GetValueOnAnyThread();
+		if (Value < 0 || Value > 2)
+		{
+			Value = 2;
+		}
+		return (ENonNullableBehavior)Value;
+	}
+}
+
 /*-----------------------------------------------------------------------------
 	FObjectPropertyBase.
 -----------------------------------------------------------------------------*/
@@ -661,11 +681,11 @@ void FObjectPropertyBase::SetObjectPtrPropertyValue(void* PropertyValueAddress,
 	if (Ptr || !HasAnyPropertyFlags(CPF_NonNullable))
 	{
 		SetObjectPtrPropertyValueUnchecked(PropertyValueAddress, Ptr);
-	}
+}
 	else
-	{
+{
 		UE_LOG(LogProperty, Verbose /*Warning*/, TEXT("Trying to assign null object value to non-nullable \"%s\""), *GetFullName());
-	}
+}
 }
 
 void FObjectPropertyBase::SetObjectPropertyValue_InContainer(void* ContainerAddress, UObject* Value, int32 ArrayIndex) const
@@ -687,7 +707,7 @@ void FObjectPropertyBase::SetObjectPtrPropertyValue_InContainer(void* ContainerA
 		SetObjectPtrPropertyValueUnchecked_InContainer(ContainerAddress, Ptr, ArrayIndex);
 	}
 	else
-	{
+{
 		UE_LOG(LogProperty, Verbose /*Warning*/, TEXT("Trying to assign null object value to non-nullable \"%s\""), *GetFullName());
 	}
 }
@@ -702,6 +722,49 @@ bool FObjectPropertyBase::AllowObjectTypeReinterpretationTo(const FObjectPropert
 	return false;
 }
 
+UObject* FObjectPropertyBase::ConstructDefaultObjectValueIfNecessary(UObject* ExistingValue) const
+{
+	UObject* NewDefaultObjectValue = nullptr;
+
+	FUObjectSerializeContext* SerializeContext = FUObjectThreadContext::Get().GetSerializeContext();
+	UObject* Outer = SerializeContext ? SerializeContext->SerializedObject : nullptr;
+	if (!Outer)
+	{
+		Outer = GetTransientPackage();
+	}
+
+	if (ExistingValue)
+	{
+		UClass* ExistingValueClass = ExistingValue->GetClass();
+		// Sanity check to make sure the existing value class matches the property class
+		if (ExistingValueClass && (ExistingValueClass->IsChildOf(PropertyClass) || ExistingValueClass->GetAuthoritativeClass()->IsChildOf(PropertyClass)))
+		{
+			if (ExistingValue->IsTemplate() && 	// Existing value is a template so we can construct a new value with it as the archetype
+				ExistingValue->GetOuter() != Outer) // Unless the template's Outer is the same as the new Outer in which case the template (ExistingValue) IS the object we can reuse
+			{
+				// We probably got here because an object value failed to load (missing import class) and the property is left with a template of default subobject
+				NewDefaultObjectValue = NewObject<UObject>(Outer, ExistingValue->GetClass(), NAME_None, RF_NoFlags, ExistingValue);
+			}
+			else
+			{
+				// Existing value is not a template or a template is what this property was pointing to so we can use it directly
+				// Similar to the above condition but the property was not referencing an instanced value in which case it's ok to leave the CDO default here
+				NewDefaultObjectValue = ExistingValue;
+			}
+		}
+	}
+
+	if (!NewDefaultObjectValue && !PropertyClass->HasAnyClassFlags(CLASS_Abstract))
+	{
+		// Existing value did not exist or it could not be used as a template
+		// Existing value may be null in case we were serializing an array of UObjects that failed to load (missing import class). Since the array is first pre-allocated with null values
+		// it will not have any existing objects to instantiate
+		NewDefaultObjectValue = NewObject<UObject>(Outer, PropertyClass);
+	}
+
+	return NewDefaultObjectValue;
+}
+
 void FObjectPropertyBase::CheckValidObject(void* ValueAddress, TObjectPtr<UObject> OldValue) const
 {
 	const TObjectPtr<UObject> Object = GetObjectPtrPropertyValue(ValueAddress);
@@ -745,31 +808,77 @@ void FObjectPropertyBase::CheckValidObject(void* ValueAddress, TObjectPtr<UObjec
 		bool bIsReplacingClassRefs = PropertyClass && PropertyClass->HasAnyClassFlags(CLASS_NewerVersionExists) != ObjectClass->HasAnyClassFlags(CLASS_NewerVersionExists);
 		if (!bIsReplacingClassRefs && !IsDeferringValueLoad())
 		{
+			UObject* DefaultValue = nullptr;
+
 			FUObjectSerializeContext* SerializeContext = FUObjectThreadContext::Get().GetSerializeContext();
+			UObject* Outer = SerializeContext ? SerializeContext->SerializedObject : nullptr;
+			if (!Outer)
+			{
+				Outer = GetTransientPackage();
+			}
 			if (!HasAnyPropertyFlags(CPF_NonNullable))
 			{
 				UE_LOG(LogProperty, Warning,
 					TEXT("Serialized %s for a property of %s. Reference will be nulled.\n    ReferencingObject = %s\n    Property = %s\n    Item = %s"),
 					*ObjectClass->GetFullName(),
 					*PropertyClass->GetFullName(),
-					*GetFullNameSafe(SerializeContext ? SerializeContext->SerializedObject : nullptr),
+					*GetFullNameSafe(Outer),
 					*GetFullName(),
 					*Object.GetFullName()
 				);
 			}
 			else
 			{
-				UE_LOG(LogProperty, Warning,
-					TEXT("Serialized %s for a non-nullable property of %s. Reference will be nulled - will cause a runtime error if accessed.\n    ReferencingObject = %s\n    Property = %s\n    Item = %s"),
-					*ObjectClass->GetFullName(),
-					*PropertyClass->GetFullName(),
-					*GetFullNameSafe(SerializeContext ? SerializeContext->SerializedObject : nullptr),
-					*GetFullName(),
-					*Object.GetFullName()
-				);
+				using UE::CoreUObject::Private::ENonNullableBehavior;
+				using UE::CoreUObject::Private::GetNonNullableBehavior;
+
+				ENonNullableBehavior NonNullableBehavior = GetNonNullableBehavior();
+				if (NonNullableBehavior == ENonNullableBehavior::CreateDefaultObjectIfPossible)
+				{
+					DefaultValue = ConstructDefaultObjectValueIfNecessary(OldValue);
+				}
+
+				if (DefaultValue)
+				{
+					UE_LOG(LogProperty, Warning,
+						TEXT("Serialized %s for a non-nullable property of %s. Reference will be defaulted to %s (previously: %s).\n    ReferencingObject = %s\n    Property = %s\n    Item = %s"),
+						*ObjectClass->GetFullName(),
+						*PropertyClass->GetFullName(),
+						*GetFullNameSafe(DefaultValue),
+						*GetFullNameSafe(OldValue),
+						*GetFullNameSafe(Outer),
+						*GetFullName(),
+						*Object.GetFullName()
+					);
+				}
+				else if (NonNullableBehavior == ENonNullableBehavior::LogWarning)
+				{
+					UE_LOG(LogProperty, Warning,
+						TEXT("Serialized %s for a non-nullable property of %s. Reference will be nulled (previously: %s) - will cause a runtime error if accessed.\n    ReferencingObject = %s\n    Property = %s\n    Item = %s"),
+						*ObjectClass->GetFullName(),
+						*PropertyClass->GetFullName(),
+						*GetFullNameSafe(OldValue),
+						*GetFullNameSafe(Outer),
+						*GetFullName(),
+						*Object.GetFullName()
+					);
+				}
+				else
+				{
+					UE_LOG(LogProperty, Error,
+						TEXT("Serialized %s for a non-nullable property of %s. Reference will be nulled%s (previously: %s) - will cause a runtime error if accessed.\n    ReferencingObject = %s\n    Property = %s\n    Item = %s"),
+						*ObjectClass->GetFullName(),
+						*PropertyClass->GetFullName(),
+						(NonNullableBehavior == ENonNullableBehavior::CreateDefaultObjectIfPossible) ? *FString::Printf(TEXT(" as %s is abstract"), *PropertyClass->GetName()) : TEXT(""),
+						*GetFullNameSafe(OldValue),
+						*GetFullNameSafe(Outer),
+						*GetFullName(),
+						*Object.GetFullName()
+					);
+				}
 			}
 
-			SetObjectPropertyValueUnchecked(ValueAddress, nullptr);
+			SetObjectPropertyValueUnchecked(ValueAddress, DefaultValue);
 		}
 	}
 }
diff --git a/Engine/Source/Runtime/CoreUObject/Private/UObject/PropertyObject.cpp b/Engine/Source/Runtime/CoreUObject/Private/UObject/PropertyObject.cpp
--- a/Engine/Source/Runtime/CoreUObject/Private/UObject/PropertyObject.cpp
+++ b/Engine/Source/Runtime/CoreUObject/Private/UObject/PropertyObject.cpp
@@ -254,12 +254,45 @@ void FObjectProperty::PostSerializeObjectItem(FArchive& SerializingArchive, void
 		!SerializingArchive.IsTransacting() && // Don't create new objects when loading from the transaction buffer
 		!SerializingArchive.IsCountingMemory())
 	{
-		UE_LOG(LogProperty, Warning,
-			TEXT("Failed to serialize value for non-nullable property %s. Reference will be nulled - will cause a runtime error if accessed."),
-			*GetFullName()
-		);
+		UObject* DefaultValue = nullptr;
 
-		SetObjectPropertyValueUnchecked(Value, nullptr);
+		using UE::CoreUObject::Private::ENonNullableBehavior;
+		using UE::CoreUObject::Private::GetNonNullableBehavior;
+
+		ENonNullableBehavior NonNullableBehavior = GetNonNullableBehavior();
+		if (NonNullableBehavior == ENonNullableBehavior::CreateDefaultObjectIfPossible)
+		{
+			DefaultValue = ConstructDefaultObjectValueIfNecessary(CurrentValue);
+		}
+
+		if (DefaultValue)
+		{
+			UE_LOG(LogProperty, Warning,
+				TEXT("Failed to serialize value for non-nullable property %s (previously: %s). Reference will be defaulted to %s."),
+				*GetFullName(),
+				*GetFullNameSafe(CurrentValue),
+				*GetFullNameSafe(DefaultValue)
+			);
+		}
+		else if (NonNullableBehavior == ENonNullableBehavior::LogWarning)
+		{
+			UE_LOG(LogProperty, Warning,
+				TEXT("Failed to serialize value for non-nullable property %s (previously: %s). Reference will be nulled - will cause a runtime error if accessed."),
+				*GetFullNameSafe(CurrentValue),
+				*GetFullName()
+			);
+		}
+		else
+		{
+			UE_LOG(LogProperty, Error,
+				TEXT("Failed to serialize value for non-nullable property %s (previously: %s). Reference will be nulled - will cause a runtime error if accessed."),
+				*GetFullNameSafe(CurrentValue),
+				*GetFullName()
+			);
+		}
+
+		SetObjectPropertyValueUnchecked(Value, DefaultValue);
+		ObjectValue = DefaultValue;
 	}
 
 	if (ObjectValue != CurrentValue)
@@ -349,27 +382,27 @@ TObjectPtr<UObject> FObjectProperty::GetObjectPtrPropertyValue_InContainer(const
 	TObjectPtr<UObject> Result = nullptr;
 	GetWrappedUObjectPtrValues<FObjectPtr>(&Result, ContainerAddress, EPropertyMemoryAccess::InContainer, ArrayIndex, 1);
 	return Result;
-}
+	}
 
 void FObjectProperty::SetObjectPtrPropertyValueUnchecked(void* PropertyValueAddress, TObjectPtr<UObject> Ptr) const
-{
+	{
 	SetPropertyValue(PropertyValueAddress, Ptr);
 }
 
 void FObjectProperty::SetObjectPropertyValueUnchecked(void* PropertyValueAddress, UObject* Value) const
-{
-	SetPropertyValue(PropertyValueAddress, Value);
-}
+	{
+		SetPropertyValue(PropertyValueAddress, Value);
+	}
 
 void FObjectProperty::SetObjectPtrPropertyValueUnchecked_InContainer(void* ContainerAddress, TObjectPtr<UObject> Ptr, int32 ArrayIndex) const
-{
+	{
 	SetWrappedUObjectPtrValues<FObjectPtr>(ContainerAddress, EPropertyMemoryAccess::InContainer, &Ptr, ArrayIndex, 1);
 }
 
 void FObjectProperty::SetObjectPropertyValueUnchecked_InContainer(void* ContainerAddress, UObject* Value, int32 ArrayIndex) const
-{
-	SetWrappedUObjectPtrValues<FObjectPtr>(ContainerAddress, EPropertyMemoryAccess::InContainer, &Value, ArrayIndex, 1);
-}
+	{
+		SetWrappedUObjectPtrValues<FObjectPtr>(ContainerAddress, EPropertyMemoryAccess::InContainer, &Value, ArrayIndex, 1);
+	}
 
 void FObjectProperty::CopySingleValueToScriptVM(void* Dest, const void* Src) const
 {
diff --git a/Engine/Source/Runtime/CoreUObject/Public/UObject/UnrealType.h b/Engine/Source/Runtime/CoreUObject/Public/UObject/UnrealType.h
--- a/Engine/Source/Runtime/CoreUObject/Public/UObject/UnrealType.h
+++ b/Engine/Source/Runtime/CoreUObject/Public/UObject/UnrealType.h
@@ -89,6 +89,18 @@ struct FPropertyVisitorPath;
 
 COREUOBJECT_API DECLARE_LOG_CATEGORY_EXTERN(LogType, Log, All);
 
+namespace UE::CoreUObject::Private
+{
+	enum class ENonNullableBehavior
+	{
+		LogWarning                    = 0,
+		LogError                      = 1,
+		CreateDefaultObjectIfPossible = 2
+	};
+
+	COREUOBJECT_API ENonNullableBehavior GetNonNullableBehavior();
+}
+
 /*-----------------------------------------------------------------------------
 	FProperty.
 -----------------------------------------------------------------------------*/
@@ -2668,6 +2680,14 @@ protected:
 	virtual bool AllowCrossLevel() const;
 	// End of FObjectPropertyBase interface
 
+	/**
+	 * Constructs a new object if the existing one is missing or is not compatible with the property class
+	 * Used for making sure non-nullable properties have valid values.
+	 * @param ExistingValue Previous object value (can be null)
+	 * @return Non-null object that was assigned to this property value address
+	 */
+	UObject* ConstructDefaultObjectValueIfNecessary(UObject* ExistingValue) const;
+
 	// Disable false positive buffer overrun warning during pgoprofile linking step
 	PRAGMA_DISABLE_BUFFER_OVERRUN_WARNING
 	/* Helper functions for UObject property types that wrap the object pointer in a smart pointer */
diff --git a/Engine/Source/Runtime/CoreUObject/Tests/ObjectPropertyTest.cpp b/Engine/Source/Runtime/CoreUObject/Tests/ObjectPropertyTest.cpp
--- a/Engine/Source/Runtime/CoreUObject/Tests/ObjectPropertyTest.cpp
+++ b/Engine/Source/Runtime/CoreUObject/Tests/ObjectPropertyTest.cpp
@@ -94,7 +94,9 @@ TEST_CASE("UE::CoreUObject::FObjectProperty::CheckValidAddressNonNullable")
 #endif
 	UClass* Class = UObjectPtrTestClassWithRef::StaticClass();
 	FObjectProperty* Property = CastField<FObjectProperty>(Class->FindPropertyByName(TEXT("ObjectPtrNonNullable")));
+	FObjectProperty* AbstractProperty = CastField<FObjectProperty>(Class->FindPropertyByName(TEXT("ObjectPtrAbstractNonNullable")));
 	REQUIRE(Property != nullptr);
+	REQUIRE(AbstractProperty != nullptr);
 
 	UPackage* TestPackage = NewObject<UPackage>(nullptr, TEXT("/Temp/TestPackageName"), RF_Transient);
 	UPackage* OtherTestPackage = NewObject<UPackage>(nullptr, TEXT("/Temp/CheckValidAddressNonNullableOther"), RF_Transient);
@@ -105,31 +107,48 @@ TEST_CASE("UE::CoreUObject::FObjectProperty::CheckValidAddressNonNullable")
 	};
 	UObjectPtrTestClassWithRef* Obj = NewObject<UObjectPtrTestClassWithRef>(TestPackage, TEXT("Object"));
 	UObjectPtrTestClass* Other = NewObject<UObjectPtrTestClass>(Obj, TEXT("Other"));
+	UObjectPtrAbstractDerivedTestClass* AbstractDerivedOther = NewObject<UObjectPtrAbstractDerivedTestClass>(Obj, TEXT("AbstractDerivedOther"));
 
 #if UE_WITH_OBJECT_HANDLE_TRACKING && UE_WITH_OBJECT_HANDLE_LATE_RESOLVE
 	FObjectHandle Handle = MakeUnresolvedHandle(Other);
+	FObjectHandle AbstractDerivedHandle = MakeUnresolvedHandle(AbstractDerivedOther);
 	TObjectPtr<UObjectPtrTestClass> ObjectPtr = *reinterpret_cast<TObjectPtr<UObjectPtrTestClass>*>(&Handle);
+	TObjectPtr<UObjectPtrAbstractDerivedTestClass> AbstractDerivedObjectPtr = *reinterpret_cast<TObjectPtr<UObjectPtrAbstractDerivedTestClass>*>(&AbstractDerivedHandle);
 #else
 	TObjectPtr<UObjectPtrTestClass> ObjectPtr = Other;
+	TObjectPtr<UObjectPtrAbstractDerivedTestClass> AbstractDerivedObjectPtr = AbstractDerivedOther;
 #endif
 
 	//property is already null should stay null
 	CHECK(!Obj->ObjectPtrNonNullable);
+	CHECK(!Obj->ObjectPtrAbstractNonNullable);
 	Property->CheckValidObject(&Obj->ObjectPtrNonNullable, ObjectPtr);
+	AbstractProperty->CheckValidObject(&Obj->ObjectPtrAbstractNonNullable, AbstractDerivedObjectPtr);
 	CHECK(!Obj->ObjectPtrNonNullable);
+	CHECK(!Obj->ObjectPtrAbstractNonNullable);
 
 	//valid assignment
 	Obj->ObjectPtrNonNullable = ObjectPtr;
+	Obj->ObjectPtrAbstractNonNullable = AbstractDerivedObjectPtr;
 	Property->CheckValidObject(&Obj->ObjectPtrNonNullable, nullptr);
+	AbstractProperty->CheckValidObject(&Obj->ObjectPtrAbstractNonNullable, nullptr);
 	CHECK(Obj->ObjectPtrNonNullable == ObjectPtr);
+	CHECK(Obj->ObjectPtrAbstractNonNullable == AbstractDerivedObjectPtr);
 
 	// Disable property warnings that will fire because we're deliberately setting non-nullable properties to null
 	LOG_SCOPE_VERBOSITY_OVERRIDE(LogProperty, ELogVerbosity::NoLogging);
 
+	using UE::CoreUObject::Private::ENonNullableBehavior;
+	using UE::CoreUObject::Private::GetNonNullableBehavior;
+	ENonNullableBehavior NonNullableBehavior = GetNonNullableBehavior();
+
 	bAllowRead = true; //has resolve the old value to construct a new default value for the property
+
 	//assign a bad value to the pointer
 	Obj->ObjectPtrNonNullable = reinterpret_cast<UObjectPtrTestClass*>(OtherTestPackage);
+	Obj->ObjectPtrAbstractNonNullable = reinterpret_cast<UObjectPtrAbstractTestClass*>(OtherTestPackage);
 	CHECK(Obj->ObjectPtrNonNullable != nullptr);
+	CHECK(Obj->ObjectPtrAbstractNonNullable != nullptr);
 
 	FWarnFilterScope _([](const TCHAR* Message, ELogVerbosity::Type Verbosity, const FName& Category)
 		{
@@ -140,15 +159,41 @@ TEST_CASE("UE::CoreUObject::FObjectProperty::CheckValidAddressNonNullable")
 			return false;
 		});
 	Property->CheckValidObject(&Obj->ObjectPtrNonNullable, ObjectPtr);
-	CHECK(Obj->ObjectPtrNonNullable == nullptr); //non nullable properties should be nulled if invalid
+	AbstractProperty->CheckValidObject(&Obj->ObjectPtrAbstractNonNullable, AbstractDerivedObjectPtr);
+	if (NonNullableBehavior == ENonNullableBehavior::CreateDefaultObjectIfPossible)
+	{
+		CHECK(Obj->ObjectPtrNonNullable == ObjectPtr); //non nullable properties should be assigned the old value
+		CHECK(Obj->ObjectPtrAbstractNonNullable == AbstractDerivedObjectPtr); //abstract non nullable properties should be assigned the old value
+	}
+	else
+	{
+		CHECK(Obj->ObjectPtrNonNullable == nullptr); //non nullable properties should be nulled if invalid
+		CHECK(Obj->ObjectPtrAbstractNonNullable == nullptr); //abstract non nullable properties should be nulled
+	}
 
 	//assign a bad value to the pointer
 	Obj->ObjectPtrNonNullable = reinterpret_cast<UObjectPtrTestClass*>(Obj);
+	Obj->ObjectPtrAbstractNonNullable = reinterpret_cast<UObjectPtrAbstractTestClass*>(Obj);
 	CHECK(Obj->ObjectPtrNonNullable != nullptr);
+	CHECK(Obj->ObjectPtrAbstractNonNullable != nullptr);
 
-	//null is required for invalid non nullable properties
-	Property->CheckValidObject(&Obj->ObjectPtrNonNullable, nullptr);
-	CHECK(Obj->ObjectPtrNonNullable == nullptr);
+	if (NonNullableBehavior == ENonNullableBehavior::CreateDefaultObjectIfPossible)
+	{
+		//new value is required for non nullable properties
+		Property->CheckValidObject(&Obj->ObjectPtrNonNullable, nullptr);
+		AbstractProperty->CheckValidObject(&Obj->ObjectPtrAbstractNonNullable, nullptr);
+		CHECK(Obj->ObjectPtrNonNullable != nullptr);
+		CHECK(Obj->ObjectPtrNonNullable->IsA(UObjectPtrTestClass::StaticClass()));
+		CHECK(Obj->ObjectPtrAbstractNonNullable == nullptr);
+	}
+	else
+	{
+		//null is required for invalid non nullable properties
+		Property->CheckValidObject(&Obj->ObjectPtrNonNullable, nullptr);
+		AbstractProperty->CheckValidObject(&Obj->ObjectPtrAbstractNonNullable, nullptr);
+		CHECK(Obj->ObjectPtrNonNullable == nullptr);
+		CHECK(Obj->ObjectPtrAbstractNonNullable == nullptr);
+	}
 }
 
 class FMockArchive : public FArchive
diff --git a/Engine/Source/Runtime/CoreUObject/Tests/ObjectPtrTestClass.cpp b/Engine/Source/Runtime/CoreUObject/Tests/ObjectPtrTestClass.cpp
--- a/Engine/Source/Runtime/CoreUObject/Tests/ObjectPtrTestClass.cpp
+++ b/Engine/Source/Runtime/CoreUObject/Tests/ObjectPtrTestClass.cpp
@@ -26,6 +26,37 @@ IMPLEMENT_CORE_INTRINSIC_CLASS(UObjectPtrTestClass, UObject,
 	});
 #endif
 
+#if WITH_EDITORONLY_DATA
+IMPLEMENT_CORE_INTRINSIC_CLASS(UObjectPtrAbstractTestClass, UObject,
+	{
+		auto MetaData = Class->GetOutermost()->GetMetaData();
+		if (MetaData)
+		{
+			MetaData->SetValue(Class, TEXT("LoadBehavior"), TEXT("LazyOnDemand"));
+		}
+	}
+);
+#else
+IMPLEMENT_CORE_INTRINSIC_CLASS(UObjectPtrAbstractTestClass, UObject,
+	{ 
+	});
+#endif
+
+#if WITH_EDITORONLY_DATA
+IMPLEMENT_CORE_INTRINSIC_CLASS(UObjectPtrAbstractDerivedTestClass, UObjectPtrAbstractTestClass,
+	{
+		auto MetaData = Class->GetOutermost()->GetMetaData();
+		if (MetaData)
+		{
+			MetaData->SetValue(Class, TEXT("LoadBehavior"), TEXT("LazyOnDemand"));
+		}
+	}
+);
+#else
+IMPLEMENT_CORE_INTRINSIC_CLASS(UObjectPtrAbstractDerivedTestClass, UObjectPtrAbstractTestClass,
+	{ 
+	});
+#endif
 
 #if WITH_EDITORONLY_DATA
 IMPLEMENT_CORE_INTRINSIC_CLASS(UObjectPtrTestClassWithRef, UObject,
@@ -43,7 +74,6 @@ IMPLEMENT_CORE_INTRINSIC_CLASS(UObjectPtrTestClassWithRef, UObject,
 		}
 
 		{
-
 			UECodeGen_Private::FObjectPropertyParams Params = { };
 			Params.NameUTF8 = "ObjectPtrNonNullable";
 			Params.Offset = STRUCT_OFFSET(UObjectPtrTestClassWithRef, ObjectPtrNonNullable);
@@ -54,6 +84,16 @@ IMPLEMENT_CORE_INTRINSIC_CLASS(UObjectPtrTestClassWithRef, UObject,
 			Property->PropertyClass = UObjectPtrTestClass::StaticClass();
 		}
 
+		{
+			UECodeGen_Private::FObjectPropertyParams Params = { };
+			Params.NameUTF8 = "ObjectPtrAbstractNonNullable";
+			Params.Offset = STRUCT_OFFSET(UObjectPtrTestClassWithRef, ObjectPtrAbstractNonNullable);
+			Params.PropertyFlags = EPropertyFlags::CPF_NonNullable | CPF_TObjectPtrWrapper;
+			Params.ObjectFlags = RF_Public | RF_Transient | RF_MarkAsNative;
+			Params.ClassFunc = nullptr;
+			auto Property = new FObjectProperty(Class, Params);
+			Property->PropertyClass = UObjectPtrAbstractTestClass::StaticClass();
+		}
 
 		{
 			UECodeGen_Private::FArrayPropertyParams Params = { };
@@ -99,6 +139,17 @@ IMPLEMENT_CORE_INTRINSIC_CLASS(UObjectPtrTestClassWithRef, UObject,
 			Property->PropertyClass = UObjectPtrTestClass::StaticClass();
 		}
 
+		{
+			UECodeGen_Private::FObjectPropertyParams Params = { };
+			Params.NameUTF8 = "ObjectPtrAbstractNonNullable";
+			Params.Offset = STRUCT_OFFSET(UObjectPtrTestClassWithRef, ObjectPtrAbstractNonNullable);
+			Params.PropertyFlags = EPropertyFlags::CPF_NonNullable | CPF_TObjectPtrWrapper;
+			Params.ObjectFlags = RF_Public | RF_Transient | RF_MarkAsNative;
+			Params.ClassFunc = nullptr;
+			auto Property = new FObjectProperty(Class, Params);
+			Property->PropertyClass = UObjectPtrAbstractTestClass::StaticClass();
+		}
+
 		{
 			UECodeGen_Private::FArrayPropertyParams Params = { };
 			Params.NameUTF8 = "ArrayObjPtr";
diff --git a/Engine/Source/Runtime/CoreUObject/Tests/ObjectPtrTestClass.h b/Engine/Source/Runtime/CoreUObject/Tests/ObjectPtrTestClass.h
--- a/Engine/Source/Runtime/CoreUObject/Tests/ObjectPtrTestClass.h
+++ b/Engine/Source/Runtime/CoreUObject/Tests/ObjectPtrTestClass.h
@@ -16,6 +16,18 @@ public:
 
 };
 
+//abstract test class for testing TObjectPtr resolve behavior
+class UObjectPtrAbstractTestClass : public UObject
+{
+	DECLARE_CLASS_INTRINSIC(UObjectPtrAbstractTestClass, UObject, CLASS_MatchedSerializers | CLASS_Abstract, TEXT("/Script/CoreUObject"))
+};
+
+//derived-from-abstract test class for testing TObjectPtr resolve behavior
+class UObjectPtrAbstractDerivedTestClass : public UObjectPtrAbstractTestClass
+{
+	DECLARE_CLASS_INTRINSIC(UObjectPtrAbstractDerivedTestClass, UObjectPtrAbstractTestClass, CLASS_MatchedSerializers, TEXT("/Script/CoreUObject"))
+};
+
 //test class with typed reference to another class
 class UObjectPtrTestClassWithRef : public UObject
 {
@@ -24,6 +36,7 @@ class UObjectPtrTestClassWithRef : public UObject
 public:
 	TObjectPtr<UObjectPtrTestClass> ObjectPtr;
 	TObjectPtr<UObjectPtrTestClass> ObjectPtrNonNullable;
+	TObjectPtr<UObjectPtrAbstractTestClass> ObjectPtrAbstractNonNullable;
 	TArray<TObjectPtr<UObjectPtrTestClass>> ArrayObjPtr;
 };
 
