Trivial changes to force an integration to Release-32.00.

#rnx

diff --git a/Engine/Source/Runtime/CoreUObject/Private/UObject/PropertyObject.cpp b/Engine/Source/Runtime/CoreUObject/Private/UObject/PropertyObject.cpp
--- a/Engine/Source/Runtime/CoreUObject/Private/UObject/PropertyObject.cpp
+++ b/Engine/Source/Runtime/CoreUObject/Private/UObject/PropertyObject.cpp
@@ -246,7 +246,7 @@ void FObjectProperty::SerializeItem(FStructuredArchive::FSlot Slot, void* Value,
 
 void FObjectProperty::PostSerializeObjectItem(FArchive& SerializingArchive, void* Value, UObject* CurrentValue, UObject* ObjectValue, EObjectPropertyOptions Options /*= EObjectPropertyOptions::None*/) const
 {
-	// Make sure non-nullable properties don't end up with null values
+	// Make sure non-nullable properties don't end up with null values.
 	if (!(Options & EObjectPropertyOptions::AllowNullValuesOnNonNullableProperty) &&
 		!ObjectValue && HasAnyPropertyFlags(CPF_NonNullable) &&
 		!SerializingArchive.IsSerializingDefaults() && // null values when Serializing CDOs are allowed, they will be fixed up later
