EditorBulkDataTests: Fix non-unity compile error
#rnx

diff --git a/Engine/Source/Runtime/CoreUObject/Tests/Serialization/EditorBulkDataTests.cpp b/Engine/Source/Runtime/CoreUObject/Tests/Serialization/EditorBulkDataTests.cpp
--- a/Engine/Source/Runtime/CoreUObject/Tests/Serialization/EditorBulkDataTests.cpp
+++ b/Engine/Source/Runtime/CoreUObject/Tests/Serialization/EditorBulkDataTests.cpp
@@ -1,5 +1,9 @@
 // Copyright Epic Games, Inc. All Rights Reserved.
 
+#if WITH_LOW_LEVEL_TESTS
+#include "Misc/CoreMiscDefines.h" // WITH_EDITORONLY_DATA
+#endif
+
 #if WITH_LOW_LEVEL_TESTS && WITH_EDITORONLY_DATA
 
 #include "Algo/AllOf.h"
