Add a new feature allowing bulkdata cooked output to be split into smaller files to help with data organization and patching efforts.

- To enable the feature call FBulkData::SetCookedIndex with a FBulkDataCookedIndex set to a value between 1 - 255. Zero is currently reserved as the default/off state and can be quickly accessed via FBulkDataCookedIndex::Default.
-- Note that we might change the default value in the future, the main reason to keep it as zero for now is that it means FChunkId values will remain unchanged for bulkdata files not using the feature.
- When a bulkdata object has a cooked index it will output to a file with that value based on the following format  <packagename>.CookedIndex.<extension> so a normal bulkdata payload with a cooked index of 5 would end up writing to <packagename>.005.ubulk.
-- This allows the calling systems to control which bulkdata payloads go to which sub files.
- We currently do not support memory mapped payloads or payloads with the duplicate non optional flags. Support and testing for this will be added later.
- Tested saving/editing/loading packages with bulkdata in the editor (vector fields), build/cook/run normal builds, build/cook/run with feature enabled then running the new code with data produced from non modified code and running non modified exe on data generated with the new code.

### IPackageResourceManager
- Added overloads for most methods that take EPackageSegment that also take a FBulkDataCookedIndex and deprecated the older versions.
- Not all methods have been ported over, just the ones I could test but the rest will need the same treatment at some point.

### FLinkerSave
- Now stores each set of bulkdata, optional bulkdata and memory mapped payloads in separate archives, one per cooked index.
- Added a method ::HasCookedIndexBulkData that returns if any of the normal bulkdata payloads contain a non default cooked index. This is used for some paranoid checks when saving packages to the workspace domain.

diff --git a/Engine/Plugins/Editor/DataValidation/Source/DataValidation/Private/PackageFileValidator.cpp b/Engine/Plugins/Editor/DataValidation/Source/DataValidation/Private/PackageFileValidator.cpp
--- a/Engine/Plugins/Editor/DataValidation/Source/DataValidation/Private/PackageFileValidator.cpp
+++ b/Engine/Plugins/Editor/DataValidation/Source/DataValidation/Private/PackageFileValidator.cpp
@@ -185,7 +185,7 @@ bool UPackageFileValidator::TryResolvePackagePath(FName PackageName, FPackagePat
 		return false;
 	}
 
-	if (!IPackageResourceManager::Get().DoesPackageExist(PackagePath, EPackageSegment::Header, &OutPackagePath))
+	if (!IPackageResourceManager::Get().DoesPackageExist(PackagePath, FBulkDataCookedIndex::Default, EPackageSegment::Header, &OutPackagePath))
 	{
 		return false;
 	}
diff --git a/Engine/Source/Developer/IoStoreUtilities/Private/IoStoreUtilities.cpp b/Engine/Source/Developer/IoStoreUtilities/Private/IoStoreUtilities.cpp
--- a/Engine/Source/Developer/IoStoreUtilities/Private/IoStoreUtilities.cpp
+++ b/Engine/Source/Developer/IoStoreUtilities/Private/IoStoreUtilities.cpp
@@ -821,36 +821,51 @@ struct FCookedFileStatData
 	EFileType FileType = Invalid;
 };
 
+/**
+ * Finds where the extension starts in a path.
+ * 
+ * We can't just find the left most '.' in the file name because the path might
+ * contain a bulkdata cooked index (@see FBulkDataCookedIndex) and we don't want
+ * to include the index part. We are attempting to match up the extensions with
+ * those found in the constructor of FCookedFileStatMap.
+ * 
+ * For example take the following filenames and the extensions we want to find:
+ * <packagename>.m.ubulk   -> .m.ubulk
+ * <packagename>.o.ubulk   -> .o.ubulk
+ * <packagename>.001.ubulk ->.ubulk
+ * 
+ * So we want to find the left most '.' unless that is followed by only numeric
+ * values in which case we have found a cooked index and we want to find the '.'
+ * after those values.
+*/
 static int32 GetFullExtensionStartIndex(FStringView Path)
 {
 	int32 ExtensionStartIndex = -1;
 	for (int32 Index = Path.Len() - 1; Index >= 0; --Index)
 	{
+		// Check if we reached the end of the filename
 		if (FPathViews::IsSeparator(Path[Index]))
 		{
 			break;
 		}
 		else if (Path[Index] == '.')
 		{
+			if (ExtensionStartIndex != -1)
+			{
+				// As we have already found a '.' we need to check for the cooked index
+				FStringView Extension = Path.SubStr(Index + 1, (ExtensionStartIndex - Index) - 1);
+				if (UE::String::IsNumericOnlyDigits(Extension))
+				{
+					return ExtensionStartIndex;
+				}
+			}
+
 			ExtensionStartIndex = Index;
 		}
 	}
 	return ExtensionStartIndex;
 }
 
-static FStringView GetBaseFilenameWithoutAnyExtension(FStringView Path)
-{
-	int32 ExtensionStartIndex = GetFullExtensionStartIndex(Path);
-	if (ExtensionStartIndex < 0)
-	{
-		return FStringView();
-	}
-	else
-	{
-		return Path.Left(ExtensionStartIndex);
-	}
-}
-
 static FStringView GetFullExtension(FStringView Path)
 {
 	int32 ExtensionStartIndex = GetFullExtensionStartIndex(Path);
@@ -1014,6 +1029,7 @@ struct FContainerTargetFile
 	uint64 IdealOrder = 0;
 	FIoChunkId ChunkId;
 	FIoHash ChunkHash;
+	FBulkDataCookedIndex BulkDataCookedIndex;
 	TArray<uint8> PackageHeaderData;
 	EContainerChunkType ChunkType;
 	bool bForceUncompressed = false;
@@ -1587,6 +1603,7 @@ static void CreateDiskLayout(
 				SortedTargetFiles.Add(&TargetFile);
 			}
 		}
+
 		check(ShaderTargetFilesMap.Num() == ContainerTarget->GlobalShaders.Num() + ContainerTarget->SharedShaders.Num() + ContainerTarget->UniqueShaders.Num() + ContainerTarget->InlineShaders.Num());
 		Algo::Sort(SortedTargetFiles, [](const FContainerTargetFile* A, const FContainerTargetFile* B)
 		{
@@ -1602,6 +1619,10 @@ static void CreateDiskLayout(
 			{
 				return A->Package->DiskLayoutOrder < B->Package->DiskLayoutOrder;
 			}
+			if (A->BulkDataCookedIndex != B->BulkDataCookedIndex)
+			{
+				return A->BulkDataCookedIndex < B->BulkDataCookedIndex;
+			}
 			check(A == B)
 			return false;
 		});
@@ -2604,6 +2625,9 @@ void InitializeContainerTargetsAndPackages(
 				return false;
 			}
 
+			// TODO - Would be nicer to parse this from the package info rather than inferring it from the path
+			OutTargetFile.BulkDataCookedIndex = FBulkDataCookedIndex::ParseFromPath(SourceFile.NormalizedPath);
+
 			switch (CookedFileStatData->FileType)
 			{
 			case FCookedFileStatData::PackageData:
@@ -2615,17 +2639,17 @@ void InitializeContainerTargetsAndPackages(
 				break;
 			case FCookedFileStatData::BulkData:
 				OutTargetFile.ChunkType = EContainerChunkType::BulkData;
-				OutTargetFile.ChunkId = CreateIoChunkId(OutTargetFile.Package->GlobalPackageId.Value(), 0, EIoChunkType::BulkData);
+				OutTargetFile.ChunkId = CreateBulkDataIoChunkId(OutTargetFile.Package->GlobalPackageId.Value(), 0, OutTargetFile.BulkDataCookedIndex.GetValue(), EIoChunkType::BulkData);
 				OutTargetFile.Package->TotalBulkDataSize += CookedFileStatData->FileSize;
 				break;
 			case FCookedFileStatData::OptionalBulkData:
 				OutTargetFile.ChunkType = EContainerChunkType::OptionalBulkData;
-				OutTargetFile.ChunkId = CreateIoChunkId(OutTargetFile.Package->GlobalPackageId.Value(), 0, EIoChunkType::OptionalBulkData);
+				OutTargetFile.ChunkId = CreateBulkDataIoChunkId(OutTargetFile.Package->GlobalPackageId.Value(), 0, OutTargetFile.BulkDataCookedIndex.GetValue(), EIoChunkType::OptionalBulkData);
 				Package->TotalBulkDataSize += CookedFileStatData->FileSize;
 				break;
 			case FCookedFileStatData::MemoryMappedBulkData:
 				OutTargetFile.ChunkType = EContainerChunkType::MemoryMappedBulkData;
-				OutTargetFile.ChunkId = CreateIoChunkId(OutTargetFile.Package->GlobalPackageId.Value(), 0, EIoChunkType::MemoryMappedBulkData);
+				OutTargetFile.ChunkId = CreateBulkDataIoChunkId(OutTargetFile.Package->GlobalPackageId.Value(), 0, OutTargetFile.BulkDataCookedIndex.GetValue(), EIoChunkType::MemoryMappedBulkData);
 				Package->TotalBulkDataSize += CookedFileStatData->FileSize;
 				break;
 			case FCookedFileStatData::OptionalSegmentPackageData:
@@ -2636,7 +2660,7 @@ void InitializeContainerTargetsAndPackages(
 				break;
 			case FCookedFileStatData::OptionalSegmentBulkData:
 				OutTargetFile.ChunkType = EContainerChunkType::OptionalSegmentBulkData;
-				OutTargetFile.ChunkId = CreateIoChunkId(OutTargetFile.Package->GlobalPackageId.Value(), 1, EIoChunkType::BulkData);
+				OutTargetFile.ChunkId = CreateBulkDataIoChunkId(OutTargetFile.Package->GlobalPackageId.Value(), 1, OutTargetFile.BulkDataCookedIndex.GetValue(), EIoChunkType::BulkData);
 				break;
 			default:
 				UE_LOG(LogIoStore, Fatal, TEXT("Unexpected file type %d for file '%s'"), CookedFileStatData->FileType, *OutTargetFile.NormalizedSourcePath);
@@ -2744,6 +2768,9 @@ void InitializeContainerTargetsAndPackages(
 		}
 		OutTargetFile.Package->PackageStoreEntry = *PackageStoreEntry;
 
+		// TODO - Would be nicer to parse this from the package info rather than inferring it from the path
+		OutTargetFile.BulkDataCookedIndex = FBulkDataCookedIndex::ParseFromPath(SourceFile.NormalizedPath);
+
 		if (Extension == TEXT(".m.ubulk"))
 		{
 			OutTargetFile.ChunkType = EContainerChunkType::MemoryMappedBulkData;
@@ -5267,7 +5294,7 @@ int32 CreateTarget(const FIoStoreArguments& Arguments, const FIoStoreWriterSetti
 			if (TargetFile.Package->PackageStoreEntry.IsAutoOptional())
 			{
 				// Auto optional packages replace the non-optional part when the container is mounted
-				ChunkId = CreateIoChunkId(TargetFile.Package->GlobalPackageId.Value(), 0, EIoChunkType::BulkData);
+				ChunkId = CreateBulkDataIoChunkId(TargetFile.Package->GlobalPackageId.Value(), 0, TargetFile.BulkDataCookedIndex.GetValue(), EIoChunkType::BulkData);
 			}
 			bIsOptionalSegmentChunk = true;
 			break;
diff --git a/Engine/Source/Developer/IoStoreUtilities/Private/PackageStoreOptimizer.cpp b/Engine/Source/Developer/IoStoreUtilities/Private/PackageStoreOptimizer.cpp
--- a/Engine/Source/Developer/IoStoreUtilities/Private/PackageStoreOptimizer.cpp
+++ b/Engine/Source/Developer/IoStoreUtilities/Private/PackageStoreOptimizer.cpp
@@ -639,6 +639,7 @@ void FPackageStoreOptimizer::ProcessDataResources(const FCookedHeaderData& Cooke
 		Entry.DuplicateSerialOffset = DataResource.DuplicateSerialOffset;
 		Entry.SerialSize = DataResource.SerialSize;
 		Entry.Flags = DataResource.LegacyBulkDataFlags;
+		Entry.CookedIndex = DataResource.CookedIndex;
 	}
 }
 
diff --git a/Engine/Source/Editor/UnrealEd/Private/Commandlets/AssetRegistryGenerator.cpp b/Engine/Source/Editor/UnrealEd/Private/Commandlets/AssetRegistryGenerator.cpp
--- a/Engine/Source/Editor/UnrealEd/Private/Commandlets/AssetRegistryGenerator.cpp
+++ b/Engine/Source/Editor/UnrealEd/Private/Commandlets/AssetRegistryGenerator.cpp
@@ -100,6 +100,50 @@ FName GetPackageNameFromDependencyPackageName(const FName RawPackageFName)
 	return PackageFName;
 }
 
+/**
+ * Checks if the provided file path is in the format supported by the BulkData CookedIndex system.
+ * We expect two extensions in the path, the first will be all numbers and the second will be one
+ * of the bulkdata types supported by the system, 'i.e <PackageName>.001.ubulk'.
+ * @see FBulkDataCookedIndex for more info.
+ */
+static bool HasBulkDataCookedIndexExtension(FStringView Path)
+{
+	// Check that the extension at the end of the file is one of the bulkdata types that supports this feature
+	if (!(Path.EndsWith(TEXT(".ubulk")) || Path.EndsWith(TEXT(".uptnl"))) || Path.Len() < 10)
+	{
+		return false;
+	}
+
+	// If the number of max digits changes then our assumptions about ExtensionSize should be reconsidered
+	static_assert(FBulkDataCookedIndex::MAX_DIGITS == 3);
+
+	// A valid extension of this type will always be 10 characters long
+	const int32 ExtensionSize = 10;
+
+	if (Path.Len() < ExtensionSize)
+	{
+		return false;
+	}
+
+	const int32 ExtensionStart = Path.Len() - ExtensionSize;
+
+	if (Path[ExtensionStart] != TEXT('.'))
+	{
+		return false;
+	}
+
+	// Make sure that the first extension only has numeric characters, 0-9
+	for (int32 Index = 1; Index < 4; Index++)
+	{
+		if (Path[ExtensionStart + Index] < TEXT('0') || Path[ExtensionStart + Index] > TEXT('9'))
+		{
+			return false;
+		}
+	}
+
+	return true;
+}
+
 class FDefaultPakFileRules
 {
 public:
@@ -1908,6 +1952,13 @@ public:
 				BaseAssetSourcePathBuffer.RemoveFromEnd(TEXT(".m.ubulk"));
 				BaseAssetSourcePath = &BaseAssetSourcePathBuffer;
 			}
+			else if (HasBulkDataCookedIndexExtension(StandardAssetSourcePath))
+			{
+				// 10 characters equals '.XXX.ubulk' or '.XXX.uptnl' extensions
+				BaseAssetSourcePathBuffer = StandardAssetSourcePath.LeftChop(10);
+				BaseAssetSourcePath = &BaseAssetSourcePathBuffer;
+			}
+
 			if (FPackageName::TryConvertFilenameToLongPackageName(*BaseAssetSourcePath, PackageName))
 			{
 				PackageExtensions.AddUnique(PackageName, StandardAssetSourcePath);
diff --git a/Engine/Source/Editor/UnrealEd/Private/EditorDomain/EditorDomain.cpp b/Engine/Source/Editor/UnrealEd/Private/EditorDomain/EditorDomain.cpp
--- a/Engine/Source/Editor/UnrealEd/Private/EditorDomain/EditorDomain.cpp
+++ b/Engine/Source/Editor/UnrealEd/Private/EditorDomain/EditorDomain.cpp
@@ -215,10 +215,10 @@ bool FEditorDomain::SupportsPackageOnlyPaths()
 	return true;
 }
 
-bool FEditorDomain::DoesPackageExist(const FPackagePath& PackagePath, EPackageSegment PackageSegment,
-	FPackagePath* OutUpdatedPath)
+bool FEditorDomain::DoesPackageExist(const FPackagePath& PackagePath, FBulkDataCookedIndex CookedIndex,
+	EPackageSegment PackageSegment, FPackagePath* OutUpdatedPath)
 {
-	return Workspace->DoesPackageExist(PackagePath, PackageSegment, OutUpdatedPath);
+	return Workspace->DoesPackageExist(PackagePath, CookedIndex, PackageSegment, OutUpdatedPath);
 }
 
 FEditorDomain::FLocks::FLocks(FEditorDomain& InOwner)
@@ -384,15 +384,17 @@ void FEditorDomain::MarkLoadedFromEditorDomain(const FPackagePath& PackagePath,
 	PackageSource->bHasRecordInEditorDomain = true;
 }
 
-int64 FEditorDomain::FileSize(const FPackagePath& PackagePath, EPackageSegment PackageSegment,
-	FPackagePath* OutUpdatedPath)
+int64 FEditorDomain::FileSize(const FPackagePath& PackagePath, FBulkDataCookedIndex CookedIndex,
+	EPackageSegment PackageSegment, FPackagePath* OutUpdatedPath)
 {
 	using namespace UE::EditorDomain;
 	using namespace UE::DerivedData;
 
+	checkf(CookedIndex.IsDefault(), TEXT("Editor Domain does not support bulkdata cooked indices"));
+
 	if (PackageSegment != EPackageSegment::Header)
 	{
-		return Workspace->FileSize(PackagePath, PackageSegment, OutUpdatedPath);
+		return Workspace->FileSize(PackagePath, FBulkDataCookedIndex::Default, PackageSegment, OutUpdatedPath);
 	}
 
 	TOptional<UE::DerivedData::FRequestOwner> Owner;
@@ -403,13 +405,13 @@ int64 FEditorDomain::FileSize(const FPackagePath& PackagePath, EPackageSegment P
 		FName PackageName = PackagePath.GetPackageFName();
 		if (PackageName.IsNone())
 		{
-			return Workspace->FileSize(PackagePath, PackageSegment, OutUpdatedPath);
+			return Workspace->FileSize(PackagePath, FBulkDataCookedIndex::Default, PackageSegment, OutUpdatedPath);
 		}
 
 		bool bReenteredLock;
 		if (!TryFindOrAddPackageSource(ScopeLock, bReenteredLock, PackageName, PackageSource) || PackageSource->Source == EPackageSource::Workspace)
 		{
-			return Workspace->FileSize(PackagePath, PackageSegment, OutUpdatedPath);
+			return Workspace->FileSize(PackagePath, FBulkDataCookedIndex::Default, PackageSegment, OutUpdatedPath);
 		}
 		PackageSource->SetHasLoaded();
 
@@ -454,7 +456,7 @@ int64 FEditorDomain::FileSize(const FPackagePath& PackagePath, EPackageSegment P
 			else if (PackageSource->Source == FEditorDomain::EPackageSource::Workspace || bLoadFromWorkspace)
 			{
 				EditorDomain->MarkLoadedFromWorkspaceDomain(PackagePath, PackageSource, bHasRecordInEditorDomain);
-				FileSize = EditorDomain->Workspace->FileSize(PackagePath, PackageSegment, OutUpdatedPath);
+				FileSize = EditorDomain->Workspace->FileSize(PackagePath, FBulkDataCookedIndex::Default, PackageSegment, OutUpdatedPath);
 			}
 			else
 			{
@@ -473,27 +475,29 @@ int64 FEditorDomain::FileSize(const FPackagePath& PackagePath, EPackageSegment P
 	return FileSize;
 }
 
-FOpenPackageResult FEditorDomain::OpenReadPackage(const FPackagePath& PackagePath, EPackageSegment PackageSegment,
-	FPackagePath* OutUpdatedPath)
+FOpenPackageResult FEditorDomain::OpenReadPackage(const FPackagePath& PackagePath, FBulkDataCookedIndex CookedIndex,
+	EPackageSegment PackageSegment, FPackagePath* OutUpdatedPath)
 {
 	using namespace UE::EditorDomain;
 	using namespace UE::DerivedData;
 
+	checkf(CookedIndex.IsDefault(), TEXT("Editor Domain does not support bulkdata cooked indices"));
+
 	FScopeLock ScopeLock(&Locks->Lock);
 	if (PackageSegment != EPackageSegment::Header)
 	{
-		return Workspace->OpenReadPackage(PackagePath, PackageSegment, OutUpdatedPath);
+		return Workspace->OpenReadPackage(PackagePath, CookedIndex, PackageSegment, OutUpdatedPath);
 	}
 	FName PackageName = PackagePath.GetPackageFName();
 	if (PackageName.IsNone())
 	{
-		return Workspace->OpenReadPackage(PackagePath, PackageSegment, OutUpdatedPath);
+		return Workspace->OpenReadPackage(PackagePath, CookedIndex, PackageSegment, OutUpdatedPath);
 	}
 	TRefCountPtr<FPackageSource> PackageSource;
 	bool bReenteredLock;
 	if (!TryFindOrAddPackageSource(ScopeLock, bReenteredLock, PackageName, PackageSource) || (PackageSource->Source == EPackageSource::Workspace))
 	{
-		return Workspace->OpenReadPackage(PackagePath, PackageSegment, OutUpdatedPath);
+		return Workspace->OpenReadPackage(PackagePath, CookedIndex, PackageSegment, OutUpdatedPath);
 	}
 	PackageSource->SetHasLoaded();
 
@@ -531,28 +535,30 @@ FOpenPackageResult FEditorDomain::OpenReadPackage(const FPackagePath& PackagePat
 	return FOpenPackageResult{ TUniquePtr<FArchive>(Result), Format, bNeedsEngineVersionChecks};
 }
 
-FOpenAsyncPackageResult FEditorDomain::OpenAsyncReadPackage(const FPackagePath& PackagePath, EPackageSegment PackageSegment)
+FOpenAsyncPackageResult FEditorDomain::OpenAsyncReadPackage(const FPackagePath& PackagePath, FBulkDataCookedIndex CookedIndex, EPackageSegment PackageSegment)
 {
 	using namespace UE::EditorDomain;
 	using namespace UE::DerivedData;
 
+	checkf(CookedIndex.IsDefault(), TEXT("Editor Domain does not support bulkdata cooked indices"));
+
 	FScopeLock ScopeLock(&Locks->Lock);
 	if (PackageSegment != EPackageSegment::Header)
 	{
-		return Workspace->OpenAsyncReadPackage(PackagePath, PackageSegment);
+		return Workspace->OpenAsyncReadPackage(PackagePath, CookedIndex, PackageSegment);
 	}
 
 	FName PackageName = PackagePath.GetPackageFName();
 	if (PackageName.IsNone())
 	{
-		return Workspace->OpenAsyncReadPackage(PackagePath, PackageSegment);
+		return Workspace->OpenAsyncReadPackage(PackagePath, CookedIndex, PackageSegment);
 	}
 	TRefCountPtr<FPackageSource> PackageSource;
 	bool bReenteredLock;
 	if (!TryFindOrAddPackageSource(ScopeLock, bReenteredLock, PackageName, PackageSource) ||
 		(PackageSource->Source == EPackageSource::Workspace))
 	{
-		return Workspace->OpenAsyncReadPackage(PackagePath, PackageSegment);
+		return Workspace->OpenAsyncReadPackage(PackagePath, CookedIndex, PackageSegment);
 	}
 	PackageSource->SetHasLoaded();
 
diff --git a/Engine/Source/Editor/UnrealEd/Private/EditorDomain/EditorDomain.h b/Engine/Source/Editor/UnrealEd/Private/EditorDomain/EditorDomain.h
--- a/Engine/Source/Editor/UnrealEd/Private/EditorDomain/EditorDomain.h
+++ b/Engine/Source/Editor/UnrealEd/Private/EditorDomain/EditorDomain.h
@@ -142,14 +142,14 @@ public:
 	// IPackageResourceManager interface
 	virtual bool SupportsLocalOnlyPaths() override;
 	virtual bool SupportsPackageOnlyPaths() override;
-	virtual bool DoesPackageExist(const FPackagePath& PackagePath, EPackageSegment PackageSegment,
-		FPackagePath* OutUpdatedPath = nullptr) override;
-	virtual int64 FileSize(const FPackagePath& PackagePath, EPackageSegment PackageSegment,
-		FPackagePath* OutUpdatedPath = nullptr) override;
-	virtual FOpenPackageResult OpenReadPackage(const FPackagePath& PackagePath, EPackageSegment PackageSegment,
-		FPackagePath* OutUpdatedPath = nullptr) override;
+	virtual bool DoesPackageExist(const FPackagePath& PackagePath, FBulkDataCookedIndex CookedIndex,
+		EPackageSegment PackageSegment, FPackagePath* OutUpdatedPath = nullptr) override;
+	virtual int64 FileSize(const FPackagePath& PackagePath, FBulkDataCookedIndex CookedIndex,
+		EPackageSegment PackageSegment, FPackagePath* OutUpdatedPath = nullptr) override;
+	virtual FOpenPackageResult OpenReadPackage(const FPackagePath& PackagePath, FBulkDataCookedIndex CookedIndex,
+		EPackageSegment PackageSegment, FPackagePath* OutUpdatedPath = nullptr) override;
 	virtual FOpenAsyncPackageResult OpenAsyncReadPackage(const FPackagePath& PackagePath,
-		EPackageSegment PackageSegment) override;
+		FBulkDataCookedIndex CookedIndex, EPackageSegment PackageSegment) override;
 	virtual IMappedFileHandle* OpenMappedHandleToPackage(const FPackagePath& PackagePath,
 		EPackageSegment PackageSegment, FPackagePath* OutUpdatedPath = nullptr) override;
 	virtual bool TryMatchCaseOnDisk(const FPackagePath& PackagePath, FPackagePath* OutNormalizedPath = nullptr) override;
diff --git a/Engine/Source/Editor/UnrealEd/Private/EditorDomain/EditorDomainArchive.cpp b/Engine/Source/Editor/UnrealEd/Private/EditorDomain/EditorDomainArchive.cpp
--- a/Engine/Source/Editor/UnrealEd/Private/EditorDomain/EditorDomainArchive.cpp
+++ b/Engine/Source/Editor/UnrealEd/Private/EditorDomain/EditorDomainArchive.cpp
@@ -522,7 +522,7 @@ bool FEditorDomainReadArchive::TryCreateFallbackData(FEditorDomain& EditorDomain
 	// will fail here because we are creating the file handle on the callback thread and using it on the 
 	// interface thread which is on the game thread. Need to make all editor platforms support that, or
 	// have this fallback occur when WaitForReady is called.
-	FOpenPackageResult Result = Workspace.OpenReadPackage(Segments.GetPackagePath(), EPackageSegment::Header);
+	FOpenPackageResult Result = Workspace.OpenReadPackage(Segments.GetPackagePath(), FBulkDataCookedIndex::Default, EPackageSegment::Header);
 	if (Result.Archive)
 	{
 		InnerArchive = MoveTemp(Result.Archive);
@@ -809,7 +809,7 @@ void FEditorDomainAsyncReadFileHandle::CreateSegmentData(bool bValid)
 bool FEditorDomainAsyncReadFileHandle::TryCreateFallbackData(FEditorDomain& EditorDomain)
 {
 	IPackageResourceManager& Workspace = *EditorDomain.Workspace;
-	FOpenAsyncPackageResult Result = Workspace.OpenAsyncReadPackage(Segments.GetPackagePath(), EPackageSegment::Header);
+	FOpenAsyncPackageResult Result = Workspace.OpenAsyncReadPackage(Segments.GetPackagePath(), FBulkDataCookedIndex::Default, EPackageSegment::Header);
 	check(Result.Handle);
 	InnerArchive = MoveTemp(Result.Handle);
 	PackageFormat = Result.Format;
diff --git a/Engine/Source/Programs/AutomationTool/Scripts/CopyBuildToStagingDirectory.Automation.cs b/Engine/Source/Programs/AutomationTool/Scripts/CopyBuildToStagingDirectory.Automation.cs
--- a/Engine/Source/Programs/AutomationTool/Scripts/CopyBuildToStagingDirectory.Automation.cs
+++ b/Engine/Source/Programs/AutomationTool/Scripts/CopyBuildToStagingDirectory.Automation.cs
@@ -495,6 +495,38 @@ namespace AutomationScripts
 			Logger.LogDebug("End Deployment Context **************");
 		}
 
+		private static bool IsBulkDataExtension(string Extension)
+		{
+			return Extension == ".ubulk" || Extension == ".uptnl";
+		}
+
+		private static bool IsOptionalSegment(string Extension)
+		{
+			return Extension == ".o";
+		}
+
+		private static bool IsMemoryMapped(string Extension)
+		{
+			return Extension == ".m";
+		}
+
+		private static string StripPackageExtension(string FullPath)
+		{
+			string Extension = Path.GetExtension(FullPath);
+			string TrimmedPath = FullPath.Substring(0, FullPath.Length - Extension.Length);
+			string SecondExtension = Path.GetExtension(TrimmedPath);
+
+			if (IsBulkDataExtension(Extension) || IsOptionalSegment(SecondExtension) || IsMemoryMapped(SecondExtension))
+			{
+				if(SecondExtension.Length != 0)
+				{
+					TrimmedPath = FullPath.Substring(0, TrimmedPath.Length - SecondExtension.Length);
+				}
+			}
+
+			return TrimmedPath;
+		}
+		
 		private static string GetInternationalizationPreset(ProjectParams Params, ConfigHierarchy PlatformGameConfig, bool bMustExist = true)
 		{
 			// Initialize internationalization preset.
@@ -4741,13 +4773,7 @@ namespace AutomationScripts
 					HashSet<ChunkDefinition> PakList = new HashSet<ChunkDefinition>();
 
 					string OriginalFilename = StagingFile.Key;
-					string NoExtension = CombinePaths(Path.GetDirectoryName(OriginalFilename), Path.GetFileNameWithoutExtension(OriginalFilename));
-					string AdditionalExtension = Path.GetExtension(NoExtension);
-					if (AdditionalExtension == ".m" || AdditionalExtension == ".o")
-					{
-						// Hack around .m.ubulk files having a double extension
-						NoExtension = CombinePaths(Path.GetDirectoryName(OriginalFilename), Path.GetFileNameWithoutExtension(NoExtension));
-					}
+					string NoExtension = StripPackageExtension(OriginalFilename);
 					string OriginalReplaceSlashes = OriginalFilename.Replace('/', '\\');
 					string NoExtensionReplaceSlashes = NoExtension.Replace('/', '\\');
 
diff --git a/Engine/Source/Runtime/Core/Public/IO/IoChunkId.h b/Engine/Source/Runtime/Core/Public/IO/IoChunkId.h
--- a/Engine/Source/Runtime/Core/Public/IO/IoChunkId.h
+++ b/Engine/Source/Runtime/Core/Public/IO/IoChunkId.h
@@ -45,6 +45,18 @@ enum class EIoChunkType : uint8
 
 CORE_API FString LexToString(const EIoChunkType Type);
 
+inline bool IsBulkDataType(const EIoChunkType Type)
+{
+	if (Type == EIoChunkType::BulkData || Type == EIoChunkType::OptionalBulkData || Type == EIoChunkType::MemoryMappedBulkData)
+	{
+		return true;
+	}
+	else
+	{
+		return false;
+	}
+}
+
 /**
  * Identifier to a chunk of data.
  */
@@ -143,10 +155,12 @@ private:
 	uint8	Id[12];
 };
 
+
 /** Creates a chunk identifier (generic -- prefer specialized versions where possible). */
 inline FIoChunkId CreateIoChunkId(uint64 ChunkId, uint16 ChunkIndex, EIoChunkType IoChunkType)
 {
 	checkSlow(IoChunkType != EIoChunkType::ExternalFile);	// Use CreateExternalFileChunkId() instead
+	checkfSlow(IsBulkDataType(IoChunkType) == false, TEXT("Bulkdata types should call CreateBulkDataIoChunkId instead"));
 
 	uint8 Data[12] = {0};
 
@@ -160,6 +174,36 @@ inline FIoChunkId CreateIoChunkId(uint64 ChunkId, uint16 ChunkIndex, EIoChunkTyp
 	return IoChunkId;
 }
 
+
+
+/**
+ * Used to create a FIoChunkId that references a bulkdata chunk.
+ * 
+ * TODO Cannot use FBulkDataChunkGroup directly here as it is in CoreUObject
+ * 
+ * @param ChunkId		FPackageId
+ * @param ChunkIndex	MultiOutputIindex [0 = disabled, 1 enabled]
+ * @param ChunkGroup	FBulkDataCookedIndex
+ * @param IoChunkType	EIoChunkType
+ */
+inline FIoChunkId CreateBulkDataIoChunkId(uint64 ChunkId, uint16 ChunkIndex, uint8 ChunkGroup, EIoChunkType IoChunkType)
+{
+	checkSlow(IoChunkType != EIoChunkType::ExternalFile);	// Use CreateExternalFileChunkId() instead
+	checkfSlow(IsBulkDataType(IoChunkType), TEXT("CreateBulkDataIoChunkId is only intended for bulkdata types"));
+	
+	uint8 Data[12] = { 0 };
+
+	*reinterpret_cast<uint64*>(&Data[0]) = ChunkId;
+	*reinterpret_cast<uint16*>(&Data[8]) = NETWORK_ORDER16(ChunkIndex);
+	*reinterpret_cast<uint8*>(&Data[10]) = ChunkGroup;
+	*reinterpret_cast<uint8*>(&Data[11]) = static_cast<uint8>(IoChunkType);
+
+	FIoChunkId IoChunkId;
+	IoChunkId.Set(Data, 12);
+
+	return IoChunkId;
+}
+
 /** Returns a package data I/O chunk ID for the specified package ID. */
 CORE_API FIoChunkId CreatePackageDataChunkId(const FPackageId& PackageId);
 
diff --git a/Engine/Source/Runtime/CoreUObject/Private/Misc/PackagePath.cpp b/Engine/Source/Runtime/CoreUObject/Private/Misc/PackagePath.cpp
--- a/Engine/Source/Runtime/CoreUObject/Private/Misc/PackagePath.cpp
+++ b/Engine/Source/Runtime/CoreUObject/Private/Misc/PackagePath.cpp
@@ -796,7 +796,7 @@ void FPackagePath::AppendLocalFullPath(FStringBuilderBase& Builder, EPackageSegm
 		{
 			check(PackageSegment == EPackageSegment::Header);
 			FPackagePath PathWithExtension;
-			if (IPackageResourceManager::Get().DoesPackageExist(*this, EPackageSegment::Header, &PathWithExtension))
+			if (IPackageResourceManager::Get().DoesPackageExist(*this, FBulkDataCookedIndex::Default, EPackageSegment::Header, &PathWithExtension))
 			{
 				Extension = PathWithExtension.GetHeaderExtension();
 				// DoesPackageExist should not search for files with non-standard extensions, so we don't handle CustomExtensions here
diff --git a/Engine/Source/Runtime/CoreUObject/Private/Serialization/AsyncLoading.cpp b/Engine/Source/Runtime/CoreUObject/Private/Serialization/AsyncLoading.cpp
--- a/Engine/Source/Runtime/CoreUObject/Private/Serialization/AsyncLoading.cpp
+++ b/Engine/Source/Runtime/CoreUObject/Private/Serialization/AsyncLoading.cpp
@@ -7723,7 +7723,7 @@ FAsyncArchive::FAsyncArchive(const FPackagePath& InPackagePath, FLinkerLoad* InO
 	SetIsLoading(true);
 
 	LogItem(TEXT("Open"));
-	FOpenAsyncPackageResult OpenResult = IPackageResourceManager::Get().OpenAsyncReadPackage(PackagePath, EPackageSegment::Header);
+	FOpenAsyncPackageResult OpenResult = IPackageResourceManager::Get().OpenAsyncReadPackage(PackagePath, FBulkDataCookedIndex::Default, EPackageSegment::Header);
 	Handle = OpenResult.Handle.Release();
 	check(Handle); // OpenAsyncReadPackage guarantees a non-null return value; the handle will fail to read later if the path does not exist
 	if (OpenResult.Format != EPackageFormat::Binary)
@@ -8298,7 +8298,7 @@ void FAsyncArchive::FirstExportStarting()
 
 			HeaderSizeWhenReadingExportsFromSplitFile = HeaderSize;
 
-			FOpenAsyncPackageResult OpenResult = IPackageResourceManager::Get().OpenAsyncReadPackage(PackagePath, EPackageSegment::Exports);
+			FOpenAsyncPackageResult OpenResult = IPackageResourceManager::Get().OpenAsyncReadPackage(PackagePath, FBulkDataCookedIndex::Default, EPackageSegment::Exports);
 			Handle = OpenResult.Handle.Release();
 			check(Handle); // OpenAsyncReadPackage guarantees a non-null return value; the handle will fail to read later if the path does not exist
 
@@ -8322,7 +8322,7 @@ IAsyncReadRequest* FAsyncArchive::MakeEventDrivenPrecacheRequest(int64 Offset, i
 		IAsyncReadFileHandle* NewHandle;
 		{
 			double StartTime = FPlatformTime::Seconds();
-			FOpenAsyncPackageResult OpenResult = IPackageResourceManager::Get().OpenAsyncReadPackage(PackagePath, EPackageSegment::Exports);
+			FOpenAsyncPackageResult OpenResult = IPackageResourceManager::Get().OpenAsyncReadPackage(PackagePath, FBulkDataCookedIndex::Default, EPackageSegment::Exports);
 			NewHandle = OpenResult.Handle.Release();
 			check(NewHandle); // OpenAsyncReadPackage guarantees a non-null return value; the handle will fail to read later if the path does not exist
 			LogItem(TEXT("Open UExp"), Offset - HeaderSizeWhenReadingExportsFromSplitFile, BytesToRead, StartTime);
diff --git a/Engine/Source/Runtime/CoreUObject/Private/Serialization/AsyncLoading2.cpp b/Engine/Source/Runtime/CoreUObject/Private/Serialization/AsyncLoading2.cpp
--- a/Engine/Source/Runtime/CoreUObject/Private/Serialization/AsyncLoading2.cpp
+++ b/Engine/Source/Runtime/CoreUObject/Private/Serialization/AsyncLoading2.cpp
@@ -217,7 +217,8 @@ FArchive& operator<<(FArchive& Ar, FBulkDataMapEntry& BulkDataEntry)
 	Ar << BulkDataEntry.DuplicateSerialOffset;
 	Ar << BulkDataEntry.SerialSize;
 	Ar << BulkDataEntry.Flags;
-	Ar << BulkDataEntry.Pad;
+	Ar << BulkDataEntry.CookedIndex;
+	Ar.Serialize(&BulkDataEntry.Pad, 3);
 
 	return Ar;
 }
@@ -2150,8 +2151,9 @@ public:
 		const uint16 ChunkIndex = bIsOptionalSegment ? 1 : 0; 
 
 		UE::BulkData::Private::FBulkMetaData& Meta = BulkData.BulkMeta;
+		FBulkDataCookedIndex CookedIndex;
 		int64 DuplicateSerialOffset = -1;
-		SerializeBulkMeta(Meta, DuplicateSerialOffset, Params.ElementSize);
+		SerializeBulkMeta(Meta, CookedIndex, DuplicateSerialOffset, Params.ElementSize);
 
 		const bool bIsInline = Meta.HasAnyFlags(BULKDATA_PayloadAtEndOfFile) == false;
 		if (bIsInline)
@@ -2167,14 +2169,45 @@ PRAGMA_ENABLE_DEPRECATION_WARNINGS
 				Serialize(BulkData.ReallocateData(PayloadSize), PayloadSize);
 			}
 		}
+		else if (Meta.HasAnyFlags(BULKDATA_MemoryMappedPayload))
+		{
+#if UE_DISABLE_COOKEDINDEX_FOR_MEMORYMAPPED
+			CookedIndex = FBulkDataCookedIndex::Default;
+#endif //UE_DISABLE_COOKEDINDEX_FOR_MEMORYMAPPED
+
+			BulkData.BulkChunkId = CreateBulkDataIoChunkId(PackageId.Value(), ChunkIndex, CookedIndex.GetValue(), EIoChunkType::MemoryMappedBulkData);
+
+			if (Params.bAttemptMemoryMapping)
+			{
+				TIoStatusOr<FIoMappedRegion> Status = IoDispatcher.OpenMapped(BulkData.BulkChunkId, FIoReadOptions(Meta.GetOffset(), Meta.GetSize()));
+
+				if (Status.IsOk())
+				{
+					FIoMappedRegion Mapping = Status.ConsumeValueOrDie();
+					BulkData.DataAllocation.SetMemoryMappedData(&BulkData, Mapping.MappedFileHandle, Mapping.MappedFileRegion);
+				}
+				else
+				{
+					UE_LOG(LogSerialization, Warning, TEXT("Memory map bulk data from chunk '%s', offset '%lld', size '%lld' FAILED"),
+						*LexToString(BulkData.BulkChunkId), Meta.GetOffset(), Meta.GetSize());
+
+					BulkData.ForceBulkDataResident();
+				}
+			}
+		}
 		else
 		{
 			const EIoChunkType ChunkType = Meta.HasAnyFlags(BULKDATA_OptionalPayload) ? EIoChunkType::OptionalBulkData : EIoChunkType::BulkData;
-			BulkData.BulkChunkId = CreateIoChunkId(PackageId.Value(), ChunkIndex, ChunkType);
+			BulkData.BulkChunkId = CreateBulkDataIoChunkId(PackageId.Value(), ChunkIndex, CookedIndex.GetValue(), ChunkType);
 
 			if (Meta.HasAnyFlags(BULKDATA_DuplicateNonOptionalPayload))
 			{
-				const FIoChunkId OptionalChunkId = CreateIoChunkId(PackageId.Value(), ChunkIndex, EIoChunkType::OptionalBulkData);
+#if UE_DISABLE_COOKEDINDEX_FOR_NONDUPLICATE
+				CookedIndex = FBulkDataCookedIndex::Default;
+#endif //UE_DISABLE_COOKEDINDEX_FOR_NONDUPLICATE
+				
+				const FIoChunkId OptionalChunkId = CreateBulkDataIoChunkId(PackageId.Value(), ChunkIndex, CookedIndex.GetValue(), EIoChunkType::OptionalBulkData);
+
 				if (IoDispatcher.DoesChunkExist(OptionalChunkId))
 				{
 					BulkData.BulkChunkId = OptionalChunkId;
@@ -2183,27 +2216,6 @@ PRAGMA_ENABLE_DEPRECATION_WARNINGS
 					Meta.SetOffset(DuplicateSerialOffset);
 				}
 			}
-			else if (Meta.HasAnyFlags(BULKDATA_MemoryMappedPayload))
-			{
-				BulkData.BulkChunkId = CreateIoChunkId(PackageId.Value(), ChunkIndex, EIoChunkType::MemoryMappedBulkData);
-				if (Params.bAttemptMemoryMapping)
-				{
-					TIoStatusOr<FIoMappedRegion> Status = IoDispatcher.OpenMapped(BulkData.BulkChunkId, FIoReadOptions(Meta.GetOffset(), Meta.GetSize()));
-
-					if (Status.IsOk())
-					{
-						FIoMappedRegion Mapping = Status.ConsumeValueOrDie(); 
-						BulkData.DataAllocation.SetMemoryMappedData(&BulkData, Mapping.MappedFileHandle, Mapping.MappedFileRegion);
-					}
-					else
-					{
-						UE_LOG(LogSerialization, Warning, TEXT("Memory map bulk data from chunk '%s', offset '%lld', size '%lld' FAILED"),
-							*LexToString(BulkData.BulkChunkId), Meta.GetOffset(), Meta.GetSize());
-
-						BulkData.ForceBulkDataResident();
-					}
-				}
-			}
 		}
 
 		return true;
@@ -2211,7 +2223,7 @@ PRAGMA_ENABLE_DEPRECATION_WARNINGS
 	//~ End FArchive::FLinkerLoad Interface
 
 private:
-	inline void SerializeBulkMeta(UE::BulkData::Private::FBulkMetaData& Meta, int64& DuplicateSerialOffset, int32 ElementSize)
+	inline void SerializeBulkMeta(UE::BulkData::Private::FBulkMetaData& Meta, FBulkDataCookedIndex& CookedIndex, int64& DuplicateSerialOffset, int32 ElementSize)
 	{
 		using namespace UE::BulkData::Private;
 		FArchive& Ar = *this;
@@ -2222,6 +2234,7 @@ private:
 		}
 		else
 		{
+			const int32 Size = sizeof(FBulkDataMapEntry);
 			int32 EntryIndex = INDEX_NONE;
 			Ar << EntryIndex;
 			const FBulkDataMapEntry& Entry = HeaderData->BulkDataMap[EntryIndex];
@@ -2241,6 +2254,7 @@ private:
 #endif //!USE_RUNTIME_BULKDATA
 
 			DuplicateSerialOffset = Entry.DuplicateSerialOffset;
+			CookedIndex = Entry.CookedIndex;
 		}
 
 		Meta.AddFlags(static_cast<EBulkDataFlags>(BULKDATA_UsesIoDispatcher | BULKDATA_LazyLoadable));
diff --git a/Engine/Source/Runtime/CoreUObject/Private/Serialization/BulkData.cpp b/Engine/Source/Runtime/CoreUObject/Private/Serialization/BulkData.cpp
--- a/Engine/Source/Runtime/CoreUObject/Private/Serialization/BulkData.cpp
+++ b/Engine/Source/Runtime/CoreUObject/Private/Serialization/BulkData.cpp
@@ -30,6 +30,10 @@
 
 //////////////////////////////////////////////////////////////////////////////
 
+const FBulkDataCookedIndex FBulkDataCookedIndex::Default;
+
+//////////////////////////////////////////////////////////////////////////////
+
 FStringBuilderBase& LexToString(EBulkDataFlags Flags, FStringBuilderBase& Sb)
 {
 	#define TEST_AND_ADD_FLAG(Sb, Flags, Contains)\
@@ -1176,6 +1180,15 @@ uint32 FBulkData::GetBulkDataFlags() const
 	return BulkMeta.GetFlags();
 }
 
+#if WITH_EDITOR
+
+void FBulkData::SetCookedIndex(FBulkDataCookedIndex Index)
+{
+	CookedIndex = Index;
+}
+
+#endif //WITH_EDITOR
+
 /**
 * Gets the current bulk data alignment.
 *
@@ -1248,7 +1261,12 @@ void FBulkData::Serialize(FArchive& Ar, UObject* Owner, bool bAttemptFileMapping
 	
 	check(!bAttemptFileMapping || Ar.IsLoading()); // makes no sense to map unless we are loading
 
-	if (Ar.SerializeBulkData(*this, FBulkDataSerializationParams {Owner, ElementSize, FileRegionType, bAttemptFileMapping}))
+#if !WITH_EDITOR
+	FBulkDataCookedIndex CookedIndex;	// Dummy value to pass into ::SerializeBulkData, as this is the editor we will only
+										// ever be loading and so this value will not be used.
+#endif //!WITH_EDITOR
+
+	if (Ar.SerializeBulkData(*this, FBulkDataSerializationParams {Owner, ElementSize, FileRegionType, bAttemptFileMapping, CookedIndex }))
 	{
 		// Just early out when the archive overrides the serialization of bulk data
 		return;
diff --git a/Engine/Source/Runtime/CoreUObject/Private/Serialization/EditorBulkData.cpp b/Engine/Source/Runtime/CoreUObject/Private/Serialization/EditorBulkData.cpp
--- a/Engine/Source/Runtime/CoreUObject/Private/Serialization/EditorBulkData.cpp
+++ b/Engine/Source/Runtime/CoreUObject/Private/Serialization/EditorBulkData.cpp
@@ -1408,7 +1408,7 @@ FCompressedBuffer FEditorBulkData::LoadFromPackageFile() const
 	TUniquePtr<FArchive> BulkArchive;
 	if (!IsReferencingByPackagePath())
 	{
-		FOpenPackageResult Result = IPackageResourceManager::Get().OpenReadPackage(PackagePath, EPackageSegment::Header);
+		FOpenPackageResult Result = IPackageResourceManager::Get().OpenReadPackage(PackagePath, FBulkDataCookedIndex::Default, EPackageSegment::Header);
 		if (Result.Format == EPackageFormat::Binary)
 		{
 			BulkArchive = MoveTemp(Result.Archive);
@@ -1455,7 +1455,7 @@ FCompressedBuffer FEditorBulkData::LoadFromPackageTrailer() const
 	TUniquePtr<FArchive> BulkArchive;
 	if (!IsReferencingByPackagePath())
 	{
-		FOpenPackageResult Result = IPackageResourceManager::Get().OpenReadPackage(PackagePath, EPackageSegment::Header);
+		FOpenPackageResult Result = IPackageResourceManager::Get().OpenReadPackage(PackagePath, FBulkDataCookedIndex::Default, EPackageSegment::Header);
 		if (Result.Format == EPackageFormat::Binary)
 		{
 			BulkArchive = MoveTemp(Result.Archive);
@@ -1500,7 +1500,7 @@ FCompressedBuffer FEditorBulkData::LoadFromSidecarFileInternal(ErrorVerbosity Ve
 {
 	TRACE_CPUPROFILER_EVENT_SCOPE(FEditorBulkData::LoadFromSidecarFileInternal);
 
-	FOpenPackageResult Result = IPackageResourceManager::Get().OpenReadPackage(PackagePath, EPackageSegment::PayloadSidecar);
+	FOpenPackageResult Result = IPackageResourceManager::Get().OpenReadPackage(PackagePath, FBulkDataCookedIndex::Default, EPackageSegment::PayloadSidecar);
 	if (Result.Archive.IsValid() && Result.Format == EPackageFormat::Binary)
 	{
 		FPackageTrailer Trailer;
diff --git a/Engine/Source/Runtime/CoreUObject/Private/Serialization/PackageWriterToSharedBuffer.cpp b/Engine/Source/Runtime/CoreUObject/Private/Serialization/PackageWriterToSharedBuffer.cpp
--- a/Engine/Source/Runtime/CoreUObject/Private/Serialization/PackageWriterToSharedBuffer.cpp
+++ b/Engine/Source/Runtime/CoreUObject/Private/Serialization/PackageWriterToSharedBuffer.cpp
@@ -142,8 +142,6 @@ void FPackageWriterRecords::ValidateCommit(FPackage& Record, const IPackageWrite
 		uint8 HasBulkDataType[IPackageWriter::FBulkDataInfo::NumTypes]{};
 		for (FBulkData& BulkRecord : Record.BulkDatas)
 		{
-			checkf((HasBulkDataType[(int32)BulkRecord.Info.BulkDataType] & (1 << BulkRecord.Info.MultiOutputIndex)) == 0,
-				TEXT("IPackageWriter->WriteBulkData must not be called with more than one BulkData of the same type."));
 			HasBulkDataType[(int32)BulkRecord.Info.BulkDataType] |= 1 << BulkRecord.Info.MultiOutputIndex;
 		}
 	}
diff --git a/Engine/Source/Runtime/CoreUObject/Private/UObject/LinkerLoad.cpp b/Engine/Source/Runtime/CoreUObject/Private/UObject/LinkerLoad.cpp
--- a/Engine/Source/Runtime/CoreUObject/Private/UObject/LinkerLoad.cpp
+++ b/Engine/Source/Runtime/CoreUObject/Private/UObject/LinkerLoad.cpp
@@ -7064,8 +7064,9 @@ bool FLinkerLoad::SerializeBulkData(FBulkData& BulkData, const FBulkDataSerializ
 	checkf(BulkData.IsUnlocked(), TEXT("Serialize bulk data FAILED, bulk data is locked"));
 
 	FBulkMetaData& Meta = BulkData.BulkMeta;
+	FBulkDataCookedIndex CookedIndex;
 	int64 DuplicateSerialOffset = -1;
-	SerializeBulkMeta(Meta, DuplicateSerialOffset, Params.ElementSize);
+	SerializeBulkMeta(Meta, CookedIndex, DuplicateSerialOffset, Params.ElementSize);
 
 	const bool bLazyLoadable = IsAllowingLazyLoading();
 	if (bLazyLoadable)
@@ -7080,15 +7081,17 @@ bool FLinkerLoad::SerializeBulkData(FBulkData& BulkData, const FBulkDataSerializ
 
 	const bool bExternalResource = Meta.HasAnyFlags(BULKDATA_WorkspaceDomainPayload);
 	EPackageSegment Segment = GetBulkDataPackageSegmentFromFlags(Meta.GetFlags(), IsLoadingFromCookedPackage());  
-	BulkData.BulkChunkId = UE::CreatePackageResourceChunkId(PackagePath.GetPackageFName(), Segment, bExternalResource);
+	BulkData.BulkChunkId = UE::CreatePackageResourceChunkId(PackagePath.GetPackageFName(), Segment, CookedIndex, bExternalResource);
 
 	const bool bIsInline = Meta.HasAnyFlags(BULKDATA_PayloadAtEndOfFile) == false;
 	if (bIsInline)
 	{
+		checkf(CookedIndex.IsDefault(), TEXT("Inline bulkdata cannot be assigned a chunk group!"));
+
 		if (IsLoadingFromCookedPackage())
 		{
 			// Cooked packages are split into .uasset/.exp files and the offset needs to be adjusted accordingly.
-			const int64 PkgHeaderSize = IPackageResourceManager::Get().FileSize(PackagePath,  EPackageSegment::Header);
+			const int64 PkgHeaderSize = IPackageResourceManager::Get().FileSize(PackagePath, CookedIndex, EPackageSegment::Header);
 			Meta.SetOffset(Tell() - PkgHeaderSize);
 		}
 		void* Payload = BulkData.ReallocateData(Meta.GetSize());
@@ -7099,9 +7102,11 @@ bool FLinkerLoad::SerializeBulkData(FBulkData& BulkData, const FBulkDataSerializ
 		// Streaming cooked bulk data / loading from Editor Domain and referencing Workspace domain bulk data
 		if (Meta.HasAnyFlags(BULKDATA_DuplicateNonOptionalPayload))
 		{
-			if (IPackageResourceManager::Get().DoesPackageExist(PackagePath, EPackageSegment::BulkDataOptional))
+			checkf(CookedIndex.IsDefault(), TEXT("Bulkdata with duplicate non optional payloads cannot be assigned a chunk group!"));
+
+			if (IPackageResourceManager::Get().DoesPackageExist(PackagePath, CookedIndex, EPackageSegment::BulkDataOptional))
 			{
-				BulkData.BulkChunkId = UE::CreatePackageResourceChunkId(PackagePath.GetPackageFName(), EPackageSegment::BulkDataOptional, bExternalResource);
+				BulkData.BulkChunkId = UE::CreatePackageResourceChunkId(PackagePath.GetPackageFName(), EPackageSegment::BulkDataOptional, FBulkDataCookedIndex::Default, bExternalResource);
 				Meta.ClearFlags(BULKDATA_DuplicateNonOptionalPayload);
 				Meta.AddFlags(BULKDATA_OptionalPayload);
 				Meta.SetOffset(DuplicateSerialOffset);
@@ -7109,6 +7114,8 @@ bool FLinkerLoad::SerializeBulkData(FBulkData& BulkData, const FBulkDataSerializ
 		}
 		else if (Meta.HasAnyFlags(BULKDATA_MemoryMappedPayload))
 		{
+			checkf(CookedIndex.IsDefault(), TEXT("Bulkdata with memory mapped payloads cannot be assigned a chunk group!"));
+
 			if (bLazyLoadable && Params.bAttemptMemoryMapping)
 			{
 				TUniquePtr<IMappedFileHandle> MappedFile;
@@ -7128,6 +7135,8 @@ bool FLinkerLoad::SerializeBulkData(FBulkData& BulkData, const FBulkDataSerializ
 	}
 	else
 	{
+		checkf(CookedIndex.IsDefault(), TEXT("Bulkdata stored within the same file cannot be assigned a chunk group!"));
+
 		// Streaming uncooked bulk data (editor only)
 		check(IsLoadingFromCookedPackage() == false);
 
@@ -7200,7 +7209,7 @@ bool FLinkerLoad::ShouldSkipProperty(const FProperty* InProperty) const
 	return FArchiveUObject::ShouldSkipProperty(InProperty);
 }
 
-void FLinkerLoad::SerializeBulkMeta(UE::BulkData::Private::FBulkMetaData& Meta, int64& DuplicateSerialOffset, int32 ElementSize)
+void FLinkerLoad::SerializeBulkMeta(UE::BulkData::Private::FBulkMetaData& Meta, FBulkDataCookedIndex& CookedIndex, int64& DuplicateSerialOffset, int32 ElementSize)
 {
 	using namespace UE::BulkData::Private;
 	FArchive& Ar = *this;
@@ -7219,6 +7228,8 @@ void FLinkerLoad::SerializeBulkMeta(UE::BulkData::Private::FBulkMetaData& Meta,
 		Meta.SetSize(DataResource.RawSize);
 		Meta.SetSizeOnDisk(DataResource.SerialSize);
 		DuplicateSerialOffset = DataResource.DuplicateSerialOffset;
+
+		CookedIndex = DataResource.CookedIndex;
 	}
 
 #if WITH_EDITOR
diff --git a/Engine/Source/Runtime/CoreUObject/Private/UObject/LinkerSave.cpp b/Engine/Source/Runtime/CoreUObject/Private/UObject/LinkerSave.cpp
--- a/Engine/Source/Runtime/CoreUObject/Private/UObject/LinkerSave.cpp
+++ b/Engine/Source/Runtime/CoreUObject/Private/UObject/LinkerSave.cpp
@@ -714,27 +714,46 @@ bool FLinkerSave::SerializeBulkData(FBulkData& BulkData, const FBulkDataSerializ
 
 		if (bSaveBulkDataToSeparateFiles && FBulkData::HasFlags(SerializedMeta.Flags, BULKDATA_OptionalPayload))
 		{
-			SerializedMeta.Offset = OptionalBulkDataAr.Tell();
-			SerializedMeta.SizeOnDisk = BulkData.SerializePayload(OptionalBulkDataAr, SerializedMeta.Flags, RegionToUse);
+			FFileRegionMemoryWriter& Ar = GetOptionalBulkDataArchive(Params.CookedIndex);
+
+			SerializedMeta.Offset = Ar.Tell();
+			SerializedMeta.SizeOnDisk = BulkData.SerializePayload(Ar, SerializedMeta.Flags, RegionToUse);
 		}
 		else if (bSaveBulkDataToSeparateFiles && FBulkData::HasFlags(SerializedMeta.Flags, BULKDATA_MemoryMappedPayload) && bSupportsMemoryMapping)
 		{
-			if (int64 Padding = Align(MemoryMappedBulkDataAr.Tell(), MemoryMappingAlignment) - MemoryMappedBulkDataAr.Tell(); Padding > 0)
+#if UE_DISABLE_COOKEDINDEX_FOR_MEMORYMAPPED
+			UE_CLOG(!Params.CookedIndex.IsDefault(), LogLinker, Warning, TEXT("%s: Cooked Index is not supported for MemoryMappedPayloads, value will be ignored"), *LinkerRoot->GetName());
+
+			FFileRegionMemoryWriter& Ar = GetMemoryMappedBulkDataArchive(FBulkDataCookedIndex::Default);
+#else
+			FFileRegionMemoryWriter& Ar = GetMemoryMappedBulkDataArchive(Params.CookedIndex);
+#endif // UE_DISABLE_COOKEDINDEX_FOR_MEMORYMAPPED
+
+			if (int64 Padding = Align(Ar.Tell(), MemoryMappingAlignment) - Ar.Tell(); Padding > 0)
 			{
 				TArray<uint8> Zeros;
 				Zeros.SetNumZeroed(int32(Padding));
-				MemoryMappedBulkDataAr.Serialize(Zeros.GetData(), Padding);
+				Ar.Serialize(Zeros.GetData(), Padding);
 			}
-			SerializedMeta.Offset = MemoryMappedBulkDataAr.Tell();
-			SerializedMeta.SizeOnDisk = BulkData.SerializePayload(MemoryMappedBulkDataAr, SerializedMeta.Flags, RegionToUse);
+			SerializedMeta.Offset = Ar.Tell();
+			SerializedMeta.SizeOnDisk = BulkData.SerializePayload(Ar, SerializedMeta.Flags, RegionToUse);
 		}
 		else
 		{
 			if (bSaveBulkDataToSeparateFiles && FBulkData::HasFlags(SerializedMeta.Flags, BULKDATA_DuplicateNonOptionalPayload))
 			{
+#if UE_DISABLE_COOKEDINDEX_FOR_NONDUPLICATE
+				UE_CLOG(!Params.CookedIndex.IsDefault(), LogLinker, Warning, TEXT("%s: Cooked Index is not supported for DuplicateNonOptionalPayloads, value will be ignored"), *LinkerRoot->GetName());
+
+				FFileRegionMemoryWriter& OptionalAr = GetOptionalBulkDataArchive(FBulkDataCookedIndex::Default);
+#else
+				FFileRegionMemoryWriter& OptionalAr = GetOptionalBulkDataArchive(Params.CookedIndex);
+#endif // UE_DISABLE_COOKEDINDEX_FOR_NONDUPLICATE
+
+
 				SerializedMeta.DuplicateFlags = SerializedMeta.Flags;
-				SerializedMeta.DuplicateOffset = OptionalBulkDataAr.Tell();
-				SerializedMeta.DuplicateSizeOnDisk = BulkData.SerializePayload(OptionalBulkDataAr, SerializedMeta.Flags, RegionToUse);
+				SerializedMeta.DuplicateOffset = OptionalAr.Tell();
+				SerializedMeta.DuplicateSizeOnDisk = BulkData.SerializePayload(OptionalAr, SerializedMeta.Flags, RegionToUse);
 
 				FBulkData::ClearBulkDataFlagsOn(SerializedMeta.DuplicateFlags, BULKDATA_DuplicateNonOptionalPayload);
 				FBulkData::SetBulkDataFlagsOn(SerializedMeta.DuplicateFlags, BULKDATA_OptionalPayload);
@@ -747,8 +766,10 @@ bool FLinkerSave::SerializeBulkData(FBulkData& BulkData, const FBulkDataSerializ
 			}
 			else
 			{
-				SerializedMeta.Offset = BulkDataAr.Tell();
-				SerializedMeta.SizeOnDisk = BulkData.SerializePayload(BulkDataAr, SerializedMeta.Flags, RegionToUse);
+				FFileRegionMemoryWriter& Ar = GetBulkDataArchive(Params.CookedIndex);
+				
+				SerializedMeta.Offset = Ar.Tell();
+				SerializedMeta.SizeOnDisk = BulkData.SerializePayload(Ar, SerializedMeta.Flags, RegionToUse);
 			}
 		}
 
@@ -770,6 +791,7 @@ bool FLinkerSave::SerializeBulkData(FBulkData& BulkData, const FBulkDataSerializ
 	}
 
 	FObjectDataResource& DataResource = DataResourceMap.AddDefaulted_GetRef();
+	DataResource.CookedIndex			= Params.CookedIndex;
 	DataResource.RawSize				= PayloadSize;
 	DataResource.SerialSize				= SerializedMeta.SizeOnDisk;
 	DataResource.SerialOffset			= SerializedMeta.Offset;
@@ -787,6 +809,82 @@ bool FLinkerSave::SerializeBulkData(FBulkData& BulkData, const FBulkDataSerializ
 	return true;
 }
 
+void FLinkerSave::ForEachBulkDataCookedIndex(TUniqueFunction<void(FBulkDataCookedIndex, FFileRegionMemoryWriter&)>&& Func, EBulkDataPayloadType Type) const
+{
+	const TMap<FBulkDataCookedIndex, TUniquePtr<FFileRegionMemoryWriter>>& Map = GetArchives(Type);
+	for (const TPair<FBulkDataCookedIndex, TUniquePtr<FFileRegionMemoryWriter>>& It : Map)
+	{
+		check(It.Value);
+		Func(It.Key, *It.Value);
+	}
+}
+
+FFileRegionMemoryWriter& FLinkerSave::GetBulkDataArchive(FBulkDataCookedIndex CookedIndex)
+{
+	TUniquePtr<FFileRegionMemoryWriter>& Ar = BulkDataAr.FindOrAdd(CookedIndex);
+	if (!Ar.IsValid())
+	{
+		Ar = MakeUnique<FFileRegionMemoryWriter>();
+	}
+	return *Ar.Get();
+}
+
+FFileRegionMemoryWriter& FLinkerSave::GetOptionalBulkDataArchive(FBulkDataCookedIndex CookedIndex)
+{
+	TUniquePtr<FFileRegionMemoryWriter>& Ar = OptionalBulkDataAr.FindOrAdd(CookedIndex);
+	if (!Ar.IsValid())
+	{
+		Ar = MakeUnique<FFileRegionMemoryWriter>();
+	}
+	return *Ar.Get();
+}
+
+FFileRegionMemoryWriter& FLinkerSave::GetMemoryMappedBulkDataArchive(FBulkDataCookedIndex CookedIndex)
+{
+	TUniquePtr<FFileRegionMemoryWriter>& Ar = MemoryMappedBulkDataAr.FindOrAdd(CookedIndex);
+	if (!Ar.IsValid())
+	{
+		Ar = MakeUnique<FFileRegionMemoryWriter>();
+	}
+	return *Ar.Get();
+}
+
+bool FLinkerSave::HasCookedIndexBulkData() const
+{
+	for (const TPair<FBulkDataCookedIndex, TUniquePtr<FFileRegionMemoryWriter>>& Iter : BulkDataAr)
+	{
+		if (!Iter.Key.IsDefault())
+		{
+			return true;
+		}
+	}
+
+	return false;
+}
+
+const TMap<FBulkDataCookedIndex, TUniquePtr<FFileRegionMemoryWriter>>& FLinkerSave::GetArchives(EBulkDataPayloadType Type) const
+{
+	switch (Type)
+	{
+		case EBulkDataPayloadType::Inline:
+		case EBulkDataPayloadType::AppendToExports:
+		case EBulkDataPayloadType::MemoryMapped:
+			return MemoryMappedBulkDataAr;
+			break;
+		case EBulkDataPayloadType::BulkSegment:
+			return BulkDataAr;
+			break;
+		case EBulkDataPayloadType::Optional:
+			return OptionalBulkDataAr;
+			break;
+		default:
+			checkNoEntry();
+	}
+
+	static TMap<FBulkDataCookedIndex, TUniquePtr<FFileRegionMemoryWriter>> NoData;
+	return NoData;
+}
+
 void FLinkerSave::OnPostSaveBulkData()
 {
 #if WITH_EDITOR
diff --git a/Engine/Source/Runtime/CoreUObject/Private/UObject/ObjectResource.cpp b/Engine/Source/Runtime/CoreUObject/Private/UObject/ObjectResource.cpp
--- a/Engine/Source/Runtime/CoreUObject/Private/UObject/ObjectResource.cpp
+++ b/Engine/Source/Runtime/CoreUObject/Private/UObject/ObjectResource.cpp
@@ -388,6 +388,10 @@ void FObjectDataResource::Serialize(FStructuredArchive::FSlot Slot, TArray<FObje
 		FStructuredArchive::FRecord Record = Slot.EnterRecord();
 		
 		Record << SA_VALUE(TEXT("Flags"), D.Flags);
+		if (Version >= (uint32)FObjectDataResource::EVersion::AddedCookedIndex)
+		{
+			Record << SA_VALUE(TEXT("CookedIndex"), D.CookedIndex);
+		}
 		Record << SA_VALUE(TEXT("SerialOffset"), D.SerialOffset);
 		Record << SA_VALUE(TEXT("DuplicateSerialOffset"), D.DuplicateSerialOffset);
 		Record << SA_VALUE(TEXT("SerialSize"), D.SerialSize);
diff --git a/Engine/Source/Runtime/CoreUObject/Private/UObject/PackageResourceIoDispatcherBackend.cpp b/Engine/Source/Runtime/CoreUObject/Private/UObject/PackageResourceIoDispatcherBackend.cpp
--- a/Engine/Source/Runtime/CoreUObject/Private/UObject/PackageResourceIoDispatcherBackend.cpp
+++ b/Engine/Source/Runtime/CoreUObject/Private/UObject/PackageResourceIoDispatcherBackend.cpp
@@ -1,24 +1,30 @@
 // Copyright Epic Games, Inc. All Rights Reserved.
 
 #include "UObject/PackageResourceIoDispatcherBackend.h"
-#include "UObject/PackageResourceManager.h"
+
 #include "Async/AsyncFileHandle.h"
 #include "Async/MappedFileHandle.h"
 #include "Misc/PackageSegment.h"
 #include "Misc/ScopeLock.h"
 #include "ProfilingDebugging/IoStoreTrace.h"
+#include "Serialization/BulkDataCookedIndex.h"
+#include "UObject/PackageResourceManager.h"
 
 namespace UE
 {
 
-FIoChunkId CreatePackageResourceChunkId(const FName& PackageName, EPackageSegment Segment, bool bExternalResource)
+static_assert(sizeof(FBulkDataCookedIndex::ValueType) == sizeof(uint8)); // If FBulkDataCookedIndex changes type then we need update CreatePackageResourceChunkId
+constexpr uint8 CookedIndexByteIdx = 8; // Position in the byte array that we are storing the cooked index in
+
+FIoChunkId CreatePackageResourceChunkId(const FName& PackageName, EPackageSegment Segment, const FBulkDataCookedIndex& CookedIndex, bool bExternalResource)
 {
 	const int32 Index = PackageName.GetComparisonIndex().ToUnstableInt();
 	const int32 Number = PackageName.GetNumber();
 	
 	uint8 Id[12] = {0};
-	FMemory::Memcpy(Id, &Index, sizeof(int32));
-	FMemory::Memcpy(Id + sizeof(int32), &Number, sizeof(int32));
+	FMemory::Memcpy(Id, &Index, sizeof(int32)); // Bytes 0-3
+	FMemory::Memcpy(Id + sizeof(int32), &Number, sizeof(int32)); // Bytes 4-7
+	Id[CookedIndexByteIdx] = CookedIndex.GetValue();
 	Id[9] = uint8(Segment);
 	Id[10] = uint8(bExternalResource);
 	Id[11] = uint8(EIoChunkType::PackageResource);
@@ -58,6 +64,18 @@ bool TryGetPackagePathFromChunkId(const FIoChunkId& ChunkId, FPackagePath& OutPa
 	return false;
 }
 
+bool TryGetPackagePathFromChunkId(const FIoChunkId& ChunkId, FPackagePath& OutPath, EPackageSegment& OutSegment, bool& bExternal, FBulkDataCookedIndex& OutCookedIndex)
+{
+	FName PackageName;
+	if (TryGetPackageNameFromChunkId(ChunkId, PackageName, OutSegment, bExternal))
+	{
+		OutCookedIndex = FBulkDataCookedIndex(ChunkId.GetData()[CookedIndexByteIdx]);
+		return FPackagePath::TryFromPackageName(PackageName, OutPath);
+	}
+
+	return false;
+}
+
 inline EAsyncIOPriorityAndFlags ConvertToAsyncIoPriority(const int32 IoDispatcherProperty)
 {
 	if (IoDispatcherProperty < IoDispatcherPriority_Low)
@@ -212,15 +230,18 @@ bool FPackageResourceIoBackend::Resolve(FIoRequestImpl* Request)
 	FPackagePath Path;
 	EPackageSegment Segment;
 	bool bExternalResource = false;
+	FBulkDataCookedIndex CookedIndex;
 
-	if (TryGetPackagePathFromChunkId(ChunkId, Path, Segment, bExternalResource) == false)
+	if (TryGetPackagePathFromChunkId(ChunkId, Path, Segment, bExternalResource, CookedIndex) == false)
 	{
 		return false;
 	}
-	
+
+	checkf(!bExternalResource || CookedIndex.IsDefault(), TEXT("Cannot use 'CookedIndices' with packages in the workspace domain"));
+
 	TUniquePtr<IAsyncReadFileHandle> FileHandle = bExternalResource
 		? ResourceMgr.OpenAsyncReadExternalResource(EPackageExternalResource::WorkspaceDomainFile, Path.GetPackageName()).Handle
-		: ResourceMgr.OpenAsyncReadPackage(Path, Segment).Handle;
+		: ResourceMgr.OpenAsyncReadPackage(Path, CookedIndex, Segment).Handle;
 
 	if (FileHandle.IsValid() == false)
 	{
@@ -303,13 +324,15 @@ bool FPackageResourceIoBackend::DoesChunkExist(const FIoChunkId& ChunkId) const
 	FPackagePath Path;
 	EPackageSegment Segment;
 	bool bExternalResource = false;
+	FBulkDataCookedIndex CookedIndex;
 
-	if (TryGetPackagePathFromChunkId(ChunkId, Path, Segment, bExternalResource) == false)
+	if (TryGetPackagePathFromChunkId(ChunkId, Path, Segment, bExternalResource, CookedIndex) == false)
 	{
 		return false;
 	}
+
 	
-	return ResourceMgr.DoesPackageExist(Path, Segment);
+	return ResourceMgr.DoesPackageExist(Path, CookedIndex, Segment);
 }
 
 TIoStatusOr<uint64> FPackageResourceIoBackend::GetSizeForChunk(const FIoChunkId& ChunkId) const
@@ -317,13 +340,15 @@ TIoStatusOr<uint64> FPackageResourceIoBackend::GetSizeForChunk(const FIoChunkId&
 	FPackagePath Path;
 	EPackageSegment Segment;
 	bool bExternalResource = false;
+	FBulkDataCookedIndex CookedIndex;
 
-	if (TryGetPackagePathFromChunkId(ChunkId, Path, Segment, bExternalResource) == false)
+	if (TryGetPackagePathFromChunkId(ChunkId, Path, Segment, bExternalResource, CookedIndex) == false)
 	{
 		return FIoStatus(EIoErrorCode::NotFound);
 	}
+
 	
-	if (int64 FileSize = ResourceMgr.FileSize(Path, Segment); FileSize > 0)
+	if (int64 FileSize = ResourceMgr.FileSize(Path, CookedIndex, Segment); FileSize > 0)
 	{
 		return static_cast<uint64>(FileSize);
 	}
@@ -348,8 +373,9 @@ TIoStatusOr<FIoMappedRegion> FPackageResourceIoBackend::OpenMapped(const FIoChun
 	FPackagePath Path;
 	EPackageSegment Segment;
 	bool bExternalResource = false;
+	FBulkDataCookedIndex ChunkGroup;
 
-	if (TryGetPackagePathFromChunkId(ChunkId, Path, Segment, bExternalResource) == false)
+	if (TryGetPackagePathFromChunkId(ChunkId, Path, Segment, bExternalResource, ChunkGroup) == false)
 	{
 		return FIoStatus(EIoErrorCode::NotFound);
 	}
diff --git a/Engine/Source/Runtime/CoreUObject/Private/UObject/PackageResourceIoDispatcherBackend.h b/Engine/Source/Runtime/CoreUObject/Private/UObject/PackageResourceIoDispatcherBackend.h
--- a/Engine/Source/Runtime/CoreUObject/Private/UObject/PackageResourceIoDispatcherBackend.h
+++ b/Engine/Source/Runtime/CoreUObject/Private/UObject/PackageResourceIoDispatcherBackend.h
@@ -6,6 +6,7 @@
 #include "Misc/PackageSegment.h"
 #include "Templates/SharedPointer.h"
 
+class FBulkDataCookedIndex;
 class FIoChunkId;
 class IPackageResourceManager;
 
@@ -14,7 +15,7 @@ namespace UE
 
 COREUOBJECT_API TSharedRef<IIoDispatcherBackend> MakePackageResourceIoDispatcherBackend(IPackageResourceManager& Mgr);
 
-COREUOBJECT_API FIoChunkId CreatePackageResourceChunkId(const FName& PackageName, EPackageSegment Segment, bool bExternalResource = false);
+COREUOBJECT_API FIoChunkId CreatePackageResourceChunkId(const FName& PackageName, EPackageSegment Segment, const FBulkDataCookedIndex& CookedIndex, bool bExternalResource);
 
 COREUOBJECT_API bool TryGetPackageNameFromChunkId(const FIoChunkId& ChunkId, FName& OutPackageName, EPackageSegment& OutSegment, bool& bOutExternal); 
 
diff --git a/Engine/Source/Runtime/CoreUObject/Private/UObject/PackageResourceManager.cpp b/Engine/Source/Runtime/CoreUObject/Private/UObject/PackageResourceManager.cpp
--- a/Engine/Source/Runtime/CoreUObject/Private/UObject/PackageResourceManager.cpp
+++ b/Engine/Source/Runtime/CoreUObject/Private/UObject/PackageResourceManager.cpp
@@ -69,22 +69,22 @@ void IPackageResourceManager::Shutdown()
 
 bool IPackageResourceManager::DoesPackageExist(const FPackagePath& PackagePath, FPackagePath* OutUpdatedPath)
 {
-	return DoesPackageExist(PackagePath, EPackageSegment::Header, OutUpdatedPath);
+	return DoesPackageExist(PackagePath, FBulkDataCookedIndex::Default, EPackageSegment::Header, OutUpdatedPath);
 }
 
 int64 IPackageResourceManager::FileSize(const FPackagePath& PackagePath, FPackagePath* OutUpdatedPath)
 {
-	return FileSize(PackagePath, EPackageSegment::Header, OutUpdatedPath);
+	return FileSize(PackagePath, FBulkDataCookedIndex::Default, EPackageSegment::Header, OutUpdatedPath);
 }
 
 FOpenPackageResult IPackageResourceManager::OpenReadPackage(const FPackagePath& PackagePath, FPackagePath* OutUpdatedPath)
 {
-	return OpenReadPackage(PackagePath, EPackageSegment::Header, OutUpdatedPath);
+	return OpenReadPackage(PackagePath, FBulkDataCookedIndex::Default, EPackageSegment::Header, OutUpdatedPath);
 }
 
 FOpenAsyncPackageResult IPackageResourceManager::OpenAsyncReadPackage(const FPackagePath& PackagePath)
 {
-	return OpenAsyncReadPackage(PackagePath, EPackageSegment::Header);
+	return OpenAsyncReadPackage(PackagePath, FBulkDataCookedIndex::Default, EPackageSegment::Header);
 }
 
 IMappedFileHandle* IPackageResourceManager::OpenMappedHandleToPackage(const FPackagePath& PackagePath, FPackagePath* OutUpdatedPath)
diff --git a/Engine/Source/Runtime/CoreUObject/Private/UObject/PackageResourceManagerFile.cpp b/Engine/Source/Runtime/CoreUObject/Private/UObject/PackageResourceManagerFile.cpp
--- a/Engine/Source/Runtime/CoreUObject/Private/UObject/PackageResourceManagerFile.cpp
+++ b/Engine/Source/Runtime/CoreUObject/Private/UObject/PackageResourceManagerFile.cpp
@@ -33,13 +33,13 @@ public:
 		return false;
 	}
 
-	virtual bool DoesPackageExist(const FPackagePath& PackagePath, EPackageSegment PackageSegment,
-		FPackagePath* OutUpdatedPath = nullptr) override;
-	virtual int64 FileSize(const FPackagePath& PackagePath, EPackageSegment PackageSegment,
-		FPackagePath* OutUpdatedPath = nullptr) override;
-	virtual FOpenPackageResult OpenReadPackage(const FPackagePath& PackagePath, EPackageSegment PackageSegment,
-		FPackagePath* OutUpdatedPath = nullptr) override;
-	virtual FOpenAsyncPackageResult OpenAsyncReadPackage(const FPackagePath& PackagePath,
+	virtual bool DoesPackageExist(const FPackagePath& PackagePath, FBulkDataCookedIndex CookedIndex,
+		EPackageSegment PackageSegment, FPackagePath* OutUpdatedPath = nullptr) override;
+	virtual int64 FileSize(const FPackagePath& PackagePath, FBulkDataCookedIndex CookedIndex,
+		EPackageSegment PackageSegment, FPackagePath* OutUpdatedPath = nullptr) override;
+	virtual FOpenPackageResult OpenReadPackage(const FPackagePath& PackagePath, FBulkDataCookedIndex CookedIndex,
+		EPackageSegment PackageSegment, FPackagePath* OutUpdatedPath = nullptr) override;
+	virtual FOpenAsyncPackageResult OpenAsyncReadPackage(const FPackagePath& PackagePath, FBulkDataCookedIndex CookedIndex,
 		EPackageSegment PackageSegment) override;
 	virtual IMappedFileHandle* OpenMappedHandleToPackage(const FPackagePath& PackagePath, EPackageSegment PackageSegment,
 		FPackagePath* OutUpdatedPath = nullptr) override;
@@ -67,6 +67,7 @@ protected:
 	 * assigned the chosen extension.
 	 * PackagePaths without a LocalPath (PackageNameOnly PackagePaths) will result in no calls made to the Callback.
 	 * @param PackagePath The PackagePath used to construct the LocalPath.
+	 * @param CookedIndex The cooked index used by the bulkdata, other segment types can just pass in default
 	 * @param PackageSegment The PackageSegment used to construct the LocalPath.
 	 * @param OutUpdatedPath If non-null and a LocalPath is found for which Callback returns true, will be set equal
 	 *        to a copy of PackagePath, and if a header extension was found, will have the header extension set
@@ -74,8 +75,8 @@ protected:
 	 *        check whether an extension is the correct one (because the FullLocalPath exists)
 	 */
 	template <typename CallbackType>
-	void IteratePossibleFiles(const FPackagePath& PackagePath, EPackageSegment PackageSegment,
-		FPackagePath* OutUpdatedPath, const CallbackType& Callback);
+	void IteratePossibleFiles(const FPackagePath& PackagePath, FBulkDataCookedIndex CookedIndex,
+		EPackageSegment PackageSegment, FPackagePath* OutUpdatedPath, const CallbackType& Callback);
 
 	/**
 	 * Base class of classes used in IteratePackagesInPath functions,
@@ -120,21 +121,27 @@ FPackageResourceManagerFile::FPackageResourceManagerFile()
 }
 
 template <typename CallbackType>
-void FPackageResourceManagerFile::IteratePossibleFiles(const FPackagePath& PackagePath, EPackageSegment PackageSegment,
-	FPackagePath* OutUpdatedPath, const CallbackType& Callback)
+void FPackageResourceManagerFile::IteratePossibleFiles(const FPackagePath& PackagePath, FBulkDataCookedIndex CookedIndex,
+	EPackageSegment PackageSegment, FPackagePath* OutUpdatedPath, const CallbackType& Callback)
 {
 	TStringBuilder<256> FilePath;
 	FilePath << PackagePath.GetLocalBaseFilenameWithPath();
-	const int32 BaseNameLen = FilePath.Len();
+	int32 BaseNameLen = FilePath.Len();
 	if (BaseNameLen == 0)
 	{
 		return;
 	}
 
+	if (!CookedIndex.IsDefault())
+	{
+		FilePath << CookedIndex.GetAsExtension();
+		BaseNameLen = FilePath.Len();
+	}
+
 	for (EPackageExtension Extension : PackagePath.GetPossibleExtensions(PackageSegment))
 	{
 		FilePath.RemoveSuffix(FilePath.Len() - BaseNameLen);
-		FilePath << ((Extension != EPackageExtension::Custom) ? 
+		FilePath << ((Extension != EPackageExtension::Custom) ?
 			FStringView(LexToString(Extension)) : PackagePath.GetExtensionString(EPackageSegment::Header));
 		if (Callback(FilePath.ToString(), Extension))
 		{
@@ -153,12 +160,12 @@ void FPackageResourceManagerFile::IteratePossibleFiles(const FPackagePath& Packa
 	}
 }
 
-bool FPackageResourceManagerFile::DoesPackageExist(const FPackagePath& PackagePath, EPackageSegment PackageSegment,
-	FPackagePath* OutUpdatedPath)
+bool FPackageResourceManagerFile::DoesPackageExist(const FPackagePath& PackagePath, FBulkDataCookedIndex CookedIndex,
+	EPackageSegment PackageSegment, FPackagePath* OutUpdatedPath)
 {
 	bool bResult = false;
 	IFileManager* FileManager = &IFileManager::Get();
-	IteratePossibleFiles(PackagePath, PackageSegment, OutUpdatedPath,
+	IteratePossibleFiles(PackagePath, CookedIndex, PackageSegment, OutUpdatedPath,
 		[FileManager, &bResult](const TCHAR* Filename, EPackageExtension Extension)
 		{
 			if (FileManager->FileExists(Filename))
@@ -171,13 +178,13 @@ bool FPackageResourceManagerFile::DoesPackageExist(const FPackagePath& PackagePa
 	return bResult;
 }
 
-int64 FPackageResourceManagerFile::FileSize(const FPackagePath& PackagePath, EPackageSegment PackageSegment,
-	FPackagePath* OutUpdatedPath)
+int64 FPackageResourceManagerFile::FileSize(const FPackagePath& PackagePath, FBulkDataCookedIndex CookedIndex,
+	EPackageSegment PackageSegment, FPackagePath* OutUpdatedPath)
 {
 	int64 Result = INDEX_NONE;
 
 	IFileManager* FileManager = &IFileManager::Get();
-	IteratePossibleFiles(PackagePath, PackageSegment, OutUpdatedPath,
+	IteratePossibleFiles(PackagePath, CookedIndex, PackageSegment, OutUpdatedPath,
 		[FileManager, &Result](const TCHAR* Filename, EPackageExtension Extension)
 		{
 			Result = FileManager->FileSize(Filename);
@@ -192,12 +199,12 @@ int64 FPackageResourceManagerFile::FileSize(const FPackagePath& PackagePath, EPa
 
 
 FOpenPackageResult FPackageResourceManagerFile::OpenReadPackage(const FPackagePath& PackagePath,
-	EPackageSegment PackageSegment, FPackagePath* OutUpdatedPath)
+	FBulkDataCookedIndex CookedIndex, EPackageSegment PackageSegment, FPackagePath* OutUpdatedPath)
 {
 	FOpenPackageResult Result{ nullptr, EPackageFormat::Binary, true /* bNeedsEngineVersionChecks */};
 
 	IFileManager* FileManager = &IFileManager::Get();
-	IteratePossibleFiles(PackagePath, PackageSegment, OutUpdatedPath,
+	IteratePossibleFiles(PackagePath, CookedIndex, PackageSegment, OutUpdatedPath,
 		[FileManager, &Result](const TCHAR* Filename, EPackageExtension Extension)
 		{
 #if !WITH_TEXT_ARCHIVE_SUPPORT
@@ -227,7 +234,7 @@ FOpenPackageResult FPackageResourceManagerFile::OpenReadPackage(const FPackagePa
 }
 
 FOpenAsyncPackageResult FPackageResourceManagerFile::OpenAsyncReadPackage(const FPackagePath& PackagePath,
-	EPackageSegment PackageSegment)
+	FBulkDataCookedIndex CookedIndex, EPackageSegment PackageSegment)
 {
 	FOpenAsyncPackageResult Result { nullptr, EPackageFormat::Binary, true /* bNeedsEngineVersionChecks */ };
 
@@ -240,7 +247,7 @@ FOpenAsyncPackageResult FPackageResourceManagerFile::OpenAsyncReadPackage(const
 	else if (Extensions.Num() > 1)
 	{
 		FPackagePath UpdatedPackagePath;
-		if (DoesPackageExist(PackagePath, PackageSegment, &UpdatedPackagePath))
+		if (DoesPackageExist(PackagePath, CookedIndex, PackageSegment, &UpdatedPackagePath))
 		{
 			FStringView CustomExtension;
 			Extension = UpdatedPackagePath.GetExtension(PackageSegment, CustomExtension);
@@ -253,6 +260,11 @@ FOpenAsyncPackageResult FPackageResourceManagerFile::OpenAsyncReadPackage(const
 		FilePath << PackagePath.GetLocalBaseFilenameWithPath();
 		if (FilePath.Len() > 0)
 		{
+			if (!CookedIndex.IsDefault())
+			{
+				FilePath << CookedIndex.GetAsExtension();
+			}
+
 			FilePath << ((Extension != EPackageExtension::Custom) ?
 				FStringView(LexToString(Extension)) : PackagePath.GetExtensionString(EPackageSegment::Header));
 			Result.Handle.Reset(FPlatformFileManager::Get().GetPlatformFile().OpenAsyncRead(FilePath.ToString()));
@@ -273,7 +285,7 @@ IMappedFileHandle* FPackageResourceManagerFile::OpenMappedHandleToPackage(const
 	IMappedFileHandle* Result = nullptr;
 
 	IPlatformFile* PlatformFile = &FPlatformFileManager::Get().GetPlatformFile();
-	IteratePossibleFiles(PackagePath, PackageSegment, OutUpdatedPath,
+	IteratePossibleFiles(PackagePath, FBulkDataCookedIndex::Default, PackageSegment, OutUpdatedPath,
 		[PlatformFile, &Result](const TCHAR* Filename, EPackageExtension Extension)
 		{
 			Result = PlatformFile->OpenMapped(Filename);
@@ -287,7 +299,7 @@ bool FPackageResourceManagerFile::TryMatchCaseOnDisk(const FPackagePath& Package
 	IPlatformFile* PlatformFile = &FPlatformFileManager::Get().GetPlatformFile();
 	FString FilenameOnDisk;
 	EPackageExtension ExtensionOnDisk;
-	IteratePossibleFiles(PackagePath, EPackageSegment::Header, nullptr,
+	IteratePossibleFiles(PackagePath, FBulkDataCookedIndex::Default, EPackageSegment::Header, nullptr,
 		[PlatformFile, &FilenameOnDisk, &ExtensionOnDisk](const TCHAR* Filename, EPackageExtension Extension)
 		{
 			// TODO: Optimize this function to only hit the disk once by creating a
@@ -332,7 +344,7 @@ TUniquePtr<FArchive> FPackageResourceManagerFile::OpenReadExternalResource(EPack
 		{
 			return TUniquePtr<FArchive>();
 		}
-		return OpenReadPackage(PackagePath, EPackageSegment::Header).Archive;
+		return OpenReadPackage(PackagePath, FBulkDataCookedIndex::Default, EPackageSegment::Header).Archive;
 	}
 	default:
 		checkNoEntry();
@@ -351,7 +363,7 @@ bool FPackageResourceManagerFile::DoesExternalResourceExist(EPackageExternalReso
 		{
 			return false;
 		}
-		return DoesPackageExist(PackagePath, EPackageSegment::Header);
+		return DoesPackageExist(PackagePath, FBulkDataCookedIndex::Default, EPackageSegment::Header);
 	}
 	default:
 		checkNoEntry();
@@ -371,7 +383,7 @@ FOpenAsyncPackageResult FPackageResourceManagerFile::OpenAsyncReadExternalResour
 		{
 			return FOpenAsyncPackageResult{ TUniquePtr<IAsyncReadFileHandle>(new FAsyncReadFileHandleNull()), EPackageFormat::Binary };
 		}
-		return OpenAsyncReadPackage(PackagePath, EPackageSegment::Header);
+		return OpenAsyncReadPackage(PackagePath, FBulkDataCookedIndex::Default, EPackageSegment::Header);
 	}
 	default:
 		checkNoEntry();
diff --git a/Engine/Source/Runtime/CoreUObject/Private/UObject/SavePackage2.cpp b/Engine/Source/Runtime/CoreUObject/Private/UObject/SavePackage2.cpp
--- a/Engine/Source/Runtime/CoreUObject/Private/UObject/SavePackage2.cpp
+++ b/Engine/Source/Runtime/CoreUObject/Private/UObject/SavePackage2.cpp
@@ -2275,7 +2275,8 @@ ESavePackageResult WriteExports(FStructuredArchive::FRecord& StructuredArchiveRo
 	
 	if (Linker.IsCooking() == false)
 	{
-		VirtualExportsFileOffset += Linker.GetBulkDataArchive().TotalSize();
+		check(!Linker.HasCookedIndexBulkData());
+		VirtualExportsFileOffset += Linker.GetBulkDataArchive(FBulkDataCookedIndex::Default).TotalSize();
 	}
 
 	const bool bIsOptionalRealm = SaveContext.GetCurrentHarvestingRealm() == ESaveRealm::Optional;
@@ -2286,10 +2287,11 @@ ESavePackageResult WriteExports(FStructuredArchive::FRecord& StructuredArchiveRo
 	{
 		// Saving non-inline bulk data to the end of the package (Editor)
 
+		check(!Linker.HasCookedIndexBulkData());
 		check(Linker.IsCooking() == false);
 		check(bIsOptionalRealm == false);
 
-		FFileRegionMemoryWriter& Ar = Linker.GetBulkDataArchive();
+		FFileRegionMemoryWriter& Ar = Linker.GetBulkDataArchive(FBulkDataCookedIndex::Default);
 		if (const int64 TotalSize = Ar.TotalSize(); TotalSize > 0)
 		{
 			FIoBuffer Buffer(FIoBuffer::AssumeOwnership, Ar.ReleaseOwnership(), TotalSize);
@@ -2320,10 +2322,18 @@ ESavePackageResult WriteExports(FStructuredArchive::FRecord& StructuredArchiveRo
 	const FPackageId PackageId = FPackageId::FromName(PackageName);
 	const uint16 MultiOutputIndex = bIsOptionalRealm ? 1 : 0;
 
-	auto GetFilePath = [&SaveContext, bIsOptionalRealm](EPackageExtension Ext) -> FString
+	auto GetFilePath = [&SaveContext, bIsOptionalRealm](EPackageExtension Ext, FBulkDataCookedIndex CookedIndex = FBulkDataCookedIndex::Default) -> FString
 	{
-		FString FileExt = bIsOptionalRealm ? FString(TEXT(".o")) + LexToString(Ext) : LexToString(Ext);
-		return FPathViews::ChangeExtension(SaveContext.GetFilename(), FileExt);
+			if (CookedIndex.IsDefault())
+			{
+				FString FileExt = bIsOptionalRealm ? FString(TEXT(".o")) + LexToString(Ext) : LexToString(Ext);
+				return FPathViews::ChangeExtension(SaveContext.GetFilename(), FileExt);
+			}
+			else
+			{
+				FString FileExt = CookedIndex.GetAsExtension() + LexToString(Ext);
+				return FPathViews::ChangeExtension(SaveContext.GetFilename(), FileExt);
+			}
 	};
 
 	auto WriteToPackageWriter = [&SaveContext, PackageWriter, bIsOptionalRealm](FFileRegionMemoryWriter& Ar, IPackageWriter::FBulkDataInfo Info) -> int64
@@ -2341,38 +2351,51 @@ ESavePackageResult WriteExports(FStructuredArchive::FRecord& StructuredArchiveRo
 	};
 
 	FScopedSlowTask Feedback(3.0f);
-	SaveContext.TotalPackageSizeUncompressed += WriteToPackageWriter(Linker.GetBulkDataArchive(), IPackageWriter::FBulkDataInfo
-	{
-		PackageName,
-		IPackageWriter::FBulkDataInfo::BulkSegment,
-		GetFilePath(EPackageExtension::BulkDataDefault),
-		CreateIoChunkId(PackageId.Value(), MultiOutputIndex, EIoChunkType::BulkData),
-		MultiOutputIndex
-	});
+
+	Linker.ForEachBulkDataCookedIndex([&SaveContext, &PackageName, &PackageId, &WriteToPackageWriter, &GetFilePath, MultiOutputIndex](FBulkDataCookedIndex CookedIndex, FFileRegionMemoryWriter& Ar)
+		{
+			SaveContext.TotalPackageSizeUncompressed += WriteToPackageWriter(Ar, IPackageWriter::FBulkDataInfo
+				{
+					PackageName,
+					IPackageWriter::FBulkDataInfo::BulkSegment,
+					GetFilePath(EPackageExtension::BulkDataDefault, CookedIndex),
+					CreateBulkDataIoChunkId(PackageId.Value(), MultiOutputIndex, CookedIndex.GetValue(), EIoChunkType::BulkData),
+					MultiOutputIndex
+				});
+		}, EBulkDataPayloadType::BulkSegment);
+
 	Feedback.EnterProgressFrame();
 
 	// @note FH: temporarily do not handle optional bulk data into editor optional packages, proper support will be added soon
 	if (bIsOptionalRealm == false)
 	{
-		SaveContext.TotalPackageSizeUncompressed += WriteToPackageWriter(Linker.GetOptionalBulkDataArchive(), IPackageWriter::FBulkDataInfo
-		{
-			PackageName,
-			IPackageWriter::FBulkDataInfo::Optional,
-			GetFilePath(EPackageExtension::BulkDataOptional),
-			CreateIoChunkId(PackageId.Value(), MultiOutputIndex, EIoChunkType::OptionalBulkData),
-			MultiOutputIndex
-		});
+		Linker.ForEachBulkDataCookedIndex([&SaveContext, &PackageName, &PackageId, &WriteToPackageWriter, &GetFilePath, MultiOutputIndex](FBulkDataCookedIndex CookedIndex, FFileRegionMemoryWriter& Ar)
+			{
+				SaveContext.TotalPackageSizeUncompressed += WriteToPackageWriter(Ar, IPackageWriter::FBulkDataInfo
+					{
+						PackageName,
+						IPackageWriter::FBulkDataInfo::Optional,
+						GetFilePath(EPackageExtension::BulkDataOptional, CookedIndex),
+						CreateBulkDataIoChunkId(PackageId.Value(), MultiOutputIndex, CookedIndex.GetValue(), EIoChunkType::OptionalBulkData),
+						MultiOutputIndex
+					});
+			}, EBulkDataPayloadType::Optional);
 	}
+
 	Feedback.EnterProgressFrame();
 
-	SaveContext.TotalPackageSizeUncompressed += WriteToPackageWriter(Linker.GetMemoryMappedBulkDataArchive(), IPackageWriter::FBulkDataInfo
-	{
-		PackageName,
-		IPackageWriter::FBulkDataInfo::Mmap,
-		GetFilePath(EPackageExtension::BulkDataMemoryMapped),
-		CreateIoChunkId(PackageId.Value(), MultiOutputIndex, EIoChunkType::MemoryMappedBulkData),
-		MultiOutputIndex
-	});
+	Linker.ForEachBulkDataCookedIndex([&SaveContext, &PackageName, &PackageId, &WriteToPackageWriter, &GetFilePath, MultiOutputIndex](FBulkDataCookedIndex CookedIndex, FFileRegionMemoryWriter& Ar)
+		{
+			SaveContext.TotalPackageSizeUncompressed += WriteToPackageWriter(Ar, IPackageWriter::FBulkDataInfo
+				{
+					PackageName,
+					IPackageWriter::FBulkDataInfo::Mmap,
+					GetFilePath(EPackageExtension::BulkDataMemoryMapped),
+					CreateBulkDataIoChunkId(PackageId.Value(), MultiOutputIndex, CookedIndex.GetValue(), EIoChunkType::MemoryMappedBulkData),
+					MultiOutputIndex
+				});
+		}, EBulkDataPayloadType::MemoryMapped);
+
 	Feedback.EnterProgressFrame();
 
 	return ESavePackageResult::Success;
diff --git a/Engine/Source/Runtime/CoreUObject/Public/Serialization/AsyncLoading2.h b/Engine/Source/Runtime/CoreUObject/Public/Serialization/AsyncLoading2.h
--- a/Engine/Source/Runtime/CoreUObject/Public/Serialization/AsyncLoading2.h
+++ b/Engine/Source/Runtime/CoreUObject/Public/Serialization/AsyncLoading2.h
@@ -14,6 +14,7 @@
 #include "IO/PackageId.h"
 #include "Misc/AssertionMacros.h"
 #include "Serialization/Archive.h"
+#include "Serialization/BulkDataCookedIndex.h"
 #include "Serialization/CustomVersion.h"
 #include "Serialization/MappedName.h"
 #include "Templates/TypeHash.h"
@@ -365,11 +366,18 @@ struct FBulkDataMapEntry
 	int64 DuplicateSerialOffset = 0;
 	int64 SerialSize = 0;
 	uint32 Flags = 0;
-	uint32 Pad = 0;
+	FBulkDataCookedIndex CookedIndex;
+	uint8 Pad[3] = { 0, 0, 0 };
 	
 	COREUOBJECT_API friend FArchive& operator<<(FArchive& Ar, FBulkDataMapEntry& BulkDataEntry);
 };
 
+// If we change the size of FBulkDataCookedIndex we will need to update FBulkDataMapEntry
+static_assert(sizeof(FBulkDataCookedIndex) == sizeof(uint8));
+
+// We don't want to grow the size of FBulkDataMapEntry accidently
+static_assert(sizeof(FBulkDataMapEntry) <= 32, "The memory layout of FBulkDataMapEntry now exceeds 32 bytes, was this intended?");
+
 COREUOBJECT_API void FindAllRuntimeScriptPackages(TArray<UPackage*>& OutPackages);
 
 /**
diff --git a/Engine/Source/Runtime/CoreUObject/Public/Serialization/BulkData.h b/Engine/Source/Runtime/CoreUObject/Public/Serialization/BulkData.h
--- a/Engine/Source/Runtime/CoreUObject/Public/Serialization/BulkData.h
+++ b/Engine/Source/Runtime/CoreUObject/Public/Serialization/BulkData.h
@@ -26,6 +26,7 @@
 #include "Misc/PackagePath.h"
 #include "Misc/PackageSegment.h"
 #include "Serialization/Archive.h"
+#include "Serialization/BulkDataCookedIndex.h"
 #include "Serialization/CustomVersion.h"
 #include "Serialization/FileRegions.h"
 #include "Templates/Function.h"
@@ -458,6 +459,8 @@ struct FBulkDataSerializationParams
 	EFileRegionType RegionType;
 	/** Flag indicating whether to try to memory map the bulk data payload or not. */
 	bool bAttemptMemoryMapping = false;
+	/** The group to which the cooked version of the bulkdata payload should be assigned*/
+	FBulkDataCookedIndex CookedIndex;
 };
 
 /**
@@ -714,6 +717,9 @@ public:
 	*/
 	COREUOBJECT_API uint32 GetBulkDataFlags() const;
 
+#if WITH_EDITOR
+	COREUOBJECT_API void SetCookedIndex(FBulkDataCookedIndex Index);
+#endif //WITH_EDITOR
 	/**
 	* Gets the current bulk data alignment.
 	*
@@ -1007,9 +1013,11 @@ private:
 protected:
 #if WITH_EDITOR
 	/** Archive associated with bulk data for serialization																*/
-	FArchive*									AttachedAr = nullptr;
+	FArchive*				AttachedAr = nullptr;
 	/** Used to make sure the linker doesn't get garbage collected at runtime for things with attached archives			*/
-	FLinkerLoad*								Linker = nullptr;
+	FLinkerLoad*			Linker = nullptr;
+	/** The index to which the cooked version of the bulkdata payload should be assigned*/
+	FBulkDataCookedIndex	CookedIndex;
 #endif // WITH_EDITOR
 	   //
 #if !USE_RUNTIME_BULKDATA
diff --git a/Engine/Source/Runtime/CoreUObject/Public/Serialization/BulkDataCookedIndex.h b/Engine/Source/Runtime/CoreUObject/Public/Serialization/BulkDataCookedIndex.h
new file mode 100644
--- /dev/null
+++ b/Engine/Source/Runtime/CoreUObject/Public/Serialization/BulkDataCookedIndex.h
@@ -0,0 +1,144 @@
+// Copyright Epic Games, Inc. All Rights Reserved.
+
+#pragma once
+
+#include "CoreTypes.h"
+#include "Containers/StringView.h"
+#include "Containers/UnrealString.h"
+#include "Misc/PathViews.h"
+#include "Serialization/StructuredArchive.h"
+#include "String/LexFromString.h"
+#include "String/Numeric.h"
+
+// When this is enabled the default value for FBulkDataCookedIndex, meaning that the bulkdata does not use the
+// system will be 0, if it is disabled we will use the max value allowed for FBulkDataCookedIndex::Type.
+// If the default is zero then we cannot use <packagename>.000.ubulk as will start indexing from .001.ubulk instead
+// but it also means that the FChunkId for a payload not using the system will remain unchanged.
+#define UE_DEFAULT_ZERO 1
+
+
+#define UE_DISABLE_COOKEDINDEX_FOR_MEMORYMAPPED 1
+
+#define UE_DISABLE_COOKEDINDEX_FOR_NONDUPLICATE 1
+
+enum class EBulkDataPayloadType : uint8
+{
+	Inline,				// Stored inside the export data in .uexp
+	AppendToExports,	// Stored after the export data in .uexp
+	BulkSegment,		// Stored in .ubulk
+	Optional,			// Stored in .uptnl
+	MemoryMapped,		// Stored in .m.bulk
+};
+
+class FBulkDataCookedIndex
+{
+public:
+	// It is likely that we will want to expand the number of bits that this system currently uses when addressed
+	// via FIoChunkIds in the future. The following constants and aliases make it easier to track places in the
+	// code base that make assumptions about this so we can safely update them all at once.
+	using ValueType = uint8;
+	constexpr static int32 MAX_DIGITS = 3;
+
+	COREUOBJECT_API static const FBulkDataCookedIndex Default;
+
+	FBulkDataCookedIndex() = default;
+	explicit FBulkDataCookedIndex(ValueType InValue)
+		: Value(InValue)
+	{
+
+	}
+
+	~FBulkDataCookedIndex() = default;
+
+	bool IsDefault() const
+	{
+#if UE_DEFAULT_ZERO
+		return Value == 0;
+#else
+		return Value == TNumericLimits<ValueType>::Max();
+#endif //UE_DEFAULT_ZERO
+	}
+
+	FString GetAsExtension() const
+	{
+		if (IsDefault())
+		{
+			return FString();
+		}
+		else
+		{
+			return FString::Printf(TEXT(".%03hhu"), Value);
+		}
+	}
+
+	ValueType GetValue() const
+	{
+		return Value;
+	}
+
+	bool operator == (const FBulkDataCookedIndex& Other) const
+	{
+		return Value == Other.Value;
+	}
+
+	bool operator < (const FBulkDataCookedIndex& Other) const
+	{
+		return Value < Other.Value;
+	}
+
+	friend void operator<<(FStructuredArchive::FSlot Slot, FBulkDataCookedIndex& ChunkGroup)
+	{
+		Slot << ChunkGroup.Value;
+	}
+
+	friend uint32 GetTypeHash(const FBulkDataCookedIndex& ChunkGroup)
+	{
+		return GetTypeHash(ChunkGroup.Value);
+	}
+
+	// TODO: Unit tests
+	static FBulkDataCookedIndex ParseFromPath(FStringView Path)
+	{
+		int32 ExtensionStartIndex = -1;
+		for (int32 Index = Path.Len() - 1; Index >= 0; --Index)
+		{
+			if (FPathViews::IsSeparator(Path[Index]))
+			{
+				return FBulkDataCookedIndex();
+			}
+			else if (Path[Index] == '.')
+			{
+				if (ExtensionStartIndex != -1)
+				{
+					FStringView Extension = Path.SubStr(Index + 1, (ExtensionStartIndex - Index) - 1);
+					if (UE::String::IsNumericOnlyDigits(Extension))
+					{
+						ValueType Value = 0;
+						LexFromString(Value, Extension);
+
+						return FBulkDataCookedIndex(Value);
+					}
+					else
+					{
+						return FBulkDataCookedIndex();
+					}
+				}
+				else
+				{
+					ExtensionStartIndex = Index;
+				}
+			}
+		}
+
+		return FBulkDataCookedIndex();
+	}
+
+private:
+#if UE_DEFAULT_ZERO
+	ValueType Value = 0;
+#else
+	ValueType Value = TNumericLimits<ValueType>::Max();
+#endif //UE_DEFAULT_ZERO
+};
+
+#undef UE_DEFAULT_ZERO
diff --git a/Engine/Source/Runtime/CoreUObject/Public/UObject/LinkerLoad.h b/Engine/Source/Runtime/CoreUObject/Public/UObject/LinkerLoad.h
--- a/Engine/Source/Runtime/CoreUObject/Public/UObject/LinkerLoad.h
+++ b/Engine/Source/Runtime/CoreUObject/Public/UObject/LinkerLoad.h
@@ -414,7 +414,7 @@ public:
 
 private:
 
-	void SerializeBulkMeta(UE::BulkData::Private::FBulkMetaData& Meta, int64& DuplicateSerialOffset, int32 ElementSize);
+	void SerializeBulkMeta(UE::BulkData::Private::FBulkMetaData& Meta, FBulkDataCookedIndex& CookedIndex, int64& DuplicateSerialOffset, int32 ElementSize);
 
 	// Variables used during async linker creation.
 
diff --git a/Engine/Source/Runtime/CoreUObject/Public/UObject/LinkerSave.h b/Engine/Source/Runtime/CoreUObject/Public/UObject/LinkerSave.h
--- a/Engine/Source/Runtime/CoreUObject/Public/UObject/LinkerSave.h
+++ b/Engine/Source/Runtime/CoreUObject/Public/UObject/LinkerSave.h
@@ -13,6 +13,7 @@
 #include "IO/IoHash.h"
 #include "Serialization/Archive.h"
 #include "Serialization/ArchiveUObject.h"
+#include "Serialization/BulkDataCookedIndex.h"
 #include "Serialization/FileRegionArchive.h"
 #include "Templates/Function.h"
 #include "Templates/RefCounting.h"
@@ -287,26 +288,39 @@ public:
 #endif // WITH_EDITORONLY_DATA
 
 	virtual bool SerializeBulkData(FBulkData& BulkData, const FBulkDataSerializationParams& Params) override;
-	
+
+	void ForEachBulkDataCookedIndex(TUniqueFunction<void(FBulkDataCookedIndex, FFileRegionMemoryWriter&)>&& Func, EBulkDataPayloadType Type) const;
+
+	FFileRegionMemoryWriter& GetBulkDataArchive(FBulkDataCookedIndex CookedIndex);
+	FFileRegionMemoryWriter& GetOptionalBulkDataArchive(FBulkDataCookedIndex CookedIndex);
+	FFileRegionMemoryWriter& GetMemoryMappedBulkDataArchive(FBulkDataCookedIndex CookedIndex);
+
+	bool HasCookedIndexBulkData() const;
+
+	UE_DEPRECATED(5.5, "Use the overload that takes a FBulkDataCookedIndex")
 	FFileRegionMemoryWriter& GetBulkDataArchive()
 	{
-		return BulkDataAr;
+		return GetBulkDataArchive(FBulkDataCookedIndex::Default);
 	}
 
+	UE_DEPRECATED(5.5, "Use the overload that takes a FBulkDataCookedIndex")
 	FFileRegionMemoryWriter& GetOptionalBulkDataArchive()
 	{
-		return OptionalBulkDataAr;
+		return GetOptionalBulkDataArchive(FBulkDataCookedIndex::Default);
 	}
 
+	UE_DEPRECATED(5.5, "Use the overload that takes a FBulkDataCookedIndex")
 	FFileRegionMemoryWriter& GetMemoryMappedBulkDataArchive()
 	{
-		return MemoryMappedBulkDataAr;
+		return GetMemoryMappedBulkDataArchive(FBulkDataCookedIndex::Default);
 	}
 
 protected:
 	/** Set the filename being saved to */
 	void SetFilename(FStringView InFilename);
 
+	const TMap<FBulkDataCookedIndex, TUniquePtr<FFileRegionMemoryWriter>>& GetArchives(EBulkDataPayloadType Type) const;
+
 private:
 	/** Optional log output to bubble errors back up. */
 	FOutputDevice* LogOutput = nullptr;
@@ -322,11 +336,12 @@ private:
 #endif //WITH_EDITOR
 
 	/** Default bulk data archive. */
-	FFileRegionMemoryWriter BulkDataAr;
+	TMap<FBulkDataCookedIndex, TUniquePtr<FFileRegionMemoryWriter>> BulkDataAr; // Can't use TUniqueObj as FFileRegionMemoryWriter isn't movable
 	/** Optional bulk data archive. */
-	FFileRegionMemoryWriter OptionalBulkDataAr;
+	TMap<FBulkDataCookedIndex, TUniquePtr<FFileRegionMemoryWriter>> OptionalBulkDataAr;
 	/** Memory mapped bulk data archive. */
-	FFileRegionMemoryWriter MemoryMappedBulkDataAr;
+	TMap<FBulkDataCookedIndex, TUniquePtr<FFileRegionMemoryWriter>> MemoryMappedBulkDataAr;
+
 	const TMap<UObject*, TSet<FProperty*>>* TransientPropertyOverrides = nullptr;
 	/** Alignment for memory mapped data .*/ 
 	int64 MemoryMappingAlignment = -1;
diff --git a/Engine/Source/Runtime/CoreUObject/Public/UObject/ObjectResource.h b/Engine/Source/Runtime/CoreUObject/Public/UObject/ObjectResource.h
--- a/Engine/Source/Runtime/CoreUObject/Public/UObject/ObjectResource.h
+++ b/Engine/Source/Runtime/CoreUObject/Public/UObject/ObjectResource.h
@@ -9,6 +9,7 @@
 #include "Misc/CString.h"
 #include "Misc/EnumClassFlags.h"
 #include "Misc/Guid.h"
+#include "Serialization/BulkDataCookedIndex.h"
 #include "Serialization/StructuredArchive.h"
 #include "Serialization/StructuredArchiveAdapters.h"
 #include "Serialization/StructuredArchiveSlots.h"
@@ -553,12 +554,15 @@ struct FObjectDataResource
 	{
 		Invalid,
 		Initial,
+		AddedCookedIndex,
 		LatestPlusOne,
 		Latest = LatestPlusOne - 1
 	};
 
 	/** Data resource flags. */
 	EObjectDataResourceFlags Flags = EObjectDataResourceFlags::None;
+	/** The cooked index to which the bulkdata payload belongs */
+	FBulkDataCookedIndex CookedIndex;
 	/** Location of the data in the underlying storage type. */
 	int64 SerialOffset = -1;
 	/** Location of the data in the underlying storage type if this data is duplicated. */
diff --git a/Engine/Source/Runtime/CoreUObject/Public/UObject/PackageResourceManager.h b/Engine/Source/Runtime/CoreUObject/Public/UObject/PackageResourceManager.h
--- a/Engine/Source/Runtime/CoreUObject/Public/UObject/PackageResourceManager.h
+++ b/Engine/Source/Runtime/CoreUObject/Public/UObject/PackageResourceManager.h
@@ -18,6 +18,7 @@
 #include "Misc/PackageSegment.h"
 #include "Modules/ModuleInterface.h"
 #include "Serialization/Archive.h"
+#include "Serialization/BulkDataCookedIndex.h"
 #include "Templates/Function.h"
 #include "Templates/SharedPointer.h"
 #include "Templates/UniquePtr.h"
@@ -340,21 +341,47 @@ public:
 
 public:
 	// Internal API used by low-level PackageResourceManager users
+	UE_DEPRECATED(5.5, "Use the overload that takes a FBulkDataCookedIndex")
+	bool DoesPackageExist(const FPackagePath& PackagePath, EPackageSegment PackageSegment,
+		FPackagePath* OutUpdatedPath = nullptr)
+	{
+		return DoesPackageExist(PackagePath, FBulkDataCookedIndex::Default, PackageSegment, OutUpdatedPath);
+	}
 
 	/** DoesPackageExist that takes a PackageSegment */
-	virtual bool DoesPackageExist(const FPackagePath& PackagePath, EPackageSegment PackageSegment,
-		FPackagePath* OutUpdatedPath = nullptr) = 0;
+	virtual bool DoesPackageExist(const FPackagePath& PackagePath, FBulkDataCookedIndex CookedIndex,
+		EPackageSegment PackageSegment, FPackagePath* OutUpdatedPath = nullptr) = 0;
+
+	UE_DEPRECATED(5.5, "Use the overload that takes a FBulkDataCookedIndex")
+	int64 FileSize(const FPackagePath& PackagePath, EPackageSegment PackageSegment,
+		FPackagePath* OutUpdatedPath = nullptr)
+	{
+		return FileSize(PackagePath, FBulkDataCookedIndex::Default, PackageSegment, OutUpdatedPath);
+	}
 
 	/** FileSize that takes a PackageSegment */
-	virtual int64 FileSize(const FPackagePath& PackagePath, EPackageSegment PackageSegment,
-		FPackagePath* OutUpdatedPath = nullptr) = 0;
+	virtual int64 FileSize(const FPackagePath& PackagePath, FBulkDataCookedIndex CookedIndex,
+		EPackageSegment PackageSegment, FPackagePath* OutUpdatedPath = nullptr) = 0;
+
+	UE_DEPRECATED(5.5, "Use the overload that takes a FBulkDataCookedIndex")
+	FOpenPackageResult OpenReadPackage(const FPackagePath& PackagePath, EPackageSegment PackageSegment,
+		FPackagePath* OutUpdatedPath = nullptr)
+	{
+		return OpenReadPackage(PackagePath, FBulkDataCookedIndex::Default, PackageSegment, OutUpdatedPath);
+	}
 
 	/** OpenReadPackage that takes a PackageSegment */
-	virtual FOpenPackageResult OpenReadPackage(const FPackagePath& PackagePath, EPackageSegment PackageSegment,
-		FPackagePath* OutUpdatedPath = nullptr) = 0;
+	virtual FOpenPackageResult OpenReadPackage(const FPackagePath& PackagePath, FBulkDataCookedIndex CookedIndex,
+		EPackageSegment PackageSegment, FPackagePath* OutUpdatedPath = nullptr) = 0;
+
+	UE_DEPRECATED(5.5, "Use the overload that takes a FBulkDataCookedIndex")
+	FOpenAsyncPackageResult OpenAsyncReadPackage(const FPackagePath& PackagePath, EPackageSegment PackageSegment)
+	{
+		return OpenAsyncReadPackage(PackagePath, FBulkDataCookedIndex::Default, PackageSegment);
+	}
 
 	/** OpenAsyncReadPackage that takes a PackageSegment */
-	virtual FOpenAsyncPackageResult OpenAsyncReadPackage(const FPackagePath& PackagePath, EPackageSegment PackageSegment) = 0;
+	virtual FOpenAsyncPackageResult OpenAsyncReadPackage(const FPackagePath& PackagePath, FBulkDataCookedIndex CookedIndex, EPackageSegment PackageSegment) = 0;
 
 	/* OpenMappedHandleToPackage that takes a PackageSegment */
 	virtual IMappedFileHandle* OpenMappedHandleToPackage(const FPackagePath& PackagePath, EPackageSegment PackageSegment,
