[Config] No longer skip EditorOnly properties when not in Editor (introduced in CL 400612) so standalone mode (-game) will be consistent with Editor and PIE.

diff --git a/Engine/Source/Runtime/CoreUObject/Private/UObject/Obj.cpp b/Engine/Source/Runtime/CoreUObject/Private/UObject/Obj.cpp
--- a/Engine/Source/Runtime/CoreUObject/Private/UObject/Obj.cpp
+++ b/Engine/Source/Runtime/CoreUObject/Private/UObject/Obj.cpp
@@ -3059,12 +3059,6 @@ void UObject::LoadConfig(UClass* ConfigClass/*=nullptr*/, const TCHAR* InFilenam
 			continue;
 		}
 
-		// Don't load config properties that are marked editoronly if not in the editor
-		if ((Property->PropertyFlags & CPF_EditorOnly) && !GIsEditor)
-		{
-			continue;
-		}
-
 		const bool bGlobalConfig = EnumHasAnyFlags(Property->PropertyFlags, CPF_GlobalConfig);
 		UClass* OwnerClass = Property->GetOwnerClass();
 
