Added missing GetLongPackageFName functions to SoftObjectPtr variations.

The underlying FSoftObjectPath has a GetLongPackageFName, which is sometimes useful to call via the owning soft pointer.

diff --git a/Engine/Source/Runtime/CoreUObject/Public/UObject/SoftObjectPtr.h b/Engine/Source/Runtime/CoreUObject/Public/UObject/SoftObjectPtr.h
--- a/Engine/Source/Runtime/CoreUObject/Public/UObject/SoftObjectPtr.h
+++ b/Engine/Source/Runtime/CoreUObject/Public/UObject/SoftObjectPtr.h
@@ -98,6 +98,12 @@ public:
 		return ToSoftObjectPath().GetLongPackageName();
 	}
 	
+	/** Returns /package/path name, leaving off the asset name */
+	FORCEINLINE FName GetLongPackageFName() const
+	{
+		return ToSoftObjectPath().GetLongPackageFName();
+	}
+
 	/** Returns assetname string, leaving off the /package/path. part */
 	FORCEINLINE FString GetAssetName() const
 	{
@@ -495,6 +501,12 @@ public:
 		return ToSoftObjectPath().GetLongPackageName();
 	}
 	
+	/** Returns /package/path name, leaving off the asset name */
+	FORCEINLINE FName GetLongPackageFName() const
+	{
+		return ToSoftObjectPath().GetLongPackageFName();
+	}
+
 	/** Returns assetname string, leaving off the /package/path part */
 	FORCEINLINE FString GetAssetName() const
 	{
@@ -851,6 +863,12 @@ public:
 		return ToSoftObjectPath().GetLongPackageName();
 	}
 	
+	/** Returns /package/path name, leaving off the asset name */
+	FORCEINLINE FName GetLongPackageFName() const
+	{
+		return ToSoftObjectPath().GetLongPackageFName();
+	}
+
 	/** Returns assetname string, leaving off the /package/path part */
 	FORCEINLINE FString GetAssetName() const
 	{
