[LLM] Properly clear LLM tagset scope

#tests LLM RR
#rnx

diff --git a/Engine/Source/Runtime/Core/Private/HAL/LowLevelMemTracker.cpp b/Engine/Source/Runtime/Core/Private/HAL/LowLevelMemTracker.cpp
--- a/Engine/Source/Runtime/Core/Private/HAL/LowLevelMemTracker.cpp
+++ b/Engine/Source/Runtime/Core/Private/HAL/LowLevelMemTracker.cpp
@@ -3903,6 +3903,10 @@ FLLMPauseScope::~FLLMPauseScope()
 	}
 }
 
+FLLMClearScope::FLLMClearScope(ELLMTagSet InTagSet, ELLMTracker InTracker)
+	: FLLMScope(nullptr, false, InTagSet, InTracker)
+{
+}
 
 FLLMScopeFromPtr::FLLMScopeFromPtr(void* Ptr, ELLMTracker InTracker)
 {
diff --git a/Engine/Source/Runtime/Core/Private/UObject/UnrealNames.cpp b/Engine/Source/Runtime/Core/Private/UObject/UnrealNames.cpp
--- a/Engine/Source/Runtime/Core/Private/UObject/UnrealNames.cpp
+++ b/Engine/Source/Runtime/Core/Private/UObject/UnrealNames.cpp
@@ -382,8 +382,8 @@ private:
 		if (!Buffer)
 		{
 			LLM_SCOPE(ELLMTag::FName);
-			LLM_TAGSET_SCOPE(ELLMTag::EngineMisc, ELLMTagSet::Assets);
-			LLM_TAGSET_SCOPE(ELLMTag::EngineMisc, ELLMTagSet::AssetClasses);
+			LLM_TAGSET_SCOPE_CLEAR(ELLMTagSet::Assets);
+			LLM_TAGSET_SCOPE_CLEAR(ELLMTagSet::AssetClasses);
 			UE_TRACE_METADATA_CLEAR_SCOPE();
 			constexpr uint32 BufferSizeBytes = NumWords * sizeof(WordType);
 			Buffer = (WordType*) FMemory::MallocZeroed(BufferSizeBytes, alignof(WordType));
@@ -660,8 +660,8 @@ private:
 	static uint8* AllocBlock()
 	{
 		LLM_SCOPE(ELLMTag::FName);
-		LLM_TAGSET_SCOPE(ELLMTag::EngineMisc, ELLMTagSet::Assets);
-		LLM_TAGSET_SCOPE(ELLMTag::EngineMisc, ELLMTagSet::AssetClasses);
+		LLM_TAGSET_SCOPE_CLEAR(ELLMTagSet::Assets);
+		LLM_TAGSET_SCOPE_CLEAR(ELLMTagSet::AssetClasses);
 		UE_TRACE_METADATA_CLEAR_SCOPE();
 		return (uint8*)FMemory::Malloc(BlockSizeBytes, alignof(FNameEntry));
 	}
@@ -1002,8 +1002,8 @@ public:
 	void Initialize(FNameEntryAllocator& InEntries)
 	{
 		LLM_SCOPE(ELLMTag::FName);
-		LLM_TAGSET_SCOPE(ELLMTag::EngineMisc, ELLMTagSet::Assets);
-		LLM_TAGSET_SCOPE(ELLMTag::EngineMisc, ELLMTagSet::AssetClasses);
+		LLM_TAGSET_SCOPE_CLEAR(ELLMTagSet::Assets);
+		LLM_TAGSET_SCOPE_CLEAR(ELLMTagSet::AssetClasses);
 		UE_TRACE_METADATA_CLEAR_SCOPE();
 		Entries = &InEntries;
 
@@ -1407,8 +1407,8 @@ private:
 	void Grow(const uint32 NewCapacity)
 	{
 		LLM_SCOPE(ELLMTag::FName);
-		LLM_TAGSET_SCOPE(ELLMTag::EngineMisc, ELLMTagSet::Assets);
-		LLM_TAGSET_SCOPE(ELLMTag::EngineMisc, ELLMTagSet::AssetClasses);
+		LLM_TAGSET_SCOPE_CLEAR(ELLMTagSet::Assets);
+		LLM_TAGSET_SCOPE_CLEAR(ELLMTagSet::AssetClasses);
 		UE_TRACE_METADATA_CLEAR_SCOPE();
 		TArrayView<FNameSlot> OldSlots(Slots, Capacity());
 		const uint32 OldUsedSlots = UsedSlots;
diff --git a/Engine/Source/Runtime/Core/Public/HAL/LowLevelMemTracker.h b/Engine/Source/Runtime/Core/Public/HAL/LowLevelMemTracker.h
--- a/Engine/Source/Runtime/Core/Public/HAL/LowLevelMemTracker.h
+++ b/Engine/Source/Runtime/Core/Public/HAL/LowLevelMemTracker.h
@@ -376,6 +376,9 @@ extern FName LLMGetTagStat(ELLMTag Tag);
 #define LLM_PLATFORM_SCOPE_BYNAME(Tag) 								static FName PREPROCESSOR_JOIN(LLMScope_Name,__LINE__)(Tag);\
 																	FLLMScope SCOPE_NAME(PREPROCESSOR_JOIN(LLMScope_Name,__LINE__), false /* bIsStatTag */, ELLMTagSet::None, ELLMTracker::Platform);
 #define LLM_PLATFORM_SCOPE_BYTAG(TagDeclName)						FLLMScope SCOPE_NAME(PREPROCESSOR_JOIN(LLMTagDeclaration_, TagDeclName).GetUniqueName(), false /* bIsStatTag */, ELLMTagSet::None, ELLMTracker::Platform);
+#define LLM_SCOPE_CLEAR()											FLLMClearScope SCOPE_NAME(ELLMTagSet::None, ELLMTracker::Default);\
+																	UE_MEMSCOPE(0)
+#define LLM_TAGSET_SCOPE_CLEAR(TagSet)								FLLMClearScope SCOPE_NAME(TagSet, ELLMTracker::Default);
 
  /**
  * LLM Pause scope macros
@@ -1037,6 +1040,13 @@ protected:
 	bool bEnabled;
 };
 
+/** LLM Scope to clear top tag while in scope */
+class FLLMClearScope : public FLLMScope
+{
+public:
+	CORE_API FLLMClearScope(ELLMTagSet InTagSet, ELLMTracker InTracker);
+};
+
 /** LLM scope for inheriting tag from the given address. */
 class FLLMScopeFromPtr
 {
@@ -1113,6 +1123,8 @@ inline bool FLowLevelMemTracker::IsEnabled()
 #define LLM_PLATFORM_SCOPE(...)
 #define LLM_PLATFORM_SCOPE_BYNAME(...)
 #define LLM_PLATFORM_SCOPE_BYTAG(...)
+#define LLM_SCOPE_CLEAR()
+#define LLM_TAGSET_SCOPE_CLEAR(...)
 #define LLM_REALLOC_SCOPE(...)
 #define LLM_REALLOC_PLATFORM_SCOPE(...)
 #define LLM_SCOPED_PAUSE_TRACKING(...)
diff --git a/Engine/Source/Runtime/CoreUObject/Private/Serialization/AsyncLoading2.cpp b/Engine/Source/Runtime/CoreUObject/Private/Serialization/AsyncLoading2.cpp
--- a/Engine/Source/Runtime/CoreUObject/Private/Serialization/AsyncLoading2.cpp
+++ b/Engine/Source/Runtime/CoreUObject/Private/Serialization/AsyncLoading2.cpp
@@ -9789,8 +9789,8 @@ void FAsyncLoadingThread2::FlushLoading(TConstArrayView<int32> RequestIDs)
 	if (IsAsyncLoadingPackages())
 	{
 		LLM_SCOPE(ELLMTag::AsyncLoading);
-		LLM_TAGSET_SCOPE(ELLMTag::EngineMisc, ELLMTagSet::Assets);
-		LLM_TAGSET_SCOPE(ELLMTag::EngineMisc, ELLMTagSet::AssetClasses);
+		LLM_TAGSET_SCOPE_CLEAR(ELLMTagSet::Assets);
+		LLM_TAGSET_SCOPE_CLEAR(ELLMTagSet::AssetClasses);
 		UE_TRACE_METADATA_CLEAR_SCOPE();
 		// Prevent objects still being loaded from being accessible from the game thread.
 		TGuardValue GuardVisibilityFilter(FUObjectThreadContext::Get().AsyncVisibilityFilter, EInternalObjectFlags::AsyncLoadingPhase1);
diff --git a/Engine/Source/Runtime/CoreUObject/Private/UObject/UObjectArray.cpp b/Engine/Source/Runtime/CoreUObject/Private/UObject/UObjectArray.cpp
--- a/Engine/Source/Runtime/CoreUObject/Private/UObject/UObjectArray.cpp
+++ b/Engine/Source/Runtime/CoreUObject/Private/UObject/UObjectArray.cpp
@@ -230,8 +230,8 @@ void FUObjectArray::AllocateUObjectIndex(UObjectBase* Object, EInternalObjectFla
 {
 	LLM_SCOPE(ELLMTag::UObject);
 	// Clear asset scopes
-	LLM_TAGSET_SCOPE(ELLMTag::EngineMisc, ELLMTagSet::Assets);
-	LLM_TAGSET_SCOPE(ELLMTag::EngineMisc, ELLMTagSet::AssetClasses);
+	LLM_TAGSET_SCOPE_CLEAR(ELLMTagSet::Assets);
+	LLM_TAGSET_SCOPE_CLEAR(ELLMTagSet::AssetClasses);
 	UE_TRACE_METADATA_CLEAR_SCOPE();
 
 	int32 Index = INDEX_NONE;
diff --git a/Engine/Source/Runtime/CoreUObject/Private/UObject/UObjectHash.cpp b/Engine/Source/Runtime/CoreUObject/Private/UObject/UObjectHash.cpp
--- a/Engine/Source/Runtime/CoreUObject/Private/UObject/UObjectHash.cpp
+++ b/Engine/Source/Runtime/CoreUObject/Private/UObject/UObjectHash.cpp
@@ -1145,8 +1145,8 @@ FORCEINLINE static UPackage* UnassignExternalPackageFromObject(FUObjectHashTable
 void ShrinkUObjectHashTables()
 {
 	LLM_SCOPE_BYTAG(UObjectHash);
-	LLM_TAGSET_SCOPE(ELLMTag::EngineMisc, ELLMTagSet::Assets);
-	LLM_TAGSET_SCOPE(ELLMTag::EngineMisc, ELLMTagSet::AssetClasses);
+	LLM_TAGSET_SCOPE_CLEAR(ELLMTagSet::Assets);
+	LLM_TAGSET_SCOPE_CLEAR(ELLMTagSet::AssetClasses);
 	UE_TRACE_METADATA_CLEAR_SCOPE();
 	TRACE_CPUPROFILER_EVENT_SCOPE(ShrinkUObjectHashTables);
 	FUObjectHashTables& ThreadHash = FUObjectHashTables::Get();
@@ -1580,8 +1580,8 @@ void HashObject(UObjectBase* Object)
 	if (Name != NAME_None)
 	{
 		LLM_SCOPE_BYTAG(UObjectHash);
-		LLM_TAGSET_SCOPE(ELLMTag::EngineMisc, ELLMTagSet::Assets);
-		LLM_TAGSET_SCOPE(ELLMTag::EngineMisc, ELLMTagSet::AssetClasses);
+		LLM_TAGSET_SCOPE_CLEAR(ELLMTagSet::Assets);
+		LLM_TAGSET_SCOPE_CLEAR(ELLMTagSet::AssetClasses);
 		UE_TRACE_METADATA_CLEAR_SCOPE();
 #if !UE_BUILD_TEST && !UE_BUILD_SHIPPING
 		SCOPE_CYCLE_COUNTER(STAT_Hash_HashObject);
@@ -1630,8 +1630,8 @@ void UnhashObject(UObjectBase* Object)
 	if (Name != NAME_None)
 	{
 		LLM_SCOPE_BYTAG(UObjectHash);
-		LLM_TAGSET_SCOPE(ELLMTag::EngineMisc, ELLMTagSet::Assets);
-		LLM_TAGSET_SCOPE(ELLMTag::EngineMisc, ELLMTagSet::AssetClasses);
+		LLM_TAGSET_SCOPE_CLEAR(ELLMTagSet::Assets);
+		LLM_TAGSET_SCOPE_CLEAR(ELLMTagSet::AssetClasses);
 		UE_TRACE_METADATA_CLEAR_SCOPE();
 #if !UE_BUILD_TEST && !UE_BUILD_SHIPPING
 		SCOPE_CYCLE_COUNTER(STAT_Hash_UnhashObject);
@@ -1670,8 +1670,8 @@ void UnhashObject(UObjectBase* Object)
 void HashObjectExternalPackage(UObjectBase* Object, UPackage* Package)
 {
 	LLM_SCOPE_BYTAG(UObjectHash);
-	LLM_TAGSET_SCOPE(ELLMTag::EngineMisc, ELLMTagSet::Assets);
-	LLM_TAGSET_SCOPE(ELLMTag::EngineMisc, ELLMTagSet::AssetClasses);
+	LLM_TAGSET_SCOPE_CLEAR(ELLMTagSet::Assets);
+	LLM_TAGSET_SCOPE_CLEAR(ELLMTagSet::AssetClasses);
 	UE_TRACE_METADATA_CLEAR_SCOPE();
 	if (Package)
 	{
@@ -1697,8 +1697,8 @@ void HashObjectExternalPackage(UObjectBase* Object, UPackage* Package)
 void UnhashObjectExternalPackage(class UObjectBase* Object)
 {
 	LLM_SCOPE_BYTAG(UObjectHash);
-	LLM_TAGSET_SCOPE(ELLMTag::EngineMisc, ELLMTagSet::Assets);
-	LLM_TAGSET_SCOPE(ELLMTag::EngineMisc, ELLMTagSet::AssetClasses);
+	LLM_TAGSET_SCOPE_CLEAR(ELLMTagSet::Assets);
+	LLM_TAGSET_SCOPE_CLEAR(ELLMTagSet::AssetClasses);
 	UE_TRACE_METADATA_CLEAR_SCOPE();
 	FUObjectHashTables& ThreadHash = FUObjectHashTables::Get();
 	FHashTableLock LockHash(ThreadHash);
diff --git a/Engine/Source/Runtime/Engine/Private/ContentStreaming.cpp b/Engine/Source/Runtime/Engine/Private/ContentStreaming.cpp
--- a/Engine/Source/Runtime/Engine/Private/ContentStreaming.cpp
+++ b/Engine/Source/Runtime/Engine/Private/ContentStreaming.cpp
@@ -372,8 +372,8 @@ FStreamingManagerCollection& IStreamingManager::Get()
 {
 	if (StreamingManagerCollection == nullptr)
 	{
-		LLM_TAGSET_SCOPE(ELLMTag::EngineMisc, ELLMTagSet::Assets);
-		LLM_TAGSET_SCOPE(ELLMTag::EngineMisc, ELLMTagSet::AssetClasses);
+		LLM_TAGSET_SCOPE_CLEAR(ELLMTagSet::Assets);
+		LLM_TAGSET_SCOPE_CLEAR(ELLMTagSet::AssetClasses);
 		UE_TRACE_METADATA_CLEAR_SCOPE();
 		// Since this is a lazily created static global variable we create it
 		// in the open.
