[AutoRTFM] Use transactionally-safe critical sections during garbage collection reference tracking.

diff --git a/Engine/Source/Runtime/CoreUObject/Private/UObject/GarbageCollection.cpp b/Engine/Source/Runtime/CoreUObject/Private/UObject/GarbageCollection.cpp
--- a/Engine/Source/Runtime/CoreUObject/Private/UObject/GarbageCollection.cpp
+++ b/Engine/Source/Runtime/CoreUObject/Private/UObject/GarbageCollection.cpp
@@ -14,6 +14,8 @@
 #include "HAL/IConsoleManager.h"
 #include "Misc/App.h"
 #include "Misc/CoreDelegates.h"
+#include "Misc/TransactionallySafeCriticalSection.h"
+#include "Misc/TransactionallySafeScopeLock.h"
 #include "Tasks/Task.h"
 #include "UObject/ScriptInterface.h"
 #include "UObject/UObjectAllocator.h"
@@ -103,7 +105,6 @@ static bool GObjCurrentPurgeObjectIndexNeedsReset = true;
 
 /** Contains a list of objects that stayed marked as unreachable after the last reachability analysis */
 static TArray<UE::GC::FUnreachableObject> GUnreachableObjects;
-static FCriticalSection GUnreachableObjectsCritical;
 static int32 GUnrechableObjectIndex = 0;
 
 static UE::GC::Private::FStats GGCStats;
@@ -1229,7 +1230,7 @@ private:
 
 		void Push(void* Page)
 		{
-			FScopeLock Scope(&Lock);
+			FTransactionallySafeScopeLock Scope(&Lock);
 			Pages.Add(Page);
 		}
 
@@ -1240,7 +1241,7 @@ private:
 
 			if (Worker.Num > NumWorkerInitialPages)
 			{
-				FScopeLock Scope(&Lock);
+				FTransactionallySafeScopeLock Scope(&Lock);
 				Pages.Append(Worker.Pages + NumWorkerInitialPages, Worker.Num - NumWorkerInitialPages);
 				Worker.Num = NumWorkerInitialPages;
 			}
@@ -1248,11 +1249,11 @@ private:
 
 		void* Pop()
 		{
-			FScopeLock Scope(&Lock);
+			FTransactionallySafeScopeLock Scope(&Lock);
 			return Pages.IsEmpty() ? nullptr : Pages.Pop(EAllowShrinking::No);
 		}
 
-		FCriticalSection Lock;
+		FTransactionallySafeCriticalSection Lock;
 		TArray<void*> Pages;
 	};	
 
