[AutoRTFM] Mark UObjectBaseUtility::AddToCluster as always open.

This is part of the garbage collector, which should always be outside of RTFM.

This CL also removes the prior transactionally-safe changes made to GarbageCollection.cpp in https://p4-swarm.epicgames.net/changes/36568224 . These changes shouldn't be necessary, since GC runs in the open, and we no longer trip on this section of code when RTFM is enabled.

diff --git a/Engine/Source/Runtime/CoreUObject/Private/UObject/GarbageCollection.cpp b/Engine/Source/Runtime/CoreUObject/Private/UObject/GarbageCollection.cpp
--- a/Engine/Source/Runtime/CoreUObject/Private/UObject/GarbageCollection.cpp
+++ b/Engine/Source/Runtime/CoreUObject/Private/UObject/GarbageCollection.cpp
@@ -14,8 +14,6 @@
 #include "HAL/IConsoleManager.h"
 #include "Misc/App.h"
 #include "Misc/CoreDelegates.h"
-#include "Misc/TransactionallySafeCriticalSection.h"
-#include "Misc/TransactionallySafeScopeLock.h"
 #include "Tasks/Task.h"
 #include "UObject/ScriptInterface.h"
 #include "UObject/UObjectAllocator.h"
@@ -1230,7 +1228,7 @@ private:
 
 		void Push(void* Page)
 		{
-			FTransactionallySafeScopeLock Scope(&Lock);
+			FScopeLock Scope(&Lock);
 			Pages.Add(Page);
 		}
 
@@ -1241,7 +1239,7 @@ private:
 
 			if (Worker.Num > NumWorkerInitialPages)
 			{
-				FTransactionallySafeScopeLock Scope(&Lock);
+				FScopeLock Scope(&Lock);
 				Pages.Append(Worker.Pages + NumWorkerInitialPages, Worker.Num - NumWorkerInitialPages);
 				Worker.Num = NumWorkerInitialPages;
 			}
@@ -1249,11 +1247,11 @@ private:
 
 		void* Pop()
 		{
-			FTransactionallySafeScopeLock Scope(&Lock);
+			FScopeLock Scope(&Lock);
 			return Pages.IsEmpty() ? nullptr : Pages.Pop(EAllowShrinking::No);
 		}
 
-		FTransactionallySafeCriticalSection Lock;
+		FCriticalSection Lock;
 		TArray<void*> Pages;
 	};	
 
diff --git a/Engine/Source/Runtime/CoreUObject/Private/UObject/UObjectClusters.cpp b/Engine/Source/Runtime/CoreUObject/Private/UObject/UObjectClusters.cpp
--- a/Engine/Source/Runtime/CoreUObject/Private/UObject/UObjectClusters.cpp
+++ b/Engine/Source/Runtime/CoreUObject/Private/UObject/UObjectClusters.cpp
@@ -6,6 +6,7 @@
 
 #include "UObject/UObjectClusters.h"
 #include "HAL/ThreadSafeCounter.h"
+#include "AutoRTFM/AutoRTFM.h"
 #include "Containers/LockFreeList.h"
 #include "Stats/Stats.h"
 #include "HAL/IConsoleManager.h"
@@ -872,6 +873,7 @@ void CreateClustersFromPackage(FLinkerLoad* PackageLinker, TArray<UObject*>& Out
 	}
 }
 
+UE_AUTORTFM_ALWAYS_OPEN
 void UObjectBaseUtility::AddToCluster(UObjectBaseUtility* ClusterRootOrObjectFromCluster, bool bAddAsMutableObject /* = false */)
 {
 	FUObjectCluster* Cluster = GUObjectClusters.GetObjectCluster(ClusterRootOrObjectFromCluster);
