Replaced usage of FInternationalization::ForUseOnlyByLocMacroAndGraphNodeTextLiterals_CreateText with FText::AsLocalizable_Advanced

#rnx

diff --git a/Engine/Plugins/Experimental/LocalizableMessage/Source/LocalizableMessage/Private/LocalizableMessageProcessor.cpp b/Engine/Plugins/Experimental/LocalizableMessage/Source/LocalizableMessage/Private/LocalizableMessageProcessor.cpp
--- a/Engine/Plugins/Experimental/LocalizableMessage/Source/LocalizableMessage/Private/LocalizableMessageProcessor.cpp
+++ b/Engine/Plugins/Experimental/LocalizableMessage/Source/LocalizableMessage/Private/LocalizableMessageProcessor.cpp
@@ -62,7 +62,7 @@ FText FLocalizableMessageProcessor::Localize(const FLocalizableMessage& Message,
 	FText LocalizedText;
 	if (!Message.IsEmpty())
 	{
-		UE_AUTORTFM_OPEN{ LocalizedText = FInternationalization::ForUseOnlyByLocMacroAndGraphNodeTextLiterals_CreateText(*Message.DefaultText, TEXT(""), *Message.Key); };
+		UE_AUTORTFM_OPEN{ LocalizedText = FText::AsLocalizable_Advanced(FTextKey(), Message.Key, Message.DefaultText); };
 	}
 	if (FormatArguments.Num() > 0)
 	{
diff --git a/Engine/Plugins/Experimental/PythonScriptPlugin/Source/PythonScriptPlugin/Private/PyCore.cpp b/Engine/Plugins/Experimental/PythonScriptPlugin/Source/PythonScriptPlugin/Private/PyCore.cpp
--- a/Engine/Plugins/Experimental/PythonScriptPlugin/Source/PythonScriptPlugin/Private/PyCore.cpp
+++ b/Engine/Plugins/Experimental/PythonScriptPlugin/Source/PythonScriptPlugin/Private/PyCore.cpp
@@ -1754,7 +1754,7 @@ PyObject* CreateLocalizedText(PyObject* InSelf, PyObject* InArgs)
 		return nullptr;
 	}
 
-	return PyConversion::Pythonize(FInternationalization::Get().ForUseOnlyByLocMacroAndGraphNodeTextLiterals_CreateText(*Source, *Namespace, *Key));
+	return PyConversion::Pythonize(FText::AsLocalizable_Advanced(Namespace, Key, MoveTemp(Source)));
 }
 
 PyObject* CreateLocalizedTextFromStringTable(PyObject* InSelf, PyObject* InArgs)
diff --git a/Engine/Source/Editor/EditorWidgets/Private/STextPropertyEditableTextBox.cpp b/Engine/Source/Editor/EditorWidgets/Private/STextPropertyEditableTextBox.cpp
--- a/Engine/Source/Editor/EditorWidgets/Private/STextPropertyEditableTextBox.cpp
+++ b/Engine/Source/Editor/EditorWidgets/Private/STextPropertyEditableTextBox.cpp
@@ -1525,7 +1525,7 @@ void STextPropertyEditableTextBox::HandleLocalizableCheckStateChanged(ECheckBoxS
 					NewKey
 					);
 
-				EditableTextProperty->SetText(TextIndex, FInternationalization::Get().ForUseOnlyByLocMacroAndGraphNodeTextLiterals_CreateText(*PropertyValue.ToString(), *NewNamespace, *NewKey));
+				EditableTextProperty->SetText(TextIndex, FText::AsLocalizable_Advanced(NewNamespace, NewKey, PropertyValue.ToString()));
 			}
 		}
 	}
diff --git a/Engine/Source/Editor/UnrealEd/Private/Kismet2/EnumEditorUtils.cpp b/Engine/Source/Editor/UnrealEd/Private/Kismet2/EnumEditorUtils.cpp
--- a/Engine/Source/Editor/UnrealEd/Private/Kismet2/EnumEditorUtils.cpp
+++ b/Engine/Source/Editor/UnrealEd/Private/Kismet2/EnumEditorUtils.cpp
@@ -595,8 +595,7 @@ void FEnumEditorUtils::EnsureAllDisplayNamesExist(UUserDefinedEnum* Enum)
 					const FString PackageNamespace = GIsEditor ? TextNamespaceUtil::EnsurePackageNamespace(Enum) : TextNamespaceUtil::GetPackageNamespace(Enum);
 					const FString TextNamespace = TextNamespaceUtil::BuildFullNamespace(FString(), PackageNamespace, true);
 					const FString TextKey = FGuid::NewGuid().ToString();
-					const FString EnumEntryDisplayName = EnumEntryName.ToString();
-					DisplayNameToSet = FInternationalization::ForUseOnlyByLocMacroAndGraphNodeTextLiterals_CreateText(*EnumEntryDisplayName, *TextNamespace, *TextKey);
+					DisplayNameToSet = FText::AsLocalizable_Advanced(TextNamespace, TextKey, EnumEntryName.ToString());
 				}
 #else  // USE_STABLE_LOCALIZATION_KEYS
 				DisplayNameToSet = FText::FromName(EnumEntryName);
@@ -633,7 +632,7 @@ void FEnumEditorUtils::UpgradeDisplayNamesFromMetaData(UUserDefinedEnum* Enum)
 					const FString PackageNamespace = GIsEditor ? TextNamespaceUtil::EnsurePackageNamespace(Enum) : TextNamespaceUtil::GetPackageNamespace(Enum);
 					const FString TextNamespace = TextNamespaceUtil::BuildFullNamespace(TEXT(""), PackageNamespace, true);
 					const FString TextKey = FGuid::NewGuid().ToString();
-					DisplayNameToSet = FInternationalization::ForUseOnlyByLocMacroAndGraphNodeTextLiterals_CreateText(*MetaDataEntryDisplayName, *TextNamespace, *TextKey);
+					DisplayNameToSet = FText::AsLocalizable_Advanced(TextNamespace, TextKey, MetaDataEntryDisplayName);
 				}
 #else  // USE_STABLE_LOCALIZATION_KEYS
 				DisplayNameToSet = FText::FromName(EnumEntryName);
diff --git a/Engine/Source/Runtime/Core/Private/Logging/StructuredLog.cpp b/Engine/Source/Runtime/Core/Private/Logging/StructuredLog.cpp
--- a/Engine/Source/Runtime/Core/Private/Logging/StructuredLog.cpp
+++ b/Engine/Source/Runtime/Core/Private/Logging/StructuredLog.cpp
@@ -386,7 +386,7 @@ FLogTemplate* FLogTemplate::CreateLocalized(const TCHAR* TextNamespace, const TC
 	checkf(FAsciiSet::HasNone(Format, Invalid),
 		TEXT("Log format does not currently allow escapes (`) or argument modifiers (|). [[%s]]"), Format);
 
-	FTextFormat TextFormat(FInternationalization::ForUseOnlyByLocMacroAndGraphNodeTextLiterals_CreateText(Format, TextNamespace, TextKey));
+	FTextFormat TextFormat(FText::AsLocalizable_Advanced(TextNamespace, TextKey, Format));
 
 	const bool bFindFields = !!Fields;
 	if (bFindFields)
diff --git a/Engine/Source/Runtime/CoreUObject/Private/Internationalization/TextPackageNamespaceUtil.cpp b/Engine/Source/Runtime/CoreUObject/Private/Internationalization/TextPackageNamespaceUtil.cpp
--- a/Engine/Source/Runtime/CoreUObject/Private/Internationalization/TextPackageNamespaceUtil.cpp
+++ b/Engine/Source/Runtime/CoreUObject/Private/Internationalization/TextPackageNamespaceUtil.cpp
@@ -325,12 +325,12 @@ bool TextNamespaceUtil::EditTextProperty_Direct(UPackage* InPackage, void* InTex
 
 	const FString ProposedNamespace = (InEditAction == ETextEditAction::Namespace ? InEditValue : CurrentTextId.GetNamespace().ToString());
 	const FString ProposedKey = (InEditAction == ETextEditAction::Key ? InEditValue : CurrentTextId.GetKey().ToString());
-	const FString SourceString = (InEditAction == ETextEditAction::SourceString ? InEditValue : CurrentSourceString ? *CurrentSourceString : FString());
+	FString SourceString = (InEditAction == ETextEditAction::SourceString ? InEditValue : CurrentSourceString ? *CurrentSourceString : FString());
 
 	FString StableNamespace;
 	FString StableKey;
 	GetTextIdForEdit(InPackage, InEditAction, SourceString, ProposedNamespace, ProposedKey, StableNamespace, StableKey, InTextKeyGenerator, bApplyPackageNamespace);
 
-	InTextProperty->SetPropertyValue(InTextValue, FInternationalization::ForUseOnlyByLocMacroAndGraphNodeTextLiterals_CreateText(*SourceString, *StableNamespace, *StableKey));
+	InTextProperty->SetPropertyValue(InTextValue, FText::AsLocalizable_Advanced(StableNamespace, StableKey, MoveTemp(SourceString)));
 	return true;
 }
diff --git a/Engine/Source/Runtime/CoreUObject/Private/Serialization/PropertyLocalizationDataGathering.cpp b/Engine/Source/Runtime/CoreUObject/Private/Serialization/PropertyLocalizationDataGathering.cpp
--- a/Engine/Source/Runtime/CoreUObject/Private/Serialization/PropertyLocalizationDataGathering.cpp
+++ b/Engine/Source/Runtime/CoreUObject/Private/Serialization/PropertyLocalizationDataGathering.cpp
@@ -728,7 +728,7 @@ private:
 				bIsParsingText = true;
 
 				SerializeExpr(iCode, Ar);
-				const FString SourceString = MoveTemp(LastParsedString);
+				FString SourceString = MoveTemp(LastParsedString);
 
 				SerializeExpr(iCode, Ar);
 				const FString TextKey = MoveTemp(LastParsedString);
@@ -738,7 +738,7 @@ private:
 
 				bIsParsingText = false;
 
-				const FText TextInstance = FInternationalization::ForUseOnlyByLocMacroAndGraphNodeTextLiterals_CreateText(*SourceString, *TextNamespace, *TextKey);
+				const FText TextInstance = FText::AsLocalizable_Advanced(TextNamespace, TextKey, MoveTemp(SourceString));
 				if (!PropertyLocalizationDataGatherer.IsDefaultTextInstance(TextInstance))
 				{
 					PropertyLocalizationDataGatherer.GatherTextInstance(TextInstance, FString::Printf(TEXT("%s [Script Bytecode]"), SourceDescription), bTreatAsEditorOnlyData);
diff --git a/Engine/Source/Runtime/CoreUObject/Private/UObject/Class.cpp b/Engine/Source/Runtime/CoreUObject/Private/UObject/Class.cpp
--- a/Engine/Source/Runtime/CoreUObject/Private/UObject/Class.cpp
+++ b/Engine/Source/Runtime/CoreUObject/Private/UObject/Class.cpp
@@ -288,7 +288,7 @@ struct FDisplayNameHelper
  */
 FText UField::GetDisplayNameText() const
 {
-	static const FString Namespace = TEXT("UObjectDisplayNames");
+	static const FTextKey Namespace = TEXT("UObjectDisplayNames");
 	static const FName NAME_DisplayName(TEXT("DisplayName"));
 
 	const FString Key = GetFullGroupName(false);
@@ -299,7 +299,7 @@ FText UField::GetDisplayNameText() const
 		NativeDisplayName = FName::NameToDisplayString(FDisplayNameHelper::Get(*this), false);
 	}
 
-	return FInternationalization::ForUseOnlyByLocMacroAndGraphNodeTextLiterals_CreateText(*NativeDisplayName, *Namespace, *Key);
+	return FText::AsLocalizable_Advanced(Namespace, Key, MoveTemp(NativeDisplayName));
 }
 
 /**
@@ -344,7 +344,7 @@ FText UField::GetToolTipText(bool bShortTooltip) const
 		{
 			FormatNativeToolTip(NativeToolTip, true);
 		}
-		LocalizedToolTip = FInternationalization::ForUseOnlyByLocMacroAndGraphNodeTextLiterals_CreateText(*NativeToolTip, *Namespace, *Key);
+		LocalizedToolTip = FText::AsLocalizable_Advanced(Namespace, Key, MoveTemp(NativeToolTip));
 	}
 
 	return LocalizedToolTip;
@@ -475,7 +475,7 @@ const FString& UField::GetMetaData(const FName& Key) const
 	return MetaDataString;
 }
 
-FText UField::GetMetaDataText(const TCHAR* MetaDataKey, const FString LocalizationNamespace, const FString LocalizationKey) const
+FText UField::GetMetaDataText(const TCHAR* MetaDataKey, const FTextKey LocalizationNamespace, const FTextKey LocalizationKey) const
 {
 	FString DefaultMetaData;
 
@@ -493,12 +493,12 @@ FText UField::GetMetaDataText(const TCHAR* MetaDataKey, const FString Localizati
 	FText LocalizedMetaData;
 	if (!DefaultMetaData.IsEmpty())
 	{
-		LocalizedMetaData = FInternationalization::ForUseOnlyByLocMacroAndGraphNodeTextLiterals_CreateText(*DefaultMetaData, *LocalizationNamespace, *LocalizationKey);
+		LocalizedMetaData = FText::AsLocalizable_Advanced(LocalizationNamespace, LocalizationKey, MoveTemp(DefaultMetaData));
 	}
 	return LocalizedMetaData;
 }
 
-FText UField::GetMetaDataText(const FName& MetaDataKey, const FString LocalizationNamespace, const FString LocalizationKey) const
+FText UField::GetMetaDataText(const FName& MetaDataKey, const FTextKey LocalizationNamespace, const FTextKey LocalizationKey) const
 {
 	FString DefaultMetaData;
 
@@ -516,7 +516,7 @@ FText UField::GetMetaDataText(const FName& MetaDataKey, const FString Localizati
 	FText LocalizedMetaData;
 	if (!DefaultMetaData.IsEmpty())
 	{
-		LocalizedMetaData = FInternationalization::ForUseOnlyByLocMacroAndGraphNodeTextLiterals_CreateText(*DefaultMetaData, *LocalizationNamespace, *LocalizationKey);
+		LocalizedMetaData = FText::AsLocalizable_Advanced(LocalizationNamespace, LocalizationKey, MoveTemp(DefaultMetaData));
 	}
 	return LocalizedMetaData;
 }
diff --git a/Engine/Source/Runtime/CoreUObject/Private/UObject/Enum.cpp b/Engine/Source/Runtime/CoreUObject/Private/UObject/Enum.cpp
--- a/Engine/Source/Runtime/CoreUObject/Private/UObject/Enum.cpp
+++ b/Engine/Source/Runtime/CoreUObject/Private/UObject/Enum.cpp
@@ -676,7 +676,7 @@ FText UEnum::GetDisplayNameTextByIndex(int32 NameIndex) const
 
 #if WITH_EDITOR
 	// In the editor, use metadata and localization to look up names
-	static const FString Namespace = TEXT("UObjectDisplayNames");
+	static const FTextKey Namespace = TEXT("UObjectDisplayNames");
 	const FString Key = GetFullGroupName(false) + TEXT(".") + RawName;
 
 	FString NativeDisplayName;
@@ -691,7 +691,7 @@ FText UEnum::GetDisplayNameTextByIndex(int32 NameIndex) const
 
 	if (!NativeDisplayName.IsEmpty())
 	{
-		return FInternationalization::ForUseOnlyByLocMacroAndGraphNodeTextLiterals_CreateText(*NativeDisplayName, *Namespace, *Key);
+		return FText::AsLocalizable_Advanced(Namespace, Key, MoveTemp(NativeDisplayName));
 	}
 #endif
 
@@ -916,8 +916,8 @@ FText UEnum::GetToolTipTextByIndex(int32 NameIndex) const
 	FText LocalizedToolTip;
 	FString NativeToolTip = GetMetaData( TEXT("ToolTip"), NameIndex );
 
-	static const FString Namespace = TEXT("UObjectToolTips");
-	FString Key = GetFullGroupName(false) + TEXT(".") + GetNameStringByIndex(NameIndex);
+	static const FTextKey Namespace = TEXT("UObjectToolTips");
+	const FString Key = GetFullGroupName(false) + TEXT(".") + GetNameStringByIndex(NameIndex);
 		
 	if ( !FText::FindTextInLiveTable_Advanced( Namespace, Key, /*OUT*/LocalizedToolTip, &NativeToolTip ) )
 	{
@@ -928,7 +928,7 @@ FText UEnum::GetToolTipTextByIndex(int32 NameIndex) const
 			NativeToolTip.TrimEndInline();
 		}
 
-		LocalizedToolTip = FInternationalization::ForUseOnlyByLocMacroAndGraphNodeTextLiterals_CreateText(*NativeToolTip, *Namespace, *Key);
+		LocalizedToolTip = FText::AsLocalizable_Advanced(Namespace, Key, MoveTemp(NativeToolTip));
 	}
 
 	return LocalizedToolTip;
diff --git a/Engine/Source/Runtime/CoreUObject/Private/UObject/Field.cpp b/Engine/Source/Runtime/CoreUObject/Private/UObject/Field.cpp
--- a/Engine/Source/Runtime/CoreUObject/Private/UObject/Field.cpp
+++ b/Engine/Source/Runtime/CoreUObject/Private/UObject/Field.cpp
@@ -672,7 +672,7 @@ const FString& FField::GetMetaData(const FName& Key) const
 	return *ValuePtr;
 }
 
-const FText FField::GetMetaDataText(const TCHAR* MetaDataKey, const FString LocalizationNamespace, const FString LocalizationKey) const
+const FText FField::GetMetaDataText(const TCHAR* MetaDataKey, const FTextKey LocalizationNamespace, const FTextKey LocalizationKey) const
 {
 	FString DefaultMetaData;
 
@@ -690,12 +690,12 @@ const FText FField::GetMetaDataText(const TCHAR* MetaDataKey, const FString Loca
 	FText LocalizedMetaData;
 	if (!DefaultMetaData.IsEmpty())
 	{
-		LocalizedMetaData = FInternationalization::ForUseOnlyByLocMacroAndGraphNodeTextLiterals_CreateText(*DefaultMetaData, *LocalizationNamespace, *LocalizationKey);
+		LocalizedMetaData = FText::AsLocalizable_Advanced(LocalizationNamespace, LocalizationKey, MoveTemp(DefaultMetaData));
 	}
 	return LocalizedMetaData;
 }
 
-const FText FField::GetMetaDataText(const FName& MetaDataKey, const FString LocalizationNamespace, const FString LocalizationKey) const
+const FText FField::GetMetaDataText(const FName& MetaDataKey, const FTextKey LocalizationNamespace, const FTextKey LocalizationKey) const
 {
 	FString DefaultMetaData;
 
@@ -713,7 +713,7 @@ const FText FField::GetMetaDataText(const FName& MetaDataKey, const FString Loca
 	FText LocalizedMetaData;
 	if (!DefaultMetaData.IsEmpty())
 	{
-		LocalizedMetaData = FInternationalization::ForUseOnlyByLocMacroAndGraphNodeTextLiterals_CreateText(*DefaultMetaData, *LocalizationNamespace, *LocalizationKey);
+		LocalizedMetaData = FText::AsLocalizable_Advanced(LocalizationNamespace, LocalizationKey, MoveTemp(DefaultMetaData));
 	}
 	return LocalizedMetaData;
 }
@@ -904,7 +904,7 @@ struct FFieldDisplayNameHelper
 */
 FText FField::GetDisplayNameText() const
 {
-	static const FString Namespace = TEXT("UObjectDisplayNames");
+	static const FTextKey Namespace = TEXT("UObjectDisplayNames");
 	static const FName NAME_DisplayName(TEXT("DisplayName"));
 
 	const FString Key = GetFullGroupName(false);
@@ -919,7 +919,7 @@ FText FField::GetDisplayNameText() const
 		NativeDisplayName = FName::NameToDisplayString(FFieldDisplayNameHelper::Get(*this), IsA<FBoolProperty>());
 	}
 
-	return FInternationalization::ForUseOnlyByLocMacroAndGraphNodeTextLiterals_CreateText(*NativeDisplayName, *Namespace, *Key);
+	return FText::AsLocalizable_Advanced(Namespace, Key, MoveTemp(NativeDisplayName));
 }
 
 /**
@@ -965,7 +965,7 @@ FText FField::GetToolTipText(bool bShortTooltip) const
 				NativeToolTip.TrimEndInline();
 			}
 		}
-		LocalizedToolTip = FInternationalization::ForUseOnlyByLocMacroAndGraphNodeTextLiterals_CreateText(*NativeToolTip, *Namespace, *Key);
+		LocalizedToolTip = FText::AsLocalizable_Advanced(Namespace, Key, MoveTemp(NativeToolTip));
 	}
 
 	return LocalizedToolTip;
diff --git a/Engine/Source/Runtime/CoreUObject/Private/UObject/ScriptCore.cpp b/Engine/Source/Runtime/CoreUObject/Private/UObject/ScriptCore.cpp
--- a/Engine/Source/Runtime/CoreUObject/Private/UObject/ScriptCore.cpp
+++ b/Engine/Source/Runtime/CoreUObject/Private/UObject/ScriptCore.cpp
@@ -3430,7 +3430,7 @@ DEFINE_FUNCTION(UObject::execTextConst)
 			FString Namespace;
 			Stack.Step(Stack.Object, &Namespace);
 
-			*(FText*)RESULT_PARAM = FInternationalization::ForUseOnlyByLocMacroAndGraphNodeTextLiterals_CreateText(*SourceString, *Namespace, *KeyString);
+			*(FText*)RESULT_PARAM = FText::AsLocalizable_Advanced(Namespace, KeyString, MoveTemp(SourceString));
 		}
 		break;
 
diff --git a/Engine/Source/Runtime/CoreUObject/Public/UObject/Class.h b/Engine/Source/Runtime/CoreUObject/Public/UObject/Class.h
--- a/Engine/Source/Runtime/CoreUObject/Public/UObject/Class.h
+++ b/Engine/Source/Runtime/CoreUObject/Public/UObject/Class.h
@@ -251,8 +251,8 @@ class UField : public UObject
 	 * @param LocalizationKey			Key to lookup in the localization manager
 	 * @return							Localized metadata if available, defaults to whatever is provided via GetMetaData
 	 */
-	COREUOBJECT_API FText GetMetaDataText(const TCHAR* MetaDataKey, const FString LocalizationNamespace = FString(), const FString LocalizationKey = FString()) const;
-	COREUOBJECT_API FText GetMetaDataText(const FName& MetaDataKey, const FString LocalizationNamespace = FString(), const FString LocalizationKey = FString()) const;
+	COREUOBJECT_API FText GetMetaDataText(const TCHAR* MetaDataKey, const FTextKey LocalizationNamespace = FTextKey(), const FTextKey LocalizationKey = FTextKey()) const;
+	COREUOBJECT_API FText GetMetaDataText(const FName& MetaDataKey, const FTextKey LocalizationNamespace = FTextKey(), const FTextKey LocalizationKey = FTextKey()) const;
 
 	/**
 	 * Sets the metadata value associated with the key
diff --git a/Engine/Source/Runtime/CoreUObject/Public/UObject/Field.h b/Engine/Source/Runtime/CoreUObject/Public/UObject/Field.h
--- a/Engine/Source/Runtime/CoreUObject/Public/UObject/Field.h
+++ b/Engine/Source/Runtime/CoreUObject/Public/UObject/Field.h
@@ -794,8 +794,8 @@ public:
 	* @param LocalizationKey			Key to lookup in the localization manager
 	* @return							Localized metadata if available, defaults to whatever is provided via GetMetaData
 	*/
-	COREUOBJECT_API const FText GetMetaDataText(const TCHAR* MetaDataKey, const FString LocalizationNamespace = FString(), const FString LocalizationKey = FString()) const;
-	COREUOBJECT_API const FText GetMetaDataText(const FName& MetaDataKey, const FString LocalizationNamespace = FString(), const FString LocalizationKey = FString()) const;
+	COREUOBJECT_API const FText GetMetaDataText(const TCHAR* MetaDataKey, const FTextKey LocalizationNamespace = FTextKey(), const FTextKey LocalizationKey = FTextKey()) const;
+	COREUOBJECT_API const FText GetMetaDataText(const FName& MetaDataKey, const FTextKey LocalizationNamespace = FTextKey(), const FTextKey LocalizationKey = FTextKey()) const;
 
 	/**
 	* Sets the metadata value associated with the key
diff --git a/Engine/Source/Runtime/Slate/Private/Framework/Commands/Commands.cpp b/Engine/Source/Runtime/Slate/Private/Framework/Commands/Commands.cpp
--- a/Engine/Source/Runtime/Slate/Private/Framework/Commands/Commands.cpp
+++ b/Engine/Source/Runtime/Slate/Private/Framework/Commands/Commands.cpp
@@ -14,8 +14,8 @@ void MakeUICommand_InternalUseOnly( FBindingContext* This, TSharedPtr< FUIComman
 		This->AsShared(),
 		OutCommand,
 		InCommandName,
-		FInternationalization::ForUseOnlyByLocMacroAndGraphNodeTextLiterals_CreateText( FriendlyName, *Namespace, InCommandName ),
-		FInternationalization::ForUseOnlyByLocMacroAndGraphNodeTextLiterals_CreateText( InDescription, *Namespace, InCommandNameUnderscoreTooltip ),
+		FText::AsLocalizable_Advanced( Namespace, InCommandName, FriendlyName ),
+		FText::AsLocalizable_Advanced( Namespace, InCommandNameUnderscoreTooltip, InDescription ),
 		FSlateIcon( This->GetStyleSetName(), ISlateStyle::Join( This->GetContextName(), DotCommandName ) ),
 		CommandType,
 		InDefaultChord,
