[ALT]
 - Fix race when trying to serialize the same empty delegate from multiple threads

#rnx

diff --git a/Engine/Source/Runtime/CoreUObject/Private/UObject/PropertyMulticastDelegate.cpp b/Engine/Source/Runtime/CoreUObject/Private/UObject/PropertyMulticastDelegate.cpp
--- a/Engine/Source/Runtime/CoreUObject/Private/UObject/PropertyMulticastDelegate.cpp
+++ b/Engine/Source/Runtime/CoreUObject/Private/UObject/PropertyMulticastDelegate.cpp
@@ -610,12 +610,14 @@ void FMulticastSparseDelegateProperty::SerializeItemInternal(FArchive& Ar, void*
 			}
 			else
 			{
-				Ar << EmptyDelegate;
+				FMulticastScriptDelegate DummyDelegate;
+				Ar << DummyDelegate;
 			}
 		}
 		else
 		{
-			Ar << EmptyDelegate;
+			FMulticastScriptDelegate DummyDelegate;
+			Ar << DummyDelegate;
 		}
 	}
 }
