[AutoRTFM] Mark FUObjectAllocator::AllocateUObject as always open.

This is more efficient than wrapping the entire function body in an OPEN block, as the compiler can omit the closed-form version of the function entirely.

diff --git a/Engine/Source/Runtime/CoreUObject/Private/UObject/UObjectAllocator.cpp b/Engine/Source/Runtime/CoreUObject/Private/UObject/UObjectAllocator.cpp
--- a/Engine/Source/Runtime/CoreUObject/Private/UObject/UObjectAllocator.cpp
+++ b/Engine/Source/Runtime/CoreUObject/Private/UObject/UObjectAllocator.cpp
@@ -37,24 +37,23 @@ void FUObjectAllocator::DisablePersistentAllocator()
  * @param Size size of uobject to allocate
  * @param Alignment alignment of uobject to allocate
  * @param bAllowPermanent if true, allow allocation in the permanent object pool, if it fits
- * @return newly allocated UObjectBase (not really a UObjectBase yet, no constructor like thing has been called).
+ * @return newly allocated UObjectBase (not really a UObjectBase yet, no constructor-like thing has been called).
  */
+UE_AUTORTFM_ALWAYS_OPEN
 UObjectBase* FUObjectAllocator::AllocateUObject(int32 Size, int32 Alignment, bool bAllowPermanent)
 {
 	void* Result = nullptr;
-	// we want to perform this allocation uninstrumented so the GC can clean this up if the transaction is aborted
-	UE_AUTORTFM_OPEN
+
+	// We want to perform this allocation uninstrumented, so the GC can clean this up if the transaction is aborted.
+	if (bAllowPermanent && !GPersistentAllocatorIsDisabled)
 	{
-		if (bAllowPermanent && !GPersistentAllocatorIsDisabled)
-		{
-			// this allocation might go over the reserved memory amount and default to FMemory::Malloc, so we are moving it into the ARTFM scope
-			Result = GetPersistentLinearAllocator().Allocate(Size, Alignment);
-		}
-		else
-		{
-			Result = FMemory::Malloc(Size, Alignment);
-		}
-	};
+		// This allocation might go over the reserved memory amount and default to FMemory::Malloc, so we are moving it into the AutoRTFM scope.
+		Result = GetPersistentLinearAllocator().Allocate(Size, Alignment);
+	}
+	else
+	{
+		Result = FMemory::Malloc(Size, Alignment);
+	}
 
 	return (UObjectBase*)Result;
 }
