[TSAN]
 - Fix TSAN warning by not reading the global unless we're on the game thread

#rnx

diff --git a/Engine/Source/Runtime/CoreUObject/Private/UObject/SavePackage/SavePackageUtilities.cpp b/Engine/Source/Runtime/CoreUObject/Private/UObject/SavePackage/SavePackageUtilities.cpp
--- a/Engine/Source/Runtime/CoreUObject/Private/UObject/SavePackage/SavePackageUtilities.cpp
+++ b/Engine/Source/Runtime/CoreUObject/Private/UObject/SavePackage/SavePackageUtilities.cpp
@@ -183,7 +183,7 @@ namespace UE
 			// That global is only meant to be set and read from the game-thread...
 			// otherwise it could interfere with normal operations coming from
 			// other threads like async loading, etc...
-			return GIsSavingPackage && IsInGameThread();
+			return IsInGameThread() && GIsSavingPackage;
 		}
 		return InOuter->GetPackage()->HasAnyPackageFlags(PKG_IsSaving);
 	}
