fixed nullptr check

diff --git a/Engine/Source/Runtime/CoreUObject/Private/StructUtils/UserDefinedStruct.cpp b/Engine/Source/Runtime/CoreUObject/Private/StructUtils/UserDefinedStruct.cpp
--- a/Engine/Source/Runtime/CoreUObject/Private/StructUtils/UserDefinedStruct.cpp
+++ b/Engine/Source/Runtime/CoreUObject/Private/StructUtils/UserDefinedStruct.cpp
@@ -373,7 +373,7 @@ uint32 UUserDefinedStruct::GetStructTypeHash(const void* Src) const
 void UUserDefinedStruct::RecursivelyPreload()
 {
 	FLinkerLoad* Linker = GetLinker();
-	if( Linker && (NULL == PropertyLink) )
+	if( Linker && (nullptr == PropertyLink) )
 	{
 		TArray<UObject*> AllChildMembers;
 		GetObjectsWithOuter(this, AllChildMembers);
@@ -384,7 +384,7 @@ void UUserDefinedStruct::RecursivelyPreload()
 		}
 
 		Linker->Preload(this);
-		if (NULL == PropertyLink)
+		if (nullptr == PropertyLink)
 		{
 			StaticLink(true);
 		}
