Utility function to convert FRemoteObjectId to FObjectPtr

diff --git a/Engine/Source/Runtime/CoreUObject/Public/UObject/RemoteObjectTransfer.h b/Engine/Source/Runtime/CoreUObject/Public/UObject/RemoteObjectTransfer.h
--- a/Engine/Source/Runtime/CoreUObject/Public/UObject/RemoteObjectTransfer.h
+++ b/Engine/Source/Runtime/CoreUObject/Public/UObject/RemoteObjectTransfer.h
@@ -51,5 +51,17 @@ namespace UE::RemoteObject::Transfer
 	*/
 	COREUOBJECT_API UObject* MigrateObjectFromRemoteServer(FRemoteObjectId ObjectId, FRemoteServerId CurrentOwnerServerId, UObject* DestinationOuter = nullptr);
 
+	/**
+	* Utility that converts FRemoteObjectId to TObjectPtr of a given type. Will not resolve the object if it's not local.
+	*/
+	FORCEINLINE FObjectPtr RemoteObjectIdToObjectPtr(FRemoteObjectId RemoteId)
+	{
+#if UE_WITH_REMOTE_OBJECT_HANDLE
+		return FObjectPtr(UE::CoreUObject::Private::FRemoteObjectHandlePrivate::FromIdNoResolve(RemoteId));
+#else
+		return FObjectPtr();
+#endif
+	}
+
 } // namespace UE::RemoveObject::Transfer
 
