UObjectBase destructor and UObjectBase::IsValid* functions  will no longer try to resolve ClassPrivate pointer when destroying and verifying objects

diff --git a/Engine/Source/Runtime/CoreUObject/Private/UObject/UObjectBase.cpp b/Engine/Source/Runtime/CoreUObject/Private/UObject/UObjectBase.cpp
--- a/Engine/Source/Runtime/CoreUObject/Private/UObject/UObjectBase.cpp
+++ b/Engine/Source/Runtime/CoreUObject/Private/UObject/UObjectBase.cpp
@@ -157,7 +157,7 @@ UObjectBase::UObjectBase(UClass* InClass,
 UObjectBase::~UObjectBase()
 {
 	// If not initialized, skip out.
-	if( UObjectInitialized() && ClassPrivate && !GIsCriticalError )
+	if( UObjectInitialized() && ClassPrivate.GetNoResolve() && !GIsCriticalError)
 	{
 		// Validate it.
 		check(IsValidLowLevelForDestruction());
@@ -321,7 +321,7 @@ bool UObjectBase::IsValidLowLevelForDestruction() const
 		UE_LOG(LogUObjectBase, Warning, TEXT("NULL object"));
 		return false;
 	}
-	if (!ClassPrivate)
+	if (!ClassPrivate.GetNoResolve())
 	{
 		UE_LOG(LogUObjectBase, Warning, TEXT("Object is not registered"));
 		return false;
@@ -361,20 +361,21 @@ bool UObjectBase::IsValidLowLevelFast(bool bRecursive /*= true*/) const
 	}
 
 	// These should all be 0.
-	UObject* Outer = UE::CoreUObject::Private::FObjectHandleUtils::GetNonAccessTrackedOuterNoResolve((UObject*)this);
-	const UPTRINT CheckZero = (GetFlagsInternal() & ~RF_AllFlags) | ((UPTRINT)ClassPrivate & AlignmentCheck) | ((UPTRINT)Outer & AlignmentCheck);
+	UObject* Outer = OuterPrivate.GetNoResolve();
+	UClass* Class = ClassPrivate.GetNoResolve();
+	const UPTRINT CheckZero = (GetFlagsInternal() & ~RF_AllFlags) | ((UPTRINT)Class & AlignmentCheck) | ((UPTRINT)Outer & AlignmentCheck);
 	if (!!CheckZero)
 	{
 		UE_LOG(LogUObjectBase, Error, TEXT("Object flags are invalid or either Class or Outer is misaligned"));
     return false;
 	}
 	// These should all be non-NULL (except CDO-alignment check which should be 0)
-	if (ClassPrivate == nullptr || ClassPrivate->ClassDefaultObject == nullptr || ((UPTRINT)ClassPrivate->ClassDefaultObject & AlignmentCheck) != 0)
+	if (Class == nullptr || Class->ClassDefaultObject == nullptr || ((UPTRINT)Class->ClassDefaultObject & AlignmentCheck) != 0)
 	{
 #if WITH_LIVE_CODING
 		// When live coding is re-instancing blueprint generated classes, we have to clear out the default object so it can get 
 		// GC'ed and deleted prior to live coding completing the patching process (of the destructor specifically)
-		if (ClassPrivate == nullptr || !ClassPrivate->HasAnyClassFlags(CLASS_NewerVersionExists))
+		if (Class == nullptr || !Class->HasAnyClassFlags(CLASS_NewerVersionExists))
 #endif
 		{
 			UE_LOG(LogUObjectBase, Error, TEXT("Class pointer is invalid or CDO is invalid."));
@@ -382,7 +383,7 @@ bool UObjectBase::IsValidLowLevelFast(bool bRecursive /*= true*/) const
 		}
 	}
 	// Avoid infinite recursion so call IsValidLowLevelFast on the class object with bRecursive = false.
-	if (bRecursive && !ClassPrivate->IsValidLowLevelFast(false))
+	if (bRecursive && !Class->IsValidLowLevelFast(false))
 	{
 		UE_LOG(LogUObjectBase, Error, TEXT("Class object failed IsValidLowLevelFast test."));
 		return false;
