Minor code optimization in FPackageHarvester::EnterConditionalOptionalObjectScope.

diff --git a/Engine/Source/Runtime/CoreUObject/Private/UObject/SavePackage/PackageHarvester.cpp b/Engine/Source/Runtime/CoreUObject/Private/UObject/SavePackage/PackageHarvester.cpp
--- a/Engine/Source/Runtime/CoreUObject/Private/UObject/SavePackage/PackageHarvester.cpp
+++ b/Engine/Source/Runtime/CoreUObject/Private/UObject/SavePackage/PackageHarvester.cpp
@@ -10,6 +10,8 @@
 #include "UObject/SavePackage/SavePackageUtilities.h"
 #include "UObject/UObjectGlobals.h"
 #include "UObject/UObjectHash.h"
+#include "Serialization/ArchiveShowReferences.h"
+#include "Misc/OutputDeviceDebug.h"
 
 EObjectMark GenerateMarksForObject(const UObject* InObject, FSaveContext& SaveContext)
 {
@@ -291,25 +293,32 @@ FPackageHarvester::FHarvestScope FPackageHarvester::EnterConditionalOptionalObje
 		return Scope;
 	}
 
-	// No matter the realm referring to it, optional objects are added into the optional realm
-	ForEachExportHarvestingRealm([this, Object](ESaveRealm HarvestingRealm)
+	// It is illegal for a non-optional used-in-game object to reference an optional object
+	// Give an error if the referer is not optional and is not editoronly.
+	// An example of an editor-only object refererring to an optional object is a UObjectRedirector.
+	UObject* Ref = CurrentExportDependencies.CurrentExport;
+	if (Ref)
+	{
+		// TODO: Change IsEditorOnlyObject to also test NeedsLoadForClient and NeedsLoadForServer, 
+		// and then change this location to call IsEditorOnlyObject
+		bool bEditorOnly = Ref->IsEditorOnly() || (!Ref->NeedsLoadForClient() && !Ref->NeedsLoadForServer());
+		if (!bEditorOnly)
 		{
-			// It is illegal for a non-optional used-in-game object to reference an optional object
-			// Give an error if the referer is not optional and is not editoronly.
-			// An example of an editor-only object refererring to an optional object is a UObjectRedirector.
-			UObject* Ref = CurrentExportDependencies.CurrentExport;
-			if (HarvestingRealm != ESaveRealm::Optional && Ref)
+			ForEachExportHarvestingRealm([this, Ref, Object](ESaveRealm HarvestingRealm)
 			{
-				// TODO: Change IsEditorOnlyObject to also test NeedsLoadForClient and NeedsLoadForServer, 
-				// and then change this location to call IsEditorOnlyObject
-				bool bEditorOnly = Ref->IsEditorOnly() || (!Ref->NeedsLoadForClient() && !Ref->NeedsLoadForServer());
-				if (!bEditorOnly)
+				if (HarvestingRealm != ESaveRealm::Optional)
 				{
+					TArray<UObject*> Exclude;
+					FOutputDeviceDebug Output;
+					FArchiveShowReferences ShowReferences(Output, Ref, Ref, Exclude);
+
 					SaveContext.RecordIllegalReference(Ref, Object, EIllegalRefReason::ReferenceToOptional);
 				}
-			}
-		});
+			});
+		}
+	}
 
+	// No matter the realm referring to it, optional objects are added into the optional realm
 	CurrentExportHarvestingRealms.Reset();
 	CurrentExportHarvestingRealms.Add(ESaveRealm::Optional);
 	return Scope;
