renamed link list builder accessors

diff --git a/Engine/Source/Runtime/Core/Public/Containers/LinkedListBuilder.h b/Engine/Source/Runtime/Core/Public/Containers/LinkedListBuilder.h
--- a/Engine/Source/Runtime/Core/Public/Containers/LinkedListBuilder.h
+++ b/Engine/Source/Runtime/Core/Public/Containers/LinkedListBuilder.h
@@ -5,32 +5,32 @@
 #include "CoreTypes.h"
 #include "Misc/AssertionMacros.h"
 
-template<typename InElementType, InElementType* InElementType::* NextLink>
-struct TLinkedListBuilderNextLinkCustom
+template<typename InElementType>
+struct TLinkedListBuilderNextLink
 {
 	using ElementType = InElementType;
 
 	FORCEINLINE static ElementType** GetNextPtr(ElementType& Element)
 	{
-		return &(Element.*NextLink);
+		return &(Element.Next);
 	}
 };
 
-template<typename InElementType>
-struct TLinkedListBuilderNextLinkDefault
+template<typename InElementType, InElementType* InElementType::* NextLink>
+struct TLinkedListBuilderNextLinkMemberVar
 {
 	using ElementType = InElementType;
 
 	FORCEINLINE static ElementType** GetNextPtr(ElementType& Element)
 	{
-		return &(Element.Next);
+		return &(Element.*NextLink);
 	}
 };
 
 /**
  * Single linked list builder
  */
-template<typename InElementType, typename InLinkAccessor = TLinkedListBuilderNextLinkDefault<InElementType>>
+template<typename InElementType, typename InLinkAccessor = TLinkedListBuilderNextLink<InElementType>>
 struct TLinkedListBuilder
 {
 public:
diff --git a/Engine/Source/Runtime/CoreUObject/Public/UObject/UnrealType.h b/Engine/Source/Runtime/CoreUObject/Public/UObject/UnrealType.h
--- a/Engine/Source/Runtime/CoreUObject/Public/UObject/UnrealType.h
+++ b/Engine/Source/Runtime/CoreUObject/Public/UObject/UnrealType.h
@@ -6617,10 +6617,10 @@ struct FPropertyChangedChainEvent : public FPropertyChangedEvent
 
 namespace UEProperty_Private
 {
-	using FPropertyListBuilderPropertyLink = TLinkedListBuilder<FProperty, TLinkedListBuilderNextLinkCustom<FProperty, &FProperty::PropertyLinkNext>>;
-	using FPropertyListBuilderRefLink = TLinkedListBuilder<FProperty, TLinkedListBuilderNextLinkCustom<FProperty, &FProperty::NextRef>>;
-	using FPropertyListBuilderDestructorLink = TLinkedListBuilder<FProperty, TLinkedListBuilderNextLinkCustom<FProperty, &FProperty::DestructorLinkNext>>;
-	using FPropertyListBuilderPostConstructLink = TLinkedListBuilder<FProperty, TLinkedListBuilderNextLinkCustom<FProperty, &FProperty::PostConstructLinkNext>>;
+	using FPropertyListBuilderPropertyLink = TLinkedListBuilder<FProperty, TLinkedListBuilderNextLinkMemberVar<FProperty, &FProperty::PropertyLinkNext>>;
+	using FPropertyListBuilderRefLink = TLinkedListBuilder<FProperty, TLinkedListBuilderNextLinkMemberVar<FProperty, &FProperty::NextRef>>;
+	using FPropertyListBuilderDestructorLink = TLinkedListBuilder<FProperty, TLinkedListBuilderNextLinkMemberVar<FProperty, &FProperty::DestructorLinkNext>>;
+	using FPropertyListBuilderPostConstructLink = TLinkedListBuilder<FProperty, TLinkedListBuilderNextLinkMemberVar<FProperty, &FProperty::PostConstructLinkNext>>;
 }
 
 /*-----------------------------------------------------------------------------
