FStore in the asset data system needs to be eagerly incremented and deferred decrement under AutoRTFM (so that resources are not destroyed before a transaction has committed successfully).

diff --git a/Engine/Source/Runtime/CoreUObject/Private/AssetRegistry/AssetDataTagMap.cpp b/Engine/Source/Runtime/CoreUObject/Private/AssetRegistry/AssetDataTagMap.cpp
--- a/Engine/Source/Runtime/CoreUObject/Private/AssetRegistry/AssetDataTagMap.cpp
+++ b/Engine/Source/Runtime/CoreUObject/Private/AssetRegistry/AssetDataTagMap.cpp
@@ -969,11 +969,14 @@ namespace FixedTagPrivate
 
 	void FStore::Release() const
 	{
-		if (RefCount.Decrement() == 0)
+		UE_AUTORTFM_ONCOMMIT(this)
 		{
-			delete this;
-		}
-	} 
+			if (RefCount.Decrement() == 0)
+			{
+				delete this;
+			}
+		};
+	}
 
 	FStore::~FStore()
 	{
diff --git a/Engine/Source/Runtime/CoreUObject/Public/AssetRegistry/AssetDataTagMapSerializationDetails.h b/Engine/Source/Runtime/CoreUObject/Public/AssetRegistry/AssetDataTagMapSerializationDetails.h
--- a/Engine/Source/Runtime/CoreUObject/Public/AssetRegistry/AssetDataTagMapSerializationDetails.h
+++ b/Engine/Source/Runtime/CoreUObject/Public/AssetRegistry/AssetDataTagMapSerializationDetails.h
@@ -5,7 +5,7 @@
 #include "AssetRegistry/AssetData.h"
 #include "AssetRegistry/AssetDataTagMap.h"
 #include "Async/Async.h"
-#include "Misc/TransactionallySafeCounter.h"
+#include "AutoRTFM/AutoRTFM.h"
 #include "HAL/ThreadSafeCounter.h"
 
 struct FAssetRegistrySerializationOptions;
@@ -63,7 +63,19 @@ namespace FixedTagPrivate
 		const uint32 Index;
 		void* Data = nullptr;
 
-		void AddRef() const { RefCount.Increment(); }
+		void AddRef() const
+		{
+			UE_AUTORTFM_OPEN
+			{
+				RefCount.Increment();
+			};
+
+			UE_AUTORTFM_ONABORT(this)
+			{
+				RefCount.Decrement();
+			};
+		}
+
 		COREUOBJECT_API void Release() const;
 		
 		const ANSICHAR* GetAnsiString(uint32 Idx) const { return &AnsiStrings[AnsiStringOffsets[Idx]]; }
@@ -71,10 +83,17 @@ namespace FixedTagPrivate
 
 	private:
 		friend class FStoreManager;
-		explicit FStore(uint32 InIndex) : Index(InIndex) {}
+
+		explicit FStore(uint32 InIndex) : Index(InIndex)
+		{
+			// `IsClosed` will fold away at compile time to true|false meaning this check will only impact
+			// transactionally compiled code!
+			check(AutoRTFM::IsClosed() || !AutoRTFM::IsOnCurrentTransactionStack(this));
+		}
+
 		~FStore();
 
-		mutable TTransactionallySafeCounter<FThreadSafeCounter> RefCount;
+		mutable FThreadSafeCounter RefCount;
 	};
 
 	struct FOptions
