Added verbose logging for localized package resolution

#rnx

diff --git a/Engine/Source/Runtime/CoreUObject/Private/Internationalization/PackageLocalizationManager.cpp b/Engine/Source/Runtime/CoreUObject/Private/Internationalization/PackageLocalizationManager.cpp
--- a/Engine/Source/Runtime/CoreUObject/Private/Internationalization/PackageLocalizationManager.cpp
+++ b/Engine/Source/Runtime/CoreUObject/Private/Internationalization/PackageLocalizationManager.cpp
@@ -96,27 +96,45 @@ FName FPackageLocalizationManager::FindLocalizedPackageName(const FName InSource
 {
 	PerformLazyInitialization();
 
+	FName LocalizedPackageName;
+
 	if (ActiveCache.IsValid())
 	{
-		return ActiveCache->FindLocalizedPackageName(InSourcePackageName);
+		LocalizedPackageName = ActiveCache->FindLocalizedPackageName(InSourcePackageName);
+	}
+	else
+	{
+		UE_LOG(LogPackageLocalizationManager, Warning, TEXT("Localized package requested for '%s' before the package localization manager cache was ready. Falling back to a non-cached look-up..."), *InSourcePackageName.ToString());
+
+		const FString CurrentCultureName = FInternationalization::Get().GetCurrentLanguage()->GetName();
+		LocalizedPackageName = FindLocalizedPackageNameNoCache(InSourcePackageName, CurrentCultureName);
 	}
-	
-	UE_LOG(LogPackageLocalizationManager, Warning, TEXT("Localized package requested for '%s' before the package localization manager cache was ready. Falling back to a non-cached look-up..."), *InSourcePackageName.ToString());
-	const FString CurrentCultureName = FInternationalization::Get().GetCurrentCulture()->GetName();
-	return FindLocalizedPackageNameNoCache(InSourcePackageName, CurrentCultureName);
+
+	UE_CLOG(!LocalizedPackageName.IsNone(), LogPackageLocalizationManager, Verbose, TEXT("Resolved localized package '%s' for source package '%s'"), *LocalizedPackageName.ToString(), *InSourcePackageName.ToString());
+
+	return LocalizedPackageName;
 }
 
 FName FPackageLocalizationManager::FindLocalizedPackageNameForCulture(const FName InSourcePackageName, const FString& InCultureName)
 {
 	PerformLazyInitialization();
 
+	FName LocalizedPackageName;
+
 	if (ActiveCache.IsValid())
 	{
-		return ActiveCache->FindLocalizedPackageNameForCulture(InSourcePackageName, InCultureName);
+		LocalizedPackageName = ActiveCache->FindLocalizedPackageNameForCulture(InSourcePackageName, InCultureName);
 	}
+	else
+	{
+		UE_LOG(LogPackageLocalizationManager, Warning, TEXT("Localized package requested for '%s' before the package localization manager cache was ready. Falling back to a non-cached look-up..."), *InSourcePackageName.ToString());
+
+		LocalizedPackageName = FindLocalizedPackageNameNoCache(InSourcePackageName, InCultureName);
+	}
+
+	UE_CLOG(!LocalizedPackageName.IsNone(), LogPackageLocalizationManager, Verbose, TEXT("Resolved localized package '%s' for source package '%s'"), *LocalizedPackageName.ToString(), *InSourcePackageName.ToString());
 
-	UE_LOG(LogPackageLocalizationManager, Warning, TEXT("Localized package requested for '%s' before the package localization manager cache was ready. Falling back to a non-cached look-up..."), *InSourcePackageName.ToString());
-	return FindLocalizedPackageNameNoCache(InSourcePackageName, InCultureName);
+	return LocalizedPackageName;
 }
 
 FName FPackageLocalizationManager::FindLocalizedPackageNameNoCache(const FName InSourcePackageName, const FString& InCultureName) const
