Removing extraneous Initialize() calls inside FCoreRedirects to avoid tripping some multithreading ensures.

diff --git a/Engine/Source/Runtime/CoreUObject/Private/UObject/CoreRedirects.cpp b/Engine/Source/Runtime/CoreUObject/Private/UObject/CoreRedirects.cpp
--- a/Engine/Source/Runtime/CoreUObject/Private/UObject/CoreRedirects.cpp
+++ b/Engine/Source/Runtime/CoreUObject/Private/UObject/CoreRedirects.cpp
@@ -1715,7 +1715,7 @@ bool FCoreRedirects::IsKnownMissing(ECoreRedirectFlags Type, const FCoreRedirect
 
 bool FCoreRedirects::AddKnownMissing(ECoreRedirectFlags Type, const FCoreRedirectObjectName& ObjectName, ECoreRedirectFlags Channel)
 {
-	Initialize();
+	ensureMsgf(IsInitialized(), TEXT("FCoreRedirects must be initialized on the game thread before use."));
 
 	check((Channel & ~ECoreRedirectFlags::Option_MissingLoad) == ECoreRedirectFlags::None);
 	FCoreRedirect NewRedirect(Type | ECoreRedirectFlags::Category_Removed | Channel, ObjectName, FCoreRedirectObjectName());
@@ -2070,7 +2070,7 @@ bool FCoreRedirectTest::RunTest(const FString& Parameters)
 bool FCoreRedirects::ReadRedirectsFromIni(const FString& IniName)
 {
 	FCoreRedirectorScopeLockForWrite ScopeLock(RWLock);
-	Initialize();
+	ensureMsgf(IsInitialized(), TEXT("FCoreRedirects must be initialized on the game thread before use."));
 
 	if (GConfig)
 	{
@@ -2221,7 +2221,7 @@ bool FCoreRedirects::ReadRedirectsFromIni(const FString& IniName)
 bool FCoreRedirects::AddRedirectList(TArrayView<const FCoreRedirect> Redirects, const FString& SourceString)
 {
 	FCoreRedirectorScopeLockForWrite ScopeLock(RWLock);
-	Initialize();
+	ensureMsgf(IsInitialized(), TEXT("FCoreRedirects must be initialized on the game thread before use."));
 	return AddRedirectListUnderWriteLock(Redirects, SourceString, ScopeLock);
 }
 
