Skipping CreateInstanceDataObject calls for already created IDOs.
This is to fix a triggering ensure that checks the seek pointer of the archive.

diff --git a/Engine/Source/Runtime/CoreUObject/Private/UObject/PropertyBagRepository.cpp b/Engine/Source/Runtime/CoreUObject/Private/UObject/PropertyBagRepository.cpp
--- a/Engine/Source/Runtime/CoreUObject/Private/UObject/PropertyBagRepository.cpp
+++ b/Engine/Source/Runtime/CoreUObject/Private/UObject/PropertyBagRepository.cpp
@@ -1110,7 +1110,9 @@ void FScopedIDOSerializationContext::FinishCreatingInstanceDataObject() const
 		// when using the linker, the repository will handle offsets
 		FPropertyBagRepository::Get().CreateInstanceDataObject(Object, Archive);
 	}
-	else
+	// Check that we haven't already created an IDO for this object so we can skip further work.
+	// This double processing can happen at re-instancing time during CPFUO.
+	else if (FPropertyBagRepository::Get().FindInstanceDataObject(Object) == nullptr)
 	{
 		check(Archive);
 		const int64 PostSerializeOffset = Archive->Tell();
