Downgrade most logging in AsyncLoading2 to be Log rather than Display. We keep logs at Display when they are signalling bad loading behaviour or soft warnings to users

#rnx

diff --git a/Engine/Source/Runtime/CoreUObject/Private/Serialization/AsyncLoading2.cpp b/Engine/Source/Runtime/CoreUObject/Private/Serialization/AsyncLoading2.cpp
--- a/Engine/Source/Runtime/CoreUObject/Private/Serialization/AsyncLoading2.cpp
+++ b/Engine/Source/Runtime/CoreUObject/Private/Serialization/AsyncLoading2.cpp
@@ -1340,7 +1340,7 @@ public:
 						FLoadedPackageRef* OldPackageRef = Packages.Find(OldPackageId);
 						if (OldPackageRef)
 						{
-							UE_LOG(LogStreaming, Display,
+							UE_LOG(LogStreaming, Log,
 								TEXT("FGlobalImportStore:AddPackageRef: Dropping stale reference to package %s (0x%s) that has been renamed to %s (0x%s)"),
 								*OldPackageRef->GetOriginalPackageName().ToString(),
 								*LexToString(OldPackageId),
@@ -1382,7 +1382,7 @@ public:
 			{
 				if (Package->IsUnreachable() || PackageRef.GetOriginalPackageName() != Package->GetFName())
 				{
-					UE_CLOG(!Package->IsUnreachable(), LogStreaming, Display,
+					UE_CLOG(!Package->IsUnreachable(), LogStreaming, Log,
 						TEXT("FGlobalImportStore:AddPackageRef: Dropping renamed package %s before reloading %s (0x%s)"),
 						*Package->GetName(),
 						*PackageRef.GetOriginalPackageName().ToString(),
@@ -4338,7 +4338,7 @@ void FAsyncLoadingThread2::InitializeLoading()
 	
 	AsyncThreadReady.Increment();
 
-	UE_LOG(LogStreaming, Display, TEXT("AsyncLoading2 - Initialized"));
+	UE_LOG(LogStreaming, Log, TEXT("AsyncLoading2 - Initialized"));
 }
 
 void FAsyncLoadingThread2::UpdatePackagePriority(FAsyncLoadingThreadState2& ThreadState, FAsyncPackage2* Package)
@@ -5491,7 +5491,7 @@ void FAsyncPackage2::ImportPackagesRecursiveInner(FAsyncLoadingThreadState2& Thr
 		{
 			if (!ImportedPackageRef.HasPackage()) // If we found a package it's not actually missing but we can't load it anyway
 			{
-				UE_ASYNC_PACKAGE_CLOG(!ImportedPackageUPackageName.IsNone(), Display, Desc, TEXT("ImportPackages: SkipPackage"),
+				UE_ASYNC_PACKAGE_CLOG(!ImportedPackageUPackageName.IsNone(), Log, Desc, TEXT("ImportPackages: SkipPackage"),
 					TEXT("Skipping non mounted imported package %s (0x%s)"), *ImportedPackageNameToLoad.ToString(), *LexToString(ImportedPackageId));
 				if (!ImportedPackageUPackageName.IsNone())
 				{
@@ -7338,7 +7338,7 @@ void FAsyncPackage2::ConditionalReleasePartialRequests(FAsyncLoadingThreadState2
 			{
 				int32 RequestId = CurrentSyncLoadContext->RequestIDs[Index];
 				// Release the sync context request tied to our package, allowing the current flush to exit.
-				UE_LOG(LogStreaming, Display, TEXT("Package %s has reached state %s > %s, releasing request %d to allow recursive sync load to finish"),
+				UE_LOG(LogStreaming, Log, TEXT("Package %s has reached state %s > %s, releasing request %d to allow recursive sync load to finish"),
 					*Desc.UPackageName.ToString(),
 					LexToString(AsyncPackageLoadingState),
 					LexToString(RequesterState),
@@ -7359,7 +7359,7 @@ void FAsyncPackage2::ConditionalReleasePartialRequests(FAsyncLoadingThreadState2
 			// gives the opportunity to the caller to get a pointer to an object inside the package even if it's not fully loaded yet.
 			if (AsyncPackageLoadingState == EAsyncPackageLoadingState2::ExportsDone && CurrentSyncLoadContext->RequestingPackage->PostLoadGroup)
 			{
-				UE_LOG(LogStreaming, Display, TEXT("Merging postload groups of package %s with requester package %s"),
+				UE_LOG(LogStreaming, Log, TEXT("Merging postload groups of package %s with requester package %s"),
 					*Desc.UPackageName.ToString(),
 					*CurrentSyncLoadContext->RequestingPackage->Desc.UPackageName.ToString());
 
@@ -8510,7 +8510,7 @@ FAsyncLoadingThread2::FAsyncLoadingThread2(FIoDispatcher& InIoDispatcher, IAsync
 	MainThreadEventQueue.SetOwnerThread(GameThreadState.Get());
 	FAsyncLoadingThreadState2::Set(GameThreadState.Get());
 
-	UE_LOG(LogStreaming, Display, TEXT("AsyncLoading2 - Created: Event Driven Loader: %s, Async Loading Thread: %s, Async Post Load: %s"),
+	UE_LOG(LogStreaming, Log, TEXT("AsyncLoading2 - Created: Event Driven Loader: %s, Async Loading Thread: %s, Async Post Load: %s"),
 		GEventDrivenLoaderEnabled ? TEXT("true") : TEXT("false"),
 		FAsyncLoadingThreadSettings::Get().bAsyncLoadingThreadEnabled ? TEXT("true") : TEXT("false"),
 		FAsyncLoadingThreadSettings::Get().bAsyncPostLoadEnabled ? TEXT("true") : TEXT("false"));
@@ -8569,7 +8569,7 @@ void FAsyncLoadingThread2::StartThread()
 		UE::Trace::ThreadGroupEnd();
 	}
 
-	UE_LOG(LogStreaming, Display, TEXT("AsyncLoading2 - Thread Started: %s, IsInitialLoad: %s"),
+	UE_LOG(LogStreaming, Log, TEXT("AsyncLoading2 - Thread Started: %s, IsInitialLoad: %s"),
 		FAsyncLoadingThreadSettings::Get().bAsyncLoadingThreadEnabled ? TEXT("true") : TEXT("false"),
 		GIsInitialLoad ? TEXT("true") : TEXT("false"));
 }
@@ -8797,7 +8797,7 @@ void FAsyncLoadingThread2::SuspendLoading()
 	int32 OldCount = SuspendRequestedCount.fetch_add(1);
 	if (OldCount == 0)
 	{
-		UE_LOG(LogStreaming, Display, TEXT("Suspending async loading"));
+		UE_LOG(LogStreaming, Log, TEXT("Suspending async loading"));
 		if (IsMultithreaded())
 		{
 			TRACE_LOADTIME_SUSPEND_ASYNC_LOADING();
@@ -8819,7 +8819,7 @@ void FAsyncLoadingThread2::ResumeLoading()
 	UE_CLOG(OldCount < 1, LogStreaming, Fatal, TEXT("Trying to resume async loading when it's not suspended"));
 	if (OldCount == 1)
 	{
-		UE_LOG(LogStreaming, Display, TEXT("Resuming async loading"));
+		UE_LOG(LogStreaming, Log, TEXT("Resuming async loading"));
 		if (IsMultithreaded())
 		{
 			ThreadResumedEvent->Wait();
@@ -9117,7 +9117,7 @@ void FAsyncLoadingThread2::NotifyRegistrationComplete()
 	GlobalImportStore.RegistrationComplete();
 	bHasRegisteredAllScriptObjects = true;
 
-	UE_LOG(LogStreaming, Display,
+	UE_LOG(LogStreaming, Log,
 		TEXT("AsyncLoading2 - NotifyRegistrationComplete: Registered %d public script object entries (%.2f KB)"),
 		GlobalImportStore.GetStoredScriptObjectsCount(), (float)GlobalImportStore.GetStoredScriptObjectsAllocatedSize() / 1024.f);
 }
