Update PropertyValueInitialized state during PostEditChangeProperty

diff --git a/Engine/Source/Runtime/CoreUObject/Private/UObject/Obj.cpp b/Engine/Source/Runtime/CoreUObject/Private/UObject/Obj.cpp
--- a/Engine/Source/Runtime/CoreUObject/Private/UObject/Obj.cpp
+++ b/Engine/Source/Runtime/CoreUObject/Private/UObject/Obj.cpp
@@ -36,6 +36,7 @@
 #include "UObject/AssetRegistryTagsContext.h"
 #include "UObject/LazyObjectPtr.h"
 #include "UObject/SoftObjectPtr.h"
+#include "UObject/PropertyBagRepository.h"
 #include "UObject/PropertyPortFlags.h"
 #include "UObject/UnrealType.h"
 #include "UObject/ObjectRedirector.h"
@@ -467,6 +468,40 @@ void UObject::PostEditChangeProperty(FPropertyChangedEvent& PropertyChangedEvent
 {
 	FCoreUObjectDelegates::OnObjectPropertyChanged.Broadcast(this, PropertyChangedEvent);
 
+#if WITH_EDITORONLY_DATA
+	const FProperty* Property = PropertyChangedEvent.MemberProperty;
+	// Set property initialized state when a property is changed and IDO is supported
+	if (UE::IsInstanceDataObjectSupportEnabled(this))
+	{
+		UE::FPropertyBagRepository& Repository = UE::FPropertyBagRepository::Get();
+
+		UObject* ChangedIDO = UE::IsClassOfInstanceDataObjectClass(this->GetClass()) ?
+								this : Repository.FindInstanceDataObject(this);
+		if (ChangedIDO)
+		{
+			int32 ChangedIndex = 0;
+			switch (PropertyChangedEvent.ChangeType)
+			{
+			case EPropertyChangeType::ArrayAdd:
+			case EPropertyChangeType::ArrayRemove:
+			case EPropertyChangeType::ArrayMove:
+				ChangedIndex = PropertyChangedEvent.GetArrayIndex(PropertyChangedEvent.GetPropertyName().ToString());
+				ChangedIndex = (ChangedIndex == -1) ? 0 : ChangedIndex; // If we didn't get a valid array index revert to index 0 as it's the API default
+				UE::SetPropertyValueInitialized(ChangedIDO->GetClass(), ChangedIDO, Property, ChangedIndex);
+				break;
+			case EPropertyChangeType::ArrayClear: // Clearing the entire array. No need for an index. Clear also counts as initialization of the property
+			case EPropertyChangeType::ValueSet:
+			case EPropertyChangeType::Interactive:
+				UE::SetPropertyValueInitialized(ChangedIDO->GetClass(), ChangedIDO, Property);
+				break;
+			default:
+				// No op for now
+				break;
+			}
+		}
+	}
+#endif
+
 	// Snapshot the transaction buffer for this object if this was from an interactive change
 	// This allows listeners to be notified of intermediate changes of state
 	if (PropertyChangedEvent.ChangeType == EPropertyChangeType::Interactive)
