Update Details Panel Alerts to track Uninitialized Non_Concrete properties

diff --git a/Engine/Source/Runtime/CoreUObject/Internal/UObject/InstanceDataObjectUtils.h b/Engine/Source/Runtime/CoreUObject/Internal/UObject/InstanceDataObjectUtils.h
--- a/Engine/Source/Runtime/CoreUObject/Internal/UObject/InstanceDataObjectUtils.h
+++ b/Engine/Source/Runtime/CoreUObject/Internal/UObject/InstanceDataObjectUtils.h
@@ -24,13 +24,16 @@ namespace UE
 bool IsInstanceDataObjectSupportEnabled();
 /** Query if InstanceDataObject support is enabled for a specific object. */
 bool IsInstanceDataObjectSupportEnabled(const UObject* Object);
+/** Query if uninitialized alert UI is enabled **/
+UE_API bool IsUninitializedAlertUIEnabled();
+
 bool StructContainsLooseProperties(const UStruct* Struct);
 
 /** Query if placeholder support is enabled for a specific import class type. */
 bool CanCreatePropertyBagPlaceholderTypeForImportClass(const UClass* ImportType);
 
 /** Helper to check if a class is an IDO class. */
-bool IsClassOfInstanceDataObjectClass(UStruct* Class);
+UE_API bool IsClassOfInstanceDataObjectClass(const UStruct* Class);
 
 /** Generate a UClass that contains the union of the properties of PropertyTree and OwnerClass. */
 UClass* CreateInstanceDataObjectClass(const FPropertyPathNameTree* PropertyTree, const FUnknownEnumNames* EnumNames, UClass* OwnerClass, UObject* Outer);
diff --git a/Engine/Source/Runtime/CoreUObject/Private/UObject/InstanceDataObjectUtils.cpp b/Engine/Source/Runtime/CoreUObject/Private/UObject/InstanceDataObjectUtils.cpp
--- a/Engine/Source/Runtime/CoreUObject/Private/UObject/InstanceDataObjectUtils.cpp
+++ b/Engine/Source/Runtime/CoreUObject/Private/UObject/InstanceDataObjectUtils.cpp
@@ -151,6 +151,13 @@ namespace UE
 		TEXT("Comma separated list of property types that will be excluded from loose properties in IDOs.")
 	);
 
+	bool bEnableUninitializedUI = false;
+	FAutoConsoleVariableRef EnableUninitializedUI(
+		TEXT("IDO.EnableUninitializedAlertUI"),
+		bEnableUninitializedUI,
+		TEXT("Enables alert information for uninitalized properties. Requires IDO.Enable=true")
+	);
+
 	static TSet<FString> GetExcludedLoosePropertyTypes()
 	{
 		TArray<FString> Result;
@@ -163,6 +170,11 @@ namespace UE
 		return bEnableIDOSupport;
 	}
 
+	bool IsUninitializedAlertUIEnabled()
+	{
+		return bEnableIDOSupport && bEnableUninitializedUI;
+	}
+
 	bool IsInstanceDataObjectSupportEnabled(const UObject* InObject)
 	{
 		if(!IsInstanceDataObjectSupportEnabled() || !InObject)
@@ -220,7 +232,7 @@ namespace UE
 		return ImportClass && ImportClass->GetFName() == NAME_VerseClass;
 	}
 
-	bool IsClassOfInstanceDataObjectClass(UStruct* Class)
+	bool IsClassOfInstanceDataObjectClass(const UStruct* Class)
 	{
 		return Class->IsA(UInstanceDataObjectClass::StaticClass()) || Class->IsA(UInstanceDataObjectStruct::StaticClass());
 	}
