Adding support for remote objects to FWeakObjectPtr / FObjectKey

Moved FRemoteObjectId from UObjectBase to FUObjectItem to serve as SerialNumber when remote objects are enabled
Added the ability to override SerialNumber and RemoteId when constructing default subobjects of remote objects
Changed remote object serialization so that the basic information about serialized objects is stored in a header
Fixed patching of deserialized FObjectPtrs for remote objects (previously the deserializing archive would store FObjectPtr* which is not always safe)

diff --git a/Engine/Plugins/Runtime/GameplayAbilities/Source/GameplayAbilities/Private/Serialization/GameplayEffectContextNetSerializer.cpp b/Engine/Plugins/Runtime/GameplayAbilities/Source/GameplayAbilities/Private/Serialization/GameplayEffectContextNetSerializer.cpp
--- a/Engine/Plugins/Runtime/GameplayAbilities/Source/GameplayAbilities/Private/Serialization/GameplayEffectContextNetSerializer.cpp
+++ b/Engine/Plugins/Runtime/GameplayAbilities/Source/GameplayAbilities/Private/Serialization/GameplayEffectContextNetSerializer.cpp
@@ -47,7 +47,11 @@ struct FGameplayEffectContextNetSerializer
 
 	struct FQuantizedType
 	{
+#if UE_WITH_REMOTE_OBJECT_HANDLE
+		alignas(16) uint8 EffectContext[176];
+#else
 		alignas(16) uint8 EffectContext[128];
+#endif
 		// 320 should be large enough and makes struct a multiple of 16 bytes
 		alignas(16) uint8 HitResult[320];
 		uint32 ReplicationFlags;
@@ -517,7 +521,11 @@ void FGameplayEffectContextNetSerializer::FNetSerializerRegistryDelegates::OnPos
 	// Setup serializer for most members. The HitResult requires special care and is handled later.
 	{
 		const UStruct* GEStruct = FGameplayEffectContextAccessorForNetSerializer::StaticStruct();
+#if UE_WITH_REMOTE_OBJECT_HANDLE
+		if (GEStruct->GetStructureSize() != 176 || GEStruct->GetMinAlignment() != 8)
+#else
 		if (GEStruct->GetStructureSize() != 128 || GEStruct->GetMinAlignment() != 8)
+#endif
 		{
 			LowLevelFatalError(TEXT("%s Size: %d Alignment: %d"), TEXT("FGameplayEffectContext layout has changed. Need to update FGameplayEffectContextNetSerializer."), GEStruct->GetStructureSize(), GEStruct->GetMinAlignment());
 		}
diff --git a/Engine/Plugins/Runtime/GameplayAbilities/Source/GameplayAbilities/Private/Serialization/MinimalReplicationTagCountNetSerializer.cpp b/Engine/Plugins/Runtime/GameplayAbilities/Source/GameplayAbilities/Private/Serialization/MinimalReplicationTagCountNetSerializer.cpp
--- a/Engine/Plugins/Runtime/GameplayAbilities/Source/GameplayAbilities/Private/Serialization/MinimalReplicationTagCountNetSerializer.cpp
+++ b/Engine/Plugins/Runtime/GameplayAbilities/Source/GameplayAbilities/Private/Serialization/MinimalReplicationTagCountNetSerializer.cpp
@@ -251,7 +251,11 @@ void FMinimalReplicationTagCountMapNetSerializer::FNetSerializerRegistryDelegate
 
 void FMinimalReplicationTagCountMapNetSerializer::FNetSerializerRegistryDelegates::OnPostFreezeNetSerializerRegistry()
 {
+#if UE_WITH_REMOTE_OBJECT_HANDLE
+	constexpr SIZE_T ExpectedSizeOfFMinimalReplicationTagCountMap = 112;
+#else
 	constexpr SIZE_T ExpectedSizeOfFMinimalReplicationTagCountMap = 104;
+#endif
 	constexpr SIZE_T ExpectedAlignOfFMinimalReplicationTagCountMap = 8;
 
 	constexpr SIZE_T ActualSizeOfFMinimalReplicationTagCountMap = sizeof(FMinimalReplicationTagCountMap);
diff --git a/Engine/Plugins/Runtime/GameplayAbilities/Source/GameplayAbilities/Public/GameplayPrediction.h b/Engine/Plugins/Runtime/GameplayAbilities/Source/GameplayAbilities/Public/GameplayPrediction.h
--- a/Engine/Plugins/Runtime/GameplayAbilities/Source/GameplayAbilities/Public/GameplayPrediction.h
+++ b/Engine/Plugins/Runtime/GameplayAbilities/Source/GameplayAbilities/Public/GameplayPrediction.h
@@ -387,7 +387,14 @@ struct GAMEPLAYABILITIES_API FPredictionKey
 		return ((InKey.Current << 1) | (InKey.bIsServerInitiated & 1));
 	}
 
-	uint64 GetPredictiveConnectionKey() const { return BitCast<uint64>(PredictiveConnectionObjectKey); }
+	uint64 GetPredictiveConnectionKey() const 
+	{ 
+#if UE_WITH_REMOTE_OBJECT_HANDLE
+		return BitCast<uint64>(PredictiveConnectionObjectKey.GetRemoteId());
+#else
+		return BitCast<uint64>(PredictiveConnectionObjectKey);
+#endif
+	}
 
 private:
 	friend UE::Net::FPredictionKeyNetSerializer;
diff --git a/Engine/Plugins/Runtime/GameplayAbilities/Source/GameplayAbilities/Public/Serialization/GameplayEffectContextNetSerializer.h b/Engine/Plugins/Runtime/GameplayAbilities/Source/GameplayAbilities/Public/Serialization/GameplayEffectContextNetSerializer.h
--- a/Engine/Plugins/Runtime/GameplayAbilities/Source/GameplayAbilities/Public/Serialization/GameplayEffectContextNetSerializer.h
+++ b/Engine/Plugins/Runtime/GameplayAbilities/Source/GameplayAbilities/Public/Serialization/GameplayEffectContextNetSerializer.h
@@ -14,7 +14,11 @@ struct FGameplayEffectContextNetSerializerConfig : public FNetSerializerConfig
 namespace UE::Net
 {
 
+#if UE_WITH_REMOTE_OBJECT_HANDLE
+constexpr SIZE_T GetGameplayEffectContextNetSerializerSafeQuantizedSize() { return 512; }
+#else
 constexpr SIZE_T GetGameplayEffectContextNetSerializerSafeQuantizedSize() { return 464; }
+#endif
 
 UE_NET_DECLARE_SERIALIZER(FGameplayEffectContextNetSerializer, GAMEPLAYABILITIES_API);
 
diff --git a/Engine/Source/Runtime/CoreUObject/Private/UObject/Class.cpp b/Engine/Source/Runtime/CoreUObject/Private/UObject/Class.cpp
--- a/Engine/Source/Runtime/CoreUObject/Private/UObject/Class.cpp
+++ b/Engine/Source/Runtime/CoreUObject/Private/UObject/Class.cpp
@@ -4961,9 +4961,17 @@ FName UClass::GetDefaultObjectName() const
 //
 // Register the native class.
 //
-void UClass::DeferredRegister(UClass *UClassStaticClass,const TCHAR* PackageName,const TCHAR* Name)
+void UClass::DeferredRegister(UClass *UClassStaticClass, const TCHAR* PackageName, const TCHAR* Name
+#if UE_WITH_REMOTE_OBJECT_HANDLE
+	, FRemoteObjectId RemoteId
+#endif
+)
 {
-	Super::DeferredRegister(UClassStaticClass,PackageName,Name);
+	Super::DeferredRegister(UClassStaticClass, PackageName, Name
+#if UE_WITH_REMOTE_OBJECT_HANDLE
+		, RemoteId
+#endif
+	);
 
 	// Get stashed registration info.
 
diff --git a/Engine/Source/Runtime/CoreUObject/Private/UObject/Obj.cpp b/Engine/Source/Runtime/CoreUObject/Private/UObject/Obj.cpp
--- a/Engine/Source/Runtime/CoreUObject/Private/UObject/Obj.cpp
+++ b/Engine/Source/Runtime/CoreUObject/Private/UObject/Obj.cpp
@@ -1279,6 +1279,10 @@ bool UObject::ConditionalFinishDestroy()
 #endif
 		FinishDestroy();
 
+#if UE_WITH_REMOTE_OBJECT_HANDLE
+		UE::RemoteObject::Private::UnmapObject(this);
+#endif
+
 		// Make sure this object can't be accessed via weak pointers after it's been FinishDestroyed
 		GUObjectArray.ResetSerialNumber(this);
 
diff --git a/Engine/Source/Runtime/CoreUObject/Private/UObject/ObjectHandlePrivate.h b/Engine/Source/Runtime/CoreUObject/Private/UObject/ObjectHandlePrivate.h
--- a/Engine/Source/Runtime/CoreUObject/Private/UObject/ObjectHandlePrivate.h
+++ b/Engine/Source/Runtime/CoreUObject/Private/UObject/ObjectHandlePrivate.h
@@ -2,6 +2,7 @@
 
 #pragma once
 #include "UObject/Object.h"
+#include "UObject/UObjectArray.h"
 
 namespace UE::CoreUObject::Private
 {
@@ -28,7 +29,14 @@ namespace UE::CoreUObject::Private
 		FORCEINLINE static FRemoteObjectId GetRemoteId(const UObjectBase* Object)
 		{
 #if UE_WITH_REMOTE_OBJECT_HANDLE
-			return Object->RemoteId;
+			if (UNLIKELY(Object->InternalIndex < 0))
+			{
+				return Object->GetPendingRegistrantRemoteId();
+			}
+			else
+			{
+				return GUObjectArray.GetRemoteId(Object->InternalIndex);
+			}
 #else
 			return FRemoteObjectId();
 #endif // UE_WITH_REMOTE_OBJECT_HANDLE
diff --git a/Engine/Source/Runtime/CoreUObject/Private/UObject/ObjectPathId.cpp b/Engine/Source/Runtime/CoreUObject/Private/UObject/ObjectPathId.cpp
--- a/Engine/Source/Runtime/CoreUObject/Private/UObject/ObjectPathId.cpp
+++ b/Engine/Source/Runtime/CoreUObject/Private/UObject/ObjectPathId.cpp
@@ -40,10 +40,12 @@ assert_equality<sizeof(TArray<FMinimalName, TInlineAllocator<3>>), 0> B;
 
 namespace UE::CoreUObject::Private
 {
+#if !UE_WITH_REMOTE_OBJECT_HANDLE
 	FObjectKey MakeObjectKey(int32 ObjectIndex, int32 ObjectSerialNumber)
 	{
 		return FObjectKey(ObjectIndex, ObjectSerialNumber);
 	}
+#endif
 
 	FStoredObjectPath::FStoredObjectPath(TConstArrayView<FMinimalName> InNames)
 	{
@@ -376,8 +378,14 @@ namespace UE::CoreUObject::Private
 
 		int32 ObjectIndex = Index;
 		int32 ObjectSerial = Number & ~WeakObjectMask;
+#if !UE_WITH_REMOTE_OBJECT_HANDLE
 		FObjectKey Key = UE::CoreUObject::Private::MakeObjectKey(ObjectIndex, ObjectSerial);
+
 		return Key.ResolveObjectPtr();
+#else
+		checkf(false, TEXT("FObjectPathId::GetWeakObjPtr() can't be used with remote object handles"));
+		return FWeakObjectPtr();
+#endif
 	}
 
 	const FStoredObjectPath& FObjectPathId::GetStoredPath() const
diff --git a/Engine/Source/Runtime/CoreUObject/Private/UObject/RemoteObject.cpp b/Engine/Source/Runtime/CoreUObject/Private/UObject/RemoteObject.cpp
--- a/Engine/Source/Runtime/CoreUObject/Private/UObject/RemoteObject.cpp
+++ b/Engine/Source/Runtime/CoreUObject/Private/UObject/RemoteObject.cpp
@@ -67,8 +67,11 @@ public:
 		using namespace UE::RemoteObject::Handle;
 
 		FScopeLock MapLock(&ObjectMapCritical);
-		FRemoteObjectStub* ExistingStub = RemoteObjects.FindChecked(Id);
-		return ExistingStub;
+		if (FRemoteObjectStub** ExistingStub = RemoteObjects.Find(Id))
+		{
+			return *ExistingStub;
+		}
+		return nullptr;
 	}
 
 	UE::RemoteObject::Handle::FRemoteObjectStub* FindOrAddRemoteObjectStub(FRemoteObjectId Id)
@@ -76,14 +79,14 @@ public:
 		using namespace UE::RemoteObject::Handle;
 
 		FScopeLock MapLock(&ObjectMapCritical);
-		FRemoteObjectStub*& ExistingStub = RemoteObjects.FindOrAdd(Id);
-		if (!ExistingStub)
+		FRemoteObjectStub*& Stub = RemoteObjects.FindOrAdd(Id);
+		if (!Stub)
 		{
-			ExistingStub = new FRemoteObjectStub();
-			ExistingStub->Id = Id;
-			ExistingStub->RemoteServerId = Id.GetServerId();
-		}
-		return ExistingStub;
+			Stub = new FRemoteObjectStub();
+			Stub->Id = Id;
+			Stub->RemoteServerId = Id.GetServerId();
+		}		
+		return Stub;
 	}
 
 	UE::RemoteObject::Handle::FRemoteObjectStub* FindOrAddRemoteObjectStub(UObject* Object, FRemoteServerId DestinationServerId)
@@ -91,13 +94,17 @@ public:
 		using namespace UE::RemoteObject::Handle;
 
 		FScopeLock MapLock(&ObjectMapCritical);
-		FRemoteObjectStub*& ExistingStub = RemoteObjects.FindOrAdd(UE::CoreUObject::Private::FObjectHandleUtils::GetRemoteId(Object));
-		if (!ExistingStub)
+		FRemoteObjectStub*& Stub = RemoteObjects.FindOrAdd(UE::CoreUObject::Private::FObjectHandleUtils::GetRemoteId(Object));
+		if (!Stub)
 		{
-			ExistingStub = new FRemoteObjectStub(Object);
+			Stub = new FRemoteObjectStub(Object);
 		}
-		ExistingStub->RemoteServerId = DestinationServerId;
-		return ExistingStub;
+		Stub->RemoteServerId = DestinationServerId;
+		FUObjectItem * ObjectItem = GUObjectArray.ObjectToObjectItem(Object);
+		checkf(ObjectItem, TEXT("Attempting to get a serial number for an object that does not exist in the global UObject array (it's possible GUObjectArray is not initialized yet, ObjectIndex=%d)"), GUObjectArray.ObjectToIndex(Object));
+		Stub->SerialNumber = ObjectItem->SerialNumber;
+
+		return Stub;
 	}
 };
 FRemoteObjectMaps* ObjectMaps = nullptr;
@@ -224,11 +231,12 @@ UObject* ResolveObject(const FRemoteObjectStub* Stub)
 	return Object;
 }
 
-COREUOBJECT_API UObject* ResolveObject(UObject* Object)
+UObject* ResolveObject(UObject* Object)
 {
 	using namespace UE::RemoteObject::Private;
 
 	FRemoteObjectStub* Stub = ObjectMaps->FindRemoteObjectStub(UE::CoreUObject::Private::FObjectHandleUtils::GetRemoteId(Object));
+	checkf(Stub, TEXT("Failed to find remote object stub for %s"), *GetPathNameSafe(Object));
 	return ResolveObject(Stub);
 }
 
@@ -241,7 +249,11 @@ namespace UE::CoreUObject::Private
 	{
 		using namespace UE::RemoteObject::Private;
 		UnmapObject((UObject*)Object);
-		Object->RemoteId = Id;
+		FUObjectItem* ObjectItem = GUObjectArray.ObjectToObjectItem(Object);
+		// ObjectItem may not exist when the UObject system hasn't been initialized yet but theoretically this function should only get called when
+		// something attempts to re-construct a default subobject that already exists so ObjectItem should always be valid
+		checkf(ObjectItem, TEXT("Attempting to change remote ID for an object that does not exist in the global UObject array (it's possible GUObjectArray is not initialized yet, ObjectIndex=%d)"), GUObjectArray.ObjectToIndex(Object));
+		ObjectItem->RemoteId = Id;
 		MapObject((UObject*)Object);
 	}
 
@@ -251,6 +263,7 @@ namespace UE::CoreUObject::Private
 		using namespace UE::RemoteObject::Private;		
 
 		FRemoteObjectStub* Stub = ObjectMaps->FindRemoteObjectStub(UE::CoreUObject::Private::FObjectHandleUtils::GetRemoteId(Object));
+		checkf(Stub, TEXT("Failed to find remote object stub for %s"), *GetPathNameSafe(Object));
 		return FRemoteObjectHandlePrivate(Stub);
 	}
 
diff --git a/Engine/Source/Runtime/CoreUObject/Private/UObject/RemoteObjectSerialization.cpp b/Engine/Source/Runtime/CoreUObject/Private/UObject/RemoteObjectSerialization.cpp
--- a/Engine/Source/Runtime/CoreUObject/Private/UObject/RemoteObjectSerialization.cpp
+++ b/Engine/Source/Runtime/CoreUObject/Private/UObject/RemoteObjectSerialization.cpp
@@ -50,6 +50,20 @@ void SaveObjectToDisk(FRemoteObjectId ObjectId, FRemoteServerId ServerId, FRemot
 namespace UE::RemoteObject::Serialization
 {
 
+const FRemoteObjectConstructionParams* FRemoteObjectConstructionOverrides::Find(FName InName, UObject* InOuter) const
+{
+	const FRemoteObjectId OuterId(InOuter);
+	// At the moment the number of serialized objects is usually pretty low (< 10) so no need for hash table lookups
+	for (const FRemoteObjectConstructionParams& RemoteObjectParams : Overrides)
+	{
+		if (RemoteObjectParams.Name == InName && RemoteObjectParams.OuterId == OuterId)
+		{
+			return &RemoteObjectParams;
+		}
+	}
+	return nullptr;
+}
+
 enum class ERemoteReferenceType : uint8
 {
 	None = 0,
@@ -123,6 +137,21 @@ public:
 		SetPortFlags(PPF_IncludeTransient);
 	}
 
+	int32 AddNameToNameMap(const FName& Name)
+	{
+		int32 NameIndex = -1;
+		if (int32* ExistingNameIndex = NameMap.Find(Name))
+		{
+			NameIndex = *ExistingNameIndex;
+		}
+		else
+		{
+			NameIndex = Names.Add(Name);
+			NameMap.Add(Name, NameIndex);
+		}
+		return NameIndex;
+	}
+
 	virtual FArchive& operator<<(UObject*& Obj) override
 	{
 		ERemoteReferenceType RefType = WriteObjectReference(Obj);
@@ -147,16 +176,7 @@ public:
 
 	virtual FArchive& operator<<(FName& Name) override
 	{
-		int32 NameIndex = -1;
-		if (int32* ExistingNameIndex = NameMap.Find(Name))
-		{
-			NameIndex = *ExistingNameIndex;
-		}
-		else
-		{
-			NameIndex = Names.Add(Name);
-			NameMap.Add(Name, NameIndex);
-		}
+		int32 NameIndex = AddNameToNameMap(Name);
 		*this << NameIndex;
 		return *this;
 	}
@@ -171,24 +191,32 @@ public:
 	virtual FArchive& operator<<(FWeakObjectPtr& Value) override { return FArchiveUObject::SerializeWeakObjectPtr(*this, Value); }
 };
 
-void SerializeRemoteObject(UObject* Object, FArchive& Ar)
+FArchive& operator<<(FArchive& Ar, FRemoteObjectHeader& Header)
+{
+	Ar << Header.Name;
+	Ar << Header.RemoteId;
+	Ar << Header.Class;
+	Ar << Header.Outer;
+	Ar << Header.Archetype;
+	return Ar;
+}
+
+void SerializeRemoteObject(UObject* Object, FArchive& Ar, FRemoteObjectHeader& OutHeader)
 {
 	using namespace UE::RemoteObject;
 	using namespace UE::RemoteObject::Private;
 	using namespace UE::CoreUObject::Private;
 
-	FName Name = Object->GetFName();
-	FRemoteObjectId Id = FObjectHandleUtils::GetRemoteId(Object);
-	UClass* Class = Object->GetClass();
+	OutHeader.Name = Object->GetFName();
+	OutHeader.RemoteId = FObjectHandleUtils::GetRemoteId(Object);
+	OutHeader.Class = Object->GetClass();
 	UObject* Outer = FObjectHandleUtils::GetNonAccessTrackedOuterNoResolve(Object);
+	OutHeader.Outer = Outer;
+	OutHeader.Archetype = Object->GetArchetype();
 
-	MapObject(Class);
+	MapObject(OutHeader.Class.Get());
 	MapObject(Outer);
-
-	Ar << Name;
-	Ar << Id;
-	Ar << Class;
-	Ar << Outer;
+	MapObject(OutHeader.Archetype.Get());
 
 	Object->Serialize(Ar);
 }
@@ -196,13 +224,6 @@ void SerializeRemoteObject(UObject* Object, FArchive& Ar)
 class FArchiveRemoteObjectReader : public FMemoryReader
 {
 	TArray<FName>& Names;
-	TMap<FRemoteObjectId, TArray<FObjectPtr*>> IdToPtrMap;
-
-	void StoreUnresolvedHandle(FRemoteObjectId Id, FObjectPtr* Ptr)
-	{
-		TArray<FObjectPtr*>& HandleList = IdToPtrMap.FindOrAdd(Id);
-		HandleList.Add(Ptr);
-	}
 
 	ERemoteReferenceType ReadObjectReference(FObjectPtr& Value)
 	{
@@ -236,10 +257,6 @@ class FArchiveRemoteObjectReader : public FMemoryReader
 #if UE_WITH_REMOTE_OBJECT_HANDLE
 				FObjectHandle Handle = FObjectHandle::FromIdNoResolve(ObjId);
 				Value = FObjectPtr(Handle);
-				if (!Value.IsResolved())
-				{
-					StoreUnresolvedHandle(ObjId, &Value);
-				}
 #endif // UE_WITH_REMOTE_OBJECT_HANDLE
 			}
 		}
@@ -286,49 +303,103 @@ public:
 	virtual FArchive& operator<<(FSoftObjectPtr& Value) override { return FArchiveUObject::SerializeSoftObjectPtr(*this, Value); }
 	virtual FArchive& operator<<(FSoftObjectPath& Value) override { return FArchiveUObject::SerializeSoftObjectPath(*this, Value); }
 	virtual FArchive& operator<<(FWeakObjectPtr& Value) override { return FArchiveUObject::SerializeWeakObjectPtr(*this, Value); }
+};
+
+/**
+* Archive that replaces unresolved FObjectPtrs (TObjectPtrs<>) with actual pointers to deserialized objects
+*/
+class FArchiveRemoteReferencePatcher : public FArchiveUObject
+{
+	TMap<FRemoteObjectId, UObject*> IdToObjectMap;
 
-	/**
-	* Fixup (resolve) all deserialized FObjectPtrs
-	*/
-	void ResolveObjectPtrs(const TArray<UObject*>& SerializedObjects)
+	void InitObjectMap(const TArray<UObject*>& DeserializedObjects)
 	{
-		for (UObject* Obj : SerializedObjects)
+		for (UObject* Obj : DeserializedObjects)
 		{
-			if (TArray<FObjectPtr*>* UnresolvedPtrs = IdToPtrMap.Find(UE::CoreUObject::Private::FObjectHandleUtils::GetRemoteId(Obj)))
-			{
-				for (FObjectPtr* Ptr : *UnresolvedPtrs)
-				{
-					*Ptr = FObjectPtr(Obj);
-				}
-			}
+			IdToObjectMap.Add(FRemoteObjectId(Obj), Obj);
 		}
 	}
+
+	void PatchObjectReference(FObjectPtr& Value)
+	{
+#if UE_WITH_REMOTE_OBJECT_HANDLE
+		using namespace UE::RemoteObject::Handle;
+		const FRemoteObjectStub* Stub = Value.GetHandleRef().ToStub();
+		if (UObject** ResolvedObject = IdToObjectMap.Find(Stub->Id))
+		{
+			Value = FObjectPtr(*ResolvedObject);
+		}
+#endif
+	}
+
+public:
+
+	FArchiveRemoteReferencePatcher(const TArray<UObject*>& DeserializedObjects)
+	{
+		ArIsObjectReferenceCollector = true;
+		SetIsPersistent(false);
+		InitObjectMap(DeserializedObjects);
+	}
+
+	virtual FArchive& operator<<(FObjectPtr& Value) override
+	{
+		if (!Value.IsResolved())
+		{
+			PatchObjectReference(Value);
+		}
+		return *this;
+	}
 };
 
-UObject* DeserializeRemoteObject(FArchive& Ar, UObject* DestinationOuter)
+FRemoteObjectConstructionOverrides::FRemoteObjectConstructionOverrides(const TArray<FRemoteObjectHeader>& InObjectHeaders)
 {
-	using namespace UE::RemoteObject;
+	using namespace UE::RemoteObject::Handle;
 	using namespace UE::RemoteObject::Private;
 
-	FName Name;
-	FRemoteObjectId Id;
-	FObjectPtr ClassPtr;
-	FObjectPtr OuterPtr;
+	Overrides.Reserve(InObjectHeaders.Num());
+	for (const FRemoteObjectHeader& Header : InObjectHeaders)
+	{
+		FRemoteObjectConstructionParams& Params = Overrides.Emplace_GetRef();
+		Params.Name = Header.Name;
+#if UE_WITH_REMOTE_OBJECT_HANDLE
+		Params.OuterId = GetRemoteObjectId(Header.Outer.GetHandle());
+#endif
+		Params.RemoteId = Header.RemoteId;
+		Params.SerialNumber = Header.SerialNumber;
+	}
+}
 
-	Ar << Name;
-	Ar << Id;
-	Ar << ClassPtr;
-	Ar << OuterPtr;
+UObject* ConstructeRemoteObject(UObject* DestinationOuter, const FRemoteObjectHeader& Header, FRemoteObjectConstructionOverrides* ConstructionOverrides)
+{
+	using namespace UE::RemoteObject::Handle;
+	using namespace UE::RemoteObject::Private;
 
-	UClass* Class = CastChecked<UClass>(ClassPtr.Get());
-	UObject* Outer = OuterPtr.Get();
+	UClass* Class = CastChecked<UClass>(Header.Class.Get());
+	UObject* Outer = Header.Outer.Get();	
+	
+	// The object may already exist in memory (it could be a default subobject of an object we've just created)
+	UObject* Object = StaticFindObjectFast(Class, Outer, Header.Name);
 
-	UObject* Object = NewObject<UObject>(Outer, Class, Name, RF_NoFlags, nullptr, false, nullptr, nullptr, Id);
+	// If not or the object is marked as remote (which means we brought it back before it was GC'd) (re)construct it
+	if (!Object || IsRemote(Object))
+	{
+		FStaticConstructObjectParameters Params(Class);
+		Params.Outer = Outer;
+		Params.Name = Header.Name;
+		Params.SerialNumber = Header.SerialNumber;
+		Params.Template = Header.Archetype.Get();
+#if UE_WITH_REMOTE_OBJECT_HANDLE
+		Params.RemoteId = Header.RemoteId;
+		Params.RemoteSubObjectOverrides = ConstructionOverrides;
+#endif
+
+		// Using StaticConstructObject_Internal to pass the extra parameters (RemoteId and SubobjectOverrides) which are not exposed to normal APIs
+		Object = StaticConstructObject_Internal(Params);
+		checkf(FRemoteObjectId(Object) == Header.RemoteId, TEXT("Created an object with a different ID:%s than requested:%s"), *FRemoteObjectId(Object).ToString(), *Header.RemoteId.ToString());
+	}
 	MapObject(Object);
 	MarkAsLocal(Object);
 
-	Object->Serialize(Ar);
-
 	return Object;
 }
 
@@ -357,11 +428,12 @@ void DeserializeNames(FArchive& Ar, TArray<FName>& OutNames)
 
 FRemoteObjectData* SerializeObjectData(UObject* Object, TSet<UObject*>& OutObjects)
 {
-	using namespace UE::RemoteObject;
+	using namespace UE::RemoteObject::Handle;
 
 	FRemoteObjectData* ObjectData = new FRemoteObjectData();
 	TArray<UObject*> ObjectsToSerialize;
 	TArray<FName> Names;
+	TArray<FRemoteObjectHeader> ObjectHeaders;
 	TMap<FName, int32> NameMap;
 	int32 Version = 0;
 	int64 HeaderOffset = 0;
@@ -380,7 +452,9 @@ FRemoteObjectData* SerializeObjectData(UObject* Object, TSet<UObject*>& OutObjec
 			if (!OutObjects.Contains(ObjectToSerialize))
 			{
 				OutObjects.Add(ObjectToSerialize);
-				SerializeRemoteObject(ObjectToSerialize, Ar);
+				FRemoteObjectHeader& Header = ObjectHeaders.Emplace_GetRef();
+				SerializeRemoteObject(ObjectToSerialize, Ar, Header);
+				Ar.AddNameToNameMap(Header.Name);
 			}
 		}
 
@@ -389,9 +463,8 @@ FRemoteObjectData* SerializeObjectData(UObject* Object, TSet<UObject*>& OutObjec
 		Ar << HeaderOffset;
 		Ar.Seek(HeaderOffset);
 
-		int32 NumObjects = OutObjects.Num();
-		Ar << NumObjects;
 		SerializeNames(Ar, Names);
+		Ar << ObjectHeaders;
 	}
 
 	return ObjectData;
@@ -399,9 +472,12 @@ FRemoteObjectData* SerializeObjectData(UObject* Object, TSet<UObject*>& OutObjec
 
 void DeserializeObjectData(FRemoteObjectData& ObjectData, UObject* DestinationOuter, TArray<UObject*>& OutObjects)
 {
+	using namespace UE::RemoteObject::Handle;
+	using namespace UE::RemoteObject::Private;
+
 	int32 Version = 0;
 	int64 HeaderOffset = 0;
-	int32 NumObjects = 0;
+	TArray<FRemoteObjectHeader> ObjectHeaders;
 	TArray<FName> Names;
 
 	FArchiveRemoteObjectReader Ar(ObjectData, Names);
@@ -412,17 +488,44 @@ void DeserializeObjectData(FRemoteObjectData& ObjectData, UObject* DestinationOu
 	const int64 ObjectDataOffset = Ar.Tell();
 	Ar.Seek(HeaderOffset);
 
-	Ar << NumObjects;
 	DeserializeNames(Ar, Names);
+	Ar << ObjectHeaders;	
 
-	Ar.Seek(ObjectDataOffset);
+	Ar.Seek(ObjectDataOffset);	
 
-	for (int32 ObjectIndex = 0; ObjectIndex < NumObjects; ++ObjectIndex)
+	if (ObjectHeaders.Num())
 	{
-		OutObjects.Add(DeserializeRemoteObject(Ar, ObjectIndex == 0 ? DestinationOuter : nullptr));
-	}
+		// Try to find any existing WeakObject serial numbers for the objects that are about to be constructed
+		for (FRemoteObjectHeader& ObjectHeader : ObjectHeaders)
+		{
+			if (FRemoteObjectStub* Stub = FindRemoteObjectStub(ObjectHeader.RemoteId))
+			{
+				ObjectHeader.SerialNumber = Stub->SerialNumber;
+			}
+		}
+
+		// Construct all objects first
+		FRemoteObjectConstructionOverrides ConstructionOverrides(ObjectHeaders);
+		OutObjects.Reserve(ObjectHeaders.Num());
+		for (int32 ObjectIndex = 0; ObjectIndex < ObjectHeaders.Num(); ++ObjectIndex)
+		{
+			// Only pass the destination outer to the first object, which is always the root object of the entire hierarchy that was serialized
+			OutObjects.Add(ConstructeRemoteObject(ObjectIndex == 0 ? DestinationOuter : nullptr, ObjectHeaders[ObjectIndex], &ConstructionOverrides));
+		}
 
-	Ar.ResolveObjectPtrs(OutObjects);
+		// Deserialize all objects
+		for (UObject* Object : OutObjects)
+		{
+			Object->Serialize(Ar);
+		}
+
+		// Patch any unresolved remote references
+		FArchiveRemoteReferencePatcher PatchAr(OutObjects);
+		for (UObject* Object : OutObjects)
+		{
+			Object->Serialize(PatchAr);
+		}
+	}
 }
 
-} // namespace UE::RemoveObject::Serialization
\ No newline at end of file
+} // namespace UE::RemoteObject::Serialization
\ No newline at end of file
diff --git a/Engine/Source/Runtime/CoreUObject/Private/UObject/RemoteObjectSerialization.h b/Engine/Source/Runtime/CoreUObject/Private/UObject/RemoteObjectSerialization.h
--- a/Engine/Source/Runtime/CoreUObject/Private/UObject/RemoteObjectSerialization.h
+++ b/Engine/Source/Runtime/CoreUObject/Private/UObject/RemoteObjectSerialization.h
@@ -8,6 +8,50 @@ struct FRemoteObjectStub;
 
 namespace UE::RemoteObject::Serialization
 {
+	/**
+	* Basic information needed to construct deserialized remote object (this information does not get serialized in UObject::Serialize())
+	*/
+	struct FRemoteObjectHeader
+	{
+		FName Name;
+		FRemoteObjectId RemoteId;
+		TObjectPtr<UClass> Class;
+		TObjectPtr<UObject> Outer;
+		TObjectPtr<UObject> Archetype;
+
+		/** Transient (SerialNumber is local-only but it's stored here for convenience when deserializing object data) */
+		int32 SerialNumber = 0;
+	};
+
+	/**
+	* Basic information needed to construct remote (sub)object
+	*/
+	struct FRemoteObjectConstructionParams
+	{
+		FName Name;
+		FRemoteObjectId OuterId;
+		FRemoteObjectId RemoteId;
+		int32 SerialNumber = 0;
+	};
+
+	/**
+	* Stores basic information for constructing remote (sub)objects.
+	* Prevents unnecessary calls to FRemoteObjectId::Generate() when constructing default subobjects (or in general objects constructed in remote objects' constructors)
+	* Sets the SerialNumber during subobject construction so that any weak pointers also constructed in constructors that may point to a subobject, has the correct SerialNumber
+	*/
+	class FRemoteObjectConstructionOverrides
+	{
+		TArray<FRemoteObjectConstructionParams> Overrides;
+
+	public:
+
+		explicit FRemoteObjectConstructionOverrides(const TArray<FRemoteObjectHeader>& InObjectHeaders);
+
+		/**
+		* Finds object construction overrides for an object that will be constricted with the specified Name and Outer
+		*/
+		const FRemoteObjectConstructionParams* Find(FName InName, UObject* InOuter) const;
+	};
 
 	FRemoteObjectData* SerializeObjectData(UObject* Object, TSet<UObject*>& OutObjects);
 	void DeserializeObjectData(FRemoteObjectData& ObjectData, UObject* DestinationOuter, TArray<UObject*>& OutReceivedObjects);
diff --git a/Engine/Source/Runtime/CoreUObject/Private/UObject/UObjectArray.cpp b/Engine/Source/Runtime/CoreUObject/Private/UObject/UObjectArray.cpp
--- a/Engine/Source/Runtime/CoreUObject/Private/UObject/UObjectArray.cpp
+++ b/Engine/Source/Runtime/CoreUObject/Private/UObject/UObjectArray.cpp
@@ -226,7 +226,7 @@ void FUObjectArray::DisableDisregardForGC()
 	}
 }
 
-void FUObjectArray::AllocateUObjectIndex(UObjectBase* Object, EInternalObjectFlags InitialFlags, int32 AlreadyAllocatedIndex, int32 SerialNumber)
+void FUObjectArray::AllocateUObjectIndex(UObjectBase* Object, EInternalObjectFlags InitialFlags, int32 AlreadyAllocatedIndex, int32 SerialNumber, FRemoteObjectId RemoteId)
 {
 	LLM_SCOPE(ELLMTag::UObject);
 	// Clear asset scopes
@@ -237,6 +237,13 @@ void FUObjectArray::AllocateUObjectIndex(UObjectBase* Object, EInternalObjectFla
 	int32 Index = INDEX_NONE;
 	check(Object->InternalIndex == INDEX_NONE);
 
+#if UE_WITH_REMOTE_OBJECT_HANDLE
+	if (!RemoteId.IsValid())
+	{
+		RemoteId = FRemoteObjectId::Generate();
+	}
+#endif
+
 	LockInternalArray();
 
 	if (AlreadyAllocatedIndex >= 0)
@@ -289,6 +296,9 @@ void FUObjectArray::AllocateUObjectIndex(UObjectBase* Object, EInternalObjectFla
 	ObjectItem->RefCount = 0;
 	ObjectItem->ClusterRootIndex = 0;
 	ObjectItem->SerialNumber = SerialNumber;
+#if UE_WITH_REMOTE_OBJECT_HANDLE
+	ObjectItem->RemoteId = RemoteId;
+#endif // UE_WITH_REMOTE_OBJECT_HANDLE
 	Object->InternalIndex = Index;
 
 	// This needs to happen after the InternalIndex is set because setting root flags may result in the object being added to UE::GC::Priate::GRoots array 
@@ -317,6 +327,9 @@ void FUObjectArray::ResetSerialNumber(UObjectBase* Object)
 	FUObjectItem* ObjectItem = IndexToObject(Index);
 	checkSlow(ObjectItem);
 	ObjectItem->SerialNumber = 0;
+#if UE_WITH_REMOTE_OBJECT_HANDLE
+	ObjectItem->RemoteId = FRemoteObjectId();
+#endif
 }
 
 /**
@@ -354,8 +367,6 @@ void FUObjectArray::FreeUObjectIndex(UObjectBase* Object)
 
 #if UE_WITH_OBJECT_HANDLE_LATE_RESOLVE
 	UE::CoreUObject::Private::FreeObjectHandle(Object);
-#elif UE_WITH_REMOTE_OBJECT_HANDLE
-	UE::RemoteObject::Private::UnmapObject((UObject*)Object);
 #endif
 
 
@@ -377,6 +388,9 @@ void FUObjectArray::FreeUObjectIndex(UObjectBase* Object)
 	ObjectItem->RefCount = 0;
 	ObjectItem->ClusterRootIndex = 0;
 	ObjectItem->SerialNumber = 0;
+#if UE_WITH_REMOTE_OBJECT_HANDLE
+	ObjectItem->RemoteId = FRemoteObjectId();
+#endif
 	Object->InternalIndex = INDEX_NONE;
 
 	// You cannot safely recycle indicies in the non-GC range
diff --git a/Engine/Source/Runtime/CoreUObject/Private/UObject/UObjectBase.cpp b/Engine/Source/Runtime/CoreUObject/Private/UObject/UObjectBase.cpp
--- a/Engine/Source/Runtime/CoreUObject/Private/UObject/UObjectBase.cpp
+++ b/Engine/Source/Runtime/CoreUObject/Private/UObject/UObjectBase.cpp
@@ -61,11 +61,21 @@ struct FPendingRegistrantInfo
 	const TCHAR*	Name;
 	const TCHAR*	PackageName;
 	class UClass*  (*StaticClassFn)();
+#if UE_WITH_REMOTE_OBJECT_HANDLE
+	const FRemoteObjectId RemoteId;
+#endif // UE_WITH_REMOTE_OBJECT_HANDLE
 
-	FPendingRegistrantInfo(class UClass* (*InStaticClassFn)(), const TCHAR* InName, const TCHAR* InPackageName)
+	FPendingRegistrantInfo(class UClass* (*InStaticClassFn)(), const TCHAR* InName, const TCHAR* InPackageName
+#if UE_WITH_REMOTE_OBJECT_HANDLE
+		, const FRemoteObjectId InRemoteId
+#endif
+		)
 		: Name(InName)
 		, PackageName(InPackageName)
 		, StaticClassFn(InStaticClassFn)
+#if UE_WITH_REMOTE_OBJECT_HANDLE
+		, RemoteId(InRemoteId)
+#endif
 	{}
 	static TMap<UObjectBase*, FPendingRegistrantInfo>& GetMap()
 	{
@@ -106,9 +116,6 @@ static TMap<FName, TArray<FPendingRegistrant*>>& GetPerModuleBootstrapMap()
 UObjectBase::UObjectBase( EObjectFlags InFlags )
 :	ObjectFlags			(InFlags)
 ,	InternalIndex		(INDEX_NONE)
-#if UE_WITH_REMOTE_OBJECT_HANDLE
-,	RemoteId(FRemoteObjectId::Generate())
-#endif // UE_WITH_REMOTE_OBJECT_HANDLE
 ,	ClassPrivate		(nullptr)
 ,	OuterPrivate		(nullptr)
 {
@@ -135,15 +142,12 @@ UObjectBase::UObjectBase(UClass* InClass,
 	FRemoteObjectId InRemoteId)
 :	ObjectFlags			(InFlags)
 ,	InternalIndex		(INDEX_NONE)
-#if UE_WITH_REMOTE_OBJECT_HANDLE
-,	RemoteId			(InRemoteId.IsValid() ? InRemoteId : FRemoteObjectId::Generate())
-#endif // UE_WITH_REMOTE_OBJECT_HANDLE
 ,	ClassPrivate		(InClass)
 ,	OuterPrivate		(InOuter)
 {
 	check(ClassPrivate);
 	// Add to global table.
-	AddObject(InName, InInternalFlags, InInternalIndex, InSerialNumber);
+	AddObject(InName, InInternalFlags, InInternalIndex, InSerialNumber, InRemoteId);
 	
 #if CSV_PROFILER_STATS && CSV_TRACK_UOBJECT_COUNT
 	UObjectStats::IncrementUObjectCount();
@@ -178,7 +182,11 @@ UObjectBase::~UObjectBase()
  *
  * @param UClassStaticClass Now that it is known, fill in ClassPrivate.  This is usually UClass::StaticClass.
  */
-void UObjectBase::DeferredRegister(UClass *UClassStaticClass,const TCHAR* PackageName,const TCHAR* InName)
+void UObjectBase::DeferredRegister(UClass *UClassStaticClass, const TCHAR* PackageName, const TCHAR* InName
+#if UE_WITH_REMOTE_OBJECT_HANDLE
+	, FRemoteObjectId RemoteId
+#endif
+)
 {
 	check(UObjectInitialized());
 	// Set object properties.
@@ -192,10 +200,21 @@ void UObjectBase::DeferredRegister(UClass *UClassStaticClass,const TCHAR* Packag
 	ClassPrivate = UClassStaticClass;
 
 	// Add to the global object table.
-	AddObject(FName(InName), EInternalObjectFlags::None);
+	AddObject(FName(InName), EInternalObjectFlags::None
+#if UE_WITH_REMOTE_OBJECT_HANDLE
+		, /*IndernalIndex = */ -1, /*SerialNumber = */ 0, RemoteId
+#endif
+	);
+
 	// At this point all compiled-in objects should have already been fully constructed so it's safe to remove the NotFullyConstructed flag
 	// which was set in FUObjectArray::AllocateUObjectIndex (called from AddObject)
-	GUObjectArray.IndexToObject(InternalIndex)->ClearFlags(EInternalObjectFlags::PendingConstruction);
+	FUObjectItem* ObjectItem = GUObjectArray.IndexToObject(InternalIndex);
+	ObjectItem->ClearFlags(EInternalObjectFlags::PendingConstruction);
+
+#if UE_WITH_REMOTE_OBJECT_HANDLE
+	checkf(ObjectItem->GetRemoteId().IsValid() && (!RemoteId.IsValid() || RemoteId == ObjectItem->GetRemoteId()), TEXT("Native object %s %s has an unexpected RemoteId:%s (expected:%s)"),
+		*GetClass()->GetName(), *GetFName().ToString(), *ObjectItem->GetRemoteId().ToString(), *RemoteId.ToString());
+#endif
 
 	// Make sure that objects disregarded for GC are part of root set.
 	check(!GUObjectArray.IsDisregardForGC(this) || GUObjectArray.IndexToObject(InternalIndex)->IsRootSet());
@@ -208,7 +227,7 @@ void UObjectBase::DeferredRegister(UClass *UClassStaticClass,const TCHAR* Packag
  *
  * @param Name name to assign to this uobject
  */
-void UObjectBase::AddObject(FName InName, EInternalObjectFlags InSetInternalFlags, int32 InInternalIndex, int32 InSerialNumber)
+void UObjectBase::AddObject(FName InName, EInternalObjectFlags InSetInternalFlags, int32 InInternalIndex, int32 InSerialNumber, FRemoteObjectId InRemoteId)
 {
 	NamePrivate = InName;
 	EInternalObjectFlags InternalFlagsToSet = InSetInternalFlags;
@@ -226,7 +245,7 @@ void UObjectBase::AddObject(FName InName, EInternalObjectFlags InSetInternalFlag
 		InternalFlagsToSet |= EInternalObjectFlags::Native;
 		ObjectFlags &= ~RF_MarkAsNative;
 	}
-	GUObjectArray.AllocateUObjectIndex(this, InternalFlagsToSet, InInternalIndex, InSerialNumber);
+	GUObjectArray.AllocateUObjectIndex(this, InternalFlagsToSet, InInternalIndex, InSerialNumber, InRemoteId);
 	check(InName != NAME_None && InternalIndex >= 0);
 	HashObject(this);
 	check(IsValidLowLevel());
@@ -486,7 +505,11 @@ void UObjectBase::Register(class UClass* (*StaticClassFn)(), const TCHAR* Packag
 	TMap<UObjectBase*, FPendingRegistrantInfo>& PendingRegistrants = FPendingRegistrantInfo::GetMap();
 
 	FPendingRegistrant* PendingRegistration = new FPendingRegistrant(this);
-	PendingRegistrants.Add(this, FPendingRegistrantInfo(StaticClassFn, InName, PackageName));
+	PendingRegistrants.Add(this, FPendingRegistrantInfo(StaticClassFn, InName, PackageName
+#if UE_WITH_REMOTE_OBJECT_HANDLE
+		, FRemoteObjectId::Generate()
+#endif
+	));
 
 #if USE_PER_MODULE_UOBJECT_BOOTSTRAP
 	if (FName(PackageName) != FName("/Script/CoreUObject"))
@@ -511,6 +534,14 @@ void UObjectBase::Register(class UClass* (*StaticClassFn)(), const TCHAR* Packag
 	}
 }
 
+#if UE_WITH_REMOTE_OBJECT_HANDLE
+FRemoteObjectId UObjectBase::GetPendingRegistrantRemoteId() const
+{
+	TMap<UObjectBase*, FPendingRegistrantInfo>& PendingRegistrants = FPendingRegistrantInfo::GetMap();
+	FPendingRegistrantInfo& Info = PendingRegistrants.FindChecked(this);
+	return Info.RemoteId;
+}
+#endif
 
 /**
  * Dequeues registrants from the list of pending registrations into an array.
@@ -574,8 +605,15 @@ void UObjectForceRegistration(UObjectBase* Object, bool bCheckForModuleRelease)
 #endif
 		const TCHAR* Name = Info->Name;
 		UClass* StaticClass = Info->StaticClassFn();
+#if UE_WITH_REMOTE_OBJECT_HANDLE
+		FRemoteObjectId RemoteId = Info->RemoteId;
+#endif
 		PendingRegistrants.Remove(Object);  // delete this first so that it doesn't try to do it twice
-		Object->DeferredRegister(StaticClass, PackageName, Name);
+		Object->DeferredRegister(StaticClass, PackageName, Name
+#if UE_WITH_REMOTE_OBJECT_HANDLE
+			, RemoteId
+#endif
+		);
 	}
 }
 
diff --git a/Engine/Source/Runtime/CoreUObject/Private/UObject/UObjectGlobals.cpp b/Engine/Source/Runtime/CoreUObject/Private/UObject/UObjectGlobals.cpp
--- a/Engine/Source/Runtime/CoreUObject/Private/UObject/UObjectGlobals.cpp
+++ b/Engine/Source/Runtime/CoreUObject/Private/UObject/UObjectGlobals.cpp
@@ -86,6 +86,7 @@
 #include "UObject/PropertyOptional.h"
 #include "UObject/VerseValueProperty.h"
 #include "UObject/ObjectHandlePrivate.h"
+#include "UObject/RemoteObjectSerialization.h"
 
 #include "Interfaces/IPluginManager.h"
 
@@ -3392,6 +3393,7 @@ UObject* StaticAllocateObject
 	bool bCanRecycleSubobjects,
 	bool* bOutRecycledSubobject,
 	UPackage* ExternalPackage,
+	int32 SerialNumber,
 	FRemoteObjectId RemoteId)
 {
 	LLM_SCOPE(ELLMTag::UObject);
@@ -3604,6 +3606,12 @@ UObject* StaticAllocateObject
 			{
 				UE::CoreUObject::Private::FObjectHandleUtils::ChangeRemoteId(Obj, RemoteId);
 			}
+			if (SerialNumber != 0)
+			{
+				FUObjectItem* SubobjectItem = GUObjectArray.ObjectToObjectItem(Obj);
+				checkf(SubobjectItem->SerialNumber == 0 || SerialNumber == 0 || SerialNumber == SubobjectItem->SerialNumber, TEXT("Attempting to change an existing subobject's %s serial number (%d -> %d)"), *Obj->GetPathName(), SubobjectItem->SerialNumber, SerialNumber);
+				SubobjectItem->SerialNumber = SubobjectItem->SerialNumber == 0 ? SerialNumber : SubobjectItem->SerialNumber;
+			}
 #endif // UE_WITH_REMOTE_OBJECT_HANDLE
 			bSubObject = true;
 		}
@@ -3617,12 +3625,18 @@ UObject* StaticAllocateObject
 
 	if (!bSubObject)
 	{
+#if UE_WITH_REMOTE_OBJECT_HANDLE
+		// We can't have two different, non-zero serial numbers
+		checkf(OldSerialNumber == 0 || SerialNumber == 0 || SerialNumber == OldSerialNumber, TEXT("Attempting to change a new object's (%s %s.%s) serial number (%d -> %d)"),
+			*InClass->GetPathName(), *GetPathNameSafe(InOuter), *InName.ToString(), OldSerialNumber, SerialNumber);
+#endif // UE_WITH_REMOTE_OBJECT_HANDLE
+
 		// perform the UObjectBase construction in the open - we expect the GC to invoke
 		// the destructor in the case of transaction abort
 		UE_AUTORTFM_OPEN
 		{
 			FMemory::Memzero((void *)Obj, TotalSize);
-			new ((void *)Obj) UObjectBase(const_cast<UClass*>(InClass), InFlags|RF_NeedInitialization, InternalSetFlags, InOuter, InName, OldIndex, OldSerialNumber, RemoteId);
+			new ((void *)Obj) UObjectBase(const_cast<UClass*>(InClass), InFlags|RF_NeedInitialization, InternalSetFlags, InOuter, InName, OldIndex, OldSerialNumber == 0 ? SerialNumber : OldSerialNumber, RemoteId);
 		};
 	}
 	else
@@ -3747,13 +3761,16 @@ FObjectInitializer::FObjectInitializer()
 	Construct_Internal();
 }	
 
-FObjectInitializer::FObjectInitializer(UObject* InObj, const FStaticConstructObjectParameters& StaticConstructParams)
+FObjectInitializer::FObjectInitializer(UObject* InObj, const FStaticConstructObjectParameters& StaticConstructParams, UE::RemoteObject::Serialization::FRemoteObjectConstructionOverrides* InRemoteSubObjectOverides)
 	: Obj(InObj)
 	, ObjectArchetype(StaticConstructParams.Template)
 	, bCopyTransientsFromClassDefaults(StaticConstructParams.bCopyTransientsFromClassDefaults)
 	, bShouldInitializePropsFromArchetype(true)
 	, InstanceGraph(StaticConstructParams.InstanceGraph)
 	, PropertyInitCallback(StaticConstructParams.PropertyInitCallback)
+#if UE_WITH_REMOTE_OBJECT_HANDLE
+	, RemoteSubObjectOverrides(InRemoteSubObjectOverides)
+#endif
 {
 	if (StaticConstructParams.SubobjectOverrides)
 	{
@@ -4510,6 +4527,45 @@ FStaticConstructObjectParameters::FStaticConstructObjectParameters(const UClass*
 {
 }
 
+#if UE_WITH_REMOTE_OBJECT_HANDLE
+namespace
+{
+	UE::RemoteObject::Serialization::FRemoteObjectConstructionOverrides* GetRemoteSubObjectOverridesFromObjectInitializer(FName InName, UObject* InOuter, int32& InOutSerialNumber, FRemoteObjectId& InOutRemoteId)
+	{
+		using namespace UE::RemoteObject::Serialization;
+
+		FUObjectThreadContext& ThreadContext = FUObjectThreadContext::Get();
+		if (ThreadContext.IsInConstructor)
+		{
+			// Try to get remote subobject overrides from FObjectInitializer so that if the object that's about to be constructed is a subobject
+			// it gets the right RemoteId and SerialNumber at the time of its construction because its owner constructor may be reading the id and/or 
+			// creating a weak pointer to the subobject (if not the owner directly maybe some code that gets called from the owner's constructor)
+
+			FObjectInitializer& ObjectInitializer = ThreadContext.TopInitializerChecked();
+			if (FRemoteObjectConstructionOverrides* RemoteSubObjectOverrides = ObjectInitializer.GetRemoteSubObjectOverrides())
+			{
+				if (const FRemoteObjectConstructionParams* ConstructionOverrides = RemoteSubObjectOverrides->Find(InName, InOuter))
+				{
+					if (ConstructionOverrides->SerialNumber != 0)
+					{
+						checkf(InOutSerialNumber == 0 || InOutSerialNumber == ConstructionOverrides->SerialNumber, TEXT("Attempting to override valid SerialNumber of %s.%s (%d -> %d)"), *GetPathNameSafe(InOuter), *InName.ToString(), InOutSerialNumber, ConstructionOverrides->SerialNumber);
+						InOutSerialNumber = ConstructionOverrides->SerialNumber;
+					}
+					if (ConstructionOverrides->RemoteId.IsValid())
+					{
+						checkf(!InOutRemoteId.IsValid() || InOutRemoteId == ConstructionOverrides->RemoteId, TEXT("Attempting to override valid RemoteId of %s.%s (%s -> %s)"), *GetPathNameSafe(InOuter), *InName.ToString(), *InOutRemoteId.ToString(), *ConstructionOverrides->RemoteId.ToString());
+						InOutRemoteId = ConstructionOverrides->RemoteId;
+					}
+				}
+				return RemoteSubObjectOverrides;
+			}
+
+		}
+		return nullptr;
+	}
+}
+#endif // UE_WITH_REMOTE_OBJECT_HANDLE
+
 UObject* StaticConstructObject_Internal(const FStaticConstructObjectParameters& Params)
 {
 	const UClass* InClass = Params.Class;
@@ -4517,6 +4573,9 @@ UObject* StaticConstructObject_Internal(const FStaticConstructObjectParameters&
 	const FName& InName = Params.Name;
 	EObjectFlags InFlags = Params.SetFlags;
 	UObject* InTemplate = Params.Template;
+	int32 SerialNumber = Params.SerialNumber;
+	FRemoteObjectId RemoteId;
+	UE::RemoteObject::Serialization::FRemoteObjectConstructionOverrides* RemoteSubObjectOverrides = nullptr;
 
 	LLM_SCOPE(ELLMTag::UObject);
 
@@ -4542,14 +4601,24 @@ UObject* StaticConstructObject_Internal(const FStaticConstructObjectParameters&
 
 	const bool bCanRecycleSubobjects = bIsNativeFromCDO && (!(InFlags & RF_DefaultSubObject) || !FUObjectThreadContext::Get().IsInConstructor);
 
+	
+#if UE_WITH_REMOTE_OBJECT_HANDLE
+	RemoteId = Params.RemoteId;
+	RemoteSubObjectOverrides = Params.RemoteSubObjectOverrides;
+	if (!RemoteSubObjectOverrides)
+	{
+		RemoteSubObjectOverrides = GetRemoteSubObjectOverridesFromObjectInitializer(InName, InOuter, SerialNumber, RemoteId);
+	}
+#endif
+
 	bool bRecycledSubobject = false;	
-	Result = StaticAllocateObject(InClass, InOuter, InName, InFlags, Params.InternalSetFlags, bCanRecycleSubobjects, &bRecycledSubobject, Params.ExternalPackage, Params.RemoteId);
+	Result = StaticAllocateObject(InClass, InOuter, InName, InFlags, Params.InternalSetFlags, bCanRecycleSubobjects, &bRecycledSubobject, Params.ExternalPackage, SerialNumber, RemoteId);
 	check(Result != nullptr);
 	// Don't call the constructor on recycled subobjects, they haven't been destroyed.
 	if (!bRecycledSubobject)
 	{		
 		STAT(FScopeCycleCounterUObject ConstructorScope(InClass->GetFName().IsNone() ? nullptr : InClass, GET_STATID(STAT_ConstructObject)));
-		(*InClass->ClassConstructor)(FObjectInitializer(Result, Params));
+		(*InClass->ClassConstructor)(FObjectInitializer(Result, Params, RemoteSubObjectOverrides));
 	}
 	
 	if (GIsEditor && 
diff --git a/Engine/Source/Runtime/CoreUObject/Private/UObject/WeakObjectPtr.cpp b/Engine/Source/Runtime/CoreUObject/Private/UObject/WeakObjectPtr.cpp
--- a/Engine/Source/Runtime/CoreUObject/Private/UObject/WeakObjectPtr.cpp
+++ b/Engine/Source/Runtime/CoreUObject/Private/UObject/WeakObjectPtr.cpp
@@ -8,6 +8,8 @@
 #include "UObject/GarbageCollection.h"
 #include "UObject/Object.h"
 #include "UObject/ObjectPtr.h"
+#include "UObject/RemoteObject.h"
+#include "UObject/RemoteObjectPrivate.h"
 
 DEFINE_LOG_CATEGORY_STATIC(LogWeakObjectPtr, Log, All);
 
@@ -30,6 +32,10 @@ void FWeakObjectPtr::operator=(const class UObject *Object)
 	{
 		ObjectIndex = GUObjectArray.ObjectToIndex((UObjectBase*)Object);
 		ObjectSerialNumber = GUObjectArray.AllocateSerialNumber(ObjectIndex);
+#if UE_WITH_REMOTE_OBJECT_HANDLE
+		ObjectRemoteId = GUObjectArray.GetRemoteId(ObjectIndex);
+#endif
+
 		checkSlow(SerialNumbersMatch());
 	}
 	else
@@ -58,6 +64,12 @@ bool FWeakObjectPtr::IsStale(bool bIncludingGarbage, bool bThreadsafeTest) const
 {
 	using namespace UE::Core::Private;
 
+#if UE_WITH_REMOTE_OBJECT_HANDLE
+	if (IsExplicitlyNull())
+	{
+		return false;
+	}
+#else
 	if (ObjectSerialNumber == 0)
 	{
 #if UE_WEAKOBJECTPTR_ZEROINIT_FIX
@@ -67,6 +79,7 @@ bool FWeakObjectPtr::IsStale(bool bIncludingGarbage, bool bThreadsafeTest) const
 #endif
 		return false;
 	}
+#endif // UE_WITH_REMOTE_OBJECT_HANDLE
 
 	if (ObjectIndex < 0)
 	{
@@ -79,6 +92,16 @@ bool FWeakObjectPtr::IsStale(bool bIncludingGarbage, bool bThreadsafeTest) const
 	}
 	if (!SerialNumbersMatch(ObjectItem))
 	{
+#if UE_WITH_REMOTE_OBJECT_HANDLE
+		if (bThreadsafeTest)
+		{
+			return true;
+		}
+		if (UE::RemoteObject::Private::FindRemoteObjectStub(ObjectRemoteId))
+		{
+			return false;
+		}
+#endif // UE_WITH_REMOTE_OBJECT_HANDLE
 		return true;
 	}
 	if (bThreadsafeTest)
@@ -144,3 +167,35 @@ void FWeakObjectPtr::Serialize(FArchive& Ar)
 {
 	FArchiveUObject::SerializeWeakObjectPtr(Ar, *this);
 }
+
+#if UE_WITH_REMOTE_OBJECT_HANDLE
+FUObjectItem* FWeakObjectPtr::TryResolveRemoteObject() const
+{
+	using namespace UE::RemoteObject::Handle;
+	using namespace UE::RemoteObject::Private;
+
+	if (FRemoteObjectStub* Stub = FindRemoteObjectStub(ObjectRemoteId))
+	{
+		if (UObject* ResolvedObject = ResolveObject(Stub))
+		{
+			int32 NewIndex = GUObjectArray.ObjectToIndex(ResolvedObject);
+			FUObjectItem* ObjectItem = GUObjectArray.IndexToObject(NewIndex);
+			checkf(ObjectItem->GetRemoteId() == ObjectRemoteId, TEXT("Remote object %s has a different id (%s) than the weak object pointer that resolved it (%s)"),
+				*ResolvedObject->GetPathName(), *ObjectItem->GetRemoteId().ToString(), *ObjectRemoteId.ToString());
+			checkf(ObjectItem->GetSerialNumber() == ObjectSerialNumber, TEXT("Resolved remote object %s has a different serial number (%d) than the weak object pointer that resolved it (%d)"),
+				*ResolvedObject->GetPathName(), ObjectItem->GetSerialNumber(), ObjectSerialNumber);
+			if (ObjectIndex != NewIndex)
+			{
+				// Currently we can't just rely on object id to resolve weak object pointers and we don't want to heep hitting the resolve path
+				// if a remote object is resolved and exists on this server so (sadly) we need to update the index, yuck.
+				// In general ObjectIndex would not be required at all if all objects were hashed with RemoteId in UObjectHash.cpp but currently we don't.
+				// ObjectIndex also serves as an optimization because it's faster than hash table lookup.
+				// With RemoteId neither ObjectIndex nor SerialNumber is used for FWeakObjectPtr comparisons/hash value calculation
+				FPlatformAtomics::AtomicStore(&const_cast<FWeakObjectPtr*>(this)->ObjectIndex, NewIndex);
+			}
+			return ObjectItem;
+		}
+	}
+	return nullptr;
+}
+#endif // UE_WITH_REMOTE_OBJECT_HANDLE
\ No newline at end of file
diff --git a/Engine/Source/Runtime/CoreUObject/Public/UObject/Class.h b/Engine/Source/Runtime/CoreUObject/Public/UObject/Class.h
--- a/Engine/Source/Runtime/CoreUObject/Public/UObject/Class.h
+++ b/Engine/Source/Runtime/CoreUObject/Public/UObject/Class.h
@@ -3208,7 +3208,11 @@ public:
 	COREUOBJECT_API virtual void Serialize(FArchive& Ar) override;
 	COREUOBJECT_API virtual void PostLoad() override;
 	COREUOBJECT_API virtual void FinishDestroy() override;
-	COREUOBJECT_API virtual void DeferredRegister(UClass *UClassStaticClass,const TCHAR* PackageName,const TCHAR* InName) override;
+	COREUOBJECT_API virtual void DeferredRegister(UClass* UClassStaticClass, const TCHAR* PackageName, const TCHAR* InName
+#if UE_WITH_REMOTE_OBJECT_HANDLE
+		, FRemoteObjectId RemoteId
+#endif // UE_WITH_REMOTE_OBJECT_HANDLE
+	) override;
 	COREUOBJECT_API virtual bool Rename(const TCHAR* NewName = nullptr, UObject* NewOuter = nullptr, ERenameFlags Flags = REN_None) override;
 	COREUOBJECT_API virtual void TagSubobjects(EObjectFlags NewFlags) override;
 	COREUOBJECT_API virtual void PostInitProperties() override;
diff --git a/Engine/Source/Runtime/CoreUObject/Public/UObject/ObjectHandle.h b/Engine/Source/Runtime/CoreUObject/Public/UObject/ObjectHandle.h
--- a/Engine/Source/Runtime/CoreUObject/Public/UObject/ObjectHandle.h
+++ b/Engine/Source/Runtime/CoreUObject/Public/UObject/ObjectHandle.h
@@ -319,26 +319,31 @@ inline bool operator!=(UE::CoreUObject::Private::FRemoteObjectHandlePrivate LHS,
 	return !(LHS == RHS);
 }
 
-/**
-* GetTypeHash for FRemoteObjectHandlePrivate is guaranteed to return the same value whether a non-null handle is resolved or not (GetTypeHash of FRemoteObjectId)
-*/
-inline uint32 GetTypeHash(UE::CoreUObject::Private::FRemoteObjectHandlePrivate Handle)
+inline FRemoteObjectId GetRemoteObjectId(UE::CoreUObject::Private::FRemoteObjectHandlePrivate Handle)
 {
 	using namespace UE::CoreUObject::Private;
 	using namespace UE::RemoteObject::Handle;
 
 	if (Handle.PointerOrHandle == 0)
 	{
-		return 0;
+		return FRemoteObjectId();
 	}
 
 	if (IsObjectHandleResolved(Handle))
 	{
 		const UObject* Obj = ReadObjectHandlePointerNoCheck(Handle);
-		return GetTypeHash(FRemoteObjectId((const UObjectBase*)Obj));
+		return FRemoteObjectId((const UObjectBase*)Obj);
 	}
 	const FRemoteObjectStub* Stub = Handle.ToStub();
-	return GetTypeHash(Stub->Id);
+	return Stub->Id;
+}
+
+/**
+* GetTypeHash for FRemoteObjectHandlePrivate is guaranteed to return the same value whether a non-null handle is resolved or not (GetTypeHash of FRemoteObjectId)
+*/
+inline uint32 GetTypeHash(UE::CoreUObject::Private::FRemoteObjectHandlePrivate Handle)
+{
+	return GetTypeHash(GetRemoteObjectId(Handle));
 }
 #endif
 
diff --git a/Engine/Source/Runtime/CoreUObject/Public/UObject/ObjectKey.h b/Engine/Source/Runtime/CoreUObject/Public/UObject/ObjectKey.h
--- a/Engine/Source/Runtime/CoreUObject/Public/UObject/ObjectKey.h
+++ b/Engine/Source/Runtime/CoreUObject/Public/UObject/ObjectKey.h
@@ -35,6 +35,9 @@ public:
 			FWeakObjectPtr Weak(Object);
 			ObjectIndex = Weak.ObjectIndex;
 			ObjectSerialNumber = Weak.ObjectSerialNumber;
+#if UE_WITH_REMOTE_OBJECT_HANDLE
+			RemoteId = Weak.ObjectRemoteId;
+#endif
 		}
 	}
 	template <
@@ -49,43 +52,71 @@ public:
 	/** Compare this key with another */
 	FORCEINLINE bool operator==(const FObjectKey& Other) const
 	{
+#if UE_WITH_REMOTE_OBJECT_HANDLE
+		return RemoteId == Other.RemoteId;
+#else
 		return ObjectIndex == Other.ObjectIndex && ObjectSerialNumber == Other.ObjectSerialNumber;
+#endif
 	}
 
 	/** Compare this key with another */
 	FORCEINLINE bool operator!=(const FObjectKey& Other) const
 	{
+#if UE_WITH_REMOTE_OBJECT_HANDLE
+		return RemoteId != Other.RemoteId;
+#else
 		return ObjectIndex != Other.ObjectIndex || ObjectSerialNumber != Other.ObjectSerialNumber;
+#endif
 	}
 
 	/** Compare this key with another */
 	FORCEINLINE bool operator<(const FObjectKey& Other) const
 	{
+#if UE_WITH_REMOTE_OBJECT_HANDLE
+		return RemoteId < Other.RemoteId;
+#else
 		return ObjectIndex < Other.ObjectIndex || (ObjectIndex == Other.ObjectIndex && ObjectSerialNumber < Other.ObjectSerialNumber);
+#endif
 	}
 
 	/** Compare this key with another */
 	FORCEINLINE bool operator<=(const FObjectKey& Other) const
 	{
+#if UE_WITH_REMOTE_OBJECT_HANDLE
+		return RemoteId <= Other.RemoteId;
+#else
 		return ObjectIndex <= Other.ObjectIndex || (ObjectIndex == Other.ObjectIndex && ObjectSerialNumber <= Other.ObjectSerialNumber);
+#endif
 	}
 
 	/** Compare this key with another */
 	FORCEINLINE bool operator>(const FObjectKey& Other) const
 	{
+#if UE_WITH_REMOTE_OBJECT_HANDLE
+		return RemoteId > Other.RemoteId;
+#else
 		return ObjectIndex > Other.ObjectIndex || (ObjectIndex == Other.ObjectIndex && ObjectSerialNumber > Other.ObjectSerialNumber);
+#endif
 	}
 
 	/** Compare this key with another */
 	FORCEINLINE bool operator>=(const FObjectKey& Other) const
 	{
+#if UE_WITH_REMOTE_OBJECT_HANDLE
+		return RemoteId >= Other.RemoteId;
+#else
 		return ObjectIndex > Other.ObjectIndex || (ObjectIndex == Other.ObjectIndex && ObjectSerialNumber >= Other.ObjectSerialNumber);
+#endif
 	}
 
 	FORCEINLINE friend FArchive& operator<<(FArchive& Ar, FObjectKey& Key)
 	{
 		check(!Ar.IsPersistent());
-		return Ar << Key.ObjectIndex << Key.ObjectSerialNumber;
+		return Ar << Key.ObjectIndex << Key.ObjectSerialNumber
+#if UE_WITH_REMOTE_OBJECT_HANDLE
+			<< Key.RemoteId
+#endif
+			;
 	}
 
 	/**
@@ -97,7 +128,9 @@ public:
 		FWeakObjectPtr WeakPtr;
 		WeakPtr.ObjectIndex = ObjectIndex;
 		WeakPtr.ObjectSerialNumber = ObjectSerialNumber;
-
+#if UE_WITH_REMOTE_OBJECT_HANDLE
+		WeakPtr.ObjectRemoteId = RemoteId;
+#endif
 		return WeakPtr.Get();
 	}
 
@@ -110,7 +143,9 @@ public:
 		FWeakObjectPtr WeakPtr;
 		WeakPtr.ObjectIndex = ObjectIndex;
 		WeakPtr.ObjectSerialNumber = ObjectSerialNumber;
-
+#if UE_WITH_REMOTE_OBJECT_HANDLE
+		WeakPtr.ObjectRemoteId = RemoteId;
+#endif
 		constexpr bool bEvenIfGarbage = true;
 		return WeakPtr.Get(bEvenIfGarbage);
 	}
@@ -130,26 +165,46 @@ public:
 		FWeakObjectPtr WeakPtr;
 		WeakPtr.ObjectIndex = ObjectIndex;
 		WeakPtr.ObjectSerialNumber = ObjectSerialNumber;
-
+#if UE_WITH_REMOTE_OBJECT_HANDLE
+		WeakPtr.ObjectRemoteId = RemoteId;
+#endif
 		return WeakPtr.GetEvenIfUnreachable();
 	}
 
 	/** Hash function */
 	[[nodiscard]] friend uint32 GetTypeHash(const FObjectKey& Key)
 	{
+#if UE_WITH_REMOTE_OBJECT_HANDLE
+		return GetTypeHash(Key.RemoteId);
+#else
 		return HashCombine(Key.ObjectIndex, Key.ObjectSerialNumber);
+#endif
 	}
 
+#if UE_WITH_REMOTE_OBJECT_HANDLE
+	/** Get globally unique id of the object */
+	FRemoteObjectId GetRemoteId() const
+	{
+		return RemoteId;
+	}
+#endif
+
 private:
+#if !UE_WITH_REMOTE_OBJECT_HANDLE
 	FObjectKey(int32 Index, int32 Serial)
 		: ObjectIndex(Index)
 		, ObjectSerialNumber(Serial)
-	{ }
+	{
+	}
 
 	friend FObjectKey UE::CoreUObject::Private::MakeObjectKey(int32 ObjectIndex, int32 ObjectSerialNumber);
+#endif
 
 	int32		ObjectIndex;
 	int32		ObjectSerialNumber;
+#if UE_WITH_REMOTE_OBJECT_HANDLE
+	FRemoteObjectId RemoteId;
+#endif // UE_WITH_REMOTE_OBJECT_HANDLE
 };
 
 /** TObjectKey is a strongly typed, immutable, copyable key which can be used to uniquely identify an object for the lifetime of the application */
diff --git a/Engine/Source/Runtime/CoreUObject/Public/UObject/RemoteObject.h b/Engine/Source/Runtime/CoreUObject/Public/UObject/RemoteObject.h
--- a/Engine/Source/Runtime/CoreUObject/Public/UObject/RemoteObject.h
+++ b/Engine/Source/Runtime/CoreUObject/Public/UObject/RemoteObject.h
@@ -24,6 +24,10 @@ namespace UE::RemoteObject::Handle
 	{
 		FRemoteObjectId Id;
 		FRemoteObjectId OuterId;
+
+		/** SerialNumber this object had on this server */
+		int32 SerialNumber = 0;
+
 		/** Server id that currently owns the object */
 		FRemoteServerId RemoteServerId;
 
diff --git a/Engine/Source/Runtime/CoreUObject/Public/UObject/RemoteObjectTypes.h b/Engine/Source/Runtime/CoreUObject/Public/UObject/RemoteObjectTypes.h
--- a/Engine/Source/Runtime/CoreUObject/Public/UObject/RemoteObjectTypes.h
+++ b/Engine/Source/Runtime/CoreUObject/Public/UObject/RemoteObjectTypes.h
@@ -52,8 +52,8 @@ private:
 	{
 		struct
 		{
-			uint64 ServerId : 10;
 			uint64 SerialNumber : 54;
+			uint64 ServerId : 10;
 		};
 		uint64 Id = 0;
 	};
@@ -68,8 +68,8 @@ public:
 	}
 
 	FRemoteObjectId(FRemoteServerId InServerId, uint64 InSerialNumber)
-		: ServerId(InServerId.Id)
-		, SerialNumber(InSerialNumber)
+		: SerialNumber(InSerialNumber)
+		, ServerId(InServerId.Id)
 	{
 	}
 
@@ -90,15 +90,19 @@ public:
 	}
 	FORCEINLINE bool operator<(const FRemoteObjectId& Other) const
 	{
-		if (ServerId < Other.ServerId)
-		{
-			return true;
-		}
-		else if (ServerId == Other.ServerId)
-		{
-			return SerialNumber < Other.SerialNumber;
-		}
-		return false;
+		return Id < Other.Id;
+	}
+	FORCEINLINE bool operator<=(const FRemoteObjectId& Other) const
+	{
+		return Id <= Other.Id;
+	}
+	FORCEINLINE bool operator>(const FRemoteObjectId& Other) const
+	{
+		return Id > Other.Id;
+	}
+	FORCEINLINE bool operator>=(const FRemoteObjectId& Other) const
+	{
+		return Id >= Other.Id;
 	}
 	FORCEINLINE bool IsValid() const
 	{
diff --git a/Engine/Source/Runtime/CoreUObject/Public/UObject/UObjectArray.h b/Engine/Source/Runtime/CoreUObject/Public/UObject/UObjectArray.h
--- a/Engine/Source/Runtime/CoreUObject/Public/UObject/UObjectArray.h
+++ b/Engine/Source/Runtime/CoreUObject/Public/UObject/UObjectArray.h
@@ -60,6 +60,11 @@ public:
 	// RefCount associated with the object preventing its destruction.
 	int32 RefCount;
 
+#if UE_WITH_REMOTE_OBJECT_HANDLE
+	// Globally unqiue id of this object
+	FRemoteObjectId RemoteId;
+#endif
+
 #if STATS || ENABLE_STATNAMEDEVENTS_UOBJECT
 	/** Stat id of this object, 0 if nobody asked for it yet */
 	mutable TStatId StatID;
@@ -125,6 +130,13 @@ public:
 		return SerialNumber;
 	}
 
+#if UE_WITH_REMOTE_OBJECT_HANDLE
+	FORCEINLINE FRemoteObjectId GetRemoteId() const
+	{
+		return RemoteId;
+	}
+#endif // UE_WITH_REMOTE_OBJECT_HANDLE
+
 	FORCEINLINE void SetFlags(EInternalObjectFlags FlagsToSet)
 	{
 		ThisThreadAtomicallySetFlag(FlagsToSet);
@@ -791,7 +803,7 @@ public:
 class FUObjectArray
 {
 	friend class UObject;
-	friend COREUOBJECT_API UObject* StaticAllocateObject(const UClass*, UObject*, FName, EObjectFlags, EInternalObjectFlags, bool, bool*, UPackage*, FRemoteObjectId);
+	friend COREUOBJECT_API UObject* StaticAllocateObject(const UClass*, UObject*, FName, EObjectFlags, EInternalObjectFlags, bool, bool*, UPackage*, int32, FRemoteObjectId);
 
 private:
 	/**
@@ -912,7 +924,7 @@ public:
 	 * @param	AlreadyAllocatedIndex already allocated internal index to use, negative value means allocate a new index
 	 * @param	SerialNumber serial number to use
 	 */
-	COREUOBJECT_API void AllocateUObjectIndex(class UObjectBase* Object, EInternalObjectFlags InitialFlags, int32 AlreadyAllocatedIndex = -1, int32 SerialNumber = 0);
+	COREUOBJECT_API void AllocateUObjectIndex(class UObjectBase* Object, EInternalObjectFlags InitialFlags, int32 AlreadyAllocatedIndex = -1, int32 SerialNumber = 0, FRemoteObjectId RemoteId = FRemoteObjectId());
 
 	/**
 	 * Returns a UObject index top to the global uobject array
@@ -1157,6 +1169,21 @@ public:
 		return ObjectItem->GetSerialNumber();
 	}
 
+#if UE_WITH_REMOTE_OBJECT_HANDLE
+
+	/**
+	* Given a UObject index return the remote id.
+	* @param Index - UObject Index
+	* @return - the remote id for this UObject
+	*/
+	FORCEINLINE FRemoteObjectId GetRemoteId(int32 Index)
+	{
+		FUObjectItem* ObjectItem = IndexToObject(Index);
+		checkSlow(ObjectItem);
+		return ObjectItem->GetRemoteId();
+	}
+#endif
+
 	/** Locks the internal object array mutex */
 	void LockInternalArray() const
 	{
diff --git a/Engine/Source/Runtime/CoreUObject/Public/UObject/UObjectBase.h b/Engine/Source/Runtime/CoreUObject/Public/UObject/UObjectBase.h
--- a/Engine/Source/Runtime/CoreUObject/Public/UObject/UObjectBase.h
+++ b/Engine/Source/Runtime/CoreUObject/Public/UObject/UObjectBase.h
@@ -63,9 +63,6 @@ class UObjectBase
 		);
 protected:
 	UObjectBase() :
-#if UE_WITH_REMOTE_OBJECT_HANDLE
-		RemoteId(this), // FRemoteObjectId is a literal type and can't have a NoInit constructor so copy the id from whatever is already set in (this) memory
-#endif // UE_WITH_REMOTE_OBJECT_HANDLE
 		ClassPrivate(NoInit),
 		NamePrivate(NoInit),  // screwy, but the name was already set and we don't want to set it again
 		OuterPrivate(NoInit)
@@ -125,7 +122,11 @@ protected:
 	 *
 	 * @param UClassStaticClass Now that it is known, fill in UClass::StaticClass() as the class
 	 */
-	COREUOBJECT_API virtual void DeferredRegister(UClass *UClassStaticClass,const TCHAR* PackageName,const TCHAR* Name);
+	COREUOBJECT_API virtual void DeferredRegister(UClass* UClassStaticClass, const TCHAR* PackageName, const TCHAR* Name
+#if UE_WITH_REMOTE_OBJECT_HANDLE
+		, FRemoteObjectId InRemoteId
+#endif
+	);
 
 private:
 	/**
@@ -136,7 +137,7 @@ private:
 	 * @param InInternalIndex already allocated internal index to use, negative value means allocate a new index
 	 * @param InSerialNumber already allocated serial number to re-use
 	 */
-	COREUOBJECT_API void AddObject(FName Name, EInternalObjectFlags InSetInternalFlags, int32 InInternalIndex = -1, int32 InSerialNumber = 0);
+	COREUOBJECT_API void AddObject(FName Name, EInternalObjectFlags InSetInternalFlags, int32 InInternalIndex = -1, int32 InSerialNumber = 0, FRemoteObjectId InRemoteId = FRemoteObjectId());
 
 	/**
 	 * Checks to see if the object appears to be valid when being destroyed
@@ -144,6 +145,10 @@ private:
 	 */
 	COREUOBJECT_API bool IsValidLowLevelForDestruction() const;
 
+#if UE_WITH_REMOTE_OBJECT_HANDLE
+	FRemoteObjectId GetPendingRegistrantRemoteId() const;
+#endif
+
 public:
 	/**
 	 * Checks to see if the object appears to be valid
@@ -327,11 +332,6 @@ private:
 	/** Index into GObjectArray...very private. */
 	int32							InternalIndex;
 
-#if UE_WITH_REMOTE_OBJECT_HANDLE
-	/** Unique Id associated with this object */
-	FRemoteObjectId					RemoteId;
-#endif // UE_WITH_REMOTE_OBJECT_HANDLE
-
 	/** Class the object belongs to. */
 	ObjectPtr_Private::TNonAccessTrackedObjectPtr<UClass>							ClassPrivate;
 
diff --git a/Engine/Source/Runtime/CoreUObject/Public/UObject/UObjectGlobals.h b/Engine/Source/Runtime/CoreUObject/Public/UObject/UObjectGlobals.h
--- a/Engine/Source/Runtime/CoreUObject/Public/UObject/UObjectGlobals.h
+++ b/Engine/Source/Runtime/CoreUObject/Public/UObject/UObjectGlobals.h
@@ -244,6 +244,10 @@ namespace UE
 
 } // namespace UE
 
+namespace UE::RemoteObject::Serialization
+{
+	class FRemoteObjectConstructionOverrides;
+} // namespace UE::RemoteObject::Serialization
 
 /** Parses a bit mask of property flags into an array of string literals that match the flags */
 COREUOBJECT_API TArray<const TCHAR*> ParsePropertyFlags(EPropertyFlags Flags);
@@ -1172,7 +1176,7 @@ bool StaticAllocateObjectErrorTests( const UClass* Class, UObject* InOuter, FNam
  * @param RemoteId		Globally unique id for this object
  * @return	a pointer to a fully initialized object of the specified class.
  */
-COREUOBJECT_API UObject* StaticAllocateObject(const UClass* Class, UObject* InOuter, FName Name, EObjectFlags SetFlags, EInternalObjectFlags InternalSetFlags = EInternalObjectFlags::None, bool bCanReuseSubobjects = false, bool* bOutReusedSubobject = nullptr, UPackage* ExternalPackage = nullptr, FRemoteObjectId RemoteId = FRemoteObjectId());
+COREUOBJECT_API UObject* StaticAllocateObject(const UClass* Class, UObject* InOuter, FName Name, EObjectFlags SetFlags, EInternalObjectFlags InternalSetFlags = EInternalObjectFlags::None, bool bCanReuseSubobjects = false, bool* bOutReusedSubobject = nullptr, UPackage* ExternalPackage = nullptr, int32 SerialNumber = 0, FRemoteObjectId RemoteId = FRemoteObjectId());
 
 /** FObjectInitializer options */
 enum class EObjectInitializerOptions
@@ -1215,7 +1219,7 @@ public:
 	}
 
 	/** Special constructor for static construct object internal that passes along the params block directly */
-	COREUOBJECT_API FObjectInitializer(UObject* InObj, const FStaticConstructObjectParameters& StaticConstructParams);
+	COREUOBJECT_API FObjectInitializer(UObject* InObj, const FStaticConstructObjectParameters& StaticConstructParams, UE::RemoteObject::Serialization::FRemoteObjectConstructionOverrides* RemoteObjectOverides = nullptr);
 
 private:
 	/** Helper for the common behaviors in the constructors */
@@ -1642,6 +1646,17 @@ private:
 	/** List of FGCObjects created during UObject construction */
 	TArray<FGCObject*> CreatedGCObjects;
 #endif // WITH_EDITORONLY_DATA
+
+#if UE_WITH_REMOTE_OBJECT_HANDLE
+	/** Existing remote id for the object */
+	UE::RemoteObject::Serialization::FRemoteObjectConstructionOverrides* RemoteSubObjectOverrides = nullptr;
+
+public:
+	UE::RemoteObject::Serialization::FRemoteObjectConstructionOverrides* GetRemoteSubObjectOverrides() const
+	{
+		return RemoteSubObjectOverrides;
+	}
+#endif // UE_WITH_REMOTE_OBJECT_HANDLE
 };
 
 /**
@@ -1686,8 +1701,16 @@ struct FStaticConstructObjectParameters
 	/** Callback for custom code to initialize properties before PostInitProperties runs */
 	TFunction<void()> PropertyInitCallback;
 
+	/** Existing serial number for the object */
+	int32 SerialNumber = 0;
+
+#if UE_WITH_REMOTE_OBJECT_HANDLE
 	/** Existing remote id for the object */
 	FRemoteObjectId RemoteId;
+
+	UE::RemoteObject::Serialization::FRemoteObjectConstructionOverrides* RemoteSubObjectOverrides = nullptr;
+#endif // UE_WITH_REMOTE_OBJECT_HANDLE
+
 private:
 	FObjectInitializer::FOverrides* SubobjectOverrides = nullptr;
 
@@ -1773,7 +1796,7 @@ COREUOBJECT_API void CheckIsClassChildOf_Internal(const UClass* Parent, const UC
  */
 template< class T >
 FUNCTION_NON_NULL_RETURN_START
-	T* NewObject(UObject* Outer, const UClass* Class, FName Name = NAME_None, EObjectFlags Flags = RF_NoFlags, UObject* Template = nullptr, bool bCopyTransientsFromClassDefaults = false, FObjectInstancingGraph* InInstanceGraph = nullptr, UPackage* ExternalPackage = nullptr, FRemoteObjectId RemoteId = FRemoteObjectId())
+	T* NewObject(UObject* Outer, const UClass* Class, FName Name = NAME_None, EObjectFlags Flags = RF_NoFlags, UObject* Template = nullptr, bool bCopyTransientsFromClassDefaults = false, FObjectInstancingGraph* InInstanceGraph = nullptr, UPackage* ExternalPackage = nullptr)
 FUNCTION_NON_NULL_RETURN_END
 {
 	if (Name == NAME_None)
@@ -1794,7 +1817,6 @@ FUNCTION_NON_NULL_RETURN_END
 	Params.bCopyTransientsFromClassDefaults = bCopyTransientsFromClassDefaults;
 	Params.InstanceGraph = InInstanceGraph;
 	Params.ExternalPackage = ExternalPackage;
-	Params.RemoteId = RemoteId;
 
 	T* Result = static_cast<T*>(StaticConstructObject_Internal(Params));
 	return Result;
diff --git a/Engine/Source/Runtime/CoreUObject/Public/UObject/WeakObjectPtr.h b/Engine/Source/Runtime/CoreUObject/Public/UObject/WeakObjectPtr.h
--- a/Engine/Source/Runtime/CoreUObject/Public/UObject/WeakObjectPtr.h
+++ b/Engine/Source/Runtime/CoreUObject/Public/UObject/WeakObjectPtr.h
@@ -13,7 +13,8 @@
 #include "UObject/UObjectArray.h"
 #include "UObject/StrongObjectPtr.h"
 #include "UObject/WeakObjectPtrFwd.h"
-
+#include "UObject/ObjectHandleDefines.h"
+#include "UObject/RemoteObjectTypes.h"
 #include <type_traits>
 
 class FArchive;
@@ -100,6 +101,9 @@ public:
 
 		ObjectIndex = InvalidWeakObjectIndex;
 		ObjectSerialNumber = 0;
+#if UE_WITH_REMOTE_OBJECT_HANDLE
+		ObjectRemoteId = FRemoteObjectId();
+#endif
 	}
 
 	/**  
@@ -123,7 +127,11 @@ public:
 	FORCEINLINE bool operator==(const FWeakObjectPtr& Other) const
 	{
 		return 
+#if UE_WITH_REMOTE_OBJECT_HANDLE
+			ObjectRemoteId == Other.ObjectRemoteId ||
+#else
 			(ObjectIndex == Other.ObjectIndex && ObjectSerialNumber == Other.ObjectSerialNumber) ||
+#endif
 			(!IsValid() && !Other.IsValid());
 	}
 
@@ -135,7 +143,11 @@ public:
 	FORCEINLINE bool operator!=(const FWeakObjectPtr& Other) const
 	{
 		return 
+#if UE_WITH_REMOTE_OBJECT_HANDLE
+			ObjectRemoteId != Other.ObjectRemoteId &&
+#else
 			(ObjectIndex != Other.ObjectIndex || ObjectSerialNumber != Other.ObjectSerialNumber) &&
+#endif
 			(IsValid() || Other.IsValid());
 	}
 #endif
@@ -146,7 +158,11 @@ public:
 	 */
 	FORCEINLINE bool HasSameIndexAndSerialNumber(const FWeakObjectPtr& Other) const
 	{
+#if UE_WITH_REMOTE_OBJECT_HANDLE
+		return ObjectRemoteId == Other.ObjectRemoteId;
+#else
 		return ObjectIndex == Other.ObjectIndex && ObjectSerialNumber == Other.ObjectSerialNumber;
+#endif
 	}
 
 	/**  
@@ -226,18 +242,25 @@ public:
 	FORCEINLINE bool IsExplicitlyNull() const
 	{
 		using namespace UE::Core::Private;
-
+#if UE_WITH_REMOTE_OBJECT_HANDLE
+		return !ObjectRemoteId.IsValid();
+#else
 #if UE_WEAKOBJECTPTR_ZEROINIT_FIX
 		return ObjectIndex == InvalidWeakObjectIndex && ObjectSerialNumber == 0;
 #else
 		return ObjectIndex == InvalidWeakObjectIndex;
 #endif
+#endif // UE_WITH_REMOTE_OBJECT_HANDLE
 	}
 
 	/** Hash function. */
 	FORCEINLINE uint32 GetTypeHash() const
 	{
+#if UE_WITH_REMOTE_OBJECT_HANDLE
+		return ::GetTypeHash(ObjectRemoteId);
+#else
 		return uint32(ObjectIndex ^ ObjectSerialNumber);
+#endif
 	}
 
 	/**
@@ -264,30 +287,78 @@ private:
 private:
 	friend struct FObjectKey;
 	
+#if UE_WITH_REMOTE_OBJECT_HANDLE
+	COREUOBJECT_API FUObjectItem* TryResolveRemoteObject() const;
+#endif
+
 	/**  
 	 * internal function to test for serial number matches
 	 * @return true if the serial number in this matches the central table
 	 */
 	FORCEINLINE_DEBUGGABLE bool SerialNumbersMatch() const
 	{
+#if UE_WITH_REMOTE_OBJECT_HANDLE
+		checkSlow(ObjectIndex >= 0); // otherwise this is a corrupted weak pointer
+		FRemoteObjectId ActualId = GUObjectArray.GetRemoteId(ObjectIndex);
+		return ActualId == ObjectRemoteId;
+#else
 		checkSlow(ObjectSerialNumber > FUObjectArray::START_SERIAL_NUMBER && ObjectIndex >= 0); // otherwise this is a corrupted weak pointer
 		int32 ActualSerialNumber = GUObjectArray.GetSerialNumber(ObjectIndex);
 		checkSlow(!ActualSerialNumber || ActualSerialNumber >= ObjectSerialNumber); // serial numbers should never shrink
 		return ActualSerialNumber == ObjectSerialNumber;
+#endif
 	}
 
 	FORCEINLINE_DEBUGGABLE bool SerialNumbersMatch(FUObjectItem* ObjectItem) const
 	{
+#if UE_WITH_REMOTE_OBJECT_HANDLE
+		checkSlow(ObjectIndex >= 0); // otherwise this is a corrupted weak pointer
+		FRemoteObjectId ActualId = ObjectItem->GetRemoteId();
+		return ActualId == ObjectRemoteId;
+#else
 		checkSlow(ObjectSerialNumber > FUObjectArray::START_SERIAL_NUMBER && ObjectIndex >= 0); // otherwise this is a corrupted weak pointer
 		const int32 ActualSerialNumber = ObjectItem->GetSerialNumber();
 		checkSlow(!ActualSerialNumber || ActualSerialNumber >= ObjectSerialNumber); // serial numbers should never shrink
 		return ActualSerialNumber == ObjectSerialNumber;
+#endif
 	}
 
 	FORCEINLINE FUObjectItem* Internal_GetObjectItem() const
 	{
 		using namespace UE::Core::Private;
 
+#if UE_WITH_REMOTE_OBJECT_HANDLE
+		if (IsExplicitlyNull())
+		{
+			return nullptr;
+		}
+		if (ObjectIndex < 0)
+		{
+			return nullptr;
+		}
+
+		FUObjectItem* ObjectItem = GUObjectArray.IndexToObject(ObjectIndex);
+		bool bResolved = false;
+		if (!ObjectItem)
+		{
+			ObjectItem = TryResolveRemoteObject();
+			if (!ObjectItem)
+			{
+				return nullptr;
+			}
+			bResolved = true;
+		}
+		// TryResolveRemoteObject already validated the serial number
+		if (!bResolved && !SerialNumbersMatch(ObjectItem))
+		{
+			ObjectItem = TryResolveRemoteObject();
+			if (!ObjectItem)
+			{
+				return nullptr;
+			}
+		}
+		return ObjectItem;
+#else
 		if (ObjectSerialNumber == 0)
 		{
 #if UE_WEAKOBJECTPTR_ZEROINIT_FIX
@@ -313,6 +384,7 @@ private:
 			return nullptr;
 		}
 		return ObjectItem;
+#endif // UE_WITH_REMOTE_OBJECT_HANDLE
 	}
 
 	/** Private (inlined) version for internal use only. */
@@ -344,7 +416,11 @@ private:
 #else
 	int32		ObjectIndex;
 	int32		ObjectSerialNumber;
-#endif
+#endif // UE_WEAKOBJECTPTR_ZEROINIT_FIX
+#if UE_WITH_REMOTE_OBJECT_HANDLE
+	FRemoteObjectId ObjectRemoteId;
+#endif // UE_WITH_REMOTE_OBJECT_HANDLE
+
 };
 
 /** Hash function. */
diff --git a/Engine/Source/Runtime/Net/Core/Private/Net/Core/NetHandle/NetHandleManager.cpp b/Engine/Source/Runtime/Net/Core/Private/Net/Core/NetHandle/NetHandleManager.cpp
--- a/Engine/Source/Runtime/Net/Core/Private/Net/Core/NetHandle/NetHandleManager.cpp
+++ b/Engine/Source/Runtime/Net/Core/Private/Net/Core/NetHandle/NetHandleManager.cpp
@@ -31,8 +31,12 @@ FNetHandle FNetHandleManager::GetNetHandle(const UObject* Object)
 	if (KeyIt)
 	{
 		FNetHandle Handle = KeyIt.Value();
+#if UE_WITH_REMOTE_OBJECT_HANDLE
+		if (Handle.Value == FRemoteObjectId(Object))
+#else
 		UObject* ResolvedObject = Handle.Value.ResolveObjectPtrEvenIfUnreachable();
 		if (ResolvedObject == Object)
+#endif
 		{
 			return Handle;
 		}
@@ -59,8 +63,12 @@ FNetHandle FNetHandleManager::GetOrCreateNetHandle(const UObject* Object)
 	if (KeyIt)
 	{
 		FNetHandle Handle = KeyIt.Value();
+#if UE_WITH_REMOTE_OBJECT_HANDLE
+		if (Handle.Value == FRemoteObjectId(Object))
+#else
 		UObject* ResolvedObject = Handle.Value.ResolveObjectPtrEvenIfUnreachable();
 		if (Object == ResolvedObject)
+#endif
 		{
 			return Handle;
 		}
diff --git a/Engine/Source/Runtime/Net/Core/Public/Net/Core/NetHandle/NetHandle.h b/Engine/Source/Runtime/Net/Core/Public/Net/Core/NetHandle/NetHandle.h
--- a/Engine/Source/Runtime/Net/Core/Public/Net/Core/NetHandle/NetHandle.h
+++ b/Engine/Source/Runtime/Net/Core/Public/Net/Core/NetHandle/NetHandle.h
@@ -69,7 +69,11 @@ private:
 
 	union
 	{
+#if UE_WITH_REMOTE_OBJECT_HANDLE
+		FRemoteObjectId Value;
+#else
 		FObjectKey Value;
+#endif
 		FInternalValue InternalValue;
 	};
 };
@@ -90,17 +94,29 @@ inline FNetHandle::FNetHandle(FNetHandle&& Other)
 }
 
 inline FNetHandle::FNetHandle(FObjectKey ObjectKey)
+#if UE_WITH_REMOTE_OBJECT_HANDLE
+: Value(ObjectKey.GetRemoteId())
+#else
 : Value(ObjectKey)
+#endif
 {
 }
 
 inline FNetHandle::FNetHandle(FInternalValue InInternalValue)
+#if UE_WITH_REMOTE_OBJECT_HANDLE
+: Value(BitCast<FRemoteObjectId>(InInternalValue))
+#else
 : Value(BitCast<FObjectKey>(InInternalValue))
+#endif
 {
 }
 
 inline FNetHandle::FNetHandle(uint64 InInternalValue)
+#if UE_WITH_REMOTE_OBJECT_HANDLE
+: Value(BitCast<FRemoteObjectId>(InInternalValue))
+#else
 : Value(BitCast<FObjectKey>(InInternalValue))
+#endif
 {
 }
 
@@ -128,12 +144,20 @@ inline uint32 FNetHandle::GetId() const
 
 inline bool FNetHandle::IsValid() const
 {
+#if UE_WITH_REMOTE_OBJECT_HANDLE
+	return Value != FRemoteObjectId();
+#else
 	return Value != FObjectKey();
+#endif
 }
 
 inline bool FNetHandle::IsCompleteHandle() const
 {
+#if UE_WITH_REMOTE_OBJECT_HANDLE
+	return Value != FRemoteObjectId();
+#else
 	return Value != FObjectKey();
+#endif
 }
 
 inline bool FNetHandle::operator==(FNetHandle Other) const
