RemoteObject cleanup:

Moved RemoteObject delegates from FCoreUObjectDelegates to UE::RemoteObject::Transfer
Added LogRemoteObject log category

diff --git a/Engine/Source/Runtime/CoreUObject/Private/UObject/RemoteObject.cpp b/Engine/Source/Runtime/CoreUObject/Private/UObject/RemoteObject.cpp
--- a/Engine/Source/Runtime/CoreUObject/Private/UObject/RemoteObject.cpp
+++ b/Engine/Source/Runtime/CoreUObject/Private/UObject/RemoteObject.cpp
@@ -15,6 +15,8 @@
 #include "Misc/CoreDelegates.h"
 #include <atomic>
 
+DEFINE_LOG_CATEGORY(LogRemoteObject);
+
 namespace UE::RemoteObject::Private
 {
 
@@ -124,13 +126,13 @@ void InitRemoteObjects()
 	checkf(ServerIdNumber != 0, TEXT("Remote Server Id=%u is not valid"), ServerIdNumber);
 	RemoteServerId = FRemoteServerId(ServerIdNumber);
 
-	if (!FCoreUObjectDelegates::RemoteObjectTransferDelegate.IsBound())
+	if (!UE::RemoteObject::Transfer::RemoteObjectTransferDelegate.IsBound())
 	{
-		FCoreUObjectDelegates::RemoteObjectTransferDelegate.BindStatic(&UE::RemoteObject::Serialization::Disk::SaveObjectToDisk);
+		UE::RemoteObject::Transfer::RemoteObjectTransferDelegate.BindStatic(&UE::RemoteObject::Serialization::Disk::SaveObjectToDisk);
 	}
-	if (!FCoreUObjectDelegates::RequestRemoteObjectDelegate.IsBound())
+	if (!UE::RemoteObject::Transfer::RequestRemoteObjectDelegate.IsBound())
 	{
-		FCoreUObjectDelegates::RequestRemoteObjectDelegate.BindStatic(&UE::RemoteObject::Serialization::Disk::LoadObjectFromDisk);
+		UE::RemoteObject::Transfer::RequestRemoteObjectDelegate.BindStatic(&UE::RemoteObject::Serialization::Disk::LoadObjectFromDisk);
 	}
 
 	FCoreDelegates::OnEndFrame.AddStatic(&TickRemoteObjectRequests);
diff --git a/Engine/Source/Runtime/CoreUObject/Private/UObject/RemoteObjectPrivate.h b/Engine/Source/Runtime/CoreUObject/Private/UObject/RemoteObjectPrivate.h
--- a/Engine/Source/Runtime/CoreUObject/Private/UObject/RemoteObjectPrivate.h
+++ b/Engine/Source/Runtime/CoreUObject/Private/UObject/RemoteObjectPrivate.h
@@ -8,6 +8,8 @@
 #include "UObject/ObjectMacros.h"
 #include "UObject/Object.h"
 
+DECLARE_LOG_CATEGORY_EXTERN(LogRemoteObject, Display, All);
+
 namespace UE::RemoteObject::Private
 {
 	/**
diff --git a/Engine/Source/Runtime/CoreUObject/Private/UObject/RemoteObjectTransfer.cpp b/Engine/Source/Runtime/CoreUObject/Private/UObject/RemoteObjectTransfer.cpp
--- a/Engine/Source/Runtime/CoreUObject/Private/UObject/RemoteObjectTransfer.cpp
+++ b/Engine/Source/Runtime/CoreUObject/Private/UObject/RemoteObjectTransfer.cpp
@@ -9,6 +9,7 @@
 #include "Serialization/MemoryWriter.h"
 #include "Serialization/MemoryReader.h"
 #include "UObject/RemoteObject.h"
+#include "UObject/RemoteObjectPrivate.h"
 #include "UObject/ObjectHandlePrivate.h"
 #include "HAL/CriticalSection.h"
 
@@ -16,6 +17,10 @@
 namespace UE::RemoteObject::Transfer
 {
 
+TDelegate<void(FRemoteObjectId ObjectId, FRemoteServerId CurrentOwnerServerId, FRemoteServerId DestinationServerId)> RequestRemoteObjectDelegate;
+TDelegate<void(FRemoteObjectId ObjectId, FRemoteServerId DestinationServerId, FRemoteObjectData& ObjectData)> RemoteObjectTransferDelegate;
+TDelegate<void()> FlushRemoteObjectRequestsDelegate;
+
 /**
  * Remote object transfer queue. 
  * Queued requests (both send and receive) are processed on the game thread since some systems (like RPCs) that are used for transfering object data are GT-only
@@ -72,7 +77,7 @@ class FRemoteObjectTransferQueue
 		virtual void Process() override
 		{
 			// Request an object from remote server
-			FCoreUObjectDelegates::RequestRemoteObjectDelegate.Execute(ObjectId, /*CurrentOwner*/ ServerId, GetGlobalServerId());
+			RequestRemoteObjectDelegate.Execute(ObjectId, /*CurrentOwner*/ ServerId, GetGlobalServerId());
 		}
 	};
 
@@ -91,7 +96,7 @@ class FRemoteObjectTransferQueue
 		virtual void Process() override
 		{
 			// Request an object from remote server to be sent to another remote server
-			FCoreUObjectDelegates::RequestRemoteObjectDelegate.Execute(ObjectId, /*CurrentOwner*/ ServerId, DestinationServerId);
+			RequestRemoteObjectDelegate.Execute(ObjectId, /*CurrentOwner*/ ServerId, DestinationServerId);
 			State = ERequestState::Complete;
 		}
 	};
@@ -111,7 +116,7 @@ class FRemoteObjectTransferQueue
 		virtual void Process() override
 		{
 			// Send the associated object data to a remote server
-			FCoreUObjectDelegates::RemoteObjectTransferDelegate.Execute(ObjectId, ServerId, *ObjectData);
+			RemoteObjectTransferDelegate.Execute(ObjectId, ServerId, *ObjectData);
 			State = ERequestState::Complete;
 		}
 	};
@@ -323,7 +328,7 @@ public:
 		FRemoteObjectDataFuture RequestFuture = FindReceiveRequest(ObjectId);
 		while (RequestFuture.IsValid() && !RequestFuture.IsReady())
 		{
-			FCoreUObjectDelegates::FlushRemoteObjectRequestsDelegate.ExecuteIfBound();
+			FlushRemoteObjectRequestsDelegate.ExecuteIfBound();
 			Tick();
 			FPlatformProcess::Sleep(0.0f);
 		}
diff --git a/Engine/Source/Runtime/CoreUObject/Private/UObject/UObjectGlobals.cpp b/Engine/Source/Runtime/CoreUObject/Private/UObject/UObjectGlobals.cpp
--- a/Engine/Source/Runtime/CoreUObject/Private/UObject/UObjectGlobals.cpp
+++ b/Engine/Source/Runtime/CoreUObject/Private/UObject/UObjectGlobals.cpp
@@ -246,10 +246,6 @@ FCoreUObjectDelegates::FShouldCookPackageForPlatform FCoreUObjectDelegates::Shou
 FCoreUObjectDelegates::FPackageCreatedForLoad FCoreUObjectDelegates::PackageCreatedForLoad;
 FCoreUObjectDelegates::FGetPrimaryAssetIdForObject FCoreUObjectDelegates::GetPrimaryAssetIdForObject;
 
-TDelegate<void(FRemoteObjectId ObjectId, FRemoteServerId CurrentOwnerServerId, FRemoteServerId DestinationServerId)> FCoreUObjectDelegates::RequestRemoteObjectDelegate;
-TDelegate<void(FRemoteObjectId ObjectId, FRemoteServerId ServerId, FRemoteObjectData& ObjectData)> FCoreUObjectDelegates::RemoteObjectTransferDelegate;
-TDelegate<void()> FCoreUObjectDelegates::FlushRemoteObjectRequestsDelegate;
-
 /** Check whether we should report progress or not */
 bool ShouldReportProgress()
 {
diff --git a/Engine/Source/Runtime/CoreUObject/Public/UObject/RemoteObjectTransfer.h b/Engine/Source/Runtime/CoreUObject/Public/UObject/RemoteObjectTransfer.h
--- a/Engine/Source/Runtime/CoreUObject/Public/UObject/RemoteObjectTransfer.h
+++ b/Engine/Source/Runtime/CoreUObject/Public/UObject/RemoteObjectTransfer.h
@@ -63,5 +63,14 @@ namespace UE::RemoteObject::Transfer
 #endif
 	}
 
+	/** Delegate that transfers object data to another server */
+	extern COREUOBJECT_API TDelegate<void(FRemoteObjectId /*ObjectId*/, FRemoteServerId /*DestinationServerId*/, FRemoteObjectData& /*ObjectData*/)> RemoteObjectTransferDelegate;
+
+	/** Delegate that requests remote object data from CurrentOwnerServerId to be transferred to DestinationServerId */
+	extern COREUOBJECT_API TDelegate<void(FRemoteObjectId /*ObjectId*/, FRemoteServerId /*CurrentOwnerServerId*/, FRemoteServerId /*DestinationServerId*/)> RequestRemoteObjectDelegate;
+
+	/** Delegate that flushes remote object requests */
+	extern COREUOBJECT_API TDelegate<void()> FlushRemoteObjectRequestsDelegate;
+
 } // namespace UE::RemoveObject::Transfer
 
diff --git a/Engine/Source/Runtime/CoreUObject/Public/UObject/UObjectGlobals.h b/Engine/Source/Runtime/CoreUObject/Public/UObject/UObjectGlobals.h
--- a/Engine/Source/Runtime/CoreUObject/Public/UObject/UObjectGlobals.h
+++ b/Engine/Source/Runtime/CoreUObject/Public/UObject/UObjectGlobals.h
@@ -3312,15 +3312,6 @@ struct FCoreUObjectDelegates
 	/** Called during cooking to see if a specific package should be cooked for a given target platform */
 	DECLARE_DELEGATE_RetVal_TwoParams(bool, FShouldCookPackageForPlatform, const UPackage*, const ITargetPlatform*);
 	static COREUOBJECT_API FShouldCookPackageForPlatform ShouldCookPackageForPlatform;
-
-	/** Delegate that transfers object data to another server */
-	static COREUOBJECT_API TDelegate<void(FRemoteObjectId /*ObjectId*/, FRemoteServerId /*DestinationServerId*/, FRemoteObjectData& /*ObjectData*/)> RemoteObjectTransferDelegate;
-
-	/** Delegate that requests remote object data from CurrentOwnerServerId to be transferred to DestinationServerId */
-	static COREUOBJECT_API TDelegate<void(FRemoteObjectId /*ObjectId*/, FRemoteServerId /*CurrentOwnerServerId*/, FRemoteServerId /*DestinationServerId*/)> RequestRemoteObjectDelegate;
-
-	/** Delegate that flushes remote object requests */
-	static COREUOBJECT_API TDelegate<void()> FlushRemoteObjectRequestsDelegate;
 };
 
 /** Allows release builds to override not verifying GC assumptions. Useful for profiling as it's hitchy. */
