Support for dumping high level overview of data serialized with FRemoteObjectWriter

diff --git a/Engine/Source/Runtime/CoreUObject/Private/UObject/RemoteObjectSerialization.cpp b/Engine/Source/Runtime/CoreUObject/Private/UObject/RemoteObjectSerialization.cpp
--- a/Engine/Source/Runtime/CoreUObject/Private/UObject/RemoteObjectSerialization.cpp
+++ b/Engine/Source/Runtime/CoreUObject/Private/UObject/RemoteObjectSerialization.cpp
@@ -15,6 +15,7 @@
 #include "UObject/ObjectHandlePrivate.h"
 #include "Misc/Paths.h"
 #include "HAL/FileManager.h"
+#include "Serialization/ArchiveSerializedPropertyChain.h"
 
 namespace UE::RemoteObject::Serialization::Disk
 {
@@ -79,14 +80,31 @@ FArchive& operator<<(FArchive& Ar, ERemoteReferenceType& RefType)
 	return Ar;
 }
 
+#ifndef UE_WITH_REMOTEOBJECTWRITER_DEBUGGING
+#define UE_WITH_REMOTEOBJECTWRITER_DEBUGGING 0
+#endif
+
 class FArchiveRemoteObjectWriter : public FMemoryWriter
 {
-	UObject* RootObject = nullptr;	
+	UObject* RootObject = nullptr;
+	FRemoteObjectData& ObjectData;
 	TArray<UObject*>& ObjectsToSerialize;
 	const TSet<UObject*>& SerializedObjects;
 	TArray<FName>& Names;
 	TMap<FName, int32>& NameMap;
 
+#if UE_WITH_REMOTEOBJECTWRITER_DEBUGGING
+	FString SerializationScope;	
+
+	struct FSerializeStats
+	{
+		int64 Size = 0;
+		int64 Count = 0;
+	};
+	TMap<FString, TMap<FProperty*, FSerializeStats>> ObjectPropertyStats;
+	TMap<FString, FSerializeStats> PathNameStats;
+#endif
+
 	ERemoteReferenceType WriteObjectReference(UObject*& Obj)
 	{
 		using namespace UE::CoreUObject::Private;
@@ -117,8 +135,18 @@ class FArchiveRemoteObjectWriter : public FMemoryWriter
 
 		if (Type == ERemoteReferenceType::PathName)
 		{
+#if UE_WITH_REMOTEOBJECTWRITER_DEBUGGING
+			const int64 StartPos = Tell();
+#endif
+
 			FString PathName = GetPathNameSafe(Obj);
 			*this << PathName;
+
+#if UE_WITH_REMOTEOBJECTWRITER_DEBUGGING
+			FSerializeStats& Stats = PathNameStats.FindOrAdd(PathName);
+			Stats.Count++;
+			Stats.Size += Tell() - StartPos;
+#endif
 		}
 
 		return Type;
@@ -129,6 +157,7 @@ public:
 	FArchiveRemoteObjectWriter(UObject* InRootObject, FRemoteObjectData& OutBytes, TArray<UObject*>& OutObjectsToSerialize, const TSet<UObject*>& InSerializedObjects, TArray<FName>& InOutNames, TMap<FName, int32>& InOutNameMap)
 		: FMemoryWriter(OutBytes.Bytes)
 		, RootObject(InRootObject)
+		, ObjectData(OutBytes)
 		, ObjectsToSerialize(OutObjectsToSerialize)
 		, SerializedObjects(InSerializedObjects)
 		, Names(InOutNames)
@@ -137,6 +166,13 @@ public:
 		SetPortFlags(PPF_IncludeTransient);
 	}
 
+	virtual ~FArchiveRemoteObjectWriter()
+	{
+#if UE_WITH_REMOTEOBJECTWRITER_DEBUGGING
+		DumpStatsToLog();
+#endif
+	}
+
 	int32 AddNameToNameMap(const FName& Name)
 	{
 		int32 NameIndex = -1;
@@ -152,6 +188,84 @@ public:
 		return NameIndex;
 	}
 
+	FORCEINLINE UObject* GetRootObject() const
+	{
+		return RootObject;
+	}
+
+#if UE_WITH_REMOTEOBJECTWRITER_DEBUGGING
+	void DumpStatsToLog()
+	{
+		UE_LOG(LogRemoteObject, Log, TEXT("Object Data stats for %s %s (Total: %d):"), *FRemoteObjectId(RootObject).ToString(), *GetFullNameSafe(RootObject), ObjectData.Bytes.Num());
+
+		for (const TPair<FString, TMap<FProperty*, FSerializeStats>>& ObjectStatsPair : ObjectPropertyStats)
+		{
+			const FString& Obj = ObjectStatsPair.Key;
+			const TMap<FProperty*, FSerializeStats>& PropStats = ObjectStatsPair.Value;
+
+			int64 Total = 0;
+			for (const TPair<FProperty*, FSerializeStats>& PropertyPair : PropStats)
+			{
+				const FSerializeStats& Stats = PropertyPair.Value;
+				Total += Stats.Size;
+			}
+
+			UE_LOG(LogRemoteObject, Log, TEXT("  Data serialized for %s (Total: %lld):"), *Obj, Total);
+			
+			for (const TPair<FProperty*, FSerializeStats>& PropertyPair : PropStats)
+			{
+				FProperty* Prop = PropertyPair.Key;
+				const FSerializeStats& Stats = PropertyPair.Value;
+				UE_LOG(LogRemoteObject, Log, TEXT("    %s: size: %lld, count: %lld"), Prop ? *GetFullNameSafe(Prop) : TEXT("Native Serialize"), Stats.Size, Stats.Count);
+			}
+		}
+
+		{
+			int64 PathNameTotalCount = 0;
+			int64 PathNameTotalSize = 0;
+			for (const TPair<FString, FSerializeStats>& StatPair : PathNameStats)
+			{
+				PathNameTotalCount += StatPair.Value.Count;
+				PathNameTotalSize += StatPair.Value.Size;
+			}
+
+			UE_LOG(LogRemoteObject, Log, TEXT("  Serialized %lld PathNames (Total: %lld):"), PathNameTotalCount, PathNameTotalSize);
+			for (const TPair<FString, FSerializeStats>& StatPair : PathNameStats)
+			{
+				UE_LOG(LogRemoteObject, Log, TEXT("    %s: size: %lld, count: %lld"), *StatPair.Key, StatPair.Value.Size, StatPair.Value.Count);
+			}
+		}
+		
+	}
+
+	void SetSerializationScope(const TCHAR* InScope)
+	{
+		SerializationScope = InScope;
+	}
+
+	virtual void Serialize(void* Data, int64 Num) override
+	{
+		FProperty* CurrentProperty = nullptr;
+		if (const FArchiveSerializedPropertyChain* PropChain = GetSerializedPropertyChain())
+		{
+			if (PropChain->GetNumProperties())
+			{
+				CurrentProperty = PropChain->GetPropertyFromRoot(0);
+			}
+		}
+
+		TMap<FProperty*, FSerializeStats>& ObjectStats = ObjectPropertyStats.FindOrAdd(SerializationScope);
+		FSerializeStats& PropertyStats = ObjectStats.FindOrAdd(CurrentProperty);
+
+		const int64 StartPos = Tell();
+		
+		FMemoryWriter::Serialize(Data, Num);
+
+		PropertyStats.Size += Tell() - StartPos;
+		PropertyStats.Count++;		
+	}
+#endif // UE_WITH_REMOTEOBJECTWRITER_DEBUGGING
+
 	virtual FArchive& operator<<(UObject*& Obj) override
 	{
 		ERemoteReferenceType RefType = WriteObjectReference(Obj);
@@ -191,6 +305,34 @@ public:
 	virtual FArchive& operator<<(FWeakObjectPtr& Value) override { return FArchiveUObject::SerializeWeakObjectPtr(*this, Value); }
 };
 
+struct FRemoteObjectWriterScope
+{
+	FArchiveRemoteObjectWriter& Writer;
+
+	FRemoteObjectWriterScope(FArchiveRemoteObjectWriter& InWriter, const TCHAR* InScope)
+		: Writer(InWriter)
+	{
+#if UE_WITH_REMOTEOBJECTWRITER_DEBUGGING
+		Writer.SetSerializationScope(InScope);
+#endif
+	}
+	FRemoteObjectWriterScope(FArchiveRemoteObjectWriter& InWriter, UObject* InObjectScope)
+		: Writer(InWriter)
+	{
+#if UE_WITH_REMOTEOBJECTWRITER_DEBUGGING
+		UObject* Root = Writer.GetRootObject();
+		UObject* RootOuter = UE::CoreUObject::Private::FObjectHandleUtils::GetNonAccessTrackedOuterNoResolve(Root);
+		Writer.SetSerializationScope(*InObjectScope->GetFullName(RootOuter));
+#endif
+	}
+	~FRemoteObjectWriterScope()
+	{
+#if UE_WITH_REMOTEOBJECTWRITER_DEBUGGING
+		Writer.SetSerializationScope(nullptr);
+#endif
+	}
+};
+
 FArchive& operator<<(FArchive& Ar, FRemoteObjectHeader& Header)
 {
 	Ar << Header.Name;
@@ -442,15 +584,20 @@ FRemoteObjectData* SerializeObjectData(UObject* Object, TSet<UObject*>& OutObjec
 		FArchiveRemoteObjectWriter Ar(Object, *ObjectData, ObjectsToSerialize, OutObjects, Names, NameMap);
 		ObjectsToSerialize.Add(Object);
 
-		Ar << Version;
-		const int64 OffsetOfHeaderOffset = Ar.Tell();
-		Ar << HeaderOffset;
+		int64 OffsetOfHeaderOffset = 0;
+		{
+			FRemoteObjectWriterScope Scope(Ar, TEXT("Header"));
+			Ar << Version;
+			OffsetOfHeaderOffset = Ar.Tell();
+			Ar << HeaderOffset;
+		}
 
 		for (int32 ObjectIndex = 0; ObjectIndex < ObjectsToSerialize.Num(); ++ObjectIndex)
 		{
 			UObject* ObjectToSerialize = ObjectsToSerialize[ObjectIndex];
 			if (!OutObjects.Contains(ObjectToSerialize))
 			{
+				FRemoteObjectWriterScope Scope(Ar, ObjectToSerialize);
 				OutObjects.Add(ObjectToSerialize);
 				FRemoteObjectHeader& Header = ObjectHeaders.Emplace_GetRef();
 				SerializeRemoteObject(ObjectToSerialize, Ar, Header);
@@ -463,8 +610,14 @@ FRemoteObjectData* SerializeObjectData(UObject* Object, TSet<UObject*>& OutObjec
 		Ar << HeaderOffset;
 		Ar.Seek(HeaderOffset);
 
-		SerializeNames(Ar, Names);
-		Ar << ObjectHeaders;
+		{
+			FRemoteObjectWriterScope Scope(Ar, TEXT("NameMap"));
+			SerializeNames(Ar, Names);
+		}
+		{
+			FRemoteObjectWriterScope Scope(Ar, TEXT("ObjectHeaders"));
+			Ar << ObjectHeaders;
+		}
 	}
 
 	return ObjectData;
