Removed accidentally submitted debugging code.

diff --git a/Engine/Source/Runtime/CoreUObject/Private/UObject/SavePackage/PackageHarvester.cpp b/Engine/Source/Runtime/CoreUObject/Private/UObject/SavePackage/PackageHarvester.cpp
--- a/Engine/Source/Runtime/CoreUObject/Private/UObject/SavePackage/PackageHarvester.cpp
+++ b/Engine/Source/Runtime/CoreUObject/Private/UObject/SavePackage/PackageHarvester.cpp
@@ -10,8 +10,6 @@
 #include "UObject/SavePackage/SavePackageUtilities.h"
 #include "UObject/UObjectGlobals.h"
 #include "UObject/UObjectHash.h"
-#include "Serialization/ArchiveShowReferences.h"
-#include "Misc/OutputDeviceDebug.h"
 
 EObjectMark GenerateMarksForObject(const UObject* InObject, FSaveContext& SaveContext)
 {
@@ -308,10 +306,6 @@ FPackageHarvester::FHarvestScope FPackageHarvester::EnterConditionalOptionalObje
 			{
 				if (HarvestingRealm != ESaveRealm::Optional)
 				{
-					TArray<UObject*> Exclude;
-					FOutputDeviceDebug Output;
-					FArchiveShowReferences ShowReferences(Output, Ref, Ref, Exclude);
-
 					SaveContext.RecordIllegalReference(Ref, Object, EIllegalRefReason::ReferenceToOptional);
 				}
 			});
