Rename remote objects if their pathname collides with local objects and preserve locally given names.

diff --git a/Engine/Source/Runtime/CoreUObject/Private/UObject/RemoteObject.cpp b/Engine/Source/Runtime/CoreUObject/Private/UObject/RemoteObject.cpp
--- a/Engine/Source/Runtime/CoreUObject/Private/UObject/RemoteObject.cpp
+++ b/Engine/Source/Runtime/CoreUObject/Private/UObject/RemoteObject.cpp
@@ -102,9 +102,10 @@ public:
 			Stub = new FRemoteObjectStub(Object);
 		}
 		Stub->RemoteServerId = DestinationServerId;
-		FUObjectItem * ObjectItem = GUObjectArray.ObjectToObjectItem(Object);
+		FUObjectItem* ObjectItem = GUObjectArray.ObjectToObjectItem(Object);
 		checkf(ObjectItem, TEXT("Attempting to get a serial number for an object that does not exist in the global UObject array (it's possible GUObjectArray is not initialized yet, ObjectIndex=%d)"), GUObjectArray.ObjectToIndex(Object));
 		Stub->SerialNumber = ObjectItem->SerialNumber;
+		Stub->Name = Object->GetFName();
 
 		return Stub;
 	}
diff --git a/Engine/Source/Runtime/CoreUObject/Private/UObject/RemoteObjectSerialization.cpp b/Engine/Source/Runtime/CoreUObject/Private/UObject/RemoteObjectSerialization.cpp
--- a/Engine/Source/Runtime/CoreUObject/Private/UObject/RemoteObjectSerialization.cpp
+++ b/Engine/Source/Runtime/CoreUObject/Private/UObject/RemoteObjectSerialization.cpp
@@ -518,16 +518,23 @@ UObject* ConstructeRemoteObject(UObject* DestinationOuter, const FRemoteObjectHe
 
 	UClass* Class = CastChecked<UClass>(Header.Class.Get());
 	UObject* Outer = Header.Outer.Get();	
-	
+	FName Name = Header.Name;
+
 	// The object may already exist in memory (it could be a default subobject of an object we've just created)
-	UObject* Object = StaticFindObjectFast(Class, Outer, Header.Name);
+	UObject* Object = StaticFindObjectFast(Class, Outer, Name);
+	if (Object && FRemoteObjectId(Object) != Header.RemoteId)
+	{
+		UE_LOG(LogRemoteObject, Warning, TEXT("Received remote object %s with identical pathname (%s) as a local object %s. Remote object will be renamed."), *Header.RemoteId.ToString(), *Object->GetPathName(), *FRemoteObjectId(Object).ToString());
+		Name = FName();
+		Object = nullptr;
+	}
 
 	// If not or the object is marked as remote (which means we brought it back before it was GC'd) (re)construct it
 	if (!Object || IsRemote(Object))
 	{
 		FStaticConstructObjectParameters Params(Class);
 		Params.Outer = Outer;
-		Params.Name = Header.Name;
+		Params.Name = Name;
 		Params.SerialNumber = Header.SerialNumber;
 		Params.Template = Header.Archetype.Get();
 #if UE_WITH_REMOTE_OBJECT_HANDLE
@@ -537,8 +544,10 @@ UObject* ConstructeRemoteObject(UObject* DestinationOuter, const FRemoteObjectHe
 
 		// Using StaticConstructObject_Internal to pass the extra parameters (RemoteId and SubobjectOverrides) which are not exposed to normal APIs
 		Object = StaticConstructObject_Internal(Params);
-		checkf(FRemoteObjectId(Object) == Header.RemoteId, TEXT("Created an object with a different ID:%s than requested:%s"), *FRemoteObjectId(Object).ToString(), *Header.RemoteId.ToString());
 	}
+
+	checkf(FRemoteObjectId(Object) == Header.RemoteId, TEXT("Created an object with a different ID:%s than requested:%s"), *FRemoteObjectId(Object).ToString(), *Header.RemoteId.ToString());
+
 	MapObject(Object);
 	MarkAsLocal(Object);
 
@@ -648,12 +657,19 @@ void DeserializeObjectData(FRemoteObjectData& ObjectData, UObject* DestinationOu
 
 	if (ObjectHeaders.Num())
 	{
+		const FRemoteServerId LocalServerId = UE::RemoteObject::GetGlobalServerId();
+
 		// Try to find any existing WeakObject serial numbers for the objects that are about to be constructed
 		for (FRemoteObjectHeader& ObjectHeader : ObjectHeaders)
 		{
 			if (FRemoteObjectStub* Stub = FindRemoteObjectStub(ObjectHeader.RemoteId))
 			{
 				ObjectHeader.SerialNumber = Stub->SerialNumber;
+				if (!Stub->Name.IsNone())
+				{
+					// Remote object could've been renamed when it was migrated so always make sure that it has the same name locally
+					ObjectHeader.Name = Stub->Name;
+				}
 			}
 		}
 
diff --git a/Engine/Source/Runtime/CoreUObject/Public/UObject/RemoteObject.h b/Engine/Source/Runtime/CoreUObject/Public/UObject/RemoteObject.h
--- a/Engine/Source/Runtime/CoreUObject/Public/UObject/RemoteObject.h
+++ b/Engine/Source/Runtime/CoreUObject/Public/UObject/RemoteObject.h
@@ -24,6 +24,7 @@ namespace UE::RemoteObject::Handle
 	{
 		FRemoteObjectId Id;
 		FRemoteObjectId OuterId;
+		FName Name;
 
 		/** SerialNumber this object had on this server */
 		int32 SerialNumber = 0;
