[AutoRTFM] Do not use the annotation key / value cache in FUObjectAnnotationSparse when transactional.

We do not want to revert the cache on abort, nor do we want to have writes to the same memory locations from open and closed code.

diff --git a/Engine/Source/Programs/AutoRTFMTests/Private/UObjectTests.cpp b/Engine/Source/Programs/AutoRTFMTests/Private/UObjectTests.cpp
--- a/Engine/Source/Programs/AutoRTFMTests/Private/UObjectTests.cpp
+++ b/Engine/Source/Programs/AutoRTFMTests/Private/UObjectTests.cpp
@@ -146,7 +146,7 @@ TEST_CASE("UObject.MarkAsReachable")
 	SetReachabilityAnalysisTimeLimit(Original);
 }
 
-TEST_CASE("UObject.TestAddAnnotation")
+TEST_CASE("FUObjectAnnotationSparse.AddAnnotation")
 {
 	struct FTestAnnotation
 	{
@@ -161,51 +161,95 @@ TEST_CASE("UObject.TestAddAnnotation")
 		{
 			return TestAnnotationNumber == 32;
 		}
+
+		bool operator==(const FTestAnnotation& Other) const
+		{
+			return Other.TestAnnotationNumber == TestAnnotationNumber;
+		}
 	};
 
-	FUObjectAnnotationSparse<FTestAnnotation, true> GTestAnnotation;
+	FUObjectAnnotationSparse<FTestAnnotation, true> AnnotationMap;
+	
+	UMyAutoRTFMTestObject* Key = NewObject<UMyAutoRTFMTestObject>();
 
-	SECTION("Create")
+	FTestAnnotation ValueA;
+	ValueA.TestAnnotationNumber = 10;
+	
+	FTestAnnotation ValueB;
+	ValueB.TestAnnotationNumber = 20;
+
+	SECTION("Add")
 	{
-		UMyAutoRTFMTestObject* Outer = NewObject<UMyAutoRTFMTestObject>();
-		UMyAutoRTFMTestObject* Object = nullptr;
+		SECTION("Commit")
+		{
+			FTestAnnotation PreGet;
+			FTestAnnotation PostGet;
 
-		AutoRTFM::Commit([&]
+			AutoRTFM::Transact([&]
 			{
-				Object = NewObject<UMyAutoRTFMTestObject>(Outer);
+				FTestAnnotation Pre = AnnotationMap.GetAnnotation(Key);
+				AnnotationMap.AddAnnotation(Key, ValueA);
+				FTestAnnotation Post = AnnotationMap.GetAnnotation(Key);
+				AutoRTFM::Open([&]
+				{
+					PreGet = Pre;
+					PostGet = Post;
+				});
+			});
 
-				FTestAnnotation Temp;
-				Temp.TestAnnotationNumber = 70;
+			REQUIRE(ValueA == AnnotationMap.GetAnnotation(Key));
+			REQUIRE(FTestAnnotation() == PreGet);
+			REQUIRE(ValueA == PostGet);
+		}
 
-				GTestAnnotation.AddAnnotation(Object, Temp);
+		SECTION("Abort")
+		{
+			AutoRTFM::Transact([&]
+			{
+				AnnotationMap.AddAnnotation(Key, ValueA);
+				AutoRTFM::AbortTransaction();
 			});
 
-		REQUIRE(nullptr != Object);
-		REQUIRE(42 == Object->Value);
-		REQUIRE(Object->IsInOuter(Outer));
-		REQUIRE(55 == Outer->Value);
-		REQUIRE(70 == GTestAnnotation.GetAnnotation(Object).TestAnnotationNumber);
+			REQUIRE(FTestAnnotation() == AnnotationMap.GetAnnotation(Key));
+		}
 	}
 
-	SECTION("Abort")
+	SECTION("Replace")
 	{
-		UMyAutoRTFMTestObject* Outer = NewObject<UMyAutoRTFMTestObject>();
-		UMyAutoRTFMTestObject* Object = nullptr;
+		AnnotationMap.AddAnnotation(Key, ValueB);
 
-		Object = NewObject<UMyAutoRTFMTestObject>(Outer);
+		SECTION("Commit")
+		{
+			FTestAnnotation PreGet;
+			FTestAnnotation PostGet;
 
-		REQUIRE(AutoRTFM::ETransactionResult::AbortedByRequest == AutoRTFM::Transact([&]
+			AutoRTFM::Transact([&]
 			{
+				FTestAnnotation Pre = AnnotationMap.GetAnnotation(Key);
+				AnnotationMap.AddAnnotation(Key, ValueA);
+				FTestAnnotation Post = AnnotationMap.GetAnnotation(Key);
+				AutoRTFM::Open([&]
+				{
+					PreGet = Pre;
+					PostGet = Post;
+				});
+			});
 
-				FTestAnnotation Temp;
-				Temp.TestAnnotationNumber = 70;
-
-				GTestAnnotation.AddAnnotation(Object, Temp);
+			REQUIRE(ValueA == AnnotationMap.GetAnnotation(Key));
+			REQUIRE(ValueB == PreGet);
+			REQUIRE(ValueA == PostGet);
+		}
 
+		SECTION("Abort")
+		{
+			AutoRTFM::Transact([&]
+			{
+				AnnotationMap.AddAnnotation(Key, ValueA);
 				AutoRTFM::AbortTransaction();
-			}));
+			});
 
-		REQUIRE(32 == GTestAnnotation.GetAnnotation(Object).TestAnnotationNumber);
+			REQUIRE(ValueB == AnnotationMap.GetAnnotation(Key));
+		}
 	}
 }
 
diff --git a/Engine/Source/Runtime/CoreUObject/Public/UObject/UObjectAnnotation.h b/Engine/Source/Runtime/CoreUObject/Public/UObject/UObjectAnnotation.h
--- a/Engine/Source/Runtime/CoreUObject/Public/UObject/UObjectAnnotation.h
+++ b/Engine/Source/Runtime/CoreUObject/Public/UObject/UObjectAnnotation.h
@@ -6,6 +6,7 @@
 
 #pragma once
 
+#include "AutoRTFM/AutoRTFM.h"
 #include "UObject/UObjectArray.h"
 #include "Misc/TransactionallySafeRWLock.h"
 #include "Misc/TransactionallySafeRWScopeLock.h"
@@ -89,20 +90,8 @@ private:
 			FTransactionallySafeScopeLock AnnotationMapLock(&AnnotationMapCritical);
 			const bool bWasEmpty = (AnnotationMap.Num() == 0);
 
-			// If we are not in a trasncation update our cache values with the correct information
-			// While in a transaction avoid updating these as it pulls from closed code memory and can cause memory stomps if we abort here
-			if (!AutoRTFM::IsTransactional())
-			{
-				AnnotationCacheKey = Object;
-				AnnotationCacheValue = MoveTemp(LocalAnnotation);
-				AnnotationMap.Add(AnnotationCacheKey, AnnotationCacheValue);
-			}
-			else
-			{
-				AnnotationMap.Add(Object, LocalAnnotation);
-				AnnotationCacheKey = nullptr;
-				AnnotationCacheValue = TAnnotation();
-			}
+			AnnotationMap.Add(Object, LocalAnnotation);
+			SetAnnotationCacheKeyAndValue(Object, MoveTemp(LocalAnnotation));
 
 			if (bWasEmpty)
 			{
@@ -157,8 +146,7 @@ public:
 		TAnnotation Result;
 		{
 			FTransactionallySafeScopeLock AnnotationMapLock(&AnnotationMapCritical);
-			AnnotationCacheKey = Object;
-			AnnotationCacheValue = TAnnotation();
+			SetAnnotationCacheKeyAndValue(Object, TAnnotation());
 			bHadElements = (AnnotationMap.Num() > 0);
 			AnnotationMap.RemoveAndCopyValue(AnnotationCacheKey, Result);
 			bIsNowEmpty = (AnnotationMap.Num() == 0);
@@ -191,8 +179,7 @@ public:
 			// Avoid holding the lock while we call GUObjectArray.RemoveUObjectDeleteListener as it could deadlock
 			{
 				FTransactionallySafeScopeLock AnnotationMapLock(&AnnotationMapCritical);
-				AnnotationCacheKey = Object;
-				AnnotationCacheValue = TAnnotation();
+				SetAnnotationCacheKeyAndValue(Object, TAnnotation());
 				bHadElements = (AnnotationMap.Num() > 0);
 				AnnotationMap.Remove(AnnotationCacheKey);
 				bIsNowEmpty = (AnnotationMap.Num() == 0);
@@ -220,8 +207,7 @@ public:
 		// Avoid holding the lock while we call GUObjectArray.RemoveUObjectDeleteListener as it could deadlock
 		{
 			FTransactionallySafeScopeLock AnnotationMapLock(&AnnotationMapCritical);
-			AnnotationCacheKey = nullptr;
-			AnnotationCacheValue = TAnnotation();
+			SetAnnotationCacheKeyAndValue(nullptr, TAnnotation());
 			bHadElements = (AnnotationMap.Num() > 0);
 			AnnotationMap.Empty();
 		}
@@ -247,10 +233,10 @@ public:
 		check(Object);
 
 		FTransactionallySafeScopeLock AnnotationMapLock(&AnnotationMapCritical);
-		if (Object != AnnotationCacheKey)
+		if (AutoRTFM::IsTransactional() || Object != AnnotationCacheKey)
 		{			
-			AnnotationCacheKey = Object;
-			TAnnotation* Entry = AnnotationMap.Find(AnnotationCacheKey);
+			TAnnotation* Entry = AnnotationMap.Find(Object);
+			SetAnnotationCacheKeyAndValue(Object, Entry ? *Entry : TAnnotation());
 			if (Entry)
 			{
 				AnnotationCacheValue = *Entry;
@@ -288,6 +274,27 @@ public:
 	}
 
 private:
+	// Sets the AnnotationCacheKey and AnnotationCacheValue if called outside of
+	// an AutoRTFM transaction.
+	// This is a no-op in a transaction as we do not want to revert the cache on
+	// abort, nor do we want to have writes to the same memory locations from
+	// open and closed code.
+	void SetAnnotationCacheKeyAndValue(const UObjectBase* Key, const TAnnotation& Value)
+	{
+		if (!AutoRTFM::IsTransactional())
+		{
+			AnnotationCacheKey = Key;
+			AnnotationCacheValue = Value;
+		}
+	}
+	void SetAnnotationCacheKeyAndValue(const UObjectBase* Key, TAnnotation&& Value)
+	{
+		if (!AutoRTFM::IsTransactional())
+		{
+			AnnotationCacheKey = Key;
+			AnnotationCacheValue = MoveTemp(Value);
+		}
+	}
 
 	/**
 	 * Map from live objects to an annotation
