Avoid soft path serialization when saving weak and strong references with an object reference collector

#rnx

diff --git a/Engine/Source/Runtime/CoreUObject/Private/UObject/SoftObjectPath.cpp b/Engine/Source/Runtime/CoreUObject/Private/UObject/SoftObjectPath.cpp
--- a/Engine/Source/Runtime/CoreUObject/Private/UObject/SoftObjectPath.cpp
+++ b/Engine/Source/Runtime/CoreUObject/Private/UObject/SoftObjectPath.cpp
@@ -331,7 +331,7 @@ void FSoftObjectPath::SerializePath(FArchive& Ar)
 {
 	bool bSerializeInternals = true;
 #if WITH_EDITOR
-	if (Ar.IsSaving())
+	if (Ar.IsSaving() && !(Ar.IsModifyingWeakAndStrongReferences() && Ar.IsObjectReferenceCollector()))
 	{
 		PreSavePath(false ? GReportSoftObjectPathRedirects : nullptr);
 	}
