Add NotForClient and NotForServer objects as candidates for optional realm harvest

#rnx

diff --git a/Engine/Source/Runtime/CoreUObject/Private/UObject/SavePackage/SaveContext.h b/Engine/Source/Runtime/CoreUObject/Private/UObject/SavePackage/SaveContext.h
--- a/Engine/Source/Runtime/CoreUObject/Private/UObject/SavePackage/SaveContext.h
+++ b/Engine/Source/Runtime/CoreUObject/Private/UObject/SavePackage/SaveContext.h
@@ -606,8 +606,8 @@ public:
 		switch (HarvestingRealm)
 		{
 		case ESaveRealm::Optional:
-			// When considering excluded objects for a platform, do not consider editor only object and not for platform objects in the optional context as excluded
-			return (EObjectMark)(GameRealmExcludedObjectMarks & ~(EObjectMark::OBJECTMARK_EditorOnly | EObjectMark::OBJECTMARK_NotForTargetPlatform));
+			// When considering excluded objects for a platform, do not consider editor only object, not for platform, not for client and not for server objects in the optional context as excluded
+			return (EObjectMark)(GameRealmExcludedObjectMarks & ~(EObjectMark::OBJECTMARK_EditorOnly | EObjectMark::OBJECTMARK_NotForTargetPlatform | EObjectMark::OBJECTMARK_NotForClient | EObjectMark::OBJECTMARK_NotForServer));
 		case ESaveRealm::Game:
 			return GameRealmExcludedObjectMarks;
 		case ESaveRealm::Editor:
