Redirect TPropertyValueIterator for FObjectPropertyBase types to the RefLink chain instead of visiting all linked properties.

diff --git a/Engine/Source/Runtime/CoreUObject/Private/UObject/UnrealType.cpp b/Engine/Source/Runtime/CoreUObject/Private/UObject/UnrealType.cpp
--- a/Engine/Source/Runtime/CoreUObject/Private/UObject/UnrealType.cpp
+++ b/Engine/Source/Runtime/CoreUObject/Private/UObject/UnrealType.cpp
@@ -53,10 +53,11 @@ FPropertyValueIterator::EPropertyValueFlags FPropertyValueIterator::GetPropertyV
 
 void FPropertyValueIterator::FillStructProperties(const UStruct* Struct, FPropertyValueStackEntry& Entry)
 {
+	const void* Owner = Entry.Owner;
 	FPropertyValueStackEntry::FValueArrayType& ValueArray = Entry.ValueArray;
-	for (TFieldIterator<FProperty> It(Struct, EFieldIteratorFlags::IncludeSuper, DeprecatedPropertyFlags, EFieldIteratorFlags::ExcludeInterfaces); It; ++It)
+
+	auto AddPropertyToValueArray = [this, Owner, &ValueArray](const FProperty* Property)
 	{
-		const FProperty* Property  = *It;
 		EPropertyValueFlags PropertyValueFlags = GetPropertyValueFlags(Property);
 		if (PropertyValueFlags != EPropertyValueFlags::None)
 		{
@@ -64,10 +65,33 @@ void FPropertyValueIterator::FillStructProperties(const UStruct* Struct, FProper
 			ValueArray.Reserve(ValueArray.Num() + Num);
 			for (int32 StaticIndex = 0; StaticIndex < Num; ++StaticIndex)
 			{
-				const void* PropertyValue = Property->ContainerPtrToValuePtr<void>(Entry.Owner, StaticIndex);
+				const void* PropertyValue = Property->ContainerPtrToValuePtr<void>(Owner, StaticIndex);
 				ValueArray.Emplace(BasePairType(Property, PropertyValue), PropertyValueFlags);
 			}
 		}
+	};
+
+	if (PropertyClass->IsChildOf(FObjectPropertyBase::StaticClass()))
+	{
+		const bool bIncludeDeprecated = EnumHasAnyFlags((EFieldIterationFlags)DeprecatedPropertyFlags, EFieldIterationFlags::IncludeDeprecated);
+
+		// Note: This includes properties that contain object references (e.g. array), along with any inherited reference properties.
+		for (FProperty* RefProp = Struct->RefLink; RefProp; RefProp = RefProp->NextRef)
+		{
+			if (!bIncludeDeprecated && RefProp->HasAnyPropertyFlags(CPF_Deprecated))
+			{
+				continue;
+			}
+
+			AddPropertyToValueArray(RefProp);
+		}
+	}
+	else
+	{
+		for (TFieldIterator<FProperty> It(Struct, EFieldIteratorFlags::IncludeSuper, DeprecatedPropertyFlags, EFieldIteratorFlags::ExcludeInterfaces); It; ++It)
+		{
+			AddPropertyToValueArray(*It);
+		}
 	}
 }
 
