Fix UStruct::Link recent changes to restore legacy UProperty link list construction behavior (ie not nulling the next element as we build it)
This fixes a bunch of race conditions occuring in runtime when we load Blueprints in the async thread

diff --git a/Engine/Source/Runtime/Core/Public/Containers/LinkedListBuilder.h b/Engine/Source/Runtime/Core/Public/Containers/LinkedListBuilder.h
--- a/Engine/Source/Runtime/Core/Public/Containers/LinkedListBuilder.h
+++ b/Engine/Source/Runtime/Core/Public/Containers/LinkedListBuilder.h
@@ -61,6 +61,13 @@ public:
 		GetNextRef(Element) = nullptr;
 	}
 
+	FORCEINLINE void AppendNoTerminate(ElementType& Element)
+	{
+		*EndPtr = &Element;
+		EndPtr = LinkAccessor::GetNextPtr(Element);
+	}
+
+
 	FORCEINLINE void Remove(ElementType& Element)
 	{
 		ElementType** PrevIt = StartPtr;
diff --git a/Engine/Source/Runtime/CoreUObject/Private/UObject/Class.cpp b/Engine/Source/Runtime/CoreUObject/Private/UObject/Class.cpp
--- a/Engine/Source/Runtime/CoreUObject/Private/UObject/Class.cpp
+++ b/Engine/Source/Runtime/CoreUObject/Private/UObject/Class.cpp
@@ -1051,7 +1051,7 @@ void UStruct::Link(FArchive& Ar, bool bRelinkExistingProperties)
 		// contain object references
 		if (Property->ContainsObjectReference(EncounteredStructProps, EPropertyObjectReferenceType::Any))
 		{
-			RefLinkBuilder.Append(*Property);
+			RefLinkBuilder.AppendNoTerminate(*Property);
 		}
 
 		const UClass* OwnerClass = Property->GetOwnerClass();
@@ -1063,13 +1063,13 @@ void UStruct::Link(FArchive& Ar, bool bRelinkExistingProperties)
 			|| bShouldHandleFinishDestroy)
 		{
 			// things in a struct that need a destructor will still be in here, even though in many cases they will also be destroyed by a native destructor on the whole struct
-			DestructorLinkBuilder.Append(*Property);
+			DestructorLinkBuilder.AppendNoTerminate(*Property);
 		}
 
 		// Link references to properties that require their values to be initialized and/or copied from CDO post-construction. Note that this includes all non-native-class-owned properties.
 		if (OwnerClass && (!bOwnedByNativeClass || (Property->HasAnyPropertyFlags(CPF_Config) && !OwnerClass->HasAnyClassFlags(CLASS_PerObjectConfig))))
 		{
-			PostConstructLinkBuilder.Append(*Property);
+			PostConstructLinkBuilder.AppendNoTerminate(*Property);
 		}
 
 #if WITH_EDITORONLY_DATA
@@ -1078,7 +1078,7 @@ void UStruct::Link(FArchive& Ar, bool bRelinkExistingProperties)
 		bHasAssetRegistrySearchableProperties |= Property->HasAnyPropertyFlags(CPF_AssetRegistrySearchable);
 #endif
 
-		PropertyLinkBuilder.Append(*Property);
+		PropertyLinkBuilder.AppendNoTerminate(*Property);
 	}
 
 	PropertyLinkBuilder.NullTerminate();
