[AutoRTFM] Migrate to UE::TScopeLock instead of FTransactionallySafeScopeLock.

These are functionally identical.

diff --git a/Engine/Source/Runtime/CoreUObject/Public/UObject/UObjectAnnotation.h b/Engine/Source/Runtime/CoreUObject/Public/UObject/UObjectAnnotation.h
--- a/Engine/Source/Runtime/CoreUObject/Public/UObject/UObjectAnnotation.h
+++ b/Engine/Source/Runtime/CoreUObject/Public/UObject/UObjectAnnotation.h
@@ -8,6 +8,8 @@
 
 #include "AutoRTFM/AutoRTFM.h"
 #include "UObject/UObjectArray.h"
+#include "Misc/ScopeLock.h"
+#include "Misc/TransactionallySafeCriticalSection.h"
 #include "Misc/TransactionallySafeRWLock.h"
 #include "Misc/TransactionallySafeRWScopeLock.h"
 
@@ -40,7 +42,7 @@ public:
 #if !(UE_BUILD_SHIPPING || UE_BUILD_TEST)
 		if (!bAutoRemove)
 		{
-			FTransactionallySafeScopeLock AnnotationMapLock(&AnnotationMapCritical);
+			UE::TScopeLock AnnotationMapLock(AnnotationMapCritical);
 			// in this case we are only verifying that the external assurances of removal are met
 			check(!AnnotationMap.Find(Object));
 		}
@@ -87,7 +89,7 @@ private:
 		}
 		else
 		{
-			FTransactionallySafeScopeLock AnnotationMapLock(&AnnotationMapCritical);
+			UE::TScopeLock AnnotationMapLock(AnnotationMapCritical);
 			const bool bWasEmpty = (AnnotationMap.Num() == 0);
 
 			AnnotationMap.Add(Object, LocalAnnotation);
@@ -145,7 +147,7 @@ public:
 		// Avoid holding the lock while we call GUObjectArray.RemoveUObjectDeleteListener as it could deadlock
 		TAnnotation Result;
 		{
-			FTransactionallySafeScopeLock AnnotationMapLock(&AnnotationMapCritical);
+			UE::TScopeLock AnnotationMapLock(AnnotationMapCritical);
 			SetAnnotationCacheKeyAndValue(Object, TAnnotation());
 			bHadElements = (AnnotationMap.Num() > 0);
 			AnnotationMap.RemoveAndCopyValue(AnnotationCacheKey, Result);
@@ -178,7 +180,7 @@ public:
 			bool bIsNowEmpty = false;
 			// Avoid holding the lock while we call GUObjectArray.RemoveUObjectDeleteListener as it could deadlock
 			{
-				FTransactionallySafeScopeLock AnnotationMapLock(&AnnotationMapCritical);
+				UE::TScopeLock AnnotationMapLock(AnnotationMapCritical);
 				SetAnnotationCacheKeyAndValue(Object, TAnnotation());
 				bHadElements = (AnnotationMap.Num() > 0);
 				AnnotationMap.Remove(AnnotationCacheKey);
@@ -206,7 +208,7 @@ public:
 
 		// Avoid holding the lock while we call GUObjectArray.RemoveUObjectDeleteListener as it could deadlock
 		{
-			FTransactionallySafeScopeLock AnnotationMapLock(&AnnotationMapCritical);
+			UE::TScopeLock AnnotationMapLock(AnnotationMapCritical);
 			SetAnnotationCacheKeyAndValue(nullptr, TAnnotation());
 			bHadElements = (AnnotationMap.Num() > 0);
 			AnnotationMap.Empty();
@@ -232,7 +234,7 @@ public:
 	{
 		check(Object);
 
-		FTransactionallySafeScopeLock AnnotationMapLock(&AnnotationMapCritical);
+		UE::TScopeLock AnnotationMapLock(AnnotationMapCritical);
 		if (AutoRTFM::IsTransactional() || Object != AnnotationCacheKey)
 		{			
 			TAnnotation* Entry = AnnotationMap.Find(Object);
@@ -264,7 +266,7 @@ public:
 	 */
 	void Reserve(int32 ExpectedNumElements)
 	{
-		FTransactionallySafeScopeLock AnnotationMapLock(&AnnotationMapCritical);
+		UE::TScopeLock AnnotationMapLock(AnnotationMapCritical);
 		AnnotationMap.Empty(ExpectedNumElements);
 	}
 
@@ -363,7 +365,7 @@ public:
 	 */
 	UObject* Find(const TAnnotation& Annotation)
 	{
-		FTransactionallySafeScopeLock InverseAnnotationMapLock(&InverseAnnotationMapCritical);
+		UE::TScopeLock InverseAnnotationMapLock(InverseAnnotationMapCritical);
 		checkSlow(!Annotation.IsDefault()); // it is not legal to search for the default annotation
 		return (UObject*)InverseAnnotationMap.FindRef(Annotation);
 	}
@@ -372,7 +374,7 @@ private:
 	template<typename T> 
 	void AddAnnotationInternal(const UObjectBase* Object, T&& Annotation)
 	{
-		FTransactionallySafeScopeLock InverseAnnotationMapLock(&InverseAnnotationMapCritical);
+		UE::TScopeLock InverseAnnotationMapLock(InverseAnnotationMapCritical);
 		if (Annotation.IsDefault())
 		{
 			RemoveAnnotation(Object); // adding the default annotation is the same as removing an annotation
@@ -417,7 +419,7 @@ public:
 	 */
 	void RemoveAnnotation(const UObjectBase *Object)
 	{
-		FTransactionallySafeScopeLock InverseAnnotationMapLock(&InverseAnnotationMapCritical);
+		UE::TScopeLock InverseAnnotationMapLock(InverseAnnotationMapCritical);
 		TAnnotation Annotation = this->GetAndRemoveAnnotation(Object);
 		if (Annotation.IsDefault())
 		{
@@ -436,7 +438,7 @@ public:
 	 */
 	void RemoveAllAnnotations()
 	{
-		FTransactionallySafeScopeLock InverseAnnotationMapLock(&InverseAnnotationMapCritical);
+		UE::TScopeLock InverseAnnotationMapLock(InverseAnnotationMapCritical);
 		Super::RemoveAllAnnotations();
 		InverseAnnotationMap.Empty();
 	}
