[AutoRTFM] Migrate to UE::TScopeLock instead of FTransactionallySafeScopeLock.

These are functionally identical.

diff --git a/Engine/Source/Runtime/CoreUObject/Private/UObject/SparseDelegate.cpp b/Engine/Source/Runtime/CoreUObject/Private/UObject/SparseDelegate.cpp
--- a/Engine/Source/Runtime/CoreUObject/Private/UObject/SparseDelegate.cpp
+++ b/Engine/Source/Runtime/CoreUObject/Private/UObject/SparseDelegate.cpp
@@ -4,7 +4,7 @@
 #include "UObject/Class.h"
 #include "UObject/Package.h"
 #include "HAL/IConsoleManager.h"
-#include "Misc/TransactionallySafeScopeLock.h"
+#include "Misc/ScopeLock.h"
 
 FSparseDelegateStorage::FObjectListener FSparseDelegateStorage::SparseDelegateObjectListener;
 FTransactionallySafeCriticalSection FSparseDelegateStorage::SparseDelegateMapCritical;
@@ -22,7 +22,7 @@ FSparseDelegateStorage::FObjectListener::~FObjectListener()
 
 void FSparseDelegateStorage::FObjectListener::NotifyUObjectDeleted(const UObjectBase* Object, int32 Index)
 {
-	FTransactionallySafeScopeLock SparseDelegateMapLock(&FSparseDelegateStorage::SparseDelegateMapCritical);
+	UE::TScopeLock SparseDelegateMapLock(FSparseDelegateStorage::SparseDelegateMapCritical);
 	FSparseDelegateStorage::SparseDelegates.Remove(Object);
 	if (FSparseDelegateStorage::SparseDelegates.Num() == 0)
 	{
@@ -32,7 +32,7 @@ void FSparseDelegateStorage::FObjectListener::NotifyUObjectDeleted(const UObject
 
 void FSparseDelegateStorage::FObjectListener::OnUObjectArrayShutdown()
 {
-	FTransactionallySafeScopeLock SparseDelegateMapLock(&FSparseDelegateStorage::SparseDelegateMapCritical);
+	UE::TScopeLock SparseDelegateMapLock(FSparseDelegateStorage::SparseDelegateMapCritical);
 	FSparseDelegateStorage::SparseDelegates.Empty();
 	DisableListener();
 }
@@ -84,7 +84,7 @@ UObject* FSparseDelegateStorage::ResolveSparseOwner(const FSparseDelegate& Spars
 
 FMulticastScriptDelegate* FSparseDelegateStorage::GetMulticastDelegate(const UObject* DelegateOwner, const FName DelegateName)
 {
-	FTransactionallySafeScopeLock SparseDelegateMapLock(&SparseDelegateMapCritical);
+	UE::TScopeLock SparseDelegateMapLock(SparseDelegateMapCritical);
 
 	if (FSparseDelegateMap* DelegateMap = SparseDelegates.Find(DelegateOwner))
 	{
@@ -98,7 +98,7 @@ FMulticastScriptDelegate* FSparseDelegateStorage::GetMulticastDelegate(const UOb
 
 TSharedPtr<FMulticastScriptDelegate> FSparseDelegateStorage::GetSharedMulticastDelegate(const UObject* DelegateOwner, const FName DelegateName)
 {
-	FTransactionallySafeScopeLock SparseDelegateMapLock(&SparseDelegateMapCritical);
+	UE::TScopeLock SparseDelegateMapLock(SparseDelegateMapCritical);
 
 	TSharedPtr<FMulticastScriptDelegate> Result;
 	if (FSparseDelegateMap* DelegateMap = SparseDelegates.Find(DelegateOwner))
@@ -113,7 +113,7 @@ TSharedPtr<FMulticastScriptDelegate> FSparseDelegateStorage::GetSharedMulticastD
 
 void FSparseDelegateStorage::SetMulticastDelegate(const UObject* DelegateOwner, const FName DelegateName, FMulticastScriptDelegate Delegate)
 {
-	FTransactionallySafeScopeLock SparseDelegateMapLock(&SparseDelegateMapCritical);
+	UE::TScopeLock SparseDelegateMapLock(SparseDelegateMapCritical);
 
 	if (SparseDelegates.Num() == 0)
 	{
@@ -136,7 +136,7 @@ bool FSparseDelegateStorage::Add(const UObject* DelegateOwner, const FName Deleg
 	bool bDelegateWasBound = false;
 	if (Delegate.IsBound())
 	{
-		FTransactionallySafeScopeLock SparseDelegateMapLock(&SparseDelegateMapCritical);
+		UE::TScopeLock SparseDelegateMapLock(SparseDelegateMapCritical);
 
 		if (SparseDelegates.Num() == 0)
 		{
@@ -162,7 +162,7 @@ bool FSparseDelegateStorage::AddUnique(const UObject* DelegateOwner, const FName
 	bool bDelegateWasBound = false;
 	if (Delegate.IsBound())
 	{
-		FTransactionallySafeScopeLock SparseDelegateMapLock(&SparseDelegateMapCritical);
+		UE::TScopeLock SparseDelegateMapLock(SparseDelegateMapCritical);
 
 		if (SparseDelegates.Num() == 0)
 		{
@@ -186,7 +186,7 @@ bool FSparseDelegateStorage::AddUnique(const UObject* DelegateOwner, const FName
 bool FSparseDelegateStorage::Contains(const UObject* DelegateOwner, const FName DelegateName, const FScriptDelegate& Delegate)
 {
 	bool bContainsDelegate = false;
-	FTransactionallySafeScopeLock SparseDelegateMapLock(&SparseDelegateMapCritical);
+	UE::TScopeLock SparseDelegateMapLock(SparseDelegateMapCritical);
 
 	if (FSparseDelegateMap* DelegateMap = SparseDelegates.Find(DelegateOwner))
 	{
@@ -205,7 +205,7 @@ bool FSparseDelegateStorage::Contains(const UObject* DelegateOwner, const FName
 bool FSparseDelegateStorage::Contains(const UObject* DelegateOwner, const FName DelegateName, const UObject* InObject, FName InFunctionName)
 {
 	bool bContainsDelegate = false;
-	FTransactionallySafeScopeLock SparseDelegateMapLock(&SparseDelegateMapCritical);
+	UE::TScopeLock SparseDelegateMapLock(SparseDelegateMapCritical);
 
 	if (FSparseDelegateMap* DelegateMap = SparseDelegates.Find(DelegateOwner))
 	{
@@ -223,7 +223,7 @@ bool FSparseDelegateStorage::Contains(const UObject* DelegateOwner, const FName
 bool FSparseDelegateStorage::Remove(const UObject* DelegateOwner, const FName DelegateName, const FScriptDelegate& Delegate)
 {
 	bool bSparseDelegateBound = false;
-	FTransactionallySafeScopeLock SparseDelegateMapLock(&SparseDelegateMapCritical);
+	UE::TScopeLock SparseDelegateMapLock(SparseDelegateMapCritical);
 
 	if (FSparseDelegateMap* DelegateMap = SparseDelegates.Find(DelegateOwner))
 	{
@@ -259,7 +259,7 @@ bool FSparseDelegateStorage::Remove(const UObject* DelegateOwner, const FName De
 bool FSparseDelegateStorage::Remove(const UObject* DelegateOwner, const FName DelegateName, const UObject* InObject, FName InFunctionName)
 {
 	bool bSparseDelegateBound = false;
-	FTransactionallySafeScopeLock SparseDelegateMapLock(&SparseDelegateMapCritical);
+	UE::TScopeLock SparseDelegateMapLock(SparseDelegateMapCritical);
 
 	if (FSparseDelegateMap* DelegateMap = SparseDelegates.Find(DelegateOwner))
 	{
@@ -295,7 +295,7 @@ bool FSparseDelegateStorage::Remove(const UObject* DelegateOwner, const FName De
 bool FSparseDelegateStorage::RemoveAll(const UObject* DelegateOwner, const FName DelegateName, const UObject* UserObject)
 {
 	bool bSparseDelegateBound = false;
-	FTransactionallySafeScopeLock SparseDelegateMapLock(&SparseDelegateMapCritical);
+	UE::TScopeLock SparseDelegateMapLock(SparseDelegateMapCritical);
 
 	if (FSparseDelegateMap* DelegateMap = SparseDelegates.Find(DelegateOwner))
 	{
@@ -329,7 +329,7 @@ bool FSparseDelegateStorage::RemoveAll(const UObject* DelegateOwner, const FName
 
 void FSparseDelegateStorage::Clear(const UObject* DelegateOwner, const FName DelegateName)
 {
-	FTransactionallySafeScopeLock SparseDelegateMapLock(&SparseDelegateMapCritical);
+	UE::TScopeLock SparseDelegateMapLock(SparseDelegateMapCritical);
 	if (FSparseDelegateMap* DelegateMap = SparseDelegates.Find(DelegateOwner))
 	{
 		if (TSharedPtr<FMulticastScriptDelegate>* MulticastDelegatePtr = DelegateMap->Find(DelegateName))
@@ -416,7 +416,7 @@ void FSparseDelegateStorage::SparseDelegateReport(const TArray<FString>& Args, U
 	uint32 BoundDelegates = 0;
 
 	{
-		FTransactionallySafeScopeLock SparseDelegateMapLock(&SparseDelegateMapCritical);
+		UE::TScopeLock SparseDelegateMapLock(SparseDelegateMapCritical);
 		for (const TPair<const UObjectBase*, FSparseDelegateMap>& ObjectAnnotation : SparseDelegates)
 		{
 			const UObject* Object = static_cast<const UObject*>(ObjectAnnotation.Key);
diff --git a/Engine/Source/Runtime/CoreUObject/Public/UObject/SparseDelegate.h b/Engine/Source/Runtime/CoreUObject/Public/UObject/SparseDelegate.h
--- a/Engine/Source/Runtime/CoreUObject/Public/UObject/SparseDelegate.h
+++ b/Engine/Source/Runtime/CoreUObject/Public/UObject/SparseDelegate.h
@@ -6,11 +6,9 @@
 #include "Containers/Map.h"
 #include "Containers/SparseArray.h"
 #include "Delegates/Delegate.h"
-#include "HAL/CriticalSection.h"
 #include "HAL/PlatformMath.h"
 #include "Misc/AssertionMacros.h"
-#include "Misc/ScopeLock.h"
-#include "Misc/TransactionallySafeScopeLock.h"
+#include "Misc/TransactionallySafeCriticalSection.h"
 #include "Templates/SharedPointer.h"
 #include "Templates/UnrealTemplate.h"
 #include "UObject/NameTypes.h"
