[StructUtils] Created TTypeBitSetBase by extracting moregeneric bits of TypedStructBitSet out into a new template class. Done something similar to FBitSetContainer and extracted read-only bits to FConstBitSetContainer.

Also the FStructTracker has been massaged a bit to make type accessing and registering more readable, and added a const getter.

#rnx

diff --git a/Engine/Source/Runtime/CoreUObject/Private/StructUtils/StructTypeBitSet.cpp b/Engine/Source/Runtime/CoreUObject/Private/StructUtils/StructTypeBitSet.cpp
--- a/Engine/Source/Runtime/CoreUObject/Private/StructUtils/StructTypeBitSet.cpp
+++ b/Engine/Source/Runtime/CoreUObject/Private/StructUtils/StructTypeBitSet.cpp
@@ -69,42 +69,58 @@ int32 FStructTracker::FindOrAddStructTypeIndex(const UStruct& InStructType)
 {
 	// Get existing index...
 	const uint32 Hash = PointerHash(&InStructType);
-	FSetElementId ElementId = StructTypeToIndexSet.FindIdByHash(Hash, Hash);
+	const FSetElementId ElementId = StructTypeToIndexSet.FindIdByHash(Hash, Hash);
 
-	if (!ElementId.IsValidId())
+	if (LIKELY(ElementId.IsValidId()))
 	{
-		checkSlow(InStructType.IsChildOf(BaseStructGetter()));
+		return ElementId.AsInteger();
+	}
+	return Register(InStructType);
+}
+
+int32 FStructTracker::Register(const UStruct& InStructType)
+{
+	const uint32 Hash = PointerHash(&InStructType);
+	checkSlow(InStructType.IsChildOf(BaseStructGetter()));
 
-		// .. or create a new one
-		ElementId = StructTypeToIndexSet.AddByHash(Hash, Hash);
-		checkSlow(ElementId.IsValidId());
+	// .. or create a new one
+	const FSetElementId ElementId = StructTypeToIndexSet.AddByHash(Hash, Hash);
+	checkSlow(ElementId.IsValidId());
 
-		const int32 NewIndex = ElementId.AsInteger();
+	const int32 NewIndex = ElementId.AsInteger();
 
-		checkSlow(StructTypesList.Num() == NewIndex);
-		StructTypesList.Add(&InStructType);
+	checkSlow(StructTypesList.Num() == NewIndex);
+	StructTypesList.Add(&InStructType);
 
-		// first-time SerializationHash initialization
-		if (SerializationHash == 0)
-		{
-			ensure(StructTypesList.Num() == 1);
-			UStruct* BaseType = BaseStructGetter();
-			CA_ASSUME(BaseType);
-			SerializationHash = GetTypeHash(BaseType->GetFullName());
-		}
-		SerializationHash = HashCombine(SerializationHash, GetTypeHash(InStructType.GetFullName()));
+	// first-time SerializationHash initialization
+	if (SerializationHash == 0)
+	{
+		ensure(StructTypesList.Num() == 1);
+		UStruct* BaseType = BaseStructGetter();
+		CA_ASSUME(BaseType);
+		SerializationHash = GetTypeHash(BaseType->GetFullName());
+	}
+	SerializationHash = HashCombine(SerializationHash, GetTypeHash(InStructType.GetFullName()));
 
-		// it's worth pointing out that we're registering a given tracker for all the hashes created along the way.
-		// This will help with loading bitsets from serialized data.
-		FBitSetSerializationDictionary::RegisterHash(SerializationHash, *this);
+	// it's worth pointing out that we're registering a given tracker for all the hashes created along the way.
+	// This will help with loading bitsets from serialized data.
+	FBitSetSerializationDictionary::RegisterHash(SerializationHash, *this);
 
 #if WITH_STRUCTUTILS_DEBUG
-		DebugStructTypeNamesList.Add(InStructType.GetFName());
-		ensure(StructTypeToIndexSet.Num() == DebugStructTypeNamesList.Num());
+	DebugStructTypeNamesList.Add(InStructType.GetFName());
+	ensure(StructTypeToIndexSet.Num() == DebugStructTypeNamesList.Num());
 #endif // WITH_STRUCTUTILS_DEBUG
-	}
 
-	return ElementId.AsInteger();
+	return NewIndex;
+}
+
+
+int32 FStructTracker::FindStructTypeIndex(const UStruct& InStructType) const
+{
+	const uint32 Hash = PointerHash(&InStructType);
+	const FSetElementId ElementId = StructTypeToIndexSet.FindIdByHash(Hash, Hash);
+
+	return ElementId.IsValidId() ? ElementId.AsInteger() : INDEX_NONE;
 }
 
 void FStructTracker::Serialize(FArchive& Ar, FStructTypeBitSet::FBitSetContainer& StructTypesBitArray)
diff --git a/Engine/Source/Runtime/CoreUObject/Public/StructUtils/StructTypeBitSet.h b/Engine/Source/Runtime/CoreUObject/Public/StructUtils/StructTypeBitSet.h
--- a/Engine/Source/Runtime/CoreUObject/Public/StructUtils/StructTypeBitSet.h
+++ b/Engine/Source/Runtime/CoreUObject/Public/StructUtils/StructTypeBitSet.h
@@ -12,22 +12,10 @@
 
 namespace FStructTypeBitSet
 {
-struct FBitSetContainer : TBitArray<>
+
+struct FConstBitSetContainer : TBitArray<>
 {
-	FBitSetContainer() = default;
-	FBitSetContainer(const TBitArray<>& Source) : TBitArray<>(Source)
-	{}
+	FConstBitSetContainer() = default;
 
-	FBitSetContainer& operator=(const TBitArray<>& Other)
-	{
-		*((TBitArray<>*)this) = Other;
-		return *this;
-	}
-
-	void SetAll(const bool bValue, const int32 Count)
-	{
-		Init(bValue, Count);
-	}
-
 	FORCEINLINE bool HasAll(const TBitArray<>& Other) const
 	{
@@ -83,21 +71,7 @@
 		return !Iterator;
 	}
 
-	FORCEINLINE void operator-=(const TBitArray<>& Other)
+	FORCEINLINE friend uint32 GetTypeHash(const FConstBitSetContainer& Instance)
 	{
-		FWordIterator ThisIterator(*this);
-		FConstWordIterator OtherIterator(Other);
-
-		while (ThisIterator && OtherIterator)
-		{
-			ThisIterator.SetWord(ThisIterator.GetWord() & ~OtherIterator.GetWord());
-
-			++ThisIterator;
-			++OtherIterator;
-		}
-	}
-
-	FORCEINLINE friend uint32 GetTypeHash(const FBitSetContainer& Instance)
-	{
 		FConstWordIterator Iterator(Instance);
 		uint32 Hash = 0;
@@ -119,6 +93,72 @@
 		return Hash;
 	}
 
+	bool Contains(const int32 Index) const
+	{
+		check(Index >= 0);
+		return (Index < Num()) && (*this)[Index];
+	}
+
+	int32 CountStoredTypes() const
+	{
+		return CountSetBits();
+	}
+
+protected:
+	// only for use by the FBitSetContainer
+	FConstBitSetContainer(const TBitArray<>& Source) : TBitArray<>(Source)
+	{
+	}
+	FConstBitSetContainer(TBitArray<>&& Source) : TBitArray<>(Forward<TBitArray<>>(Source))
+	{
+	}
+};
+
+struct FBitSetContainer : FConstBitSetContainer
+{
+	using FConstBitSetContainer::FConstBitSetContainer;
+
+	FBitSetContainer(const TBitArray<>& Source)
+		: FConstBitSetContainer(Source)
+	{
+	}
+	FBitSetContainer(TBitArray<>&& Source)
+		: FConstBitSetContainer(Forward<TBitArray<>>(Source))
+	{
+	}
+
+	FBitSetContainer& operator=(const TBitArray<>& Other)
+	{
+		static_assert(sizeof(FBitSetContainer) == sizeof(TBitArray<>)
+			, "This operator requires FBitSetContainer and TBitArray<> sizes to match (i.e. FBitSetContainer cannot have member variables nor virtual funtions)");
+		*((TBitArray<>*)this) = Other;
+		return *this;
+	}
+
+	void SetAll(const bool bValue, const int32 Count)
+	{
+		Init(bValue, Count);
+	}
+	
+	FORCEINLINE void operator-=(const TBitArray<>& Other)
+	{
+		FWordIterator ThisIterator(*this);
+		FConstWordIterator OtherIterator(Other);
+
+		while (ThisIterator && OtherIterator)
+		{
+			ThisIterator.SetWord(ThisIterator.GetWord() & ~OtherIterator.GetWord());
+
+			++ThisIterator;
+			++OtherIterator;
+		}
+	}
+
+	FORCEINLINE friend uint32 GetTypeHash(const FBitSetContainer& Instance)
+	{
+		return GetTypeHash(*static_cast<const FConstBitSetContainer*>(&Instance));
+	}
+
 	void AddAtIndex(const int32 Index)
 	{
 		PadToNum(Index + 1, false);
@@ -135,12 +175,6 @@
 		// else, it's already not present
 	}
 
-	bool Contains(const int32 Index) const
-	{
-		check(Index >= 0);
-		return (Index < Num()) && (*this)[Index];
-	}
-
 protected:
 	/**
 	 * duplication of TBitArray::SetBitNoCheck needed since it's private but it's the performant way of setting bits
@@ -192,8 +226,21 @@
 	explicit FStructTracker(const FBaseStructGetter& InBaseStructGetter);
 	~FStructTracker();
 
+	/** @return the index for the given type, or INDEX_NONE if it hasn't been registered with this struct tracker */
+	int32 FindStructTypeIndex(const UStruct& InStructType) const;
+
+	/**
+	 * Fetches the internal index representing given UStruct instance. If it hasn't been registered yet, it will be automatically added
+	 * @return the index for the given type
+	 */
 	int32 FindOrAddStructTypeIndex(const UStruct& InStructType);
 
+	/**
+	 * Registers given UStruct with the struct tracker.
+	 * @return the index for the registered type
+	 */
+	int32 Register(const UStruct& InStructType);
+
 	const UStruct* GetStructType(const int32 StructTypeIndex) const
 	{
 		return StructTypesList.IsValidIndex(StructTypeIndex) ? StructTypesList[StructTypeIndex].Get() : nullptr;
@@ -233,53 +280,11 @@
 	const FBaseStructGetter BaseStructGetter;
 };
 
-
-template<typename TBaseStruct, typename TStructTrackerWrapper, typename TUStructType = UScriptStruct>
-struct TStructTypeBitSet
+template<typename TImplementation, typename TBaseStruct, typename TStructType, bool bTestInheritanceAtRuntime=WITH_STRUCTUTILS_DEBUG>
+struct TTypeBitSetBase
 {
-	using FStructTrackerWrapper = TStructTrackerWrapper;
+	using FUStructType = TStructType;
+	using FBaseStruct = TBaseStruct;
 
-	TStructTypeBitSet() = default;
-
-	explicit TStructTypeBitSet(const TUStructType& StructType)
-	{
-		Add(StructType);
-	}
-
-	explicit TStructTypeBitSet(std::initializer_list<const TUStructType*> InitList)
-	{
-		for (const TUStructType* StructType : InitList)
-		{
-			if (StructType)
-			{
-				Add(*StructType);
-			}
-		}
-	}
-
-	explicit TStructTypeBitSet(TConstArrayView<const TUStructType*> InitList)
-	{
-		for (const TUStructType* StructType : InitList)
-		{
-			if (StructType)
-			{
-				Add(*StructType);
-			}
-		}
-	}
-
-	/** This flavor of constructor is only available for UScriptStructs */
-	template<typename T = TBaseStruct, typename = typename TEnableIf<!TIsDerivedFrom<T, UObject>::IsDerived>::Type>
-	explicit TStructTypeBitSet(TConstArrayView<FInstancedStruct> InitList)
-	{
-		for (const FInstancedStruct& StructInstance : InitList)
-		{
-			if (StructInstance.GetScriptStruct())
-			{
-				Add(*StructInstance.GetScriptStruct());
-			}
-		}
-	}
-
 	/** Retrieves indices of bits of given value. */
 	struct FIndexIterator
@@ -320,14 +325,303 @@
 		const bool bValueToCheck;
 	};
 
+	FIndexIterator GetIndexIterator(const bool bValueToCheck = true) const
+	{
+		return FIndexIterator(StructTypesBitArray, bValueToCheck);
+	}
+
+	TImplementation& GetImplementation()
+	{
+		return *reinterpret_cast<TImplementation*>(this);
+	}
+
+	const TImplementation& GetImplementation() const 
+	{
+		return *reinterpret_cast<const TImplementation*>(this);
+	}
+
+	FORCEINLINE_DEBUGGABLE static const UStruct* GetBaseUStruct()
+	{
+		if constexpr (bTestInheritanceAtRuntime)
+		{
+			static const UStruct* Instance = UE::StructUtils::GetAsUStruct<FBaseStruct>();
+			return Instance;
+		}
+		else
+		{
+			return nullptr;
+		}
+	}
+	
+	void SetAll(const bool bValue = true)
+	{
+		StructTypesBitArray.SetAll(bValue, GetImplementation().GetStructTracker().Num());
+	}
+
+	void Add(const FUStructType& InStructType)
+	{
+#if WITH_STRUCTUTILS_DEBUG
+		if constexpr (bTestInheritanceAtRuntime)
+		{
+			ensureMsgf(InStructType.IsChildOf(GetBaseUStruct())
+				, TEXT("Registering '%s' with FStructTracker while it doesn't derive from the expected struct type %s")
+				, *InStructType.GetPathName(), *GetBaseUStruct()->GetName());
+		}
+#endif // WITH_STRUCTUTILS_DEBUG
+
+		const int32 StructTypeIndex = GetImplementation().GetStructTracker().FindOrAddStructTypeIndex(InStructType);
+		StructTypesBitArray.AddAtIndex(StructTypeIndex);
+	}
+
+	void Remove(const FUStructType& InStructType)
+	{
+#if WITH_STRUCTUTILS_DEBUG
+		if constexpr (bTestInheritanceAtRuntime)
+		{
+			ensureMsgf(InStructType.IsChildOf(GetBaseUStruct())
+				, TEXT("Registering '%s' with FStructTracker while it doesn't derive from the expected struct type %s")
+				, *InStructType.GetPathName(), *GetBaseUStruct()->GetName());
+		}
+#endif // WITH_STRUCTUTILS_DEBUG
+
+		const int32 StructTypeIndex = GetImplementation().GetStructTracker().FindOrAddStructTypeIndex(InStructType);
+		StructTypesBitArray.RemoveAtIndex(StructTypeIndex);
+	}
+
+	void Reset()
+	{
+		StructTypesBitArray.Reset();
+	}
+
+	bool Contains(const FUStructType& InStructType) const
+	{
+#if WITH_STRUCTUTILS_DEBUG
+		if constexpr (bTestInheritanceAtRuntime)
+		{
+			ensureMsgf(InStructType.IsChildOf(GetBaseUStruct())
+				, TEXT("Registering '%s' with FStructTracker while it doesn't derive from the expected struct type %s")
+				, *InStructType.GetPathName(), *GetBaseUStruct()->GetName());
+		}
+#endif // WITH_STRUCTUTILS_DEBUG
+
+		const int32 StructTypeIndex = GetImplementation().GetStructTracker().FindStructTypeIndex(InStructType);
+		return StructTypeIndex != INDEX_NONE && StructTypesBitArray.Contains(StructTypeIndex);
+	}
+
+	FORCEINLINE TImplementation operator&(const TImplementation& Other) const
+	{
+		return TImplementation(TBitArray<>::BitwiseAND(StructTypesBitArray, Other.StructTypesBitArray, EBitwiseOperatorFlags::MinSize));
+	}
+
+	FORCEINLINE TImplementation operator|(const TImplementation& Other) const
+	{
+		return TImplementation(TBitArray<>::BitwiseOR(StructTypesBitArray, Other.StructTypesBitArray, EBitwiseOperatorFlags::MaxSize));
+	}
+
+	FORCEINLINE TImplementation GetOverlap(const TImplementation& Other) const
+	{
+		return GetImplementation() & Other;
+	}
+
+	FORCEINLINE bool IsEquivalent(const TImplementation& Other) const
+	{
+		return StructTypesBitArray.CompareSetBits(Other.StructTypesBitArray, /*bMissingBitValue=*/false);
+	}
+
+	FORCEINLINE bool HasAll(const TImplementation& Other) const
+	{
+		return StructTypesBitArray.HasAll(Other.StructTypesBitArray);
+	}
+
+	FORCEINLINE bool HasAny(const TImplementation& Other) const
+	{
+		return StructTypesBitArray.HasAny(Other.StructTypesBitArray);
+	}
+
+	FORCEINLINE bool HasNone(const TImplementation& Other) const
+	{
+		return !StructTypesBitArray.HasAny(Other.StructTypesBitArray);
+	}
+
+	bool IsEmpty() const 
+	{ 
+		return StructTypesBitArray.IsEmpty();
+	}
+
+	FORCEINLINE bool IsBitSet(const int32 BitIndex) const
+	{
+		return StructTypesBitArray.Contains(BitIndex);
+	}
+
+	FORCEINLINE void operator+=(const TImplementation& Other)
+	{
+		StructTypesBitArray = TBitArray<>::BitwiseOR(StructTypesBitArray, Other.StructTypesBitArray, EBitwiseOperatorFlags::MaxSize);
+	}
+
+	FORCEINLINE void operator-=(const TImplementation& Other)
+	{
+		StructTypesBitArray -= Other.StructTypesBitArray;
+	}
+
+	FORCEINLINE TImplementation operator+(const FUStructType& NewElement) const
+	{
+		TImplementation Result = GetImplementation();
+		Result.Add(NewElement);
+		return MoveTemp(Result);
+	}
+
+	FORCEINLINE TImplementation operator-(const FUStructType& NewElement) const
+	{
+		TImplementation Result = GetImplementation();
+		Result.Remove(NewElement);
+		return MoveTemp(Result);
+	}
+
+	int32 CountStoredTypes() const
+	{
+		return StructTypesBitArray.CountSetBits();
+	}
+
+	/**
+	 * note that this function is slow(ish) due to the FStructTracker utilizing WeakObjectPtrs to store types. 
+	 * @todo To be improved.
+	 */
+	template<typename TOutStructType, typename Allocator>
+	void ExportTypes(TArray<const TOutStructType*, Allocator>& OutTypes) const
+	{
+		TBitArray<>::FConstIterator It(StructTypesBitArray);
+		while (It)
+		{
+			if (It.GetValue())
+			{
+				OutTypes.Add(Cast<TOutStructType>(GetImplementation().GetStructTracker().GetStructType(It.GetIndex())));
+			}
+			++It;
+		}
+	}
+
+#if WITH_STRUCTUTILS_DEBUG
+	void DebugGetStringDesc(FOutputDevice& Ar) const
+	{
+		for (int32 Index = 0; Index < StructTypesBitArray.Num(); ++Index)
+		{
+			if (StructTypesBitArray[Index])
+			{
+				Ar.Logf(TEXT("%s, "), *GetImplementation().GetStructTracker().DebugGetStructTypeName(Index).ToString());
+			}
+		}
+	}
+
+	void DebugGetIndividualNames(TArray<FName>& OutFNames) const
+	{
+		for (int32 Index = 0; Index < StructTypesBitArray.Num(); ++Index)
+		{
+			if (StructTypesBitArray[Index])
+			{
+				OutFNames.Add(GetImplementation().GetStructTracker().DebugGetStructTypeName(Index));
+			}
+		}
+	}
+#endif // WITH_STRUCTUTILS_DEBUG
+
+	SIZE_T GetAllocatedSize() const
+	{
+		return StructTypesBitArray.GetAllocatedSize();
+	}
+
+protected:
+	/** 
+	 * A private constructor for a creating an instance straight from TBitArrays. 
+	 * @Note that this constructor needs to remain private to ensure consistency of stored values with data tracked 
+	 * by the TStructTrackerWrapper::StructTracker
+	 */
+	TTypeBitSetBase(TBitArray<>&& Source)
+		: StructTypesBitArray(Forward<TBitArray<>>(Source))
+	{
+	}
+
+	TTypeBitSetBase() = default;
+
+#if WITH_STRUCTUTILS_DEBUG
+	// unittesting purposes only
+	const TBitArray<>& DebugGetStructTypesBitArray() const { return StructTypesBitArray; }
+	TBitArray<>& DebugGetMutableStructTypesBitArray() { return StructTypesBitArray; }
+#endif // WITH_STRUCTUTILS_DEBUG
+
+	FStructTypeBitSet::FBitSetContainer StructTypesBitArray;
+};
+
+template<typename TBaseStruct, typename TStructTrackerWrapper, typename TUStructType = UScriptStruct>
+struct TStructTypeBitSet : TTypeBitSetBase<TStructTypeBitSet<TBaseStruct, TStructTrackerWrapper, TUStructType>, TBaseStruct, TUStructType>
+{
+	using Super = TTypeBitSetBase<TStructTypeBitSet, TBaseStruct, TUStructType>;
+	using FStructTrackerWrapper = TStructTrackerWrapper;
+	using FUStructType =  TUStructType;
+	using FBaseStruct = TBaseStruct;
+
+	friend Super;
+	using Super::StructTypesBitArray;
+	using Super::GetBaseUStruct;
+	using Super::Add;
+	using Super::Remove;
+	using Super::Contains;
+	using Super::operator+=;
+	using Super::operator-=;
+	using Super::operator+;
+	using Super::operator-;
+	using Super::ExportTypes;
+
+	TStructTypeBitSet() = default;
+
+	explicit TStructTypeBitSet(const FUStructType& StructType)
+	{
+		Add(StructType);
+	}
+
+	explicit TStructTypeBitSet(std::initializer_list<const FUStructType*> InitList)
+	{
+		for (const FUStructType* StructType : InitList)
+		{
+			if (StructType)
+			{
+				Add(*StructType);
+			}
+		}
+	}
+
+	explicit TStructTypeBitSet(TConstArrayView<const FUStructType*> InitList)
+	{
+		for (const FUStructType* StructType : InitList)
+		{
+			if (StructType)
+			{
+				Add(*StructType);
+			}
+		}
+	}
+
+	/** This flavor of constructor is only available for UScriptStructs */
+	template<typename T = TBaseStruct, typename = typename TEnableIf<!TIsDerivedFrom<T, UObject>::IsDerived>::Type>
+	explicit TStructTypeBitSet(TConstArrayView<FInstancedStruct> InitList)
+	{
+		for (const FInstancedStruct& StructInstance : InitList)
+		{
+			if (StructInstance.GetScriptStruct())
+			{
+				Add(*StructInstance.GetScriptStruct());
+			}
+		}
+	}
+
 private:
+
 	/** 
 	 * A private constructor for a creating an instance straight from TBitArrays. 
 	 * @Note that this constructor needs to remain private to ensure consistency of stored values with data tracked 
 	 * by the TStructTrackerWrapper::StructTracker
 	 */
-	TStructTypeBitSet(const TBitArray<>& Source)
-		: StructTypesBitArray(Source)
+	TStructTypeBitSet(TBitArray<>&& Source)
+		: Super(Forward<TBitArray<>>(Source))
 	{
 	}
 
@@ -336,15 +630,19 @@
 		StructTypesBitArray.AddAtIndex(BitToSet);
 	}
 
-	FORCEINLINE_DEBUGGABLE static const UStruct* GetBaseUStruct()
-	{
-		static const UStruct* Instance = UE::StructUtils::GetAsUStruct<TBaseStruct>();
-		return Instance;
-	}
-
 public:
 
-	static int32 GetTypeIndex(const TUStructType& InStructType)
+	FStructTracker& GetStructTracker()
+	{
+		return TStructTrackerWrapper::StructTracker;
+	}
+
+	const FStructTracker& GetStructTracker() const
+	{
+		return TStructTrackerWrapper::StructTracker;
+	}
+
+	static int32 GetTypeIndex(const FUStructType& InStructType)
 	{
 #if WITH_STRUCTUTILS_DEBUG
 		ensureMsgf(InStructType.IsChildOf(GetBaseUStruct())
@@ -371,11 +669,11 @@
 		return TypeBitSet;
 	}
 
-	static const TUStructType* GetTypeAtIndex(const int32 Index)
+	static const FUStructType* GetTypeAtIndex(const int32 Index)
 	{
-		return Cast<const TUStructType>(TStructTrackerWrapper::StructTracker.GetStructType(Index));
+		return Cast<const FUStructType>(TStructTrackerWrapper::StructTracker.GetStructType(Index));
 	}
-	
+
 	template<typename T>
 	FORCEINLINE void Add()
 	{
@@ -405,49 +703,6 @@
 		return StructTypesBitArray.Contains(StructTypeIndex);
 	}
 
-	void SetAll(const bool bValue = true)
-	{
-		StructTypesBitArray.SetAll(bValue, TStructTrackerWrapper::StructTracker.Num());
-	}
-
-	void Add(const TUStructType& InStructType)
-	{
-#if WITH_STRUCTUTILS_DEBUG
-		ensureMsgf(InStructType.IsChildOf(GetBaseUStruct())
-				, TEXT("Registering '%s' with FStructTracker while it doesn't derive from the expected struct type %s")
-				, *InStructType.GetPathName(), *GetBaseUStruct()->GetName());
-#endif // WITH_STRUCTUTILS_DEBUG
-
-		const int32 StructTypeIndex = TStructTrackerWrapper::StructTracker.FindOrAddStructTypeIndex(InStructType);
-		StructTypesBitArray.AddAtIndex(StructTypeIndex);
-	}
-
-	void Remove(const TUStructType& InStructType)
-	{
-#if WITH_STRUCTUTILS_DEBUG
-		ensureMsgf(InStructType.IsChildOf(GetBaseUStruct())
-				, TEXT("Registering '%s' with FStructTracker while it doesn't derive from the expected struct type %s")
-				, *InStructType.GetPathName(), *GetBaseUStruct()->GetName());
-#endif // WITH_STRUCTUTILS_DEBUG
-
-		const int32 StructTypeIndex = TStructTrackerWrapper::StructTracker.FindOrAddStructTypeIndex(InStructType);
-		StructTypesBitArray.RemoveAtIndex(StructTypeIndex);
-	}
-
-	void Reset() { StructTypesBitArray.Reset(); }
-
-	bool Contains(const TUStructType& InStructType) const
-	{
-#if WITH_STRUCTUTILS_DEBUG
-		ensureMsgf(InStructType.IsChildOf(GetBaseUStruct())
-				, TEXT("Registering '%s' with FStructTracker while it doesn't derive from the expected struct type %s")
-				, *InStructType.GetPathName(), *GetBaseUStruct()->GetName());
-#endif // WITH_STRUCTUTILS_DEBUG
-
-		const int32 StructTypeIndex = TStructTrackerWrapper::StructTracker.FindOrAddStructTypeIndex(InStructType);
-		return StructTypesBitArray.Contains(StructTypeIndex);
-	}
-
 	FORCEINLINE TStructTypeBitSet operator+(const TStructTypeBitSet& Other) const
 	{
 		TStructTypeBitSet Result;
@@ -455,90 +710,4 @@
 		return MoveTemp(Result);
 	}
 
-	FORCEINLINE void operator+=(const TStructTypeBitSet& Other)
-	{
-		StructTypesBitArray = TBitArray<>::BitwiseOR(StructTypesBitArray, Other.StructTypesBitArray, EBitwiseOperatorFlags::MaxSize);
-	}
-
-	FORCEINLINE void operator-=(const TStructTypeBitSet& Other)
-	{
-		StructTypesBitArray -= Other.StructTypesBitArray;
-	}
-
-	FORCEINLINE TStructTypeBitSet operator+(const TUStructType& NewElement) const
-	{
-		TStructTypeBitSet Result = *this;
-		Result.Add(NewElement);
-		return MoveTemp(Result);
-	}
-
-	FORCEINLINE TStructTypeBitSet operator-(const TUStructType& NewElement) const
-	{
-		TStructTypeBitSet Result = *this;
-		Result.Remove(NewElement);
-		return MoveTemp(Result);
-	}
-
-	FORCEINLINE TStructTypeBitSet operator-(const TStructTypeBitSet& Other) const
-	{
-		TStructTypeBitSet Result = *this;
-		Result -= Other;
-		return MoveTemp(Result);
-	}
-
-	FORCEINLINE TStructTypeBitSet operator&(const TStructTypeBitSet& Other) const
-	{
-		return TStructTypeBitSet(TBitArray<>::BitwiseAND(StructTypesBitArray, Other.StructTypesBitArray, EBitwiseOperatorFlags::MinSize));
-	}
-
-	FORCEINLINE TStructTypeBitSet operator|(const TStructTypeBitSet& Other) const
-	{
-		return TStructTypeBitSet(TBitArray<>::BitwiseOR(StructTypesBitArray, Other.StructTypesBitArray, EBitwiseOperatorFlags::MaxSize));
-	}
-
-	FORCEINLINE TStructTypeBitSet GetOverlap(const TStructTypeBitSet& Other) const
-	{
-		return *this & Other;
-	}
-
-	FORCEINLINE bool IsEquivalent(const TStructTypeBitSet& Other) const
-	{
-		return StructTypesBitArray.CompareSetBits(Other.StructTypesBitArray, /*bMissingBitValue=*/false);
-	}
-
-	FORCEINLINE bool HasAll(const TStructTypeBitSet& Other) const
-	{
-		return StructTypesBitArray.HasAll(Other.StructTypesBitArray);
-	}
-
-	FORCEINLINE bool HasAny(const TStructTypeBitSet& Other) const
-	{
-		return StructTypesBitArray.HasAny(Other.StructTypesBitArray);
-	}
-
-	FORCEINLINE bool HasNone(const TStructTypeBitSet& Other) const
-	{
-		return !StructTypesBitArray.HasAny(Other.StructTypesBitArray);
-	}
-
-	bool IsEmpty() const 
-	{ 
-		return StructTypesBitArray.IsEmpty();
-	}
-
-	FORCEINLINE bool IsBitSet(const int32 BitIndex) const
-	{
-		return StructTypesBitArray.Contains(BitIndex);
-	}
-
-	FIndexIterator GetIndexIterator(const bool bValueToCheck = true) const
-	{
-		return FIndexIterator(StructTypesBitArray, bValueToCheck);
-	}
-
-	int32 CountStoredTypes() const
-	{
-		return StructTypesBitArray.CountSetBits();
-	}
-
 	static int32 GetMaxNum() 
@@ -546,25 +715,14 @@
 		return TStructTrackerWrapper::StructTracker.Num();
 	}
 
-	FORCEINLINE bool operator==(const TStructTypeBitSet& Other) const { return StructTypesBitArray == Other.StructTypesBitArray; }
+	FORCEINLINE bool operator==(const TStructTypeBitSet& Other) const { return Super::StructTypesBitArray == Other.StructTypesBitArray; }
 	FORCEINLINE bool operator!=(const TStructTypeBitSet& Other) const { return !(*this == Other); }
 
-	/**
-	 * note that this function is slow(ish) due to the FStructTracker utilizing WeakObjectPtrs to store types. 
-	 * @todo To be improved.
-	 */
-	template<typename TOutStructType, typename Allocator>
-	void ExportTypes(TArray<const TOutStructType*, Allocator>& OutTypes) const
+	FORCEINLINE TStructTypeBitSet operator-(const TStructTypeBitSet& Other) const
 	{
-		TBitArray<>::FConstIterator It(StructTypesBitArray);
-		while (It)
-		{
-			if (It.GetValue())
-			{
-				OutTypes.Add(Cast<TOutStructType>(TStructTrackerWrapper::StructTracker.GetStructType(It.GetIndex())));
-			}
-			++It;
-		}
+		TStructTypeBitSet Result = *this;
+		Result -= Other;
+		return MoveTemp(Result);
 	}
 
 	/**
@@ -574,9 +732,9 @@
 	 * note that this function is slow(ish) due to the FStructTracker utilizing WeakObjectPtrs to store types.
 	 * @todo To be improved.
 	 */
-	void ExportTypes(TFunctionRef<bool(const TUStructType*)> Callback) const
+	void ExportTypes(TFunctionRef<bool(const FUStructType*)> Callback) const
 	{
-		TBitArray<>::FConstIterator It(StructTypesBitArray);
+		TBitArray<>::FConstIterator It(Super::StructTypesBitArray);
 		bool bKeepGoing = true;
 		while (bKeepGoing && It)
 		{
@@ -588,45 +746,22 @@
 		}
 	}
 
-	SIZE_T GetAllocatedSize() const
-	{
-		return StructTypesBitArray.GetAllocatedSize();
-	}
-
+#if WITH_STRUCTUTILS_DEBUG
+	using Super::DebugGetStringDesc;
 	FString DebugGetStringDesc() const
 	{
-#if WITH_STRUCTUTILS_DEBUG
 		FStringOutputDevice Ar;
 		DebugGetStringDesc(Ar);
 		return static_cast<FString>(Ar);
+	}
 #else
+	FString DebugGetStringDesc() const
+	{
 		return TEXT("DEBUG INFO COMPILED OUT");
+	}
 #endif //WITH_STRUCTUTILS_DEBUG
-	}
 
 #if WITH_STRUCTUTILS_DEBUG
-	void DebugGetStringDesc(FOutputDevice& Ar) const
-	{
-		for (int32 Index = 0; Index < StructTypesBitArray.Num(); ++Index)
-		{
-			if (StructTypesBitArray[Index])
-			{
-				Ar.Logf(TEXT("%s, "), *TStructTrackerWrapper::StructTracker.DebugGetStructTypeName(Index).ToString());
-			}
-		}
-	}
-
-	void DebugGetIndividualNames(TArray<FName>& OutFNames) const
-	{
-		for (int32 Index = 0; Index < StructTypesBitArray.Num(); ++Index)
-		{
-			if (StructTypesBitArray[Index])
-			{
-				OutFNames.Add(TStructTrackerWrapper::StructTracker.DebugGetStructTypeName(Index));
-			}
-		}
-	}
-
 	static FName DebugGetStructTypeName(const int32 StructTypeIndex)
 	{
 		return TStructTrackerWrapper::StructTracker.DebugGetStructTypeName(StructTypeIndex);
@@ -645,11 +780,6 @@
 	{
 		TStructTrackerWrapper::StructTracker.DebugResetStructTypeMappingInfo();
 	}
-
-protected:
-	// unittesting purposes only
-	const TBitArray<>& DebugGetStructTypesBitArray() const { return StructTypesBitArray; }
-	TBitArray<>& DebugGetMutableStructTypesBitArray() { return StructTypesBitArray; }
 #endif // WITH_STRUCTUTILS_DEBUG
 
 public:
@@ -662,7 +792,7 @@
 
 	void Serialize(FArchive& Ar)
 	{
-		TStructTrackerWrapper::StructTracker.Serialize(Ar, StructTypesBitArray);
+		TStructTrackerWrapper::StructTracker.Serialize(Ar, Super::StructTypesBitArray);
 	}
 
 	FORCEINLINE friend FArchive& operator<<(FArchive& Ar, TStructTypeBitSet& Instance)
@@ -670,9 +800,6 @@
 		Instance.Serialize(Ar);
 		return Ar;
 	}
-
-private:
-	FStructTypeBitSet::FBitSetContainer StructTypesBitArray;
 };
 
 /** 
@@ -686,7 +813,8 @@
 		using FBaseStructType = BaseType; \
 		EXPORTED_API static FStructTracker StructTracker; \
 	}; \
-	using ContainerTypeName = TStructTypeBitSet<BaseType, ContainerTypeName##StructTrackerWrapper, BaseUStructType>
+	using ContainerTypeName = TStructTypeBitSet<BaseType, ContainerTypeName##StructTrackerWrapper, BaseUStructType>; \
+	static_assert(std::is_move_constructible_v<ContainerTypeName> && std::is_move_assignable_v<ContainerTypeName>)
 
 #define DECLARE_STRUCTTYPEBITSET_EXPORTED(EXPORTED_API, ContainerTypeName, BaseStructType) _DECLARE_TYPEBITSET_IMPL(EXPORTED_API, ContainerTypeName, BaseStructType, UScriptStruct)
 #define DECLARE_STRUCTTYPEBITSET(ContainerTypeName, BaseStructType) _DECLARE_TYPEBITSET_IMPL(, ContainerTypeName, BaseStructType, UScriptStruct)
