[StructUtils] Issue 617333: PVS Studio warnings in StructTypeBitSet.h

#rnx

diff --git a/Engine/Source/Runtime/CoreUObject/Public/StructUtils/StructTypeBitSet.h b/Engine/Source/Runtime/CoreUObject/Public/StructUtils/StructTypeBitSet.h
--- a/Engine/Source/Runtime/CoreUObject/Public/StructUtils/StructTypeBitSet.h
+++ b/Engine/Source/Runtime/CoreUObject/Public/StructUtils/StructTypeBitSet.h
@@ -15,6 +15,8 @@ namespace FStructTypeBitSet
 
 struct FConstBitSetContainer : TBitArray<>
 {
+	using Super = TBitArray<>;
+
 	FConstBitSetContainer() = default;
 
 	FORCEINLINE bool HasAll(const TBitArray<>& Other) const
@@ -106,10 +108,13 @@ struct FConstBitSetContainer : TBitArray<>
 
 protected:
 	// only for use by the FBitSetContainer
-	FConstBitSetContainer(const TBitArray<>& Source) : TBitArray<>(Source)
+	explicit FConstBitSetContainer(const TBitArray<>& Source)
+		: Super(Source)
 	{
 	}
-	FConstBitSetContainer(TBitArray<>&& Source) : TBitArray<>(Forward<TBitArray<>>(Source))
+
+	explicit FConstBitSetContainer(TBitArray<>&& Source)
+		: Super(Forward<TBitArray<>>(Source))
 	{
 	}
 };
@@ -118,11 +123,11 @@ struct FBitSetContainer : FConstBitSetContainer
 {
 	using FConstBitSetContainer::FConstBitSetContainer;
 
-	FBitSetContainer(const TBitArray<>& Source)
+	explicit FBitSetContainer(const TBitArray<>& Source)
 		: FConstBitSetContainer(Source)
 	{
 	}
-	FBitSetContainer(TBitArray<>&& Source)
+	explicit FBitSetContainer(TBitArray<>&& Source)
 		: FConstBitSetContainer(Forward<TBitArray<>>(Source))
 	{
 	}
