[AutoRTFM] Migrate to UE::TScopeLock instead of FTransactionallySafeScopeLock.

These are functionally identical.

diff --git a/Engine/Source/Runtime/CoreUObject/Private/UObject/UObjectArray.cpp b/Engine/Source/Runtime/CoreUObject/Private/UObject/UObjectArray.cpp
--- a/Engine/Source/Runtime/CoreUObject/Private/UObject/UObjectArray.cpp
+++ b/Engine/Source/Runtime/CoreUObject/Private/UObject/UObjectArray.cpp
@@ -10,7 +10,6 @@
 #include "HAL/LowLevelMemStats.h"
 #include "Misc/ScopeLock.h"
 #include "ProfilingDebugging/MetadataTrace.h"
-#include "Misc/TransactionallySafeScopeLock.h"
 #include "UObject/UObjectAllocator.h"
 #include "UObject/Class.h"
 #include "UObject/GarbageCollectionInternalFlags.h"
@@ -142,7 +141,7 @@ void FUObjectArray::OpenDisregardForGC()
 void FUObjectArray::CloseDisregardForGC()
 {
 #if THREADSAFE_UOBJECTS
-	FTransactionallySafeScopeLock ObjObjectsLock(&ObjObjectsCritical);
+	UE::TScopeLock ObjObjectsLock(ObjObjectsCritical);
 #else
 	// Disregard from GC pool is only available from the game thread, at least for now
 	check(IsInGameThread());
@@ -340,7 +339,7 @@ void FUObjectArray::ResetSerialNumber(UObjectBase* Object)
 void FUObjectArray::RemoveObjectFromDeleteListeners(UObjectBase* Object)
 {
 #if THREADSAFE_UOBJECTS
-	FTransactionallySafeScopeLock UObjectDeleteListenersLock(&UObjectDeleteListenersCritical);
+	UE::TScopeLock UObjectDeleteListenersLock(UObjectDeleteListenersCritical);
 #endif
 	int32 Index = Object->InternalIndex;
 	check(Index >= 0);
@@ -431,7 +430,7 @@ void FUObjectArray::RemoveUObjectCreateListener(FUObjectCreateListener* Listener
 void FUObjectArray::AddUObjectDeleteListener(FUObjectDeleteListener* Listener)
 {
 #if THREADSAFE_UOBJECTS
-	FTransactionallySafeScopeLock UObjectDeleteListenersLock(&UObjectDeleteListenersCritical);
+	UE::TScopeLock UObjectDeleteListenersLock(UObjectDeleteListenersCritical);
 #endif
 	check(!UObjectDeleteListeners.Contains(Listener));
 	UObjectDeleteListeners.Add(Listener);
@@ -445,7 +444,7 @@ void FUObjectArray::AddUObjectDeleteListener(FUObjectDeleteListener* Listener)
 void FUObjectArray::RemoveUObjectDeleteListener(FUObjectDeleteListener* Listener)
 {
 #if THREADSAFE_UOBJECTS
-	FTransactionallySafeScopeLock UObjectDeleteListenersLock(&UObjectDeleteListenersCritical);
+	UE::TScopeLock UObjectDeleteListenersLock(UObjectDeleteListenersCritical);
 #endif
 	UObjectDeleteListeners.RemoveSingleSwap(Listener);
 }
@@ -494,14 +493,15 @@ int32 FUObjectArray::AllocateSerialNumber(int32 Index)
 	int32 SerialNumber = *SerialNumberPtr;
 	if (!SerialNumber)
 	{
-		// Open around PrimarySerialNumber as if we fail/abort a transaction we dont need to undo this, simply allow it to grow for the next use
-		UE_AUTORTFM_OPEN{
+		// Open around PrimarySerialNumber. If we fail/abort a transaction, we don't need to undo this; we simply allow it to grow for the next use.
+		UE_AUTORTFM_OPEN
+		{
 			SerialNumber = PrimarySerialNumber.Increment();
 			UE_CLOG(SerialNumber <= START_SERIAL_NUMBER, LogUObjectArray, Fatal, TEXT("UObject serial numbers overflowed (trying to allocate serial number %d)."), SerialNumber);
 			int32 ValueWas = FPlatformAtomics::InterlockedCompareExchange((int32*)SerialNumberPtr, SerialNumber, 0);
 			if (ValueWas != 0)
 			{
-				// someone else go it first, use their value
+				// Someone else got it first; use their value.
 				SerialNumber = ValueWas;
 			}
 		};
@@ -517,7 +517,7 @@ void FUObjectArray::ShutdownUObjectArray()
 {
 	{
 #if THREADSAFE_UOBJECTS
-		FTransactionallySafeScopeLock UObjectDeleteListenersLock(&UObjectDeleteListenersCritical);
+		UE::TScopeLock UObjectDeleteListenersLock(UObjectDeleteListenersCritical);
 #endif
 		for (int32 Index = UObjectDeleteListeners.Num() - 1; Index >= 0; --Index)
 		{
diff --git a/Engine/Source/Runtime/CoreUObject/Public/UObject/UObjectArray.h b/Engine/Source/Runtime/CoreUObject/Public/UObject/UObjectArray.h
--- a/Engine/Source/Runtime/CoreUObject/Public/UObject/UObjectArray.h
+++ b/Engine/Source/Runtime/CoreUObject/Public/UObject/UObjectArray.h
@@ -9,10 +9,10 @@
 #include "AutoRTFM/AutoRTFM.h"
 #include "HAL/ThreadSafeCounter.h"
 #include "Containers/LockFreeList.h"
+#include "Misc/ScopeLock.h"
 #include "Misc/TransactionallySafeCriticalSection.h"
 #include "UObject/GarbageCollectionGlobals.h"
 #include "UObject/UObjectBase.h"
-#include "Misc/TransactionallySafeScopeLock.h"
 
 #if WITH_VERSE_VM || defined(__INTELLISENSE__)
 namespace Verse
@@ -1373,9 +1373,9 @@ public:
 
     SIZE_T GetAllocatedSize() const
     {
-		FTransactionallySafeScopeLock ObjListLock(&ObjObjectsCritical);
+		UE::TScopeLock ObjListLock(ObjObjectsCritical);
 #if THREADSAFE_UOBJECTS
-		FTransactionallySafeScopeLock ListenersLock(&UObjectDeleteListenersCritical);
+		UE::TScopeLock ListenersLock(UObjectDeleteListenersCritical);
 #endif
         return ObjObjects.GetAllocatedSize() + ObjAvailableList.GetAllocatedSize() + UObjectCreateListeners.GetAllocatedSize() + UObjectDeleteListeners.GetAllocatedSize();
     }
@@ -1383,7 +1383,7 @@ public:
 	SIZE_T GetDeleteListenersAllocatedSize(int32* OutNumListeners = nullptr) const
 	{
 #if THREADSAFE_UOBJECTS
-		FTransactionallySafeScopeLock ListenersLock(&UObjectDeleteListenersCritical);
+		UE::TScopeLock ListenersLock(UObjectDeleteListenersCritical);
 #endif
 		SIZE_T AllocatedSize = 0;
 		for (FUObjectDeleteListener* Listener : UObjectDeleteListeners)
