Modified the PropertyVisitor API to include the instance data for the parent struct.

[RN] Encapsulated the property data argument of PropertyVisitor callbacks to an object that also includes the parent struct data.

diff --git a/Engine/Plugins/Experimental/InstanceDataObjectFixupTool/Source/Private/InstanceDataObjectFixupPanel.cpp b/Engine/Plugins/Experimental/InstanceDataObjectFixupTool/Source/Private/InstanceDataObjectFixupPanel.cpp
--- a/Engine/Plugins/Experimental/InstanceDataObjectFixupTool/Source/Private/InstanceDataObjectFixupPanel.cpp
+++ b/Engine/Plugins/Experimental/InstanceDataObjectFixupTool/Source/Private/InstanceDataObjectFixupPanel.cpp
@@ -228,7 +228,7 @@ FInstanceDataObjectFixupPanel::FInstanceDataObjectFixupPanel(
 static bool ObjectHasLoosePropertiesThatNeedFixup(UObject* Object)
 {
 	bool bNeedsFixup = false;
-	Object->GetClass()->Visit(Object, [&bNeedsFixup](const FPropertyVisitorPath& Path, void* Data)->EPropertyVisitorControlFlow
+	Object->GetClass()->Visit(Object, [&bNeedsFixup](const FPropertyVisitorPath& Path, const FPropertyVisitorData& Data)->EPropertyVisitorControlFlow
 	{
 		const FProperty* Property = Path.Top().Property;
 		if (!Property->HasAnyPropertyFlags(CPF_SkipSerialization) && Property->GetBoolMetaData(NAME_IsLooseMetadata))
diff --git a/Engine/Source/Runtime/CoreUObject/Private/StructUtils/InstancedStruct.cpp b/Engine/Source/Runtime/CoreUObject/Private/StructUtils/InstancedStruct.cpp
--- a/Engine/Source/Runtime/CoreUObject/Private/StructUtils/InstancedStruct.cpp
+++ b/Engine/Source/Runtime/CoreUObject/Private/StructUtils/InstancedStruct.cpp
@@ -538,14 +538,16 @@ bool FInstancedStruct::FindInnerPropertyInstance(FName PropertyName, const FProp
 	return false;
 }
 
-EPropertyVisitorControlFlow FInstancedStruct::Visit(FPropertyVisitorPath& Path, const TFunctionRef<EPropertyVisitorControlFlow(const FPropertyVisitorPath& /*Path*/, void* /*Data*/)> InFunc) const
+EPropertyVisitorControlFlow FInstancedStruct::Visit(FPropertyVisitorPath& Path, const FPropertyVisitorData& InData, const TFunctionRef<EPropertyVisitorControlFlow(const FPropertyVisitorPath& /*Path*/, const FPropertyVisitorData& /*Data*/)> InFunc) const
 {
 	if (!ScriptStruct || !StructMemory)
 	{
 		return EPropertyVisitorControlFlow::StepOver;
 	}
 
-	return ScriptStruct->Visit(Path, StructMemory, InFunc);
+	FPropertyVisitorData Data = InData.VisitPropertyData(StructMemory);
+
+	return ScriptStruct->Visit(Path, Data, InFunc);
 }
 
 void* FInstancedStruct::ResolveVisitedPathInfo(const FPropertyVisitorInfo& Info) const
diff --git a/Engine/Source/Runtime/CoreUObject/Private/UObject/Class.cpp b/Engine/Source/Runtime/CoreUObject/Private/UObject/Class.cpp
--- a/Engine/Source/Runtime/CoreUObject/Private/UObject/Class.cpp
+++ b/Engine/Source/Runtime/CoreUObject/Private/UObject/Class.cpp
@@ -786,13 +786,14 @@ void UStruct::CollectBytecodeAndPropertyReferencedObjectsRecursively()
 	}
 }
 
-EPropertyVisitorControlFlow UStruct::Visit(void* Data, const TFunctionRef<EPropertyVisitorControlFlow(const FPropertyVisitorPath& /*Path*/, void* /*Data*/)> InFunc) const
+EPropertyVisitorControlFlow UStruct::Visit(void* InData, const TFunctionRef<EPropertyVisitorControlFlow(const FPropertyVisitorPath& /*Path*/, const FPropertyVisitorData& /*Data*/)> InFunc) const
 {
 	FPropertyVisitorPath Path;
+	FPropertyVisitorData Data(InData, /*ParentStructData*/nullptr);
 	return Visit(Path, Data, InFunc);
 }
 
-EPropertyVisitorControlFlow UStruct::Visit(FPropertyVisitorPath& Path, void* Data, const TFunctionRef<EPropertyVisitorControlFlow(const FPropertyVisitorPath& /*Path*/, void* /*Data*/)> InFunc) const
+EPropertyVisitorControlFlow UStruct::Visit(FPropertyVisitorPath& Path, const FPropertyVisitorData& InData, const TFunctionRef<EPropertyVisitorControlFlow(const FPropertyVisitorPath& /*Path*/, const FPropertyVisitorData& /*Data*/)> InFunc) const
 {
 	EPropertyVisitorControlFlow RetVal = EPropertyVisitorControlFlow::StepOver;
 	for (const FProperty* Property = PropertyLink; Property; Property = Property->PropertyLinkNext)
@@ -803,7 +804,9 @@ EPropertyVisitorControlFlow UStruct::Visit(FPropertyVisitorPath& Path, void* Dat
 		// ArrayDim of one means it is just a single property, not a static array
 		if(Property->ArrayDim == 1)
 		{
-			RetVal = Property->Visit(Path, Property->ContainerPtrToValuePtr<void>(Data), InFunc);
+			FPropertyVisitorData Data(Property->ContainerPtrToValuePtr<void>(InData.PropertyData), /*ParentStructData*/InData.PropertyData);
+
+			RetVal = Property->Visit(Path, Data, InFunc);
 			if (RetVal == EPropertyVisitorControlFlow::Stop)
 			{
 				return EPropertyVisitorControlFlow::Stop;
@@ -819,7 +822,10 @@ EPropertyVisitorControlFlow UStruct::Visit(FPropertyVisitorPath& Path, void* Dat
 			for (int32 StaticArrayIndex = 0; StaticArrayIndex < Property->ArrayDim; ++StaticArrayIndex)
 			{
 				Path.Top().SetIndex(StaticArrayIndex, EPropertyVisitorInfoType::StaticArrayIndex);
-				RetVal = Property->Visit(Path, Property->ContainerPtrToValuePtr<void>(Data, StaticArrayIndex), InFunc);
+
+				FPropertyVisitorData Data(Property->ContainerPtrToValuePtr<void>(InData.PropertyData, StaticArrayIndex), /*ParentStructData*/InData.PropertyData);
+
+				RetVal = Property->Visit(Path, Data, InFunc);
 
 				if (RetVal == EPropertyVisitorControlFlow::Stop)
 				{
@@ -3664,7 +3670,7 @@ bool UScriptStruct::FindInnerPropertyInstance(FName PropertyName, const void* Da
 	return false;
 }
 
-EPropertyVisitorControlFlow UScriptStruct::Visit(FPropertyVisitorPath& Path, void* Data, const TFunctionRef<EPropertyVisitorControlFlow(const FPropertyVisitorPath& /*Path*/, void* /*Data*/)> InFunc) const
+EPropertyVisitorControlFlow UScriptStruct::Visit(FPropertyVisitorPath& Path, const FPropertyVisitorData& Data, const TFunctionRef<EPropertyVisitorControlFlow(const FPropertyVisitorPath& /*Path*/, const FPropertyVisitorData& /*Data*/)> InFunc) const
 {
 	const EPropertyVisitorControlFlow RetVal = UStruct::Visit(Path, Data, InFunc);
 	if (RetVal == EPropertyVisitorControlFlow::Stop)
diff --git a/Engine/Source/Runtime/CoreUObject/Private/UObject/Property.cpp b/Engine/Source/Runtime/CoreUObject/Private/UObject/Property.cpp
--- a/Engine/Source/Runtime/CoreUObject/Private/UObject/Property.cpp
+++ b/Engine/Source/Runtime/CoreUObject/Private/UObject/Property.cpp
@@ -1417,13 +1417,13 @@ namespace UE::CoreUObject::Private
 	}
 }
 
-EPropertyVisitorControlFlow FProperty::Visit(void* Data, const TFunctionRef<EPropertyVisitorControlFlow(const FPropertyVisitorPath&, void*)> InFunc) const
+EPropertyVisitorControlFlow FProperty::Visit(const FPropertyVisitorData& Data, const TFunctionRef<EPropertyVisitorControlFlow(const FPropertyVisitorPath&, const FPropertyVisitorData&)> InFunc) const
 {
 	FPropertyVisitorPath Path(FPropertyVisitorInfo(this));
 	return Visit(Path, Data, InFunc);
 }
 
-EPropertyVisitorControlFlow FProperty::Visit(FPropertyVisitorPath& Path, void* Data, const TFunctionRef<EPropertyVisitorControlFlow(const FPropertyVisitorPath& /*Path*/, void* /*Data*/)> InFunc) const
+EPropertyVisitorControlFlow FProperty::Visit(FPropertyVisitorPath& Path, const FPropertyVisitorData& Data, const TFunctionRef<EPropertyVisitorControlFlow(const FPropertyVisitorPath& /*Path*/, const FPropertyVisitorData& /*Data*/)> InFunc) const
 {
 	checkf(Path.Top().Property == this, TEXT("The property set in the info has to match to this property"));
 
diff --git a/Engine/Source/Runtime/CoreUObject/Private/UObject/PropertyArray.cpp b/Engine/Source/Runtime/CoreUObject/Private/UObject/PropertyArray.cpp
--- a/Engine/Source/Runtime/CoreUObject/Private/UObject/PropertyArray.cpp
+++ b/Engine/Source/Runtime/CoreUObject/Private/UObject/PropertyArray.cpp
@@ -1391,17 +1391,17 @@ bool FArrayProperty::CanSerializeFromTypeName(UE::FPropertyTypeName Type) const
 	return LocalInner->CanSerializeFromTypeName(Type.GetParameter(0));
 }
 
-EPropertyVisitorControlFlow FArrayProperty::Visit(FPropertyVisitorPath& Path, void* Data, const TFunctionRef<EPropertyVisitorControlFlow(const FPropertyVisitorPath& /*Path*/, void* /*Data*/)> InFunc) const
+EPropertyVisitorControlFlow FArrayProperty::Visit(FPropertyVisitorPath& Path, const FPropertyVisitorData& InData, const TFunctionRef<EPropertyVisitorControlFlow(const FPropertyVisitorPath& /*Path*/, const FPropertyVisitorData& /*Data*/)> InFunc) const
 {
 	// Indicate in the path that this property contains inner properties
 	Path.Top().bContainsInnerProperties = true;
 
-	EPropertyVisitorControlFlow RetVal = Super::Visit(Path, Data, InFunc);
+	EPropertyVisitorControlFlow RetVal = Super::Visit(Path, InData, InFunc);
 
 	if (RetVal == EPropertyVisitorControlFlow::StepInto)
 	{
 		checkf(Inner, TEXT("Expecting a valid inner property type"));
-		FScriptArrayHelper ArrayHelper(this, Data);
+		FScriptArrayHelper ArrayHelper(this, InData.PropertyData);
 
 		FPropertyVisitorScope Scope(Path, FPropertyVisitorInfo(Inner));
 
@@ -1409,7 +1409,10 @@ EPropertyVisitorControlFlow FArrayProperty::Visit(FPropertyVisitorPath& Path, vo
 		for (int32 ContainerIndex = 0; ContainerIndex < ArrayNum; ContainerIndex++)
 		{
 			Path.Top().SetIndex(ContainerIndex, EPropertyVisitorInfoType::ContainerIndex);
-			RetVal = Inner->Visit(Path, ArrayHelper.GetRawPtr(ContainerIndex), InFunc);
+
+			FPropertyVisitorData Data = InData.VisitPropertyData(ArrayHelper.GetRawPtr(ContainerIndex));
+
+			RetVal = Inner->Visit(Path, Data, InFunc);
 			if (RetVal == EPropertyVisitorControlFlow::Stop)
 			{
 				return EPropertyVisitorControlFlow::Stop;
diff --git a/Engine/Source/Runtime/CoreUObject/Private/UObject/PropertyBaseObject.cpp b/Engine/Source/Runtime/CoreUObject/Private/UObject/PropertyBaseObject.cpp
--- a/Engine/Source/Runtime/CoreUObject/Private/UObject/PropertyBaseObject.cpp
+++ b/Engine/Source/Runtime/CoreUObject/Private/UObject/PropertyBaseObject.cpp
@@ -889,18 +889,20 @@ bool FObjectPropertyBase::SameType(const FProperty* Other) const
 			 (PropertyClass == ((FObjectPropertyBase*)Other)->PropertyClass);
 }
 
-EPropertyVisitorControlFlow FObjectPropertyBase::Visit(FPropertyVisitorPath& Path, void* Data, const TFunctionRef<EPropertyVisitorControlFlow(const FPropertyVisitorPath& /*Path*/, void* /*Data*/)> InFunc) const
+EPropertyVisitorControlFlow FObjectPropertyBase::Visit(FPropertyVisitorPath& Path, const FPropertyVisitorData& InData, const TFunctionRef<EPropertyVisitorControlFlow(const FPropertyVisitorPath& /*Path*/, const FPropertyVisitorData& /*Data*/)> InFunc) const
 {
 	// Indicate in the path that this property contains inner properties
 	Path.Top().bContainsInnerProperties = true;
 
-	EPropertyVisitorControlFlow RetVal = Super::Visit(Path, Data, InFunc);
+	EPropertyVisitorControlFlow RetVal = Super::Visit(Path, InData, InFunc);
 
 	if (RetVal == EPropertyVisitorControlFlow::StepInto)
 	{
-		if (const TObjectPtr<UObject> Object = GetObjectPropertyValue(Data))
+		if (const TObjectPtr<UObject> Object = GetObjectPropertyValue(InData.PropertyData))
 		{
-			RetVal = Object.GetClass()->Visit(Path, Object.Get(), InFunc);
+			FPropertyVisitorData Data = InData.VisitPropertyData(Object.Get());
+
+			RetVal = Object.GetClass()->Visit(Path, Data, InFunc);
 		}
 	}
 	return RetVal;
diff --git a/Engine/Source/Runtime/CoreUObject/Private/UObject/PropertyMap.cpp b/Engine/Source/Runtime/CoreUObject/Private/UObject/PropertyMap.cpp
--- a/Engine/Source/Runtime/CoreUObject/Private/UObject/PropertyMap.cpp
+++ b/Engine/Source/Runtime/CoreUObject/Private/UObject/PropertyMap.cpp
@@ -1755,17 +1755,17 @@ bool FMapProperty::CanSerializeFromTypeName(UE::FPropertyTypeName Type) const
 	return LocalKeyProp->CanSerializeFromTypeName(Type.GetParameter(0)) && LocalValueProp->CanSerializeFromTypeName(Type.GetParameter(1));
 }
 
-EPropertyVisitorControlFlow FMapProperty::Visit(FPropertyVisitorPath& Path, void* Data, const TFunctionRef<EPropertyVisitorControlFlow(const FPropertyVisitorPath& /*Path*/, void* /*Data*/)> InFunc) const
+EPropertyVisitorControlFlow FMapProperty::Visit(FPropertyVisitorPath& Path, const FPropertyVisitorData& InData, const TFunctionRef<EPropertyVisitorControlFlow(const FPropertyVisitorPath& /*Path*/, const FPropertyVisitorData& /*Data*/)> InFunc) const
 {
 	// Indicate in the path that this property contains inner properties
 	Path.Top().bContainsInnerProperties = true;
 
-	EPropertyVisitorControlFlow RetVal = Super::Visit(Path, Data, InFunc);
+	EPropertyVisitorControlFlow RetVal = Super::Visit(Path, InData, InFunc);
 
 	if (RetVal == EPropertyVisitorControlFlow::StepInto)
 	{
 		checkf(KeyProp && ValueProp, TEXT("Expecting a valid inner property type"));
-		FScriptMapHelper MapHelper(this, Data);
+		FScriptMapHelper MapHelper(this, InData.PropertyData);
 
 		for (FScriptMapHelper::FIterator It(MapHelper); It; ++It)
 		{
@@ -1773,7 +1773,9 @@ EPropertyVisitorControlFlow FMapProperty::Visit(FPropertyVisitorPath& Path, void
 				// Visit Key
 				FPropertyVisitorScope Scope(Path, FPropertyVisitorInfo(KeyProp, It.GetLogicalIndex(), EPropertyVisitorInfoType::MapKey));
 
-				RetVal = KeyProp->Visit(Path, MapHelper.GetKeyPtr(It), InFunc);
+				FPropertyVisitorData Data = InData.VisitPropertyData(MapHelper.GetKeyPtr(It));
+
+				RetVal = KeyProp->Visit(Path, Data, InFunc);
 				if (RetVal == EPropertyVisitorControlFlow::Stop)
 				{
 					return EPropertyVisitorControlFlow::Stop;
@@ -1788,7 +1790,9 @@ EPropertyVisitorControlFlow FMapProperty::Visit(FPropertyVisitorPath& Path, void
 				// Visit Value
 				FPropertyVisitorScope Scope(Path, FPropertyVisitorInfo(ValueProp, It.GetLogicalIndex(), EPropertyVisitorInfoType::MapValue));
 
-				RetVal = ValueProp->Visit(Path, MapHelper.GetValuePtr(It), InFunc);
+				FPropertyVisitorData Data = InData.VisitPropertyData(MapHelper.GetValuePtr(It));
+
+				RetVal = ValueProp->Visit(Path, Data, InFunc);
 				if (RetVal == EPropertyVisitorControlFlow::Stop)
 				{
 					return EPropertyVisitorControlFlow::Stop;
diff --git a/Engine/Source/Runtime/CoreUObject/Private/UObject/PropertyOptional.cpp b/Engine/Source/Runtime/CoreUObject/Private/UObject/PropertyOptional.cpp
--- a/Engine/Source/Runtime/CoreUObject/Private/UObject/PropertyOptional.cpp
+++ b/Engine/Source/Runtime/CoreUObject/Private/UObject/PropertyOptional.cpp
@@ -696,16 +696,16 @@ bool FOptionalProperty::CanSerializeFromTypeName(UE::FPropertyTypeName Type) con
 	return LocalValueProperty->CanSerializeFromTypeName(Type.GetParameter(0));
 }
 
-EPropertyVisitorControlFlow FOptionalProperty::Visit(FPropertyVisitorPath& Path, void* Data, const TFunctionRef<EPropertyVisitorControlFlow(const FPropertyVisitorPath& /*Path*/, void* /*Data*/)> InFunc) const
+EPropertyVisitorControlFlow FOptionalProperty::Visit(FPropertyVisitorPath& Path, const FPropertyVisitorData& Data, const TFunctionRef<EPropertyVisitorControlFlow(const FPropertyVisitorPath& /*Path*/, const FPropertyVisitorData& /*Data*/)> InFunc) const
 {
 	// Indicate in the path that this property contains inner properties
 	Path.Top().bContainsInnerProperties = true;
 
 	EPropertyVisitorControlFlow RetVal = Super::Visit(Path, Data, InFunc);
 
-	if (RetVal == EPropertyVisitorControlFlow::StepInto)
+	if (RetVal == EPropertyVisitorControlFlow::StepInto && IsSet(Data.PropertyData))
 	{
-		if(!IsSet(Data))
+		if(!IsSet(Data.PropertyData))
 		{
 			// There is nothing to StepInto, so continue to next property
 			return EPropertyVisitorControlFlow::StepOver;
diff --git a/Engine/Source/Runtime/CoreUObject/Private/UObject/PropertySet.cpp b/Engine/Source/Runtime/CoreUObject/Private/UObject/PropertySet.cpp
--- a/Engine/Source/Runtime/CoreUObject/Private/UObject/PropertySet.cpp
+++ b/Engine/Source/Runtime/CoreUObject/Private/UObject/PropertySet.cpp
@@ -1130,24 +1130,27 @@ bool FSetProperty::CanSerializeFromTypeName(UE::FPropertyTypeName Type) const
 	return LocalElementProp->CanSerializeFromTypeName(Type.GetParameter(0));
 }
 
-EPropertyVisitorControlFlow FSetProperty::Visit(FPropertyVisitorPath& Path, void* Data, const TFunctionRef<EPropertyVisitorControlFlow(const FPropertyVisitorPath& /*Path*/, void* /*Data*/)> InFunc) const
+EPropertyVisitorControlFlow FSetProperty::Visit(FPropertyVisitorPath& Path, const FPropertyVisitorData& InData, const TFunctionRef<EPropertyVisitorControlFlow(const FPropertyVisitorPath& /*Path*/, const FPropertyVisitorData& /*Data*/)> InFunc) const
 {
 	// Indicate in the path that this property contains inner properties
 	Path.Top().bContainsInnerProperties = true;
 
-	EPropertyVisitorControlFlow RetVal = Super::Visit(Path, Data, InFunc);
+	EPropertyVisitorControlFlow RetVal = Super::Visit(Path, InData, InFunc);
 
 	if (RetVal == EPropertyVisitorControlFlow::StepInto)
 	{
 		checkf(ElementProp, TEXT("Expecting a valid inner property type"));
-		FScriptSetHelper SetHelper(this, Data);
+		FScriptSetHelper SetHelper(this, InData.PropertyData);
 
 		FPropertyVisitorScope Scope(Path, FPropertyVisitorInfo(ElementProp));
 		for (FScriptSetHelper::FIterator It(SetHelper); It; ++It)
 		{
 			// visit element
 			Path.Top().SetIndex(It.GetLogicalIndex(), EPropertyVisitorInfoType::ContainerIndex);
-			RetVal = ElementProp->Visit(Path, SetHelper.GetElementPtr(It), InFunc);
+
+			FPropertyVisitorData Data = InData.VisitPropertyData(SetHelper.GetElementPtr(It));
+
+			RetVal = ElementProp->Visit(Path, Data, InFunc);
 			if (RetVal == EPropertyVisitorControlFlow::Stop)
 			{
 				return EPropertyVisitorControlFlow::Stop;
diff --git a/Engine/Source/Runtime/CoreUObject/Private/UObject/PropertyStruct.cpp b/Engine/Source/Runtime/CoreUObject/Private/UObject/PropertyStruct.cpp
--- a/Engine/Source/Runtime/CoreUObject/Private/UObject/PropertyStruct.cpp
+++ b/Engine/Source/Runtime/CoreUObject/Private/UObject/PropertyStruct.cpp
@@ -648,7 +648,7 @@ bool FStructProperty::CanSerializeFromTypeName(UE::FPropertyTypeName Type) const
 	return false;
 }
 
-EPropertyVisitorControlFlow FStructProperty::Visit(FPropertyVisitorPath& Path, void* Data, const TFunctionRef<EPropertyVisitorControlFlow(const FPropertyVisitorPath& /*Path*/, void* /*Data*/)> InFunc) const
+EPropertyVisitorControlFlow FStructProperty::Visit(FPropertyVisitorPath& Path, const FPropertyVisitorData& Data, const TFunctionRef<EPropertyVisitorControlFlow(const FPropertyVisitorPath& /*Path*/, const FPropertyVisitorData& /*Data*/)> InFunc) const
 {
 	// Indicate in the path that this property contains inner properties
 	Path.Top().bContainsInnerProperties = true;
diff --git a/Engine/Source/Runtime/CoreUObject/Public/StructUtils/InstancedStruct.h b/Engine/Source/Runtime/CoreUObject/Public/StructUtils/InstancedStruct.h
--- a/Engine/Source/Runtime/CoreUObject/Public/StructUtils/InstancedStruct.h
+++ b/Engine/Source/Runtime/CoreUObject/Public/StructUtils/InstancedStruct.h
@@ -156,7 +156,7 @@ public:
 	void GetPreloadDependencies(TArray<UObject*>& OutDeps);
 	bool NetSerialize(FArchive& Ar, class UPackageMap* Map, bool& bOutSuccess);
 	bool FindInnerPropertyInstance(FName PropertyName, const FProperty*& OutProp, const void*& OutData) const;
-	EPropertyVisitorControlFlow Visit(FPropertyVisitorPath& Path, const TFunctionRef<EPropertyVisitorControlFlow(const FPropertyVisitorPath& /*Path*/, void* /*Data*/)> InFunc) const;
+	EPropertyVisitorControlFlow Visit(FPropertyVisitorPath& Path, const FPropertyVisitorData& Data, const TFunctionRef<EPropertyVisitorControlFlow(const FPropertyVisitorPath& /*Path*/, const FPropertyVisitorData& /*Data*/)> InFunc) const;
 	void* ResolveVisitedPathInfo(const FPropertyVisitorInfo& Info) const;
 
 	/** Returns struct type. */
diff --git a/Engine/Source/Runtime/CoreUObject/Public/UObject/Class.h b/Engine/Source/Runtime/CoreUObject/Public/UObject/Class.h
--- a/Engine/Source/Runtime/CoreUObject/Public/UObject/Class.h
+++ b/Engine/Source/Runtime/CoreUObject/Public/UObject/Class.h
@@ -742,7 +742,7 @@ public:
 	 * @param InFunc to call on each visited property, the return value controls what is the next behavior once this property has been visited
 	 * @return the new action to take one visited this property
 	 */
-	COREUOBJECT_API EPropertyVisitorControlFlow Visit(void* Data, const TFunctionRef<EPropertyVisitorControlFlow(const FPropertyVisitorPath& /*Path*/, void* /*Data*/)> InFunc) const;
+	COREUOBJECT_API EPropertyVisitorControlFlow Visit(void* Data, const TFunctionRef<EPropertyVisitorControlFlow(const FPropertyVisitorPath& /*Path*/, const FPropertyVisitorData& /*Data*/)> InFunc) const;
 
 	/**
 	 * Visits this property and allows recursion into the inner properties
@@ -753,7 +753,7 @@ public:
 	 * @param InFunc to call on each visited property, the return value controls what is the next behavior once this property has been visited
 	 * @return the new action to take one visited this property
 	 */
-	COREUOBJECT_API virtual EPropertyVisitorControlFlow Visit(FPropertyVisitorPath& Path, void* Data, const TFunctionRef<EPropertyVisitorControlFlow(const FPropertyVisitorPath& /*Path*/, void* /*Data*/)> InFunc) const;
+	COREUOBJECT_API virtual EPropertyVisitorControlFlow Visit(FPropertyVisitorPath& Path, const FPropertyVisitorData& Data, const TFunctionRef<EPropertyVisitorControlFlow(const FPropertyVisitorPath& /*Path*/, const FPropertyVisitorData& /*Data*/)> InFunc) const;
 
 	/**
 	 * Attempt to resolve the given inner path info against this outer struct to get the inner property value.
@@ -1263,7 +1263,7 @@ public:
 			return GetCapabilities().HasVisitor;
 		}
 		/** Structs property visitor signature */
-		virtual EPropertyVisitorControlFlow Visit(FPropertyVisitorPath& Path, void* Data, const TFunctionRef<EPropertyVisitorControlFlow(const FPropertyVisitorPath& /*Path*/, void* /*Data*/)> InFunc) const = 0;
+		virtual EPropertyVisitorControlFlow Visit(FPropertyVisitorPath& Path, const FPropertyVisitorData& Data, const TFunctionRef<EPropertyVisitorControlFlow(const FPropertyVisitorPath& /*Path*/, const FPropertyVisitorData& /*Data*/)> InFunc) const = 0;
 		virtual void* ResolveVisitedPathInfo(void* Data, const FPropertyVisitorInfo& Info) const = 0;
 
 	private:
@@ -1652,11 +1652,12 @@ public:
 		}
 #endif // WITH_EDITOR
 
-		virtual EPropertyVisitorControlFlow Visit(FPropertyVisitorPath& Path, void* Data, const TFunctionRef<EPropertyVisitorControlFlow(const FPropertyVisitorPath& /*Path*/, void* /*Data*/)> InFunc) const override
+		virtual EPropertyVisitorControlFlow Visit(FPropertyVisitorPath& Path, const FPropertyVisitorData& Data, const TFunctionRef<EPropertyVisitorControlFlow(const FPropertyVisitorPath& /*Path*/, const FPropertyVisitorData& /*Data*/)> InFunc) const override
 		{
 			if constexpr (TStructOpsTypeTraits<CPPSTRUCT>::WithVisitor)
 			{
-				return ((CPPSTRUCT*)Data)->Visit(Path, InFunc);
+				CPPSTRUCT* Struct = (CPPSTRUCT*)Data.PropertyData;
+				return Struct->Visit(Path, Data, InFunc);
 			}
 			else
 			{
@@ -1925,7 +1926,7 @@ public:
 
 	/* Custom visit implementation for structs */
 	using Super::Visit;
-	virtual COREUOBJECT_API EPropertyVisitorControlFlow Visit(FPropertyVisitorPath& Path, void* Data, const TFunctionRef<EPropertyVisitorControlFlow(const FPropertyVisitorPath& /*Path*/, void* /*Data*/)> InFunc) const override;
+	virtual COREUOBJECT_API EPropertyVisitorControlFlow Visit(FPropertyVisitorPath& Path, const FPropertyVisitorData& Data, const TFunctionRef<EPropertyVisitorControlFlow(const FPropertyVisitorPath& /*Path*/, const FPropertyVisitorData& /*Data*/)> InFunc) const override;
 	virtual COREUOBJECT_API void* ResolveVisitedPathInfo(void* Data, const FPropertyVisitorInfo& Info) const override;
 };
 
diff --git a/Engine/Source/Runtime/CoreUObject/Public/UObject/PropertyOptional.h b/Engine/Source/Runtime/CoreUObject/Public/UObject/PropertyOptional.h
--- a/Engine/Source/Runtime/CoreUObject/Public/UObject/PropertyOptional.h
+++ b/Engine/Source/Runtime/CoreUObject/Public/UObject/PropertyOptional.h
@@ -239,7 +239,7 @@ public:
 	virtual bool HasIntrusiveUnsetOptionalState() const override;
 	virtual bool SameType(const FProperty* Other) const override;
 
-	virtual EPropertyVisitorControlFlow Visit(FPropertyVisitorPath& Path, void* Data, const TFunctionRef<EPropertyVisitorControlFlow(const FPropertyVisitorPath& /*Path*/, void* /*Data*/)> InFunc) const override;
+	virtual EPropertyVisitorControlFlow Visit(FPropertyVisitorPath& Path, const FPropertyVisitorData& Data, const TFunctionRef<EPropertyVisitorControlFlow(const FPropertyVisitorPath& /*Path*/, const FPropertyVisitorData& /*Data*/)> InFunc) const override;
 	virtual void* ResolveVisitedPathInfo(void* Data, const FPropertyVisitorInfo& Info) const override;
 	// End of FProperty interface
 };
diff --git a/Engine/Source/Runtime/CoreUObject/Public/UObject/PropertyVisitor.h b/Engine/Source/Runtime/CoreUObject/Public/UObject/PropertyVisitor.h
--- a/Engine/Source/Runtime/CoreUObject/Public/UObject/PropertyVisitor.h
+++ b/Engine/Source/Runtime/CoreUObject/Public/UObject/PropertyVisitor.h
@@ -223,6 +223,25 @@ protected:
 	FPropertyVisitorPath& Path;
 };
 
+struct FPropertyVisitorData
+{
+	explicit FPropertyVisitorData(void* InPropertyData, void* InParentStructData)
+		: PropertyData(InPropertyData)
+		, ParentStructData(InParentStructData)
+	{}
+
+	/** Utility that constructs a new visitor data object with new property data but the same parent struct data */
+	FPropertyVisitorData VisitPropertyData(void* InPropertyData) const
+	{
+		return FPropertyVisitorData(InPropertyData, ParentStructData);
+	}
+
+	/** Data associated with the property being iterated */
+	void* PropertyData = nullptr;
+	/** Data associated with the parent struct that provided the property being iterated */
+	void* ParentStructData = nullptr;
+};
+
 namespace PropertyVisitorHelpers
 {
 
@@ -266,12 +285,18 @@ template <typename Type>
 void* ResolveVisitedPathInfo_Generic(Type* This, FPropertyVisitorPath& Path, void* Data, const FPropertyVisitorInfo& Info)
 {
 	void* FoundInnerData = nullptr;
-	This->Visit(Path, Data, [&FoundInnerData, &Info, InnerPathDepth = Path.Num() + 1](const FPropertyVisitorPath& InnerPath, void* InnerData)
+
+	FPropertyVisitorData VisitorData(Data, /*ParentStructData*/nullptr);
+
+	This->Visit(Path, VisitorData, [&FoundInnerData, &Info, InnerPathDepth = Path.Num() + 1](const FPropertyVisitorPath& InnerPath, const FPropertyVisitorData& InnerVisitorData)
 	{
 		if (InnerPath.Num() < InnerPathDepth)
 		{
 			return EPropertyVisitorControlFlow::StepInto;
 		}
+
+		void* InnerData = InnerVisitorData.PropertyData;
+
 		if (Info == InnerPath.Top())
 		{
 			FoundInnerData = InnerData;
diff --git a/Engine/Source/Runtime/CoreUObject/Public/UObject/UnrealType.h b/Engine/Source/Runtime/CoreUObject/Public/UObject/UnrealType.h
--- a/Engine/Source/Runtime/CoreUObject/Public/UObject/UnrealType.h
+++ b/Engine/Source/Runtime/CoreUObject/Public/UObject/UnrealType.h
@@ -353,7 +353,7 @@ PRAGMA_ENABLE_DEPRECATION_WARNINGS
 	 * @param InFunc to call on each visited property, the return value controls what is the next behavior once this property has been visited
 	 * @return the new action to take one visited this property
 	 */
-	COREUOBJECT_API EPropertyVisitorControlFlow Visit(void* Data, const TFunctionRef<EPropertyVisitorControlFlow(const FPropertyVisitorPath& /*Path*/, void* /*Data*/)> InFunc) const;
+	COREUOBJECT_API EPropertyVisitorControlFlow Visit(const FPropertyVisitorData& Data, const TFunctionRef<EPropertyVisitorControlFlow(const FPropertyVisitorPath& /*Path*/, const FPropertyVisitorData& /*Data*/)> InFunc) const;
 
 	/**
 	 * Visits this property and allows recursion into the inner properties
@@ -364,7 +364,7 @@ PRAGMA_ENABLE_DEPRECATION_WARNINGS
 	 * @param InFunc to call on each visited property, the return value controls what is the next behavior once this property has been visited
 	 * @return the new action to take one visited this property
 	 */
-	COREUOBJECT_API virtual EPropertyVisitorControlFlow Visit(FPropertyVisitorPath& Path, void* Data, const TFunctionRef<EPropertyVisitorControlFlow(const FPropertyVisitorPath& /*Path*/, void* /*Data*/)> InFunc) const;
+	COREUOBJECT_API virtual EPropertyVisitorControlFlow Visit(FPropertyVisitorPath& Path, const FPropertyVisitorData& Data, const TFunctionRef<EPropertyVisitorControlFlow(const FPropertyVisitorPath& /*Path*/, const FPropertyVisitorData& /*Data*/)> InFunc) const;
 
 	/**
 	 * Attempt to resolve the given inner path info against this outer property to get the inner property value.
@@ -2586,7 +2586,7 @@ public:
 	virtual FName GetID() const override;
 	virtual void InstanceSubobjects( void* Data, void const* DefaultData, UObject* Owner, struct FObjectInstancingGraph* InstanceGraph ) override;
 	virtual bool SameType(const FProperty* Other) const override;
-	virtual EPropertyVisitorControlFlow Visit(FPropertyVisitorPath& Path, void* Data, const TFunctionRef<EPropertyVisitorControlFlow(const FPropertyVisitorPath& /*Path*/, void* /*Data*/)> InFunc) const override;
+	virtual EPropertyVisitorControlFlow Visit(FPropertyVisitorPath& Path, const FPropertyVisitorData& Data, const TFunctionRef<EPropertyVisitorControlFlow(const FPropertyVisitorPath& /*Path*/, const FPropertyVisitorData& /*Data*/)> InFunc) const override;
 	virtual void* ResolveVisitedPathInfo(void* Data, const FPropertyVisitorInfo& Info) const override;
 	// End of FProperty interface
 
@@ -3563,7 +3563,7 @@ public:
 	virtual bool LoadTypeName(UE::FPropertyTypeName Type, const FPropertyTag* Tag = nullptr) override;
 	virtual void SaveTypeName(UE::FPropertyTypeNameBuilder& Type) const override;
 	virtual bool CanSerializeFromTypeName(UE::FPropertyTypeName Type) const override;
-	virtual EPropertyVisitorControlFlow Visit(FPropertyVisitorPath& Path, void* Data, const TFunctionRef<EPropertyVisitorControlFlow(const FPropertyVisitorPath& /*Path*/, void* /*Data*/)> InFunc) const override;
+	virtual EPropertyVisitorControlFlow Visit(FPropertyVisitorPath& Path, const FPropertyVisitorData& Data, const TFunctionRef<EPropertyVisitorControlFlow(const FPropertyVisitorPath& /*Path*/, const FPropertyVisitorData& /*Data*/)> InFunc) const override;
 	virtual void* ResolveVisitedPathInfo(void* Data, const FPropertyVisitorInfo& Info) const override;
 	// End of FProperty interface
 
@@ -3696,7 +3696,7 @@ public:
 	virtual bool LoadTypeName(UE::FPropertyTypeName Type, const FPropertyTag* Tag = nullptr) override;
 	virtual void SaveTypeName(UE::FPropertyTypeNameBuilder& Type) const override;
 	virtual bool CanSerializeFromTypeName(UE::FPropertyTypeName Type) const override;
-	virtual EPropertyVisitorControlFlow Visit(FPropertyVisitorPath& Path, void* Data, const TFunctionRef<EPropertyVisitorControlFlow(const FPropertyVisitorPath& /*Path*/, void* /*Data*/)> InFunc) const override;
+	virtual EPropertyVisitorControlFlow Visit(FPropertyVisitorPath& Path, const FPropertyVisitorData& Data, const TFunctionRef<EPropertyVisitorControlFlow(const FPropertyVisitorPath& /*Path*/, const FPropertyVisitorData& /*Data*/)> InFunc) const override;
 	virtual void* ResolveVisitedPathInfo(void* Data, const FPropertyVisitorInfo& Info) const override;
 	// End of FProperty interface
 
@@ -3830,7 +3830,7 @@ public:
 	virtual bool LoadTypeName(UE::FPropertyTypeName Type, const FPropertyTag* Tag = nullptr) override;
 	virtual void SaveTypeName(UE::FPropertyTypeNameBuilder& Type) const override;
 	virtual bool CanSerializeFromTypeName(UE::FPropertyTypeName Type) const override;
-	virtual EPropertyVisitorControlFlow Visit(FPropertyVisitorPath& Path, void* Data, const TFunctionRef<EPropertyVisitorControlFlow(const FPropertyVisitorPath& /*Path*/, void* /*Data*/)> InFunc) const override;
+	virtual EPropertyVisitorControlFlow Visit(FPropertyVisitorPath& Path, const FPropertyVisitorData& Data, const TFunctionRef<EPropertyVisitorControlFlow(const FPropertyVisitorPath& /*Path*/, const FPropertyVisitorData& /*Data*/)> InFunc) const override;
 	virtual void* ResolveVisitedPathInfo(void* Data, const FPropertyVisitorInfo& Info) const override;
 	// End of FProperty interface
 
@@ -6024,7 +6024,7 @@ public:
 	virtual bool LoadTypeName(UE::FPropertyTypeName Type, const FPropertyTag* Tag = nullptr) override;
 	virtual void SaveTypeName(UE::FPropertyTypeNameBuilder& Type) const override;
 	virtual bool CanSerializeFromTypeName(UE::FPropertyTypeName Type) const override;
-	virtual EPropertyVisitorControlFlow Visit(FPropertyVisitorPath& Path, void* Data, const TFunctionRef<EPropertyVisitorControlFlow(const FPropertyVisitorPath& /*Path*/, void* /*Data*/)> InFunc) const override;
+	virtual EPropertyVisitorControlFlow Visit(FPropertyVisitorPath& Path, const FPropertyVisitorData& Data, const TFunctionRef<EPropertyVisitorControlFlow(const FPropertyVisitorPath& /*Path*/, const FPropertyVisitorData& /*Data*/)> InFunc) const override;
 	virtual void* ResolveVisitedPathInfo(void* Data, const FPropertyVisitorInfo& Info) const override;
 	// End of FProperty interface
 
