[AutoRTFM] Disable memory validation for hot code

These open blocks have not upset the validator, but are in extremely hot code paths.
Disable validation for these to improve performance.

diff --git a/Engine/Source/Runtime/Core/Private/HAL/ThreadingBase.cpp b/Engine/Source/Runtime/Core/Private/HAL/ThreadingBase.cpp
--- a/Engine/Source/Runtime/Core/Private/HAL/ThreadingBase.cpp
+++ b/Engine/Source/Runtime/Core/Private/HAL/ThreadingBase.cpp
@@ -1457,7 +1457,9 @@ FQueuedThread::Run()
 FTlsAutoCleanup* FThreadSingletonInitializer::Get( TFunctionRef<FTlsAutoCleanup*()> CreateInstance, uint32& InOutTlsSlot )
 {
 	uint32 TlsSlot;
-	UE_AUTORTFM_OPEN
+	FTlsAutoCleanup* ThreadSingleton = nullptr;
+	// Memory validation disabled as this is hot code.
+	UE_AUTORTFM_OPEN_NO_VALIDATION
 	{
 		TlsSlot = (uint32)FPlatformAtomics::AtomicRead_Relaxed((int32*)&InOutTlsSlot);
 		if (TlsSlot == FPlatformTLS::InvalidTlsSlot)
@@ -1475,11 +1477,7 @@ FTlsAutoCleanup* FThreadSingletonInitializer::Get( TFunctionRef<FTlsAutoCleanup*
 				TlsSlot = ThisTlsSlot;
 			}
 		}
-	};
 
-	FTlsAutoCleanup* ThreadSingleton = nullptr;
-	UE_AUTORTFM_OPEN
-	{
 		ThreadSingleton = (FTlsAutoCleanup*)FPlatformTLS::GetTlsValue( TlsSlot );
 		if( !ThreadSingleton )
 		{
diff --git a/Engine/Source/Runtime/Core/Public/Delegates/MulticastDelegateBase.h b/Engine/Source/Runtime/Core/Public/Delegates/MulticastDelegateBase.h
--- a/Engine/Source/Runtime/Core/Public/Delegates/MulticastDelegateBase.h
+++ b/Engine/Source/Runtime/Core/Public/Delegates/MulticastDelegateBase.h
@@ -382,14 +382,16 @@ private:
 	/** Increments the lock counter for the invocation list. */
 	inline void LockInvocationList( ) const
 	{
-		UE_AUTORTFM_OPEN { ++InvocationListLockCount; };
+		// Validation disabled as the only closed write is the ctor.
+		UE_AUTORTFM_OPEN_NO_VALIDATION { ++InvocationListLockCount; };
 		UE_AUTORTFM_ONABORT(this) { --InvocationListLockCount; };
 	}
 
 	/** Decrements the lock counter for the invocation list. */
 	inline void UnlockInvocationList( ) const
 	{
-		UE_AUTORTFM_OPEN { --InvocationListLockCount; };
+		// Validation disabled as the only closed write is the ctor.
+		UE_AUTORTFM_OPEN_NO_VALIDATION { --InvocationListLockCount; };
 		UE_AUTORTFM_ONABORT(this) { ++InvocationListLockCount; };
 	}
 
diff --git a/Engine/Source/Runtime/Core/Public/HAL/FMemory.inl b/Engine/Source/Runtime/Core/Public/HAL/FMemory.inl
--- a/Engine/Source/Runtime/Core/Public/HAL/FMemory.inl
+++ b/Engine/Source/Runtime/Core/Public/HAL/FMemory.inl
@@ -22,7 +22,8 @@ FMEMORY_INLINE_FUNCTION_DECORATOR void* FMemory::Malloc(SIZE_T Count, uint32 Ali
 	// we don't end up keeping track of the writes to the allocator's internal data structures.
 	// This is because allocators are already transactional - malloc can be rolled back by
 	// calling free.
-	void* Ptr = AutoRTFM::Open([Count, Alignment]
+	// Memory validation disabled as this is hot code.
+	void* Ptr = AutoRTFM::Open<AutoRTFM::EMemoryValidationLevel::Disabled>([Count, Alignment]
 	{
 		void* Alloc = nullptr;
 		if (!FMEMORY_INLINE_GMalloc)
@@ -235,7 +236,8 @@ FMEMORY_INLINE_FUNCTION_DECORATOR void* FMemory::MallocZeroed(SIZE_T Count, uint
 FMEMORY_INLINE_FUNCTION_DECORATOR SIZE_T FMemory::QuantizeSize(SIZE_T Count, uint32 Alignment)
 {
 	SIZE_T Result;
-	UE_AUTORTFM_OPEN
+	// Memory validation disabled as this is hot code.
+	UE_AUTORTFM_OPEN_NO_VALIDATION
 	{
 		if (!FMEMORY_INLINE_GMalloc)
 		{
diff --git a/Engine/Source/Runtime/Core/Public/ProfilingDebugging/CsvProfiler.h b/Engine/Source/Runtime/Core/Public/ProfilingDebugging/CsvProfiler.h
--- a/Engine/Source/Runtime/Core/Public/ProfilingDebugging/CsvProfiler.h
+++ b/Engine/Source/Runtime/Core/Public/ProfilingDebugging/CsvProfiler.h
@@ -617,14 +617,16 @@ public:
 	FScopedCsvStatExclusive(const char * InStatName, const char* InNamedEventName = nullptr)
 		: StatName(InStatName)
 	{
-		UE_AUTORTFM_OPEN{ FCsvProfiler::BeginExclusiveStat(StatName, InNamedEventName); };
+		// Memory validation disabled as this is hot code.
+		UE_AUTORTFM_OPEN_NO_VALIDATION{ FCsvProfiler::BeginExclusiveStat(StatName, InNamedEventName); };
 		AutoRTFM::PushOnAbortHandler(this, [SN = StatName](){ FCsvProfiler::EndExclusiveStat(SN); });
 	}
 
 	~FScopedCsvStatExclusive()
 	{
 		AutoRTFM::PopOnAbortHandler(this);
-		UE_AUTORTFM_OPEN{ FCsvProfiler::EndExclusiveStat(StatName); };
+		// Memory validation disabled as this is hot code.
+		UE_AUTORTFM_OPEN_NO_VALIDATION{ FCsvProfiler::EndExclusiveStat(StatName); };
 	}
 	const char * StatName;
 };
diff --git a/Engine/Source/Runtime/Core/Public/Templates/SharedPointerInternals.h b/Engine/Source/Runtime/Core/Public/Templates/SharedPointerInternals.h
--- a/Engine/Source/Runtime/Core/Public/Templates/SharedPointerInternals.h
+++ b/Engine/Source/Runtime/Core/Public/Templates/SharedPointerInternals.h
@@ -81,7 +81,8 @@ namespace SharedPointerInternals
 
 				int32 Count = 0;
 
-				UE_AUTORTFM_OPEN
+				// Memory validation disabled as this is hot code.
+				UE_AUTORTFM_OPEN_NO_VALIDATION
 					{
 						// This reference count may be accessed by multiple threads
 						Count = SharedReferenceCount.load(std::memory_order_relaxed);
@@ -111,14 +112,16 @@ namespace SharedPointerInternals
 				// in response to the increment, so there's nothing to order with.
 
 #if defined(_MSC_VER) && (defined(_M_X64) || defined(_M_IX86))
-				UE_AUTORTFM_OPEN
+				// Memory validation disabled as this is hot code.
+				UE_AUTORTFM_OPEN_NO_VALIDATION
 				{
 					// We do a regular SC increment here because it maps to an _InterlockedIncrement (lock inc).
 					// The codegen for a relaxed fetch_add is actually much worse under MSVC (lock xadd).
 					++SharedReferenceCount;
 				};
 #else
-				UE_AUTORTFM_OPEN
+				// Memory validation disabled as this is hot code.
+				UE_AUTORTFM_OPEN_NO_VALIDATION
 				{
 					SharedReferenceCount.fetch_add(1, std::memory_order_relaxed);
 				};
@@ -147,7 +150,8 @@ namespace SharedPointerInternals
 			{
 				bool bSucceeded = false;
 
-				UE_AUTORTFM_OPEN
+				// Memory validation disabled as this is hot code.
+				UE_AUTORTFM_OPEN_NO_VALIDATION
 				{
 					// See AddSharedReference for the same reasons that std::memory_order_relaxed is used in this function.
 
@@ -256,14 +260,16 @@ namespace SharedPointerInternals
 				// See AddSharedReference for the same reasons that std::memory_order_relaxed is used in this function.
 
 #if defined(_MSC_VER) && (defined(_M_X64) || defined(_M_IX86))
-				UE_AUTORTFM_OPEN
+				// Memory validation disabled as this is hot code.
+				UE_AUTORTFM_OPEN_NO_VALIDATION
 					{
 						// We do a regular SC increment here because it maps to an _InterlockedIncrement (lock inc).
 						// The codegen for a relaxed fetch_add is actually much worse under MSVC (lock xadd).
 						++WeakReferenceCount;
 					};
 #else
-				UE_AUTORTFM_OPEN
+				// Memory validation disabled as this is hot code.
+				UE_AUTORTFM_OPEN_NO_VALIDATION
 					{
 						WeakReferenceCount.fetch_add(1, std::memory_order_relaxed);
 					};
diff --git a/Engine/Source/Runtime/CoreUObject/Public/UObject/UObjectArray.h b/Engine/Source/Runtime/CoreUObject/Public/UObject/UObjectArray.h
--- a/Engine/Source/Runtime/CoreUObject/Public/UObject/UObjectArray.h
+++ b/Engine/Source/Runtime/CoreUObject/Public/UObject/UObjectArray.h
@@ -172,7 +172,8 @@ public:
 	{
 		checkf((int32(FlagToClear) & ~int32(EInternalObjectFlags_AllFlags)) == 0, TEXT("%d is not a valid internal flag value"), int32(FlagToClear));
 		bool Result = false;
-		UE_AUTORTFM_OPEN
+		// Memory validation disabled as this is hot code.
+		UE_AUTORTFM_OPEN_NO_VALIDATION
 		{
 			FlagToClear &= ~EInternalObjectFlags_ReachabilityFlags; // reachability flags can only be cleared by GC through *_ForGC functions
 			FlagToClear &= ~EInternalObjectFlags::RefCounted; // refcounted flag is internal and must only be cleared internally by AddRef/ReleaseRef.
@@ -209,7 +210,8 @@ public:
 	{
 		checkf((int32(FlagToSet) & ~int32(EInternalObjectFlags_AllFlags)) == 0, TEXT("%d is not a valid internal flag value"), int32(FlagToSet));
 		bool Result = false;
-		UE_AUTORTFM_OPEN
+		// Memory validation disabled as this is hot code.
+		UE_AUTORTFM_OPEN_NO_VALIDATION
 		{
 			FlagToSet &= ~EInternalObjectFlags_ReachabilityFlags; // reachability flags can only be cleared by GC through *_ForGC functions
 			FlagToSet &= ~EInternalObjectFlags::RefCounted; // refcounted flag is internal and must only be set by AddRef/ReleaseRef.
@@ -310,7 +312,8 @@ public:
 
 	void AddRef()
 	{
-		UE_AUTORTFM_OPEN
+		// Memory validation disabled as this is hot code.
+		UE_AUTORTFM_OPEN_NO_VALIDATION
 		{
 			FPlatformAtomics::InterlockedIncrement(&RefCount);
 			if ((GetFlags() & EInternalObjectFlags::RefCounted) != EInternalObjectFlags::RefCounted)
@@ -322,7 +325,8 @@ public:
 
 	void ReleaseRef()
 	{
-		UE_AUTORTFM_OPEN
+		// Memory validation disabled as this is hot code.
+		UE_AUTORTFM_OPEN_NO_VALIDATION
 		{
 			// This alone is not thread-safe as we may race with AddRef and in that case we don't want ClearRootFlags to apply.
 			// We fix this by validating that the refcount is still 0 while inside the root locks in ClearRootFlags.
diff --git a/Engine/Source/Runtime/CoreUObject/Public/UObject/UObjectBaseUtility.h b/Engine/Source/Runtime/CoreUObject/Public/UObject/UObjectBaseUtility.h
--- a/Engine/Source/Runtime/CoreUObject/Public/UObject/UObjectBaseUtility.h
+++ b/Engine/Source/Runtime/CoreUObject/Public/UObject/UObjectBaseUtility.h
@@ -777,7 +777,8 @@ public:
 	FORCEINLINE TStatId GetStatID(bool bForDeferredUse = false) const
 	{
 		TStatId Result{};
-		AutoRTFM::Open([&]
+		// Memory validation disabled as this is hot code.
+		UE_AUTORTFM_OPEN_NO_VALIDATION
 		{
 #if STATS
 			const TStatId& StatID = GUObjectArray.IndexToObject(InternalIndex)->StatID;
@@ -805,7 +806,7 @@ public:
 #else
 			Result = TStatId(); // not doing stats at the moment, or ever
 #endif // STATS
-		});
+		};
 		return Result;
 	}
 
@@ -943,7 +944,8 @@ public:
 	 */
 	FORCEINLINE_STATS FScopeCycleCounterUObject(const UObjectBaseUtility *Object)
 	{
-		AutoRTFM::Open([&]
+		// Memory validation disabled as this is hot code.
+		UE_AUTORTFM_OPEN_NO_VALIDATION
 		{
 			if (Object)
 			{
@@ -962,7 +964,7 @@ public:
 				}
 #endif
 			}
-		});
+		};
 
 		AutoRTFM::PushOnAbortHandler(this, [this](){ this->Stop(); });
 	}
@@ -971,7 +973,8 @@ public:
 	 */
 	FORCEINLINE_STATS FScopeCycleCounterUObject(const UObjectBaseUtility *Object, TStatId OtherStat)
 	{
-		AutoRTFM::Open([&]
+		// Memory validation disabled as this is hot code.
+		UE_AUTORTFM_OPEN_NO_VALIDATION
 		{
 			if (Object)
 			{
@@ -993,7 +996,7 @@ public:
 				}
 #endif
 			}
-		});
+		};
 
 		AutoRTFM::PushOnAbortHandler(this, [this](){ this->Stop(); });
 	}
@@ -1008,10 +1011,11 @@ public:
 	FORCEINLINE_STATS ~FScopeCycleCounterUObject()
 	{
 		AutoRTFM::PopOnAbortHandler(this);
-		AutoRTFM::Open([&]
+		// Memory validation disabled as this is hot code.
+		UE_AUTORTFM_OPEN_NO_VALIDATION
 		{
 			Stop();
-		});
+		};
 	}
 };
 
@@ -1036,7 +1040,8 @@ public:
 #endif
 	{
 #if ENABLE_STATNAMEDEVENTS_UOBJECT && CPUPROFILERTRACE_ENABLED
-		AutoRTFM::Open([&]
+		// Memory validation disabled as this is hot code.
+		UE_AUTORTFM_OPEN_NO_VALIDATION
 		{
 			if (GCycleStatsShouldEmitNamedEvents && UE_TRACE_CHANNELEXPR_IS_ENABLED(CpuChannel) && Object)
 			{
@@ -1047,7 +1052,7 @@ public:
 					FCpuProfilerTrace::OutputBeginDynamicEvent(ObjectStatId.StatString);
 				}
 			}
-		});
+		};
 
 		if(bPop)
 		{
@@ -1066,13 +1071,14 @@ public:
 #if ENABLE_STATNAMEDEVENTS_UOBJECT && CPUPROFILERTRACE_ENABLED
 		AutoRTFM::PopOnAbortHandler(this);
 
-		AutoRTFM::Open([&]
+		// Memory validation disabled as this is hot code.
+		UE_AUTORTFM_OPEN_NO_VALIDATION
 		{
 			if (bPop)
 			{
 				FCpuProfilerTrace::OutputEndEvent();
 			}
-		});
+		};
 #endif
 	}
 };
@@ -1096,10 +1102,11 @@ public:
 	{
 		if (GHitchDetected &&  StatObject)
 		{
-			AutoRTFM::Open([&]
+			// Memory validation disabled as this is hot code.
+			UE_AUTORTFM_OPEN_NO_VALIDATION
 			{
 				ReportHitch();
-			});
+			};
 		}
 	}
 
