Fixed assignment of IndexInOwner for properties of native types

#rnx

diff --git a/Engine/Source/Runtime/CoreUObject/Private/UObject/Class.cpp b/Engine/Source/Runtime/CoreUObject/Private/UObject/Class.cpp
--- a/Engine/Source/Runtime/CoreUObject/Private/UObject/Class.cpp
+++ b/Engine/Source/Runtime/CoreUObject/Private/UObject/Class.cpp
@@ -930,6 +930,11 @@ void UStruct::Link(FArchive& Ar, bool bRelinkExistingProperties)
 
 					PropertiesSize = Property->Link(Ar);
 
+				#if WITH_EDITORONLY_DATA
+					Property->SetIndexInOwner(TotalFieldCount);
+					TotalFieldCount += Property->ArrayDim;
+				#endif
+
 					if ((bPropertyIsTransient != Property->HasAllFlags(RF_Transient)) || (PropertyName != Property->GetFName()))
 					{
 						LoopNum++;
@@ -974,11 +979,24 @@ void UStruct::Link(FArchive& Ar, bool bRelinkExistingProperties)
 	}
 	else
 	{
+	#if WITH_EDITORONLY_DATA
+		TotalFieldCount = 0;
+		if (UStruct* InheritanceSuper = GetInheritanceSuper())
+		{
+			TotalFieldCount = InheritanceSuper->TotalFieldCount;
+		}
+	#endif
+
 		for (FField* Field = ChildProperties; (Field != NULL) && (Field->GetOwner<UObject>() == this); Field = Field->Next)
 		{
 			if (FProperty* Property = CastField<FProperty>(Field))
 			{
 				Property->LinkWithoutChangingOffset(Ar);
+
+			#if WITH_EDITORONLY_DATA
+				Property->SetIndexInOwner(TotalFieldCount);
+				TotalFieldCount += Property->ArrayDim;
+			#endif
 			}
 		}
 	}
diff --git a/Engine/Source/Runtime/CoreUObject/Private/UObject/Property.cpp b/Engine/Source/Runtime/CoreUObject/Private/UObject/Property.cpp
--- a/Engine/Source/Runtime/CoreUObject/Private/UObject/Property.cpp
+++ b/Engine/Source/Runtime/CoreUObject/Private/UObject/Property.cpp
@@ -1443,17 +1443,10 @@ int32 FProperty::SetupOffset()
 	{
 		UStruct* OwnerStruct = (UStruct*)OwnerUObject;
 		Offset_Internal = Align(OwnerStruct->GetPropertiesSize(), GetMinAlignment());
-	#if WITH_EDITORONLY_DATA
-		IndexInOwner = OwnerStruct->TotalFieldCount;
-		OwnerStruct->TotalFieldCount += ArrayDim;
-	#endif
 	}
 	else
 	{
 		Offset_Internal = Align(0, GetMinAlignment());
-	#if WITH_EDITORONLY_DATA
-		IndexInOwner = 0;
-	#endif
 	}
 
 	uint32 UnsignedTotal = (uint32)Offset_Internal + (uint32)GetSize();
diff --git a/Engine/Source/Runtime/CoreUObject/Public/UObject/UnrealType.h b/Engine/Source/Runtime/CoreUObject/Public/UObject/UnrealType.h
--- a/Engine/Source/Runtime/CoreUObject/Public/UObject/UnrealType.h
+++ b/Engine/Source/Runtime/CoreUObject/Public/UObject/UnrealType.h
@@ -203,7 +203,7 @@ private:
 
 #if WITH_EDITORONLY_DATA
 	/** Index of the property within its owner, inclusive of base properties. Generated during Link(). */
-	int32		IndexInOwner = 0;
+	int32		IndexInOwner = -1;
 #endif
 
 	// In memory variables (generated during Link()).
@@ -392,11 +392,17 @@ protected:
 
 public:
 #if WITH_EDITORONLY_DATA
-	/** Return the index of the property in its owner. Only valid after Link(). */
+	/** Return the index of the property in its owner. Only valid after Link(). -1 if not linked by a UStruct. */
 	UE_INTERNAL FORCEINLINE int32 GetIndexInOwner() const
 	{
 		return IndexInOwner;
 	}
+
+	/** Set the index of the property in its owner. Set by the owner in UStruct::Link(). */
+	UE_INTERNAL FORCEINLINE void SetIndexInOwner(int32 Index)
+	{
+		IndexInOwner = Index;
+	}
 #endif
 
 	/** Return offset of property from container base. */
