Add a default data-driven subobject instancing path that does not rely on 'Instanced' property flags on generated types.

Allows for decoupling of compile-time 'Instanced' semantics to support non-native UE class types that do not statically bind references for subobject instancing.

Notes:
- This path is disabled by default; UClass subtypes can opt-in via override.

diff --git a/Engine/Source/Runtime/CoreUObject/Private/UObject/Class.cpp b/Engine/Source/Runtime/CoreUObject/Private/UObject/Class.cpp
--- a/Engine/Source/Runtime/CoreUObject/Private/UObject/Class.cpp
+++ b/Engine/Source/Runtime/CoreUObject/Private/UObject/Class.cpp
@@ -2630,9 +2630,13 @@ void UStruct::InstanceSubobjectTemplates( void* Data, void const* DefaultData, U
 	checkSlow(Data);
 	checkSlow(Owner);
 
+	// Dynamic instancing allows the instance graph to infer when to instance subobjects instead of requiring the
+	// explicit 'CPF_InstancedReference' flag. This also means the code below intentionally visits all properties
+	// in the reflink chain as a result, unless they are otherwise explicitly excluded through the instance graph.
+	const bool bUsesDynamicInstancing = InstanceGraph && InstanceGraph->CanUseDynamicInstancing();
 	for ( FProperty* Property = RefLink; Property != NULL; Property = Property->NextRef )
 	{
-		if ((Property->ContainsInstancedObjectProperty() || Property->HasAnyPropertyFlags(CPF_AllowSelfReference)) && 
+		if ((bUsesDynamicInstancing || Property->ContainsInstancedObjectProperty() || Property->HasAnyPropertyFlags(CPF_AllowSelfReference)) &&
 			(!InstanceGraph || !InstanceGraph->IsPropertyInSubobjectExclusionList(Property)))
 		{
 			Property->InstanceSubobjects( Property->ContainerPtrToValuePtr<uint8>(Data), (uint8*)Property->ContainerPtrToValuePtrForDefaults<uint8>(DefaultStruct, DefaultData), Owner, InstanceGraph );
@@ -5586,6 +5590,23 @@ FTopLevelAssetPath UClass::GetReinstancedClassPathName() const
 }
 #endif
 
+bool UClass::ShouldUseDynamicSubobjectInstancing() const
+{
+	// If any parent class supports dynamic instancing, this class also inherits that trait.
+	const UClass* CurrentClass = this;
+	while (CurrentClass)
+	{
+		if (CurrentClass->SupportsDynamicInstancedReference())
+		{
+			return true;
+		}
+
+		CurrentClass = CurrentClass->GetSuperClass();
+	}
+
+	return false;
+}
+
 #if USTRUCT_FAST_ISCHILDOF_IMPL == USTRUCT_ISCHILDOF_STRUCTARRAY
 
 	FStructBaseChain::FStructBaseChain()
diff --git a/Engine/Source/Runtime/CoreUObject/Private/UObject/CoreNative.cpp b/Engine/Source/Runtime/CoreUObject/Private/UObject/CoreNative.cpp
--- a/Engine/Source/Runtime/CoreUObject/Private/UObject/CoreNative.cpp
+++ b/Engine/Source/Runtime/CoreUObject/Private/UObject/CoreNative.cpp
@@ -133,20 +133,12 @@ FObjectInstancingGraph::FObjectInstancingGraph(bool bDisableInstancing)
 }
 
 FObjectInstancingGraph::FObjectInstancingGraph(EObjectInstancingGraphOptions InOptions)
-	: SourceRoot(nullptr)
-	, DestinationRoot(nullptr)
-	, InstancingOptions(InOptions)
-	, bCreatingArchetype(false)
-	, bLoadingObject(false)
+	: InstancingOptions(InOptions)
 {
 }
 
 FObjectInstancingGraph::FObjectInstancingGraph( UObject* DestinationSubobjectRoot, EObjectInstancingGraphOptions InOptions)
-	: SourceRoot(nullptr)
-	, DestinationRoot(nullptr)
-	, InstancingOptions(InOptions)
-	, bCreatingArchetype(false)
-	, bLoadingObject(false)
+	: InstancingOptions(InOptions)
 {
 	SetDestinationRoot(DestinationSubobjectRoot);
 }
@@ -171,6 +163,12 @@ void FObjectInstancingGraph::SetDestinationRoot(UObject* DestinationSubobjectRoo
 		//they should always have been created from the correct up-to-date template already.
 		bCreatingArchetype = false;
 	}
+
+	// If the destination root's underlying type should use dynamic instancing, enable it for this object graph.
+	if (DestinationSubobjectRoot->GetClass()->ShouldUseDynamicSubobjectInstancing())
+	{
+		bCanUseDynamicInstancing = true;
+	}
 }
 
 UObject* FObjectInstancingGraph::GetDestinationObject(UObject* SourceObject)
@@ -329,6 +327,9 @@ UObject* FObjectInstancingGraph::InstancePropertyValue(UObject* SubObjectTemplat
 	bool bCausesInstancing   = !!(Flags & EInstancePropertyValueFlags::CausesInstancing);
 	bool bAllowSelfReference = !!(Flags & EInstancePropertyValueFlags::AllowSelfReference);
 
+	checkf(Owner, TEXT("Null owner encountered while trying to instance subobject template %s. This is unexpected - a valid owner instance is required for instancing."),
+		*GetFullNameSafe(SubObjectTemplate));
+
 	checkf(CurrentValue, TEXT("Null object value encountered while trying to instance subobject of %s. Instanced object properties are never expected to be null (template: %s)."), 
 		*GetFullNameSafe(Owner), *GetNameSafe(SubObjectTemplate));
 
@@ -336,6 +337,30 @@ UObject* FObjectInstancingGraph::InstancePropertyValue(UObject* SubObjectTemplat
 	{
 		bCausesInstancing = true; // these are always instanced no matter what
 	}
+	else if (!bCausesInstancing && bCanUseDynamicInstancing)
+	{
+		// Dynamic instancing means that we'll analyze the current value to determine how instancing should proceed. At
+		// construction time for example, the current value will have been initialized to the value from the default data
+		// object, so in that case we'll compare that value against the owner's archetype to see if it should be instanced.
+		if (CurrentValue == SourceRoot)
+		{
+			// In this case, the current value was initialized to reference the source graph's root. The instancing
+			// graph will already contain the mapping from the archetype source to the root instance, so we'll allow
+			// instancing to proceed, where it will resolve the archetype and return the current root as the new value.
+			bAllowSelfReference = true;
+
+			// Must also set the flag here b/c it's required later on (unlike the instancing flag)
+			Flags |= EInstancePropertyValueFlags::AllowSelfReference;
+		}
+		else if (CurrentValue->IsIn(SourceRoot))
+		{
+			// In this case, the current value was initialized from the default data to reference an archetype that's
+			// exists within the source object graph. The instancing graph may not yet contain a mapping for this
+			// archetype, so we allow instancing to proceed, where it will either construct a new instance under the
+			// current owner, or return a reference to the instance that was already created for the source archetype.
+			bCausesInstancing = true;
+		}
+	}
 
 	if (!IsSubobjectInstancingEnabled() || // if instancing is off
 		(!bCausesInstancing && // or if this class isn't forced to be instanced and this var did not have the instance keyword
@@ -355,7 +380,7 @@ UObject* FObjectInstancingGraph::InstancePropertyValue(UObject* SubObjectTemplat
 	// if the object we're instancing the subobjects for (Owner) has the current subobject's outer in its archetype chain, and its archetype has a nullptr value
 	// for this subobject property it means that the archetype didn't instance its subobject, so we shouldn't either.
 
-	if (SubObjectTemplate == nullptr && (Owner != nullptr && Owner->IsBasedOnArchetype(CurrentValue->GetOuter())))
+	if (SubObjectTemplate == nullptr && Owner->IsBasedOnArchetype(CurrentValue->GetOuter()))
 	{
 		NewValue = nullptr;
 	}
diff --git a/Engine/Source/Runtime/CoreUObject/Public/UObject/Class.h b/Engine/Source/Runtime/CoreUObject/Public/UObject/Class.h
--- a/Engine/Source/Runtime/CoreUObject/Public/UObject/Class.h
+++ b/Engine/Source/Runtime/CoreUObject/Public/UObject/Class.h
@@ -3687,6 +3687,12 @@ public:
 	COREUOBJECT_API FTopLevelAssetPath GetReinstancedClassPathName() const;
 #endif
 
+	/**
+	 * Returns true if the object instancing graph should analyze default data at runtime to determine which of its object
+	 * properties need to be instanced at construction time, rather than exclusively relying on reflected property semantics.
+	 */
+	COREUOBJECT_API bool ShouldUseDynamicSubobjectInstancing() const;
+
 private:
 	/** 
 	 * This signature intentionally hides the method declared in UObjectBaseUtility to make it private.
@@ -3738,6 +3744,16 @@ protected:
 	 **/
 	COREUOBJECT_API virtual UObject* CreateDefaultObject();
 
+	/**
+	 * Subtypes can optionally override to return true. This will enable dynamic subobject
+	 * instancing, which will use default data to determine which properties need instancing
+	 * at construction/serialization time rather than reflected property 'Instanced' semantics.
+	 */
+	virtual bool SupportsDynamicInstancedReference() const
+	{
+		return false;
+	}
+
 #if WITH_EDITOR
 	/**
 	 * Internal helper method for GetReinstancedClassPathName(). Subclasses can choose to override the default
@@ -4037,6 +4053,14 @@ public:
 		return SubobjectInstantiationExclusionList.Contains(Property);
 	}
 
+	/**
+	 * Returns whether dynamic subobject instancing has been enabled for this object graph
+	 */
+	bool CanUseDynamicInstancing() const
+	{
+		return bCanUseDynamicInstancing;
+	}
+
 private:
 	/**
 	 * Returns whether DestinationRoot corresponds to an archetype object.
@@ -4084,7 +4108,7 @@ private:
 	 * - when creating an archetype, this would be the actor instance
 	 * - when duplicating an object, this would be the duplication source
 	 */
-	class		UObject*						SourceRoot;
+	class		UObject*						SourceRoot = nullptr;
 
 	/**
 	 * The root of the object tree that is the destination used for instancing components
@@ -4094,20 +4118,25 @@ private:
 	 * - when updating an archetype, this would be the source archetype
 	 * - when duplicating an object, this would be the copied object (destination)
 	 */
-	class		UObject*						DestinationRoot;
+	class		UObject*						DestinationRoot = nullptr;
 
 	/** Subobject instancing options */
-	EObjectInstancingGraphOptions				InstancingOptions;
+	EObjectInstancingGraphOptions				InstancingOptions = EObjectInstancingGraphOptions::None;
 
 	/**
 	 * Indicates whether we are currently instancing components for an archetype.  true if we are creating or updating an archetype.
 	 */
-	bool										bCreatingArchetype;
+	bool										bCreatingArchetype = false;
 
 	/**
 	 * true when loading object data from disk.
 	 */
-	bool										bLoadingObject;
+	bool										bLoadingObject = false;
+
+	/**
+	 * If true, allow instancing based on default data analysis at runtime when we can't instance based on reflected property semantics.
+	 */
+	bool										bCanUseDynamicInstancing = false;
 
 	/**
 	 * Maps the source (think archetype) to the destination (think instance)
