Check FName memory availability alongside UObject table availability

diff --git a/Engine/Source/Runtime/Core/Private/UObject/UnrealNames.cpp b/Engine/Source/Runtime/Core/Private/UObject/UnrealNames.cpp
--- a/Engine/Source/Runtime/Core/Private/UObject/UnrealNames.cpp
+++ b/Engine/Source/Runtime/Core/Private/UObject/UnrealNames.cpp
@@ -2567,6 +2567,12 @@ int32 FName::GetNameEntryMemorySize()
 	return GetNamePool().NumBlocks() * FNameEntryAllocator::BlockSizeBytes;
 }
 
+int32 FName::GetNameEntryMemoryEstimatedAvailable()
+{
+	FNamePool& Pool = GetNamePool();
+	return (FNameMaxBlocks - Pool.NumBlocks()) * FNameEntryAllocator::BlockSizeBytes;
+}
+
 int32 FName::GetNameTableMemorySize()
 {
 	return GetNameEntryMemorySize() + sizeof(FNamePool) + GetNamePool().NumSlots() * sizeof(FNameSlot);
diff --git a/Engine/Source/Runtime/Core/Public/UObject/NameTypes.h b/Engine/Source/Runtime/Core/Public/UObject/NameTypes.h
--- a/Engine/Source/Runtime/Core/Public/UObject/NameTypes.h
+++ b/Engine/Source/Runtime/Core/Public/UObject/NameTypes.h
@@ -1052,6 +1052,11 @@ public:
 	 */
 	CORE_API static int32 GetNameEntryMemorySize();
 
+	/**
+	 * @return Estimated remaining size the name entry table is willing to allocate.
+	 */
+	CORE_API static int32 GetNameEntryMemoryEstimatedAvailable();
+
 	/**
 	* @return Size of Name Table object as a whole
 	*/
diff --git a/Engine/Source/Runtime/CoreUObject/Private/UObject/UObjectArray.cpp b/Engine/Source/Runtime/CoreUObject/Private/UObject/UObjectArray.cpp
--- a/Engine/Source/Runtime/CoreUObject/Private/UObject/UObjectArray.cpp
+++ b/Engine/Source/Runtime/CoreUObject/Private/UObject/UObjectArray.cpp
@@ -635,6 +635,11 @@ bool verse::CanAllocateUObjects()
 	// NOTE: This is an arbitrary limit. If we have less than ~10k `UObject`s available for allocation left
 	// we're probably in a bad spot anyway. This just makes sure that there is some slack available before the
 	// limit gets hit.
+	// (The `FName` space requirement of 5MB, out of a maximum of 1GB, is chosen to match the ratio of 10k to
+	// the default maximum of 2M `UObject`s. It is checked separately because we have observed islands running
+	// out of `FName` space before hitting this `UObject` limit.)
 	static constexpr int32 MinAvailableObjectCount = 10 * 1024;
-	return GUObjectArray.GetObjectArrayEstimatedAvailable() >= MinAvailableObjectCount;
+	static constexpr int32 MinAvailableNameEntrySize = 5 * 1024 * 1024;
+	return GUObjectArray.GetObjectArrayEstimatedAvailable() >= MinAvailableObjectCount &&
+		FName::GetNameEntryMemoryEstimatedAvailable() >= MinAvailableNameEntrySize;
 }
