Restoring the old behavior of GIsGarbageCollecting (pre 26252560) where the global was set to false (through TGuardValue<bool> GuardIsGarbageCollecting(GIsGarbageCollecting, true);
 destructor) after Unhashing and Purge was complete and memory trimmed

diff --git a/Engine/Source/Runtime/CoreUObject/Private/UObject/GarbageCollection.cpp b/Engine/Source/Runtime/CoreUObject/Private/UObject/GarbageCollection.cpp
--- a/Engine/Source/Runtime/CoreUObject/Private/UObject/GarbageCollection.cpp
+++ b/Engine/Source/Runtime/CoreUObject/Private/UObject/GarbageCollection.cpp
@@ -5550,8 +5550,6 @@ void PostCollectGarbageImpl(EObjectFlags KeepFlags)
 	GIsGarbageCollectingAndLockingUObjectHashTables = false;
 	UnlockUObjectHashTables();
 
-	GIsGarbageCollecting = false;
-
 	// The hash tables lock was released when reachability analysis was done.
 	// BeginDestroy, FinishDestroy, destructors and callbacks are allowed to call functions like StaticAllocateObject and StaticFindObject.
 	// Now release the GC lock to allow async loading and other threads to perform UObject operations under the FGCScopeGuard.
@@ -5592,12 +5590,15 @@ void PostCollectGarbageImpl(EObjectFlags KeepFlags)
 		{
 			FMemory::Trim();
 		}
+	}
+
+	GIsGarbageCollecting = false;
 
+	if (!GIsIncrementalReachabilityPending)
+	{
 		// Route callbacks to verify GC assumptions
-		{
-			TRACE_CPUPROFILER_EVENT_SCOPE(BroadcastPostGarbageCollect);
-			FCoreUObjectDelegates::GetPostGarbageCollect().Broadcast();
-		}
+		TRACE_CPUPROFILER_EVENT_SCOPE(BroadcastPostGarbageCollect);
+		FCoreUObjectDelegates::GetPostGarbageCollect().Broadcast();
 	}
 
 	const double PostCollectEndTime = FPlatformTime::Seconds();
