Reducing the size of serialized data for remote objects (compiled out by default)

diff --git a/Engine/Source/Runtime/CoreUObject/Private/UObject/Class.cpp b/Engine/Source/Runtime/CoreUObject/Private/UObject/Class.cpp
--- a/Engine/Source/Runtime/CoreUObject/Private/UObject/Class.cpp
+++ b/Engine/Source/Runtime/CoreUObject/Private/UObject/Class.cpp
@@ -7783,6 +7783,12 @@ UScriptStruct* TBaseStructure<FRemoteObjectId>::Get()
 	return ScriptStruct;
 }
 
+UScriptStruct* TBaseStructure<FRemoteObjectPathName>::Get()
+{
+	static UScriptStruct* ScriptStruct = StaticGetBaseStructureInternal(TEXT("RemoteObjectPathName"));
+	return ScriptStruct;
+}
+
 UScriptStruct* TBaseStructure<FRemoteObjectData>::Get()
 {
 	static UScriptStruct* ScriptStruct = StaticGetBaseStructureInternal(TEXT("RemoteObjectData"));
diff --git a/Engine/Source/Runtime/CoreUObject/Private/UObject/RemoteObjectSerialization.cpp b/Engine/Source/Runtime/CoreUObject/Private/UObject/RemoteObjectSerialization.cpp
--- a/Engine/Source/Runtime/CoreUObject/Private/UObject/RemoteObjectSerialization.cpp
+++ b/Engine/Source/Runtime/CoreUObject/Private/UObject/RemoteObjectSerialization.cpp
@@ -17,9 +17,47 @@
 #include "HAL/FileManager.h"
 #include "Serialization/ArchiveSerializedPropertyChain.h"
 
+#ifndef UE_WITH_REMOTEOBJECTWRITER_DEBUGGING
+#define UE_WITH_REMOTEOBJECTWRITER_DEBUGGING 0
+#endif
+
+using FNameIndexType = FRemoteObjectPathName::FNameIndexType;
+
+FArchive& operator<<(FArchive& Ar, FRemoteObjectPathName& PathName)
+{
+	Ar << PathName.Names;
+	return Ar;
+}
+
 namespace UE::RemoteObject::Serialization::Disk
 {
 
+void SerializeNameTables(FArchive& Ar, FRemoteObjectData& InObjectData)
+{
+	FNameIndexType NumNames = IntCastChecked<FNameIndexType>(InObjectData.Names.Num());
+	Ar << NumNames;
+	for (FName& Name : InObjectData.Names)
+	{
+		FString NameString = Name.ToString();
+		Ar << NameString;
+	}
+	Ar << InObjectData.PathNames;
+}
+
+void DeserializeNameTables(FArchive& Ar, FRemoteObjectData& OutObjectData)
+{
+	OutObjectData.Names.Reset();
+	FNameIndexType NumNames = 0;
+	Ar << NumNames;
+	for (FNameIndexType NameIndex = 0; NameIndex < NumNames; ++NameIndex)
+	{
+		FString NameString;
+		Ar << NameString;
+		OutObjectData.Names.Add(FName(NameString, FNAME_Add));
+	}
+	Ar << OutObjectData.PathNames;
+}
+
 FString GenerateRemoteObjectFilename(FRemoteObjectId ObjectId)
 {
 	return FPaths::Combine(*FPaths::ProjectSavedDir(), *(ObjectId.ToString() + TEXT(".tmp")));
@@ -31,6 +69,7 @@ void LoadObjectFromDisk(FRemoteObjectId ObjectId, FRemoteServerId CurrentOwnerSe
 	checkf(FileReader.IsValid(), TEXT("Unable to create file reader for remote object %s"), *ObjectId.ToString());
 	FRemoteObjectData* ObjectData = new FRemoteObjectData();
 	
+	DeserializeNameTables(*FileReader, *ObjectData);
 	*FileReader << ObjectData->Bytes;
 	FileReader->Close();
 
@@ -42,6 +81,7 @@ void SaveObjectToDisk(FRemoteObjectId ObjectId, FRemoteServerId ServerId, FRemot
 	TUniquePtr<FArchive> FileWriter{ IFileManager::Get().CreateFileWriter(*GenerateRemoteObjectFilename(ObjectId)) };
 	checkf(FileWriter.IsValid(), TEXT("Unable to create file writer for remote object %s"), *ObjectId.ToString());
 	
+	SerializeNameTables(*FileWriter, ObjectData);
 	*FileWriter << ObjectData.Bytes;
 	FileWriter->Close();
 }
@@ -80,18 +120,14 @@ FArchive& operator<<(FArchive& Ar, ERemoteReferenceType& RefType)
 	return Ar;
 }
 
-#ifndef UE_WITH_REMOTEOBJECTWRITER_DEBUGGING
-#define UE_WITH_REMOTEOBJECTWRITER_DEBUGGING 0
-#endif
-
 class FArchiveRemoteObjectWriter : public FMemoryWriter
 {
 	UObject* RootObject = nullptr;
 	FRemoteObjectData& ObjectData;
 	TArray<UObject*>& ObjectsToSerialize;
 	const TSet<UObject*>& SerializedObjects;
-	TArray<FName>& Names;
-	TMap<FName, int32>& NameMap;
+	TMap<FName, FNameIndexType> NameMap;
+	TMap<UObject*, FNameIndexType> PathNameMap;
 
 #if UE_WITH_REMOTEOBJECTWRITER_DEBUGGING
 	FString SerializationScope;	
@@ -102,7 +138,6 @@ class FArchiveRemoteObjectWriter : public FMemoryWriter
 		int64 Count = 0;
 	};
 	TMap<FString, TMap<FProperty*, FSerializeStats>> ObjectPropertyStats;
-	TMap<FString, FSerializeStats> PathNameStats;
 #endif
 
 	ERemoteReferenceType WriteObjectReference(UObject*& Obj)
@@ -135,18 +170,8 @@ class FArchiveRemoteObjectWriter : public FMemoryWriter
 
 		if (Type == ERemoteReferenceType::PathName)
 		{
-#if UE_WITH_REMOTEOBJECTWRITER_DEBUGGING
-			const int64 StartPos = Tell();
-#endif
-
-			FString PathName = GetPathNameSafe(Obj);
-			*this << PathName;
-
-#if UE_WITH_REMOTEOBJECTWRITER_DEBUGGING
-			FSerializeStats& Stats = PathNameStats.FindOrAdd(PathName);
-			Stats.Count++;
-			Stats.Size += Tell() - StartPos;
-#endif
+			FNameIndexType PathNameIndex = AddPathNameToNameMap(Obj);
+			*this << PathNameIndex;
 		}
 
 		return Type;
@@ -154,16 +179,17 @@ class FArchiveRemoteObjectWriter : public FMemoryWriter
 
 public:
 
-	FArchiveRemoteObjectWriter(UObject* InRootObject, FRemoteObjectData& OutBytes, TArray<UObject*>& OutObjectsToSerialize, const TSet<UObject*>& InSerializedObjects, TArray<FName>& InOutNames, TMap<FName, int32>& InOutNameMap)
-		: FMemoryWriter(OutBytes.Bytes)
+	FArchiveRemoteObjectWriter(UObject* InRootObject, FRemoteObjectData& OutObjectData, TArray<UObject*>& OutObjectsToSerialize, const TSet<UObject*>& InSerializedObjects)
+		: FMemoryWriter(OutObjectData.Bytes)
 		, RootObject(InRootObject)
-		, ObjectData(OutBytes)
+		, ObjectData(OutObjectData)
 		, ObjectsToSerialize(OutObjectsToSerialize)
 		, SerializedObjects(InSerializedObjects)
-		, Names(InOutNames)
-		, NameMap(InOutNameMap)
 	{
-		SetPortFlags(PPF_IncludeTransient);
+		static_assert(sizeof(FNameIndexType) == sizeof(FRemoteObjectPathName::FNameIndexType), "Name Index type must match the type used for storing remote object path name indices");
+
+		SetIsPersistent(false);
+		SetUseUnversionedPropertySerialization(true);
 	}
 
 	virtual ~FArchiveRemoteObjectWriter()
@@ -173,21 +199,45 @@ public:
 #endif
 	}
 
-	int32 AddNameToNameMap(const FName& Name)
+	FNameIndexType AddNameToNameMap(const FName& Name)
 	{
-		int32 NameIndex = -1;
-		if (int32* ExistingNameIndex = NameMap.Find(Name))
+		FNameIndexType NameIndex = TNumericLimits<FNameIndexType>::Max();
+		if (FNameIndexType* ExistingNameIndex = NameMap.Find(Name))
 		{
 			NameIndex = *ExistingNameIndex;
 		}
 		else
 		{
-			NameIndex = Names.Add(Name);
+			int32 NewIndex = ObjectData.Names.Add(Name);
+			NameIndex = IntCastChecked<FNameIndexType>(NewIndex);
 			NameMap.Add(Name, NameIndex);
 		}
 		return NameIndex;
 	}
 
+	FNameIndexType AddPathNameToNameMap(UObject* Object)
+	{
+		FNameIndexType PathNameIndex = TNumericLimits<FNameIndexType>::Max();
+		if (FNameIndexType* ExistingPathNameIndex = PathNameMap.Find(Object))
+		{
+			PathNameIndex = *ExistingPathNameIndex;
+		}
+		else
+		{
+			int32 NewPathNameIndex = ObjectData.PathNames.AddDefaulted();
+			PathNameIndex = IntCastChecked<FNameIndexType>(NewPathNameIndex);
+
+			FRemoteObjectPathName& NewPathName = ObjectData.PathNames[PathNameIndex];
+			// Store individual indices of FNames of every object in this object's outer chain
+			for (UObject* OuterChain = Object; OuterChain; OuterChain = UE::CoreUObject::Private::FObjectHandleUtils::GetNonAccessTrackedOuterNoResolve(OuterChain))
+			{
+				NewPathName.Names.Add(AddNameToNameMap(OuterChain->GetFName()));
+			}
+			PathNameMap.Add(Object, PathNameIndex);
+		}
+		return PathNameIndex;
+	}
+
 	FORCEINLINE UObject* GetRootObject() const
 	{
 		return RootObject;
@@ -196,7 +246,18 @@ public:
 #if UE_WITH_REMOTEOBJECTWRITER_DEBUGGING
 	void DumpStatsToLog()
 	{
-		UE_LOG(LogRemoteObject, Log, TEXT("Object Data stats for %s %s (Total: %d):"), *FRemoteObjectId(RootObject).ToString(), *GetFullNameSafe(RootObject), ObjectData.Bytes.Num());
+		int64 NameTableSize = 0;
+		{
+			TArray<uint8> NameTableData;
+			FMemoryWriter NameTableWriter(NameTableData);			
+			UE::RemoteObject::Serialization::Disk::SerializeNameTables(NameTableWriter, ObjectData);
+			NameTableSize = NameTableData.Num();
+		}
+
+		const int64 TotalSize = NameTableSize + ObjectData.Bytes.Num();
+
+		UE_LOG(LogRemoteObject, Log, TEXT("Object Data stats for %s %s (Object Data toal: %d, total: %lld):"), *FRemoteObjectId(RootObject).ToString(), *GetFullNameSafe(RootObject), ObjectData.Bytes.Num(), TotalSize);
+		UE_LOG(LogRemoteObject, Log, TEXT("  Name Table total size: %lld (Names: %d, Paths: %d)"), NameTableSize, ObjectData.Names.Num(), ObjectData.PathNames.Num());
 
 		for (const TPair<FString, TMap<FProperty*, FSerializeStats>>& ObjectStatsPair : ObjectPropertyStats)
 		{
@@ -210,7 +271,7 @@ public:
 				Total += Stats.Size;
 			}
 
-			UE_LOG(LogRemoteObject, Log, TEXT("  Data serialized for %s (Total: %lld):"), *Obj, Total);
+			UE_LOG(LogRemoteObject, Log, TEXT("  Data serialized for %s (Total: %lld):"), Obj.Len() ? *Obj : TEXT("Native Serialize"), Total);
 			
 			for (const TPair<FProperty*, FSerializeStats>& PropertyPair : PropStats)
 			{
@@ -219,23 +280,6 @@ public:
 				UE_LOG(LogRemoteObject, Log, TEXT("    %s: size: %lld, count: %lld"), Prop ? *GetFullNameSafe(Prop) : TEXT("Native Serialize"), Stats.Size, Stats.Count);
 			}
 		}
-
-		{
-			int64 PathNameTotalCount = 0;
-			int64 PathNameTotalSize = 0;
-			for (const TPair<FString, FSerializeStats>& StatPair : PathNameStats)
-			{
-				PathNameTotalCount += StatPair.Value.Count;
-				PathNameTotalSize += StatPair.Value.Size;
-			}
-
-			UE_LOG(LogRemoteObject, Log, TEXT("  Serialized %lld PathNames (Total: %lld):"), PathNameTotalCount, PathNameTotalSize);
-			for (const TPair<FString, FSerializeStats>& StatPair : PathNameStats)
-			{
-				UE_LOG(LogRemoteObject, Log, TEXT("    %s: size: %lld, count: %lld"), *StatPair.Key, StatPair.Value.Size, StatPair.Value.Count);
-			}
-		}
-		
 	}
 
 	void SetSerializationScope(const TCHAR* InScope)
@@ -290,7 +334,7 @@ public:
 
 	virtual FArchive& operator<<(FName& Name) override
 	{
-		int32 NameIndex = AddNameToNameMap(Name);
+		FNameIndexType NameIndex = AddNameToNameMap(Name);
 		*this << NameIndex;
 		return *this;
 	}
@@ -365,7 +409,8 @@ void SerializeRemoteObject(UObject* Object, FArchive& Ar, FRemoteObjectHeader& O
 
 class FArchiveRemoteObjectReader : public FMemoryReader
 {
-	TArray<FName>& Names;
+	const TArray<FName>& Names;
+	const TArray<UObject*>& ResolvedPathNameObjects;
 
 	ERemoteReferenceType ReadObjectReference(FObjectPtr& Value)
 	{
@@ -380,12 +425,11 @@ class FArchiveRemoteObjectReader : public FMemoryReader
 			bool bNeedsResolvingWithId = true;
 			if (Type == ERemoteReferenceType::PathName)
 			{
-				FString PathName;
-				*this << PathName;
-				checkf(!PathName.IsEmpty(), TEXT("Serialized remote object with empty path name but it should've been ERemoteReferenceType::None"));
+				FNameIndexType PathNameIndex = TNumericLimits<FNameIndexType>::Max();
+				*this << PathNameIndex;
 
 				// Try to resolve path name immediately as we expect the object to exist in memory.
-				UObject* Obj = ::FindObject<UObject>(nullptr, *PathName);
+				UObject* Obj = ResolvedPathNameObjects[PathNameIndex];
 				if (Obj)
 				{
 					Value = FObjectPtr(Obj);
@@ -411,11 +455,13 @@ class FArchiveRemoteObjectReader : public FMemoryReader
 
 public:
 
-	FArchiveRemoteObjectReader(const FRemoteObjectData& InBytes, TArray<FName>& InNames)
-		: FMemoryReader(InBytes.Bytes)
-		, Names(InNames)
+	FArchiveRemoteObjectReader(const FRemoteObjectData& InObjectData, const TArray<UObject*>& InResolvedPathNames)
+		: FMemoryReader(InObjectData.Bytes)
+		, Names(InObjectData.Names)
+		, ResolvedPathNameObjects(InResolvedPathNames)
 	{
-		SetPortFlags(PPF_IncludeTransient);
+		SetIsPersistent(false);
+		SetUseUnversionedPropertySerialization(true);
 	}
 
 	virtual FArchive& operator<<(UObject*& Obj) override
@@ -429,7 +475,7 @@ public:
 
 	virtual FArchive& operator<<(FName& Name) override
 	{
-		int32 NameIndex = -1;
+		FNameIndexType NameIndex = 0;
 		*this << NameIndex;
 		Name = Names[NameIndex];
 		return *this;
@@ -554,28 +600,25 @@ UObject* ConstructeRemoteObject(UObject* DestinationOuter, const FRemoteObjectHe
 	return Object;
 }
 
-void SerializeNames(FArchive& Ar, TArray<FName>& InNames)
+void ResolvePathNames(FRemoteObjectData& InObjectData, TArray<UObject*>& OutResolvedObjects)
 {
-	int32 NumNames = InNames.Num();
-	Ar << NumNames;
-	for (FName& Name : InNames)
+	for (const FRemoteObjectPathName& PathName : InObjectData.PathNames)
 	{
-		FString NameString = Name.ToString();
-		Ar << NameString;
+		UObject* Outer = nullptr;
+		UObject* Object = nullptr;
+		for (int32 PathNameIndex = PathName.Names.Num() - 1; PathNameIndex >= 0; --PathNameIndex)
+		{
+			Object = StaticFindObjectFast(UObject::StaticClass(), Outer, InObjectData.Names[PathName.Names[PathNameIndex]]);
+			if (!Object)
+			{
+				break;
+			}
+			Outer = Object;
+		}
+		OutResolvedObjects.Add(Object);
 	}
 }
 
-void DeserializeNames(FArchive& Ar, TArray<FName>& OutNames)
-{
-	int32 NumNames = 0;
-	Ar << NumNames;
-	for (int32 NameIndex = 0; NameIndex < NumNames; ++NameIndex)
-	{
-		FString NameString;
-		Ar << NameString;
-		OutNames.Add(FName(NameString, FNAME_Add));
-	}
-}
 
 FRemoteObjectData* SerializeObjectData(UObject* Object, TSet<UObject*>& OutObjects)
 {
@@ -583,14 +626,13 @@ FRemoteObjectData* SerializeObjectData(UObject* Object, TSet<UObject*>& OutObjec
 
 	FRemoteObjectData* ObjectData = new FRemoteObjectData();
 	TArray<UObject*> ObjectsToSerialize;
-	TArray<FName> Names;
+	TArray<FRemoteObjectPathName> PathNames;
 	TArray<FRemoteObjectHeader> ObjectHeaders;
-	TMap<FName, int32> NameMap;
 	int32 Version = 0;
 	int64 HeaderOffset = 0;
 
 	{
-		FArchiveRemoteObjectWriter Ar(Object, *ObjectData, ObjectsToSerialize, OutObjects, Names, NameMap);
+		FArchiveRemoteObjectWriter Ar(Object, *ObjectData, ObjectsToSerialize, OutObjects);
 		ObjectsToSerialize.Add(Object);
 
 		int64 OffsetOfHeaderOffset = 0;
@@ -619,10 +661,6 @@ FRemoteObjectData* SerializeObjectData(UObject* Object, TSet<UObject*>& OutObjec
 		Ar << HeaderOffset;
 		Ar.Seek(HeaderOffset);
 
-		{
-			FRemoteObjectWriterScope Scope(Ar, TEXT("NameMap"));
-			SerializeNames(Ar, Names);
-		}
 		{
 			FRemoteObjectWriterScope Scope(Ar, TEXT("ObjectHeaders"));
 			Ar << ObjectHeaders;
@@ -640,19 +678,19 @@ void DeserializeObjectData(FRemoteObjectData& ObjectData, UObject* DestinationOu
 	int32 Version = 0;
 	int64 HeaderOffset = 0;
 	TArray<FRemoteObjectHeader> ObjectHeaders;
-	TArray<FName> Names;
+	TArray<UObject*> ResolvedPathNameObjects;
+
+	ResolvePathNames(ObjectData, ResolvedPathNameObjects);
 
-	FArchiveRemoteObjectReader Ar(ObjectData, Names);
+	FArchiveRemoteObjectReader Ar(ObjectData, ResolvedPathNameObjects);
 
 	Ar << Version;
 	Ar << HeaderOffset;
 
 	const int64 ObjectDataOffset = Ar.Tell();
-	Ar.Seek(HeaderOffset);
-
-	DeserializeNames(Ar, Names);
-	Ar << ObjectHeaders;	
 
+	Ar.Seek(HeaderOffset);
+	Ar << ObjectHeaders;
 	Ar.Seek(ObjectDataOffset);	
 
 	if (ObjectHeaders.Num())
diff --git a/Engine/Source/Runtime/CoreUObject/Public/UObject/Class.h b/Engine/Source/Runtime/CoreUObject/Public/UObject/Class.h
--- a/Engine/Source/Runtime/CoreUObject/Public/UObject/Class.h
+++ b/Engine/Source/Runtime/CoreUObject/Public/UObject/Class.h
@@ -4471,6 +4471,11 @@ template<> struct TBaseStructure<FRemoteObjectId>
 	static COREUOBJECT_API UScriptStruct* Get();
 };
 
+template<> struct TBaseStructure<FRemoteObjectPathName>
+{
+	static COREUOBJECT_API UScriptStruct* Get();
+};
+
 template<> struct TBaseStructure<FRemoteObjectData>
 {
 	static COREUOBJECT_API UScriptStruct* Get();
diff --git a/Engine/Source/Runtime/CoreUObject/Public/UObject/NoExportTypes.h b/Engine/Source/Runtime/CoreUObject/Public/UObject/NoExportTypes.h
--- a/Engine/Source/Runtime/CoreUObject/Public/UObject/NoExportTypes.h
+++ b/Engine/Source/Runtime/CoreUObject/Public/UObject/NoExportTypes.h
@@ -3390,11 +3390,26 @@ struct FRemoteObjectId
 	uint64 Id;
 };
 
+USTRUCT(noexport, BlueprintType)
+struct FRemoteObjectPathName
+{
+	GENERATED_BODY()
+
+	UPROPERTY()
+	TArray<uint16> Names;
+};
+
 USTRUCT(noexport, BlueprintType)
 struct FRemoteObjectData
 {
 	GENERATED_BODY()
 
+	UPROPERTY()
+	TArray<FName> Names;
+
+	UPROPERTY()
+	TArray<FRemoteObjectPathName> PathNames;
+
 	UPROPERTY()
 	TArray<uint8> Bytes;
 };
diff --git a/Engine/Source/Runtime/CoreUObject/Public/UObject/RemoteObjectTypes.h b/Engine/Source/Runtime/CoreUObject/Public/UObject/RemoteObjectTypes.h
--- a/Engine/Source/Runtime/CoreUObject/Public/UObject/RemoteObjectTypes.h
+++ b/Engine/Source/Runtime/CoreUObject/Public/UObject/RemoteObjectTypes.h
@@ -130,9 +130,20 @@ FORCEINLINE uint32 GetTypeHash(const FRemoteObjectId& ObjectId)
 }
 
 /**
-* Structure that holds remote object memory
+* Structure that represents remote object pathname (noexport type)
+*/
+struct FRemoteObjectPathName
+{
+	using FNameIndexType = uint16;
+	TArray<FNameIndexType> Names;
+};
+
+/**
+* Structure that holds remote object memory (noexport type)
 */
 struct FRemoteObjectData
 {
+	TArray<FName> Names;
+	TArray<FRemoteObjectPathName> PathNames;
 	TArray<uint8> Bytes;
 };
\ No newline at end of file
