Sever IDO dependency from Property Initialization State tracking

Major changes
* Property Initialization State tracked in out-of-band data instead of as hidden property of IDO.  IDOs not guaranteed to exist (or be enabled) to support this feature.
* Track property initialization state when objects are initialized in addition to deserialized
* Change serialization to no longer set property value initialized if there is no value that is being deserialized.  Note: recursion here not currently implemented - assumes that nested props are always initialized.
* On load, set property initialization state only if there is a value.  It is assumed that if there is no value, the object initializer would have set the state to true or in the case of CPF_RequiredParm, only set it if the template counterpart property was also initialized
* Serialization of annotations for property initialization state added to enable Undo/Redo

Deficiencies
* Assumes IDO.Enabled = 0, need to make this work for both cases I think.
* Copy/Paste/Duplicate doesn't set property init state correctly.
* The override state widget not showing up for properties of uobjects within arrays

CVar:
* `AnnotationBasedPropertyInitStateTracking` will enable the annotation based property initialization state tracking if IDOs are disabled

diff --git a/Engine/Source/Runtime/CoreUObject/Internal/UObject/InstanceDataObjectUtils.h b/Engine/Source/Runtime/CoreUObject/Internal/UObject/InstanceDataObjectUtils.h
--- a/Engine/Source/Runtime/CoreUObject/Internal/UObject/InstanceDataObjectUtils.h
+++ b/Engine/Source/Runtime/CoreUObject/Internal/UObject/InstanceDataObjectUtils.h
@@ -10,6 +10,7 @@
 
 class FFieldVariant;
 class FProperty;
+class FStructuredArchiveRecord;
 class UClass;
 class UObject;
 class UStruct;
@@ -21,7 +22,7 @@ namespace UE
 {
 
 /** Query if InstanceDataObject support is available generally. */
-bool IsInstanceDataObjectSupportEnabled();
+UE_API bool IsInstanceDataObjectSupportEnabled();
 /** Query if InstanceDataObject support is enabled for a specific object. */
 bool IsInstanceDataObjectSupportEnabled(const UObject* Object);
 /** Query if uninitialized alert UI is enabled **/
@@ -47,6 +48,9 @@ bool WasPropertyValueSerialized(const UStruct* Struct, const void* StructData, c
 /** Copy whether each property was set by serialization from one IDO to another. */
 void CopyPropertyValueSerializedData(const FFieldVariant& OldField, void* OldDataPtr, const FFieldVariant& NewField, void* NewDataPtr);
 
+
+UE_API bool SupportsPropertyInitStateTracking();
+	
 /** Query whether the property value is initialized for a property in the struct. */
 UE_API bool IsPropertyValueInitialized(const UStruct* Struct, void* StructData, const FProperty* Property, int32 ArrayIndex = 0);
 /** Set the flag that the property value is initialized for a property in the struct. */
@@ -55,6 +59,8 @@ UE_API void SetPropertyValueInitialized(const UStruct* Struct, void* StructData,
 UE_API void ClearPropertyValueInitialized(const UStruct* Struct, void* StructData, const FProperty* Property, int32 ArrayIndex = 0);
 /** Reset the property value initialized flags for every property in the struct. */
 UE_API void ResetPropertyValueInitialized(const UStruct* Struct, void* StructData);
+		
+UE_API void SerializePropertyValueInitializationState(UObject& Object, FStructuredArchiveRecord& ObjectRecord);
 
 } // UE
 
diff --git a/Engine/Source/Runtime/CoreUObject/Private/UObject/Class.cpp b/Engine/Source/Runtime/CoreUObject/Private/UObject/Class.cpp
--- a/Engine/Source/Runtime/CoreUObject/Private/UObject/Class.cpp
+++ b/Engine/Source/Runtime/CoreUObject/Private/UObject/Class.cpp
@@ -1391,15 +1391,22 @@ void UStruct::LoadTaggedPropertiesFromText(FStructuredArchive::FSlot Slot, uint8
 				Tag.ArrayIndex = ItemIndex;
 				Tag.Name = PropertyName;
 
-			#if WITH_EDITORONLY_DATA
-				if (SerializeContext->bTrackInitializedProperties)
+				#if WITH_EDITORONLY_DATA
+				const auto SetInitialized = [bTrackInitializedProperties = SerializeContext->bTrackInitializedProperties](const UStruct* Struct, void* Data, const FProperty* Property, int32 ArrayIndex)
 				{
-					UE::SetPropertyValueInitialized(this, Data, Property, ItemIndex);
-				}
-			#endif
+				
+					if (bTrackInitializedProperties)
+					{
+						UE::SetPropertyValueInitialized(Struct, Data, Property, ArrayIndex);
+					}
+				};
+				#else
+				const auto SetInitialized = [](const UStruct* Struct, void* Data, const FProperty* Property, int32 ArrayIndex) {};
+				#endif
 
 				if (Tag.SerializeType == EPropertyTagSerializeType::Skipped)
 				{
+					SetInitialized(this, Data, Property, ItemIndex);
 					continue;
 				}
 
@@ -1420,6 +1427,7 @@ void UStruct::LoadTaggedPropertiesFromText(FStructuredArchive::FSlot Slot, uint8
 				{
 				case EConvertFromTypeResult::Converted:
 				case EConvertFromTypeResult::Serialized:
+					SetInitialized(this, Data, Property, ItemIndex);
 					break;
 
 				case EConvertFromTypeResult::UseSerializeItem:
@@ -1437,6 +1445,7 @@ void UStruct::LoadTaggedPropertiesFromText(FStructuredArchive::FSlot Slot, uint8
 
 						// This property is ok.
 						Tag.SerializeTaggedProperty(ItemSlot.GetValue(), Property, DestAddress, DefaultsFromParent);
+						SetInitialized(this, Data, Property, ItemIndex);
 					}
 					break;
 
@@ -1707,18 +1716,14 @@ void UStruct::SerializeVersionedTaggedProperties(FStructuredArchive::FSlot Slot,
 					Property = CustomFindProperty(Tag.Name);
 				}
 
-			#if WITH_EDITORONLY_DATA
-				ON_SCOPE_EXIT
+				if (Tag.SerializeType == EPropertyTagSerializeType::Skipped)
 				{
+				#if WITH_EDITORONLY_DATA
 					if (SerializeContext->bTrackInitializedProperties && Property)
 					{
 						UE::SetPropertyValueInitialized(this, Data, Property, Tag.ArrayIndex);
 					}
-				};
-			#endif
-
-				if (Tag.SerializeType == EPropertyTagSerializeType::Skipped)
-				{
+				#endif
 					continue;
 				}
 
@@ -1918,6 +1923,13 @@ void UStruct::SerializeVersionedTaggedProperties(FStructuredArchive::FSlot Slot,
 					checkf(Tag.Size == Loaded,
 						TEXT("Size mismatch in %s of %s of type %s. Loaded %" INT64_FMT " bytes but expected %d. Package: %s"),
 						*Tag.Name.ToString(), *GetName(), *WriteToString<64>(Tag.GetType()), Loaded, Tag.Size, *UnderlyingArchive.GetArchiveName());
+
+				#if WITH_EDITORONLY_DATA
+					if (SerializeContext->bTrackInitializedProperties && Property)
+					{
+						UE::SetPropertyValueInitialized(this, Data, Property, Tag.ArrayIndex);
+					}
+				#endif
 				}
 				else if (Tag.Size != Loaded)
 				{
diff --git a/Engine/Source/Runtime/CoreUObject/Private/UObject/InstanceDataObjectUtils.cpp b/Engine/Source/Runtime/CoreUObject/Private/UObject/InstanceDataObjectUtils.cpp
--- a/Engine/Source/Runtime/CoreUObject/Private/UObject/InstanceDataObjectUtils.cpp
+++ b/Engine/Source/Runtime/CoreUObject/Private/UObject/InstanceDataObjectUtils.cpp
@@ -4,11 +4,13 @@
 
 #include "Serialization/ObjectReader.h"
 #include "Serialization/ObjectWriter.h"
+#include "VerseVM/VVMVerseClass.h"
 
 #if WITH_EDITORONLY_DATA
 
 #include "HAL/IConsoleManager.h"
 #include "Misc/ReverseIterate.h"
+#include "Serialization/StructuredArchive.h"
 #include "UObject/Class.h"
 #include "UObject/EnumProperty.h"
 #include "UObject/Field.h"
@@ -18,6 +20,7 @@
 #include "UObject/PropertyOptional.h"
 #include "UObject/PropertyPathNameTree.h"
 #include "UObject/UnrealType.h"
+#include "UObject/UObjectAnnotation.h"
 
 static const FName NAME_InitializedValues(ANSITEXTVIEW("_InitializedValues"));
 static const FName NAME_SerializedValues(ANSITEXTVIEW("_SerializedValues"));
@@ -151,6 +154,13 @@ namespace UE
 		TEXT("Comma separated list of property types that will be excluded from loose properties in IDOs.")
 	);
 
+	bool bUseAnnotationsForPropertyStateInitTracking = true;
+	FAutoConsoleVariableRef CVarUseAnnotationBasedPropertyStateInitTracking(
+		TEXT("AnnotationBasedPropertyInitStateTracking"),
+		bUseAnnotationsForPropertyStateInitTracking,
+		TEXT("Enables uobject annotation method to track property initialization state if IDOs disabled.")
+	);
+
 	bool bEnableUninitializedUI = false;
 	FAutoConsoleVariableRef EnableUninitializedUI(
 		TEXT("IDO.EnableUninitializedAlertUI"),
@@ -912,57 +922,173 @@ namespace UE
 		}
 	}
 
-	static const FByteProperty* FindInitializedValuesProperty(const UStruct* Struct)
+	bool SupportsPropertyInitStateTracking()
 	{
-		if (const UInstanceDataObjectClass* IdoClass = Cast<UInstanceDataObjectClass>(Struct))
+		return bUseAnnotationsForPropertyStateInitTracking || bEnableIDOSupport;
+	}
+
+	static bool IsIdoPropertyValueInitialized(const UStruct* Struct, void* StructData, const FProperty* Property, int32 ArrayIndex, const FByteProperty* InitializedValuesProperty)
+	{
+		if (!SupportsPropertyInitStateTracking())
 		{
-			return IdoClass->InitializedValuesProperty;
+			return true;
 		}
-		if (const UInstanceDataObjectStruct* IdoStruct = Cast<UInstanceDataObjectStruct>(Struct))
+		
+		const int32 PropertyIndex = Property->GetIndexInOwner() + ArrayIndex;
+		const int32 ByteIndex = PropertyIndex / 8;
+		const int32 BitOffset = PropertyIndex % 8;
+		if (ensureMsgf(ByteIndex < InitializedValuesProperty->ArrayDim,
+			TEXT("Property %s in %s has out of range index %d with capacity for %d."),
+			*Property->GetAuthoredName(), *Struct->GetPathName(), PropertyIndex, InitializedValuesProperty->ArrayDim * 8))
 		{
-			return IdoStruct->InitializedValuesProperty;
+			const uint8* PropertyDataPtr = InitializedValuesProperty->ContainerPtrToValuePtr<uint8>(StructData, ByteIndex);
+			return (*PropertyDataPtr & (1 << BitOffset)) != 0;
 		}
-		return CastField<FByteProperty>(Struct->FindPropertyByName(NAME_InitializedValues));
+		// True is a safe fallback when this function is misused.
+		return true;
 	}
 
-	static void SetPropertyValueInitializedFlag(const UStruct* Struct, void* StructData, const FProperty* Property, int32 ArrayIndex, bool bValue)
+	static void SetIdoPropertyValueInitializedFlag(const UStruct* Struct, void* StructData, const FProperty* Property, int32 ArrayIndex, bool bValue, const FByteProperty* InitializedValuesProperty)
 	{
-		if (const FByteProperty* InitializedValuesProperty = FindInitializedValuesProperty(Struct))
-		{
-			const int32 PropertyIndex = Property->GetIndexInOwner() + ArrayIndex;
-			const int32 ByteIndex = PropertyIndex / 8;
-			const int32 BitOffset = PropertyIndex % 8;
-			if (ByteIndex < InitializedValuesProperty->ArrayDim)
+		const int32 PropertyIndex = Property->GetIndexInOwner() + ArrayIndex;
+		const int32 ByteIndex = PropertyIndex / 8;
+		const int32 BitOffset = PropertyIndex % 8;
+		if (ensureMsgf(ByteIndex < InitializedValuesProperty->ArrayDim,
+			TEXT("Property %s in %s has out of range index %d with capacity for %d."),
+			*Property->GetAuthoredName(), *Struct->GetPathName(), PropertyIndex, InitializedValuesProperty->ArrayDim * 8))
+		{
+			uint8* PropertyDataPtr = InitializedValuesProperty->ContainerPtrToValuePtr<uint8>(StructData, ByteIndex);
+			if (bValue)
 			{
-				uint8* PropertyDataPtr = InitializedValuesProperty->ContainerPtrToValuePtr<uint8>(StructData, ByteIndex);
-				if (bValue)
-				{
-					*PropertyDataPtr |= (1 << BitOffset);
-				}
-				else
-				{
-					*PropertyDataPtr &= ~(1 << BitOffset);
-				}
+				*PropertyDataPtr |= (1 << BitOffset);
+			}
+			else
+			{
+				*PropertyDataPtr &= ~(1 << BitOffset);
 			}
 		}
 	}
 
-	bool IsPropertyValueInitialized(const UStruct* Struct, void* StructData, const FProperty* Property, int32 ArrayIndex)
+	static void ResetIdoPropertyValueInitialized(const UStruct* Struct, void* StructData, const FByteProperty* InitializedValuesProperty)
+	{
+		uint8* PropertyDataPtr = InitializedValuesProperty->ContainerPtrToValuePtr<uint8>(StructData);
+		FMemory::Memzero(PropertyDataPtr, InitializedValuesProperty->ArrayDim);
+	}
+
+	struct FPropertyInitializedAnnotationData
+	{
+		TBitArray<> PropertyInitialized;
+	};
+
+	struct FPropertyInitializedAnnotation
+	{
+		bool IsDefault() const
+		{
+			return !Data.IsValid();
+		}
+
+		TSharedPtr<FPropertyInitializedAnnotationData> Data;
+	};
+
+	using FPropertyInitializedAnnotationStore = FUObjectAnnotationSparse<FPropertyInitializedAnnotation, /*bAutoRemove*/ true>;
+
+	FPropertyInitializedAnnotationStore& GetPropertyInitializedAnnotations()
 	{
-		if (const FByteProperty* InitializedValuesProperty = FindInitializedValuesProperty(Struct))
+		static FPropertyInitializedAnnotationStore Annotations;
+		return Annotations;
+	}
+
+	static bool IsPropertyValueInitializedByAnnotation(const UStruct* Struct, void* StructData, const FProperty* Property, int32 ArrayIndex)
+	{
+		if (!SupportsPropertyInitStateTracking())
+		{
+			return true;
+		}
+		UObjectBase* Object = static_cast<UObjectBase*>(StructData);
+		if (TSharedPtr<FPropertyInitializedAnnotationData> Data = GetPropertyInitializedAnnotations().GetAnnotation(Object).Data)
 		{
 			const int32 PropertyIndex = Property->GetIndexInOwner() + ArrayIndex;
-			const int32 ByteIndex = PropertyIndex / 8;
-			const int32 BitOffset = PropertyIndex % 8;
-			if (ByteIndex < InitializedValuesProperty->ArrayDim)
+			if (Data->PropertyInitialized.IsValidIndex(PropertyIndex))
 			{
-				const uint8* PropertyDataPtr = InitializedValuesProperty->ContainerPtrToValuePtr<uint8>(StructData, ByteIndex);
-				return (*PropertyDataPtr & (1 << BitOffset)) != 0;
+				return Data->PropertyInitialized[PropertyIndex];
 			}
 		}
+		// Default to uninitialized when a flag has not been stored.
 		return false;
 	}
 
+	static void SetPropertyValueInitializedFlagByAnnotation(const UStruct* Struct, void* StructData, const FProperty* Property, int32 ArrayIndex, bool bValue)
+	{
+		if (!SupportsPropertyInitStateTracking())
+		{
+			return;
+		}		
+		UObjectBase* Object = static_cast<UObjectBase*>(StructData);
+		FPropertyInitializedAnnotationStore& Store = GetPropertyInitializedAnnotations();
+		TSharedPtr<FPropertyInitializedAnnotationData> Data = Store.GetAnnotation(Object).Data;
+		if (!Data)
+		{
+			Data = MakeShared<FPropertyInitializedAnnotationData>();
+			Store.AddAnnotation(Object, {Data});
+		}
+
+		const int32 PropertyIndex = Property->GetIndexInOwner() + ArrayIndex;
+		if (Data->PropertyInitialized.Num() <= PropertyIndex)
+		{
+			checkf(PropertyIndex < Struct->TotalFieldCount,
+				TEXT("Property %s in %s has out of range index %d with capacity for %d."),
+				*Property->GetAuthoredName(), *Struct->GetPathName(), PropertyIndex, Struct->TotalFieldCount);
+			Data->PropertyInitialized.SetNum(Struct->TotalFieldCount, /*bValue*/ false);
+		}
+		Data->PropertyInitialized[PropertyIndex] = bValue;
+	}
+
+	static void ResetPropertyValueInitializedByAnnotation(const UStruct* Struct, void* StructData)
+	{
+		if (!SupportsPropertyInitStateTracking())
+		{
+			return;
+		}
+		UObjectBase* Object = static_cast<UObjectBase*>(StructData);
+		if (TSharedPtr<FPropertyInitializedAnnotationData> Data = GetPropertyInitializedAnnotations().GetAnnotation(Object).Data)
+		{
+			Data->PropertyInitialized.SetRange(0, Data->PropertyInitialized.Num(), /*bValue*/ false);
+		}
+	}
+
+	bool IsPropertyValueInitialized(const UStruct* Struct, void* StructData, const FProperty* Property, int32 ArrayIndex)
+	{
+		if (const UInstanceDataObjectClass* IdoClass = Cast<UInstanceDataObjectClass>(Struct))
+		{
+			return IsIdoPropertyValueInitialized(Struct, StructData, Property, ArrayIndex, IdoClass->InitializedValuesProperty);
+		}
+		if (const UInstanceDataObjectStruct* IdoStruct = Cast<UInstanceDataObjectStruct>(Struct))
+		{
+			return IsIdoPropertyValueInitialized(Struct, StructData, Property, ArrayIndex, IdoStruct->InitializedValuesProperty);
+		}
+		if (const UClass* Class = Cast<UClass>(Struct))
+		{
+			return IsPropertyValueInitializedByAnnotation(Struct, StructData, Property, ArrayIndex);
+		}
+		return true;
+	}
+
+	static void SetPropertyValueInitializedFlag(const UStruct* Struct, void* StructData, const FProperty* Property, int32 ArrayIndex, bool bValue)
+	{
+		if (const UInstanceDataObjectClass* IdoClass = Cast<UInstanceDataObjectClass>(Struct))
+		{
+			return SetIdoPropertyValueInitializedFlag(Struct, StructData, Property, ArrayIndex, bValue, IdoClass->InitializedValuesProperty);
+		}
+		if (const UInstanceDataObjectStruct* IdoStruct = Cast<UInstanceDataObjectStruct>(Struct))
+		{
+			return SetIdoPropertyValueInitializedFlag(Struct, StructData, Property, ArrayIndex, bValue, IdoStruct->InitializedValuesProperty);
+		}
+		if (const UClass* Class = Cast<UClass>(Struct))
+		{
+			return SetPropertyValueInitializedFlagByAnnotation(Struct, StructData, Property, ArrayIndex, bValue);
+		}
+	}
+
 	void SetPropertyValueInitialized(const UStruct* Struct, void* StructData, const FProperty* Property, int32 ArrayIndex)
 	{
 		SetPropertyValueInitializedFlag(Struct, StructData, Property, ArrayIndex, /*bValue*/ true);
@@ -975,13 +1101,103 @@ namespace UE
 
 	void ResetPropertyValueInitialized(const UStruct* Struct, void* StructData)
 	{
-		if (const FByteProperty* InitializedValuesProperty = FindInitializedValuesProperty(Struct))
+		if (const UInstanceDataObjectClass* IdoClass = Cast<UInstanceDataObjectClass>(Struct))
+		{
+			return ResetIdoPropertyValueInitialized(Struct, StructData, IdoClass->InitializedValuesProperty);
+		}
+		if (const UInstanceDataObjectStruct* IdoStruct = Cast<UInstanceDataObjectStruct>(Struct))
+		{
+			return ResetIdoPropertyValueInitialized(Struct, StructData, IdoStruct->InitializedValuesProperty);
+		}
+		if (const UClass* Class = Cast<UClass>(Struct))
+		{
+			return ResetPropertyValueInitializedByAnnotation(Struct, StructData);
+		}
+	}
+
+	static void CopyPropertyInitializationState(const FFieldVariant& Destination, void* DestinationDataPtr, const FFieldVariant& Source, void* SourceDataPtr)
+	{
+		if (const FStructProperty* SourceAsStructProperty = Source.Get<FStructProperty>())
+		{
+			const FStructProperty* DestAsStructProperty = Destination.Get<FStructProperty>();
+			checkf(DestAsStructProperty, TEXT("Type mismatch between Source field and Destingation field. Expected FStructProperty"));
+			CopyPropertyValueSerializedData(DestAsStructProperty->Struct, DestinationDataPtr, SourceAsStructProperty->Struct, SourceDataPtr);
+		}
+		else if (const FArrayProperty* SourceAsArrayProperty = Source.Get<FArrayProperty>())
 		{
-			uint8* PropertyDataPtr = InitializedValuesProperty->ContainerPtrToValuePtr<uint8>(StructData);
-			FMemory::Memzero(PropertyDataPtr, InitializedValuesProperty->ArrayDim);
+			// Any properties within Arrays are assumed initialized
+		}
+		else if (const FSetProperty* SourceAsSetProperty = Source.Get<FSetProperty>())
+		{
+			// Any properties within Sets are assumed initialized
+		}
+		else if (const FMapProperty* SourceAsMapProperty = Source.Get<FMapProperty>())
+		{
+			// Any properties within Maps are assumed initialized
+		}
+		else if (UStruct* SourceAsStruct = Source.Get<UStruct>())
+		{
+			UStruct* DestAsStruct = Destination.Get<UStruct>();
+
+			// Integrity check - both should be same type
+			// Can use the same properties then
+			check(SourceAsStruct->GetClass() == DestAsStruct->GetClass());			
+			
+			for (const FProperty* Property : TFieldRange<FProperty>(SourceAsStruct))
+			{
+				for (int32 ArrayIndex = 0; ArrayIndex < Property->ArrayDim; ++ArrayIndex)
+				{
+					const bool bIsInitialized = IsPropertyValueInitializedByAnnotation(SourceAsStruct->GetClass(), SourceAsStruct, Property, ArrayIndex);
+					SetPropertyValueInitializedFlagByAnnotation(SourceAsStruct->GetClass(), SourceAsStruct, Property, ArrayIndex, bIsInitialized);
+				}
+			}
 		}
 	}
 
+	void SerializePropertyValueInitializationState(UObject& Object, FStructuredArchiveRecord& ObjectRecord)
+	{
+		// Serialization of the property init state when IDOs enabled are done via a hidden property on the IDOs
+		if (!bUseAnnotationsForPropertyStateInitTracking)
+		{
+			return;
+		}
+		const FArchiveState& ArchiveState = ObjectRecord.GetArchiveState();
+		
+		bool bHasPropertyValueInitializationState = false;
+		TSharedPtr<FPropertyInitializedAnnotationData> Data;
+		if (ArchiveState.IsSaving())
+		{
+			Data = GetPropertyInitializedAnnotations().GetAnnotation(&Object).Data;
+
+			if (Data.IsValid())
+			{
+				bHasPropertyValueInitializationState = true;
+			}
+		}
+
+		TOptional<FStructuredArchiveSlot> Slot = ObjectRecord.TryEnterField(TEXT("PropertyValueInitializationState"), bHasPropertyValueInitializationState);
+		if (Slot.IsSet())
+		{
+			*Slot << SA_ATTRIBUTE(TEXT("PropertyValueInitializationState"), bHasPropertyValueInitializationState);
+
+			if (bHasPropertyValueInitializationState)
+			{
+				if (ArchiveState.IsLoading())
+				{
+					Data = MakeShared<FPropertyInitializedAnnotationData>();
+				}
+
+				*Slot << Data->PropertyInitialized;
+				
+				if (ArchiveState.IsLoading())
+				{
+					FPropertyInitializedAnnotation Annotation;
+					Annotation.Data = Data;
+					GetPropertyInitializedAnnotations().AddAnnotation(&Object, MoveTemp(Annotation));
+				}
+			}
+		}
+	}
 } // UE
 
 #endif // WITH_EDITORONLY_DATA
diff --git a/Engine/Source/Runtime/CoreUObject/Private/UObject/Obj.cpp b/Engine/Source/Runtime/CoreUObject/Private/UObject/Obj.cpp
--- a/Engine/Source/Runtime/CoreUObject/Private/UObject/Obj.cpp
+++ b/Engine/Source/Runtime/CoreUObject/Private/UObject/Obj.cpp
@@ -77,6 +77,7 @@
 #include "UObject/InstanceDataObjectUtils.h"
 #include "UObject/PropertyBagRepository.h"
 #include "AutoRTFM/AutoRTFM.h"
+#include "VerseVM/VVMVerseClass.h"
 
 DEFINE_LOG_CATEGORY(LogObj);
 
@@ -469,34 +470,55 @@ void UObject::PostEditChangeProperty(FPropertyChangedEvent& PropertyChangedEvent
 	FCoreUObjectDelegates::OnObjectPropertyChanged.Broadcast(this, PropertyChangedEvent);
 
 #if WITH_EDITORONLY_DATA
-	const FProperty* Property = PropertyChangedEvent.MemberProperty;
-	// Set property initialized state when a property is changed and IDO is supported
-	if (UE::IsInstanceDataObjectSupportEnabled(this))
+	if (const FProperty* Property = PropertyChangedEvent.MemberProperty)
 	{
-		UE::FPropertyBagRepository& Repository = UE::FPropertyBagRepository::Get();
-
-		UObject* ChangedIDO = UE::IsClassOfInstanceDataObjectClass(this->GetClass()) ?
-								this : Repository.FindInstanceDataObject(this);
-		if (ChangedIDO)
+		if (UE::SupportsPropertyInitStateTracking())
 		{
-			int32 ChangedIndex = 0;
-			switch (PropertyChangedEvent.ChangeType)
+			// Set property initialized state when a property is changed
+
+			auto SetPropertyValueInitializedHelper = [&PropertyChangedEvent](UObject* Object, const FProperty* P)
 			{
-			case EPropertyChangeType::ArrayAdd:
-			case EPropertyChangeType::ArrayRemove:
-			case EPropertyChangeType::ArrayMove:
-				ChangedIndex = PropertyChangedEvent.GetArrayIndex(PropertyChangedEvent.GetPropertyName().ToString());
-				ChangedIndex = (ChangedIndex == -1) ? 0 : ChangedIndex; // If we didn't get a valid array index revert to index 0 as it's the API default
-				UE::SetPropertyValueInitialized(ChangedIDO->GetClass(), ChangedIDO, Property, ChangedIndex);
-				break;
-			case EPropertyChangeType::ArrayClear: // Clearing the entire array. No need for an index. Clear also counts as initialization of the property
-			case EPropertyChangeType::ValueSet:
-			case EPropertyChangeType::Interactive:
-				UE::SetPropertyValueInitialized(ChangedIDO->GetClass(), ChangedIDO, Property);
-				break;
-			default:
-				// No op for now
-				break;
+				int32 ChangedIndex = 0;
+				switch (PropertyChangedEvent.ChangeType)
+				{
+				case EPropertyChangeType::ArrayAdd:
+				case EPropertyChangeType::ArrayRemove:
+				case EPropertyChangeType::ArrayMove:
+					// No tracking for elements of container.  Any change counts as initialization
+					UE::SetPropertyValueInitialized(Object->GetClass(), Object, P);
+					break;
+				case EPropertyChangeType::ArrayClear: // Clearing the entire array. No need for an index. Clear also counts as initialization of the property
+				case EPropertyChangeType::ValueSet:
+				case EPropertyChangeType::Interactive:
+					UE::SetPropertyValueInitialized(Object->GetClass(), Object, P);
+					break;
+				default:
+					// No op for now
+						break;
+				}
+			};
+	
+			if (!UE::IsInstanceDataObjectSupportEnabled())
+			{
+				// IDOs as a general system are disabled
+				UVerseClass* VerseClass = Cast<UVerseClass>(GetClass());
+				if (VerseClass)
+				{
+					SetPropertyValueInitializedHelper(this, Property);
+				}
+			}
+			else if (UE::IsInstanceDataObjectSupportEnabled(this))
+			{
+				// IDOs as a general system are enabled AND they are enabled on this object
+				// IDO Enabled case
+				UE::FPropertyBagRepository& Repository = UE::FPropertyBagRepository::Get();
+		
+				UObject* ChangedIDO = UE::IsClassOfInstanceDataObjectClass(this->GetClass()) ?
+										this : Repository.FindInstanceDataObject(this);
+				if (ChangedIDO)
+				{
+					SetPropertyValueInitializedHelper(ChangedIDO, Property);
+				}
 			}
 		}
 	}
@@ -1728,6 +1750,9 @@ void UObject::Serialize(FStructuredArchive::FRecord Record)
 				(UnderlyingArchive.UseUnversionedPropertySerialization() || UnderlyingArchive.IsTransacting()))
 			{
 				FOverridableManager::Get().SerializeOverriddenProperties(*this, Record);
+#if WITH_EDITORONLY_DATA
+				UE::SerializePropertyValueInitializationState(*this, Record);
+#endif
 			}
 
 			SerializeScriptProperties(Record.EnterField(TEXT("Properties")));
diff --git a/Engine/Source/Runtime/CoreUObject/Private/UObject/PropertyBagRepository.cpp b/Engine/Source/Runtime/CoreUObject/Private/UObject/PropertyBagRepository.cpp
--- a/Engine/Source/Runtime/CoreUObject/Private/UObject/PropertyBagRepository.cpp
+++ b/Engine/Source/Runtime/CoreUObject/Private/UObject/PropertyBagRepository.cpp
@@ -1046,7 +1046,8 @@ FScopedIDOSerializationContext::FScopedIDOSerializationContext(UObject* InObject
 
 		// Enable tracking of initialized properties when loading an IDO, which is implied by impersonation being enabled.
 		const bool bLoadingIDO = bHasIDOSupport && SerializeContext->bImpersonateProperties;
-		SerializeContext->bTrackInitializedProperties = bLoadingIDO;
+		// Track property initialization state if IDO feature not enabled
+		SerializeContext->bTrackInitializedProperties = IsInstanceDataObjectSupportEnabled() ? bLoadingIDO : true; 
 		SerializeContext->bTrackSerializedProperties = bLoadingIDO;
 	}
 	else
diff --git a/Engine/Source/Runtime/CoreUObject/Private/UObject/UObjectGlobals.cpp b/Engine/Source/Runtime/CoreUObject/Private/UObject/UObjectGlobals.cpp
--- a/Engine/Source/Runtime/CoreUObject/Private/UObject/UObjectGlobals.cpp
+++ b/Engine/Source/Runtime/CoreUObject/Private/UObject/UObjectGlobals.cpp
@@ -87,6 +87,7 @@
 #include "UObject/VerseValueProperty.h"
 #include "UObject/ObjectHandlePrivate.h"
 #include "UObject/RemoteObjectSerialization.h"
+#include "VerseVM/VVMVerseClass.h"
 
 #include "Interfaces/IPluginManager.h"
 
@@ -4275,6 +4276,36 @@ void FObjectInitializer::InitProperties(UObject* Obj, UClass* DefaultsClass, UOb
 		bCopyTransientsFromClassDefaults = false;
 	}
 
+	auto SetupPropertyValueInitializedState = [](UObject* TargetObject, UObject* Default, FProperty* Property)
+	{
+#if WITH_EDITORONLY_DATA
+		if (!UE::SupportsPropertyInitStateTracking())
+		{
+			return;
+		}
+		UClass* TargetClass = TargetObject->GetClass();
+		UClass* DefaultsClass = Default->GetClass();
+		
+		// PropertyValueInitialized state tracking only possible for UVerseClass types at this time
+		if (Cast<UVerseClass>(TargetClass) != nullptr)
+		{
+			// For RequiredParm properties, check if the template's property is initialized.
+			// if it is, then data just copied in was initialized data.
+			const bool bHasRequiredParam = Property->HasAllPropertyFlags(CPF_RequiredParm);
+			bool bIsTemplateCounterpartInitialized = UE::IsPropertyValueInitialized(DefaultsClass, Default, Property);
+			
+			if (!bHasRequiredParam || bIsTemplateCounterpartInitialized)
+			{
+				UE::SetPropertyValueInitialized(TargetClass, TargetObject, Property);
+			}
+			else
+			{
+				UE::ClearPropertyValueInitialized(TargetClass, TargetObject, Property);
+			}
+		}
+#endif
+	};
+
 	if (!bNeedInitialize && bCanUsePostConstructLink)
 	{
 		// This is just a fast path for the below in the common case that we are not doing a duplicate or initializing a CDO and this is all native.
@@ -4307,6 +4338,7 @@ void FObjectInitializer::InitProperties(UObject* Obj, UClass* DefaultsClass, UOb
 						if (P->IsInContainer(DefaultsClass))
 						{
 							P->CopyCompleteValue_InContainer(Obj, DefaultData);
+							SetupPropertyValueInitializedState(Obj, DefaultData, P);
 						}
 					}
 				}
@@ -4336,10 +4368,12 @@ void FObjectInitializer::InitProperties(UObject* Obj, UClass* DefaultsClass, UOb
 					// This is a duplicate. The value for all transient or non-duplicatable properties should be copied
 					// from the source class's defaults.
 					P->CopyCompleteValue_InContainer(Obj, ClassDefaults);
+					SetupPropertyValueInitializedState(Obj, ClassDefaults, P);
 				}
 				else if (P->IsInContainer(DefaultsClass))
 				{
 					P->CopyCompleteValue_InContainer(Obj, DefaultData);
+					SetupPropertyValueInitializedState(Obj, DefaultData, P);					
 				}
 			}
 		}
diff --git a/Engine/Source/Runtime/CoreUObject/Tests/UObject/InstanceDataObjectUtilsTest.cpp b/Engine/Source/Runtime/CoreUObject/Tests/UObject/InstanceDataObjectUtilsTest.cpp
--- a/Engine/Source/Runtime/CoreUObject/Tests/UObject/InstanceDataObjectUtilsTest.cpp
+++ b/Engine/Source/Runtime/CoreUObject/Tests/UObject/InstanceDataObjectUtilsTest.cpp
@@ -76,6 +76,7 @@ TEST_CASE_NAMED(FInstanceDataObjectUtilsTest, "CoreUObject::Serialization::Insta
 	FName TestObjectName = MakeUniqueObjectName(nullptr, TestClass, FName(WriteToString<128>(TestClass->GetFName(), TEXT("_Instance"))));
 	UObject* Owner = NewObject<UObject>(GetTransientPackage(), TestClass, TestObjectName);
 	void* StructData = StructProperty->ContainerPtrToValuePtr<void>(Owner);
+	ResetPropertyValueInitialized(TestClass, Owner);
 
 	CHECK_FALSE(WasPropertyValueSerialized(TestClass, Owner, StructProperty));
 	MarkPropertyValueSerialized(TestClass, Owner, StructProperty);
