FWeakObjectPtr::IsValid() will no longer attempt to resolve remote objects (which means operator== won't resolve remote objects too)
FWeakObjectPtr::IsValid() returns true if it points to an existing local object or to a remote object which can be resolved

diff --git a/Engine/Source/Runtime/CoreUObject/Private/UObject/WeakObjectPtr.cpp b/Engine/Source/Runtime/CoreUObject/Private/UObject/WeakObjectPtr.cpp
--- a/Engine/Source/Runtime/CoreUObject/Private/UObject/WeakObjectPtr.cpp
+++ b/Engine/Source/Runtime/CoreUObject/Private/UObject/WeakObjectPtr.cpp
@@ -198,4 +198,14 @@ FUObjectItem* FWeakObjectPtr::TryResolveRemoteObject() const
 	}
 	return nullptr;
 }
+
+bool FWeakObjectPtr::CanBeResolved() const
+{
+	// Object can be resolved if there's an existing FRemoteObjectSub with its id (stubs are created when an object is migrated to remote server)
+	// Note that this doesn't mean the object actually exists on any of the remote server. It could've been already GC'd and leaked to disk.
+	// Future plans: If an object was leaked to disk/database we could broadcast that to all servers and have them remove FRemoteObjectStub that points to a leaked object
+	// or we could check if the leaked object exists in the database and return false if it does.
+	return !!UE::RemoteObject::Private::FindRemoteObjectStub(ObjectRemoteId);
+}
+
 #endif // UE_WITH_REMOTE_OBJECT_HANDLE
\ No newline at end of file
diff --git a/Engine/Source/Runtime/CoreUObject/Public/UObject/WeakObjectPtr.h b/Engine/Source/Runtime/CoreUObject/Public/UObject/WeakObjectPtr.h
--- a/Engine/Source/Runtime/CoreUObject/Public/UObject/WeakObjectPtr.h
+++ b/Engine/Source/Runtime/CoreUObject/Public/UObject/WeakObjectPtr.h
@@ -289,6 +289,7 @@ private:
 	
 #if UE_WITH_REMOTE_OBJECT_HANDLE
 	COREUOBJECT_API FUObjectItem* TryResolveRemoteObject() const;
+	COREUOBJECT_API bool CanBeResolved() const;
 #endif
 
 	/**  
@@ -339,7 +340,8 @@ private:
 
 		FUObjectItem* ObjectItem = GUObjectArray.IndexToObject(ObjectIndex);
 		bool bResolved = false;
-		if (!ObjectItem)
+		// if an existing object is marked as Remote it needs to be resolved before its ObjectItem is returned
+		if (!ObjectItem || ObjectItem->HasAnyFlags(EInternalObjectFlags::Remote))
 		{
 			ObjectItem = TryResolveRemoteObject();
 			if (!ObjectItem)
@@ -387,6 +389,36 @@ private:
 #endif // UE_WITH_REMOTE_OBJECT_HANDLE
 	}
 
+#if UE_WITH_REMOTE_OBJECT_HANDLE
+	/** 
+	* Private (inlined) version for internal use only. Remote object version of IsValid does not try to resolve remote objects. 
+	* WeakObjectPtr is considered valid if it points to an existing local object or to a remote object.
+	*/
+	FORCEINLINE_DEBUGGABLE bool Internal_IsValid(bool bEvenIfGarbage, bool bThreadsafeTest) const
+	{
+		using namespace UE::Core::Private;
+
+		if (IsExplicitlyNull())
+		{
+			return false;
+		}
+		if (ObjectIndex < 0)
+		{
+			return false;
+		}
+
+		FUObjectItem* ObjectItem = GUObjectArray.IndexToObject(ObjectIndex);
+		if (!ObjectItem || !SerialNumbersMatch(ObjectItem))
+		{
+			return CanBeResolved();
+		}
+		else if (bThreadsafeTest)
+		{
+			return true;
+		}
+		return GUObjectArray.IsValid(ObjectItem, bEvenIfGarbage);
+	}
+#else
 	/** Private (inlined) version for internal use only. */
 	FORCEINLINE_DEBUGGABLE bool Internal_IsValid(bool bEvenIfGarbage, bool bThreadsafeTest) const
 	{
@@ -400,6 +432,7 @@ private:
 			return (ObjectItem != nullptr) && GUObjectArray.IsValid(ObjectItem, bEvenIfGarbage);
 		}
 	}
+#endif
 
 	/** Private (inlined) version for internal use only. */
 	FORCEINLINE_DEBUGGABLE UObject* Internal_Get(bool bEvenIfGarbage) const
