GetExternalPackage() will not longer attempt to resolve OuterPrivate if it's marked as remote (this behavior is disabled by default)

diff --git a/Engine/Source/Runtime/CoreUObject/Private/UObject/ObjectHandlePrivate.h b/Engine/Source/Runtime/CoreUObject/Private/UObject/ObjectHandlePrivate.h
--- a/Engine/Source/Runtime/CoreUObject/Private/UObject/ObjectHandlePrivate.h
+++ b/Engine/Source/Runtime/CoreUObject/Private/UObject/ObjectHandlePrivate.h
@@ -13,7 +13,7 @@ namespace UE::CoreUObject::Private
 			return Object->OuterPrivate;
 		}
 
-		FORCEINLINE static UObject* GetNonAccessTrackedOuterNoResolve(UObjectBase* Object)
+		FORCEINLINE static UObject* GetNonAccessTrackedOuterNoResolve(const UObjectBase* Object)
 		{
 #if UE_WITH_REMOTE_OBJECT_HANDLE
 			return Object->OuterPrivate.GetNoResolve();
diff --git a/Engine/Source/Runtime/CoreUObject/Private/UObject/UObjectBase.cpp b/Engine/Source/Runtime/CoreUObject/Private/UObject/UObjectBase.cpp
--- a/Engine/Source/Runtime/CoreUObject/Private/UObject/UObjectBase.cpp
+++ b/Engine/Source/Runtime/CoreUObject/Private/UObject/UObjectBase.cpp
@@ -275,7 +275,7 @@ void UObjectBase::LowLevelRename(FName NewName,UObject *NewOuter)
 UPackage* UObjectBase::GetExternalPackage() const
 {
 	// if we have no outer, consider this a package, packages returns themselves as their external package
-	if (OuterPrivate == nullptr)
+	if (UE::CoreUObject::Private::FObjectHandleUtils::GetNonAccessTrackedOuterNoResolve(this) == nullptr)
 	{
 		return CastChecked<UPackage>((UObject*)(this));
 	}
@@ -290,7 +290,7 @@ UPackage* UObjectBase::GetExternalPackage() const
 UPackage* UObjectBase::GetExternalPackageInternal() const
 {
 	// if we have no outer, consider this a package, packages returns themselves as their external package
-	if (OuterPrivate == nullptr)
+	if (UE::CoreUObject::Private::FObjectHandleUtils::GetNonAccessTrackedOuterNoResolve(this) == nullptr)
 	{
 		return CastChecked<UPackage>((UObject*)(this));
 	}
