Fix property bag atomicity under transactions.

diff --git a/Engine/Source/Programs/AutoRTFMTests/Private/UECoreTests.cpp b/Engine/Source/Programs/AutoRTFMTests/Private/UECoreTests.cpp
--- a/Engine/Source/Programs/AutoRTFMTests/Private/UECoreTests.cpp
+++ b/Engine/Source/Programs/AutoRTFMTests/Private/UECoreTests.cpp
@@ -1,7 +1,7 @@
 // Copyright Epic Games, Inc. All Rights Reserved.
 
 #include "Catch2Includes.h"
-#include "AutoRTFM/AutoRTFM.h"
+#include "AutoRTFM/AutoRTFMTesting.h"
 #include "Delegates/IDelegateInstance.h"
 #include "HAL/MallocLeakDetection.h"
 #include "HAL/ThreadSingleton.h"
@@ -16,6 +16,7 @@
 #include "Misc/TransactionallySafeMutex.h"
 #include "Containers/Queue.h"
 #include "Misc/ConfigCacheIni.h"
+#include "StructUtils/PropertyBag.h"
 
 TEST_CASE("UECore.FDelegateHandle")
 {
@@ -1926,3 +1927,36 @@ TEST_CASE("UECore.FConfigFile")
 		REQUIRE(Config.IsEmpty());
 	}
 }
+
+TEST_CASE("UECore.PropertyBag")
+{
+	UPropertyBag* const Bag = NewObject<UPropertyBag>();
+
+	UScriptStruct* const SS = static_cast<UScriptStruct*>(Bag);
+
+	SS->PrepareCppStructOps();
+
+	char Data[128];
+
+	AutoRTFM::Testing::Abort([&]
+	{
+		SS->InitializeStruct(Data);
+		AutoRTFM::AbortTransaction();
+	});
+
+	AutoRTFM::Testing::Commit([&]
+	{
+		SS->InitializeStruct(Data);
+	});
+
+	AutoRTFM::Testing::Abort([&]
+	{
+		SS->DestroyStruct(Data);
+		AutoRTFM::AbortTransaction();
+	});
+
+	AutoRTFM::Testing::Commit([&]
+	{
+		SS->DestroyStruct(Data);
+	});
+}
diff --git a/Engine/Source/Runtime/CoreUObject/Private/StructUtils/PropertyBag.cpp b/Engine/Source/Runtime/CoreUObject/Private/StructUtils/PropertyBag.cpp
--- a/Engine/Source/Runtime/CoreUObject/Private/StructUtils/PropertyBag.cpp
+++ b/Engine/Source/Runtime/CoreUObject/Private/StructUtils/PropertyBag.cpp
@@ -3085,17 +3085,19 @@ void UPropertyBag::DecrementRefCount() const
 	// Do ref counting based on struct usage.
 	// This ensures that if the UPropertyBag is still valid in the C++ destructor of
 	// the last instance of the bag.
-
-	UPropertyBag* NonConstThis = const_cast<UPropertyBag*>(this);
-	const int32 OldCount = NonConstThis->RefCount.fetch_sub(1, std::memory_order_acq_rel);
-	if (OldCount == 1)
-	{
-		NonConstThis->RemoveFromRoot();
-	}
-	if (OldCount <= 0)
+	UE_AUTORTFM_ONCOMMIT(this)
 	{
-		UE_LOG(LogCore, Error, TEXT("PropertyBag: DestroyStruct is called when RefCount is %d."), OldCount);
-	}
+		UPropertyBag* NonConstThis = const_cast<UPropertyBag*>(this);
+		const int32 OldCount = NonConstThis->RefCount.fetch_sub(1, std::memory_order_acq_rel);
+		if (OldCount == 1)
+		{
+			NonConstThis->RemoveFromRoot();
+		}
+		if (OldCount <= 0)
+		{
+			UE_LOG(LogCore, Error, TEXT("PropertyBag: DestroyStruct is called when RefCount is %d."), OldCount);
+		}
+	};
 }
 
 void UPropertyBag::IncrementRefCount() const
@@ -3104,7 +3106,17 @@ void UPropertyBag::IncrementRefCount() const
 	// This ensures that if the UPropertyBag is still valid in the C++ destructor of
 	// the last instance of the bag.
 	UPropertyBag* NonConstThis = const_cast<UPropertyBag*>(this);
-	const int32 OldCount = NonConstThis->RefCount.fetch_add(1, std::memory_order_acq_rel);
+
+	const int32 OldCount = AutoRTFM::Open([&]
+	{
+		return NonConstThis->RefCount.fetch_add(1, std::memory_order_acq_rel);
+	});
+
+	UE_AUTORTFM_ONABORT(NonConstThis)
+	{
+		NonConstThis->RefCount.fetch_sub(1, std::memory_order_acq_rel);
+	};
+
 	if (OldCount == 0)
 	{
 		NonConstThis->AddToRoot();
