Never resolve remote object when accessing it through GetOuter() of its inner object (behavior disabled by default)

diff --git a/Engine/Source/Runtime/CoreUObject/Public/UObject/ObjectPtr.h b/Engine/Source/Runtime/CoreUObject/Public/UObject/ObjectPtr.h
--- a/Engine/Source/Runtime/CoreUObject/Public/UObject/ObjectPtr.h
+++ b/Engine/Source/Runtime/CoreUObject/Public/UObject/ObjectPtr.h
@@ -823,6 +823,11 @@ namespace ObjectPtr_Private
 			return ObjectPtr;
 		}
 
+		const TObjectPtr<T>& GetAccessTrackedObjectPtr() const
+		{
+			return ObjectPtr;
+		}
+
 		operator T*() const
 		{
 			return Get();
diff --git a/Engine/Source/Runtime/CoreUObject/Public/UObject/UObjectBase.h b/Engine/Source/Runtime/CoreUObject/Public/UObject/UObjectBase.h
--- a/Engine/Source/Runtime/CoreUObject/Public/UObject/UObjectBase.h
+++ b/Engine/Source/Runtime/CoreUObject/Public/UObject/UObjectBase.h
@@ -183,7 +183,11 @@ public:
 	/** Returns the UObject this object resides in */
 	FORCEINLINE UObject* GetOuter() const
 	{
+#if UE_WITH_REMOTE_OBJECT_HANDLE
+		return ObjectPtr_Private::Friend::NoAccessTrackingGetNoResolve(OuterPrivate.GetAccessTrackedObjectPtr());
+#else
 		return OuterPrivate;
+#endif
 	}
 
 	/** Returns the logical name of this object */
