Made IDO hierarchy flat: all IDOs are now parented under the transient package. This fixes issues during re-instancing where new IDOs can become parented under an old IDO.

#rnx

diff --git a/Engine/Source/Runtime/CoreUObject/Private/UObject/PropertyBagRepository.cpp b/Engine/Source/Runtime/CoreUObject/Private/UObject/PropertyBagRepository.cpp
--- a/Engine/Source/Runtime/CoreUObject/Private/UObject/PropertyBagRepository.cpp
+++ b/Engine/Source/Runtime/CoreUObject/Private/UObject/PropertyBagRepository.cpp
@@ -768,16 +768,7 @@ UObject* FPropertyBagRepository::DuplicateInstanceDataObject(UObject* SourceOwne
 		FPropertyBagAssociationData& DestData = AssociatedData.FindOrAdd(DestOwner);
 		ensure(!DestData.InstanceDataObject);
 
-		// use the IDO of the outer as the new outer if there is one, otherwise use the transient package
-		UObject* Outer = nullptr;
-		if (const FPropertyBagAssociationData* OuterData = AssociatedData.Find(DestOwner->GetOuter()))
-		{
-			Outer = OuterData->InstanceDataObject;
-		}
-		if (!Outer)
-		{
-			Outer = GetTransientPackage();
-		}
+		UObject* Outer = GetTransientPackage();
 
 		// construct InstanceDataObject
 		FStaticConstructObjectParameters Params(SourceData->InstanceDataObject->GetClass());
@@ -962,16 +953,7 @@ void FPropertyBagRepository::CreateInstanceDataObjectUnsafe(UObject* Owner, FPro
 
 	BagData.bNeedsFixup = StructContainsLooseProperties(InstanceDataObjectClass);
 
-	// use the IDO of the outer as the new outer if there is one, otherwise use the transient package
-	UObject* Outer = nullptr;
-	if (FPropertyBagAssociationData* OuterData = AssociatedData.Find(Owner->GetOuter()))
-	{
-		Outer = OuterData->InstanceDataObject;
-	}
-	if (!Outer)
-	{
-		Outer = GetTransientPackage();
-	}
+	UObject* Outer = GetTransientPackage();
 
 	// construct InstanceDataObject object
 	FStaticConstructObjectParameters Params(InstanceDataObjectClass);
