StructUtils: Support for InstancedStruct in FPropertyValueIterator.

diff --git a/Engine/Source/Runtime/CoreUObject/Private/UObject/UnrealType.cpp b/Engine/Source/Runtime/CoreUObject/Private/UObject/UnrealType.cpp
--- a/Engine/Source/Runtime/CoreUObject/Private/UObject/UnrealType.cpp
+++ b/Engine/Source/Runtime/CoreUObject/Private/UObject/UnrealType.cpp
@@ -3,6 +3,7 @@
 #include "UObject/UnrealType.h"
 #include "UObject/PropertyOptional.h"
 #include "Serialization/ArchiveUObjectFromStructuredArchive.h"
+#include "StructUtils/InstancedStruct.h"
 
 DEFINE_LOG_CATEGORY(LogType);
 
@@ -53,6 +54,16 @@ FPropertyValueIterator::EPropertyValueFlags FPropertyValueIterator::GetPropertyV
 
 void FPropertyValueIterator::FillStructProperties(const UStruct* Struct, FPropertyValueStackEntry& Entry)
 {
+	if (Struct == TBaseStructure<FInstancedStruct>::Get())
+	{
+		const FInstancedStruct* InstancedStruct = reinterpret_cast<const FInstancedStruct*>(Entry.Owner);
+		if (InstancedStruct->GetScriptStruct() && InstancedStruct->GetMemory())
+		{
+			Entry.Owner = InstancedStruct->GetMemory();
+			Struct = InstancedStruct->GetScriptStruct();
+		}
+	}
+
 	const void* Owner = Entry.Owner;
 	FPropertyValueStackEntry::FValueArrayType& ValueArray = Entry.ValueArray;
 
