Add more debug info when there's a size mismatch during UStruct::SerializeVersionedTaggedProperties

#rnx

diff --git a/Engine/Source/Runtime/CoreUObject/Private/UObject/Class.cpp b/Engine/Source/Runtime/CoreUObject/Private/UObject/Class.cpp
--- a/Engine/Source/Runtime/CoreUObject/Private/UObject/Class.cpp
+++ b/Engine/Source/Runtime/CoreUObject/Private/UObject/Class.cpp
@@ -1921,8 +1921,8 @@ void UStruct::SerializeVersionedTaggedProperties(FStructuredArchive::FSlot Slot,
 				if (bAdvanceProperty)
 				{
 					checkf(Tag.Size == Loaded,
-						TEXT("Size mismatch in %s of %s of type %s. Loaded %" INT64_FMT " bytes but expected %d. Package: %s"),
-						*Tag.Name.ToString(), *GetName(), *WriteToString<64>(Tag.GetType()), Loaded, Tag.Size, *UnderlyingArchive.GetArchiveName());
+						TEXT("Size mismatch in %s of %s of type %s. Loaded %" INT64_FMT " bytes but expected %d. Package: %s. Property '%s' has outer '%s'."),
+						*Tag.Name.ToString(), *GetName(), *WriteToString<64>(Tag.GetType()), Loaded, Tag.Size, *UnderlyingArchive.GetArchiveName(), *Property->GetName(), *Property->GetOwnerVariant().GetFullName());
 
 				#if WITH_EDITORONLY_DATA
 					if (SerializeContext->bTrackInitializedProperties && Property)
