Deprecated TObjectPtr being constructed from a reference.

diff --git a/Engine/Plugins/FX/Niagara/Source/NiagaraEditor/Private/Widgets/SNiagaraScriptInputPreviewPanel.cpp b/Engine/Plugins/FX/Niagara/Source/NiagaraEditor/Private/Widgets/SNiagaraScriptInputPreviewPanel.cpp
--- a/Engine/Plugins/FX/Niagara/Source/NiagaraEditor/Private/Widgets/SNiagaraScriptInputPreviewPanel.cpp
+++ b/Engine/Plugins/FX/Niagara/Source/NiagaraEditor/Private/Widgets/SNiagaraScriptInputPreviewPanel.cpp
@@ -113,7 +113,7 @@ void SNiagaraScriptInputPreviewPanel::Refresh()
 		}
 
 		// Since they haven't been added to the hierarchy, we sort them lexicographically
-		TransientLeftoverParameters.Sort([](TObjectPtr<UNiagaraHierarchyScriptParameter> HierarchyScriptParameterA, TObjectPtr<UNiagaraHierarchyScriptParameter> HierarchyScriptParameterB)
+		Algo::Sort(TransientLeftoverParameters, [](TObjectPtr<UNiagaraHierarchyScriptParameter> HierarchyScriptParameterA, TObjectPtr<UNiagaraHierarchyScriptParameter> HierarchyScriptParameterB)
 		{
 			return HierarchyScriptParameterA->ToString() < HierarchyScriptParameterB->ToString();
 		});
diff --git a/Engine/Plugins/Interchange/Runtime/Source/Import/Private/Material/InterchangeMaterialFactory.cpp b/Engine/Plugins/Interchange/Runtime/Source/Import/Private/Material/InterchangeMaterialFactory.cpp
--- a/Engine/Plugins/Interchange/Runtime/Source/Import/Private/Material/InterchangeMaterialFactory.cpp
+++ b/Engine/Plugins/Interchange/Runtime/Source/Import/Private/Material/InterchangeMaterialFactory.cpp
@@ -1231,7 +1231,7 @@ void UInterchangeMaterialFactory::SetupMaterial(UMaterial* Material, const FImpo
 		}
 	}
 
-	FMaterialExpressionBuilder Builder(Material, nullptr, Arguments, *Results.Get());
+	FMaterialExpressionBuilder Builder(Material, nullptr, Arguments, Results.Get());
 
 	// Substrate
 	if(UInterchangeShaderPortsAPI::HasInput(MaterialFactoryNode, SubstrateMaterial::Parameters::FrontMaterial))
@@ -2174,7 +2174,7 @@ void UInterchangeMaterialFunctionFactory::SetupMaterial(UMaterialFunction* Mater
 	using namespace UE::Interchange::MaterialFactory::Internal;
 	using namespace UE::Interchange::Materials;
 
-	FMaterialExpressionBuilder Builder(nullptr, MaterialFunction, Arguments, *Results.Get());
+	FMaterialExpressionBuilder Builder(nullptr, MaterialFunction, Arguments, Results.Get());
 
 	UMaterialExpressionFunctionOutput* Output = CreateMaterialExpression<UMaterialExpressionFunctionOutput>(nullptr, MaterialFunction);
 
diff --git a/Engine/Plugins/NNE/NNEDenoiser/Source/NNEDenoiser/Private/NNEDenoiserModelInstanceCPU.cpp b/Engine/Plugins/NNE/NNEDenoiser/Source/NNEDenoiser/Private/NNEDenoiserModelInstanceCPU.cpp
--- a/Engine/Plugins/NNE/NNEDenoiser/Source/NNEDenoiser/Private/NNEDenoiserModelInstanceCPU.cpp
+++ b/Engine/Plugins/NNE/NNEDenoiser/Source/NNEDenoiser/Private/NNEDenoiserModelInstanceCPU.cpp
@@ -35,13 +35,13 @@ namespace UE::NNEDenoiser::Private
 			return {};
 		}
 
-		if (Runtime->CanCreateModelCPU(ModelData) != INNERuntimeCPU::ECanCreateModelCPUStatus::Ok)
+		if (Runtime->CanCreateModelCPU(&ModelData) != INNERuntimeCPU::ECanCreateModelCPUStatus::Ok)
 		{
 			UE_LOG(LogNNEDenoiser, Log, TEXT("%s on CPU can not create model"), *RuntimeName);
 			return {};
 		}
 
-		TSharedPtr<UE::NNE::IModelCPU> Model = Runtime->CreateModelCPU(ModelData);
+		TSharedPtr<UE::NNE::IModelCPU> Model = Runtime->CreateModelCPU(&ModelData);
 		if (!Model.IsValid())
 		{
 			UE_LOG(LogNNEDenoiser, Log, TEXT("Could not create model using %s on CPU"), *RuntimeName);
diff --git a/Engine/Plugins/NNE/NNEDenoiser/Source/NNEDenoiser/Private/NNEDenoiserModelInstanceGPU.cpp b/Engine/Plugins/NNE/NNEDenoiser/Source/NNEDenoiser/Private/NNEDenoiserModelInstanceGPU.cpp
--- a/Engine/Plugins/NNE/NNEDenoiser/Source/NNEDenoiser/Private/NNEDenoiserModelInstanceGPU.cpp
+++ b/Engine/Plugins/NNE/NNEDenoiser/Source/NNEDenoiser/Private/NNEDenoiserModelInstanceGPU.cpp
@@ -38,13 +38,13 @@ namespace UE::NNEDenoiser::Private
 			return {};
 		}
 
-		if (Runtime->CanCreateModelGPU(ModelData) != INNERuntimeGPU::ECanCreateModelGPUStatus::Ok)
+		if (Runtime->CanCreateModelGPU(&ModelData) != INNERuntimeGPU::ECanCreateModelGPUStatus::Ok)
 		{
 			UE_LOG(LogNNEDenoiser, Log, TEXT("%s on GPU can not create model"), *RuntimeName);
 			return {};
 		}
 
-		TSharedPtr<UE::NNE::IModelGPU> Model = Runtime->CreateModelGPU(ModelData);
+		TSharedPtr<UE::NNE::IModelGPU> Model = Runtime->CreateModelGPU(&ModelData);
 		if (!Model.IsValid())
 		{
 			UE_LOG(LogNNEDenoiser, Log, TEXT("Could not create model using %s on GPU"), *RuntimeName);
diff --git a/Engine/Plugins/NNE/NNEDenoiser/Source/NNEDenoiser/Private/NNEDenoiserModelInstanceRDG.cpp b/Engine/Plugins/NNE/NNEDenoiser/Source/NNEDenoiser/Private/NNEDenoiserModelInstanceRDG.cpp
--- a/Engine/Plugins/NNE/NNEDenoiser/Source/NNEDenoiser/Private/NNEDenoiserModelInstanceRDG.cpp
+++ b/Engine/Plugins/NNE/NNEDenoiser/Source/NNEDenoiser/Private/NNEDenoiserModelInstanceRDG.cpp
@@ -31,13 +31,13 @@ namespace UE::NNEDenoiser::Private
 			return {};
 		}
 
-		if (Runtime->CanCreateModelRDG(ModelData) != INNERuntimeRDG::ECanCreateModelRDGStatus::Ok)
+		if (Runtime->CanCreateModelRDG(&ModelData) != INNERuntimeRDG::ECanCreateModelRDGStatus::Ok)
 		{
 			UE_LOG(LogNNEDenoiser, Log, TEXT("%s on RDG can not create model"), *RuntimeName);
 			return {};
 		}
 
-		TSharedPtr<UE::NNE::IModelRDG> Model = Runtime->CreateModelRDG(ModelData);
+		TSharedPtr<UE::NNE::IModelRDG> Model = Runtime->CreateModelRDG(&ModelData);
 		if (!Model.IsValid())
 		{
 			UE_LOG(LogNNEDenoiser, Log, TEXT("Could not create model using %s on RDG"), *RuntimeName);
diff --git a/Engine/Plugins/Runtime/InstancedActors/Source/InstancedActors/Private/InstancedActorsData.cpp b/Engine/Plugins/Runtime/InstancedActors/Source/InstancedActors/Private/InstancedActorsData.cpp
--- a/Engine/Plugins/Runtime/InstancedActors/Source/InstancedActors/Private/InstancedActorsData.cpp
+++ b/Engine/Plugins/Runtime/InstancedActors/Source/InstancedActors/Private/InstancedActorsData.cpp
@@ -184,7 +184,7 @@ namespace UE::InstancedActors
 			AActor* ExemplarActor = Actor.Get();
 			check(ExemplarActor);
 	
-			InstancedActorsSubsystem.UnregisterExemplarActorClass(ExemplarActor->GetClass());
+			InstancedActorsSubsystem->UnregisterExemplarActorClass(ExemplarActor->GetClass());
 
 			ExemplarActor->Destroy();
 		}
diff --git a/Engine/Plugins/Runtime/InstancedActors/Source/InstancedActors/Private/InstancedActorsSubsystem.cpp b/Engine/Plugins/Runtime/InstancedActors/Source/InstancedActors/Private/InstancedActorsSubsystem.cpp
--- a/Engine/Plugins/Runtime/InstancedActors/Source/InstancedActors/Private/InstancedActorsSubsystem.cpp
+++ b/Engine/Plugins/Runtime/InstancedActors/Source/InstancedActors/Private/InstancedActorsSubsystem.cpp
@@ -689,7 +689,7 @@ TSharedRef<UE::InstancedActors::FExemplarActorData> UInstancedActorsSubsystem::G
 	check(NewExemplarActor);
 
 	// Cache for subsequent calls
-	TSharedPtr<UE::InstancedActors::FExemplarActorData> NewExemplarActorDataPtr{new UE::InstancedActors::FExemplarActorData{*NewExemplarActor, *this}};
+	TSharedPtr<UE::InstancedActors::FExemplarActorData> NewExemplarActorDataPtr{new UE::InstancedActors::FExemplarActorData{NewExemplarActor, this}};
 	ExemplarActors.AddByHash(ActorClassHash, ActorClassKey, NewExemplarActorDataPtr);
 
 	return NewExemplarActorDataPtr.ToSharedRef();
diff --git a/Engine/Plugins/Runtime/InstancedActors/Source/InstancedActors/Public/InstancedActorsData.h b/Engine/Plugins/Runtime/InstancedActors/Source/InstancedActors/Public/InstancedActorsData.h
--- a/Engine/Plugins/Runtime/InstancedActors/Source/InstancedActors/Public/InstancedActorsData.h
+++ b/Engine/Plugins/Runtime/InstancedActors/Source/InstancedActors/Public/InstancedActorsData.h
@@ -20,13 +20,13 @@ struct FExemplarActorData
 
 	TObjectPtr<AActor> Actor;
 private:
-	FExemplarActorData(AActor& InActor, UInstancedActorsSubsystem& InInstancedActorSubsystem)
+	FExemplarActorData(AActor* InActor, UInstancedActorsSubsystem* InInstancedActorSubsystem)
 		: Actor(InActor)
 		, InstancedActorsSubsystem(InInstancedActorSubsystem)
 	{
 	}
 
-	UInstancedActorsSubsystem& InstancedActorsSubsystem;
+	UInstancedActorsSubsystem* InstancedActorsSubsystem;
 
 	friend class ::UInstancedActorsSubsystem;
 };
diff --git a/Engine/Plugins/Runtime/MassGameplay/Source/MassRepresentation/Private/MassVisualizationComponent.cpp b/Engine/Plugins/Runtime/MassGameplay/Source/MassRepresentation/Private/MassVisualizationComponent.cpp
--- a/Engine/Plugins/Runtime/MassGameplay/Source/MassRepresentation/Private/MassVisualizationComponent.cpp
+++ b/Engine/Plugins/Runtime/MassGameplay/Source/MassRepresentation/Private/MassVisualizationComponent.cpp
@@ -102,7 +102,7 @@ FStaticMeshInstanceVisualizationDescHandle UMassVisualizationComponent::FindOrAd
 
 FStaticMeshInstanceVisualizationDescHandle UMassVisualizationComponent::AddVisualDescWithISMComponent(const FStaticMeshInstanceVisualizationDesc& Desc, UInstancedStaticMeshComponent& ISMComponent)
 {
-	TObjectPtr<UInstancedStaticMeshComponent> AsObjectPtr = ISMComponent;
+	TObjectPtr<UInstancedStaticMeshComponent> AsObjectPtr = &ISMComponent;
 	return AddVisualDescWithISMComponents(Desc, MakeArrayView(&AsObjectPtr, 1));
 }
 
diff --git a/Engine/Plugins/Runtime/Metasound/Source/MetasoundFrontend/Private/MetasoundFrontendDocumentBuilder.cpp b/Engine/Plugins/Runtime/Metasound/Source/MetasoundFrontend/Private/MetasoundFrontendDocumentBuilder.cpp
--- a/Engine/Plugins/Runtime/Metasound/Source/MetasoundFrontend/Private/MetasoundFrontendDocumentBuilder.cpp
+++ b/Engine/Plugins/Runtime/Metasound/Source/MetasoundFrontend/Private/MetasoundFrontendDocumentBuilder.cpp
@@ -727,8 +727,8 @@ UMetaSoundBuilderDocument& UMetaSoundBuilderDocument::Create(const IMetaSoundDoc
 	UMetaSoundBuilderDocument* DocObject = NewObject<UMetaSoundBuilderDocument>();
 	check(DocObject);
 	DocObject->Document = InDocToCopy.GetConstDocument();
-	DocObject->MetaSoundUClass = InDocToCopy.GetBaseMetaSoundUClass();
-	DocObject->BuilderUClass = InDocToCopy.GetBuilderUClass();
+	DocObject->MetaSoundUClass = &InDocToCopy.GetBaseMetaSoundUClass();
+	DocObject->BuilderUClass = &InDocToCopy.GetBuilderUClass();
 	return *DocObject;
 }
 
@@ -4365,7 +4365,7 @@ void FMetaSoundFrontendDocumentBuilder::SetMemberMetadata(UMetaSoundFrontendMemb
 
 	TMap<FGuid, TObjectPtr<UMetaSoundFrontendMemberMetadata>>& LiteralMetadata = GetDocumentChecked().Metadata.MemberMetadata;
 	LiteralMetadata.Remove(NewMetadata.MemberID);
-	LiteralMetadata.Add(NewMetadata.MemberID, NewMetadata);
+	LiteralMetadata.Add(NewMetadata.MemberID, &NewMetadata);
 }
 
 bool FMetaSoundFrontendDocumentBuilder::SetNodeComment(const FGuid& InNodeID, FString&& InNewComment, const FGuid* InPageID)
diff --git a/Engine/Source/Runtime/CoreUObject/Public/StructUtils/SharedStruct.h b/Engine/Source/Runtime/CoreUObject/Public/StructUtils/SharedStruct.h
--- a/Engine/Source/Runtime/CoreUObject/Public/StructUtils/SharedStruct.h
+++ b/Engine/Source/Runtime/CoreUObject/Public/StructUtils/SharedStruct.h
@@ -50,19 +50,27 @@ struct COREUOBJECT_API FStructSharedMemory
 		}
 	};
 
-	static TSharedPtr<FStructSharedMemory> Create(const UScriptStruct& InScriptStruct, const uint8* InStructMemory = nullptr)
+	static TSharedPtr<FStructSharedMemory> Create(const UScriptStruct* InScriptStruct, const uint8* InStructMemory = nullptr)
 	{
+		checkf(InScriptStruct, TEXT("FStructSharedMemory::Create - script struct cannot be null"));
+
 		FStructSharedMemory* const StructMemory = CreateImpl(InScriptStruct, InStructMemory);
 
 		return MakeShareable(StructMemory, FStructSharedMemoryDeleter());
 	}
 
+	UE_DEPRECATED(5.6, "FStructSharedMemory::Create should be passed a pointer.")
+	static TSharedPtr<FStructSharedMemory> Create(const UScriptStruct& InScriptStruct, const uint8* InStructMemory = nullptr)
+	{
+		return Create(&InScriptStruct, InStructMemory);
+	}
+
 	template<typename T, typename... TArgs>
 	static TSharedPtr<FStructSharedMemory> CreateArgs(TArgs&&... InArgs)
 	{
 		UE::StructUtils::CheckStructType<T>();
 
-		FStructSharedMemory* const StructMemory = CreateImpl(*TBaseStructure<T>::Get());
+		FStructSharedMemory* const StructMemory = CreateImpl(TBaseStructure<T>::Get());
 
 		new (StructMemory->GetMutableMemory()) T(Forward<TArgs>(InArgs)...);
 		return MakeShareable(StructMemory, FStructSharedMemoryDeleter());
@@ -92,7 +100,7 @@ struct COREUOBJECT_API FStructSharedMemory
 	}
 
 private:
-	FStructSharedMemory(const UScriptStruct& InScriptStruct, const uint8* InStructMemory = nullptr)
+	explicit FStructSharedMemory(const UScriptStruct* InScriptStruct, const uint8* InStructMemory = nullptr)
 		: ScriptStruct(InScriptStruct)
 	{
 		ScriptStruct->InitializeStruct(GetMutableMemory());
@@ -103,14 +111,14 @@ private:
 		}
 	}
 
-	static FStructSharedMemory* CreateImpl(const UScriptStruct& InScriptStruct, const uint8* InStructMemory = nullptr)
+	static FStructSharedMemory* CreateImpl(const UScriptStruct* InScriptStruct, const uint8* InStructMemory = nullptr)
 	{
 		// Align RequiredSize to InScriptStruct's alignment to effectively add padding in between ScriptStruct and
 		// StructMemory. GetMemory will then round &StructMemory up past this 'padding' to the nearest aligned address.
-		const int32 RequiredSize = static_cast<int32>(Align(sizeof(FStructSharedMemory), InScriptStruct.GetMinAlignment())) + InScriptStruct.GetStructureSize();
+		const int32 RequiredSize = static_cast<int32>(Align(sizeof(FStructSharedMemory), InScriptStruct->GetMinAlignment())) + InScriptStruct->GetStructureSize();
 		// Code analysis is unable to understand correctly what we are doing here, so disabling the warning C6386: Buffer overrun while writing to...
 		CA_SUPPRESS( 6386 )
-		FStructSharedMemory* StructMemory = new(FMemory::Malloc(RequiredSize, InScriptStruct.GetMinAlignment())) FStructSharedMemory(InScriptStruct, InStructMemory);
+		FStructSharedMemory* StructMemory = new(FMemory::Malloc(RequiredSize, InScriptStruct->GetMinAlignment())) FStructSharedMemory(InScriptStruct, InStructMemory);
 		return StructMemory;
 	}
 
@@ -219,7 +227,7 @@ struct COREUOBJECT_API FSharedStruct
 	{
 		if (InScriptStruct)
 		{
-			StructMemoryPtr = FStructSharedMemory::Create(*InScriptStruct, InStructMemory);
+			StructMemoryPtr = FStructSharedMemory::Create(InScriptStruct, InStructMemory);
 		}
 		else
 		{
@@ -618,7 +626,7 @@ struct COREUOBJECT_API FConstSharedStruct
 	{
 		if (InScriptStruct)
 		{
-			StructMemoryPtr = FStructSharedMemory::Create(*InScriptStruct, InStructMemory);
+			StructMemoryPtr = FStructSharedMemory::Create(InScriptStruct, InStructMemory);
 		}
 		else
 		{
diff --git a/Engine/Source/Runtime/CoreUObject/Public/UObject/ObjectPtr.h b/Engine/Source/Runtime/CoreUObject/Public/UObject/ObjectPtr.h
--- a/Engine/Source/Runtime/CoreUObject/Public/UObject/ObjectPtr.h
+++ b/Engine/Source/Runtime/CoreUObject/Public/UObject/ObjectPtr.h
@@ -533,6 +533,7 @@ public:
 	{
 	}
 
+	UE_DEPRECATED(5.6, "Constructing a TObjectPtr from a reference is deprecated - use a pointer instead.")
 	FORCEINLINE TObjectPtr(T& Object)
 		: ObjectPtr(const_cast<std::remove_const_t<T>*>(&Object))
 	{
