Cooker: When loading properties from config (which might include per-user settings) into the objects that use config properties, respect the Untracked keyword on SoftObjectProperties, and prevent those SoftObjectProperties from being harvested for cooking.
#rnx

diff --git a/Engine/Source/Runtime/CoreUObject/Private/UObject/Obj.cpp b/Engine/Source/Runtime/CoreUObject/Private/UObject/Obj.cpp
--- a/Engine/Source/Runtime/CoreUObject/Private/UObject/Obj.cpp
+++ b/Engine/Source/Runtime/CoreUObject/Private/UObject/Obj.cpp
@@ -3108,6 +3108,7 @@ void UObject::LoadConfig(UClass* ConfigClass/*=nullptr*/, const TCHAR* InFilenam
 	{
 #if WITH_EDITOR
 		FSoftObjectPathSerializationScope SerializationScope(NAME_None, Property->GetFName(),
+			Property->GetOwnerProperty()->HasMetaData(FSoftObjectPath::NAME_Untracked) ? ESoftObjectPathCollectType::NeverCollect :
 			Property->IsEditorOnlyProperty() ? ESoftObjectPathCollectType::EditorOnlyCollect : ESoftObjectPathCollectType::AlwaysCollect,
 			ESoftObjectPathSerializeType::AlwaysSerialize);
 #endif
diff --git a/Engine/Source/Runtime/CoreUObject/Private/UObject/SoftObjectPath.cpp b/Engine/Source/Runtime/CoreUObject/Private/UObject/SoftObjectPath.cpp
--- a/Engine/Source/Runtime/CoreUObject/Private/UObject/SoftObjectPath.cpp
+++ b/Engine/Source/Runtime/CoreUObject/Private/UObject/SoftObjectPath.cpp
@@ -28,6 +28,10 @@ static FAutoConsoleVariableRef CVarEnablePathFixupOutsideEditor(TEXT("SoftObject
 #endif
 }
 
+#if WITH_EDITOR
+FName FSoftObjectPath::NAME_Untracked(TEXT("Untracked"));
+#endif
+
 // Deprecated constructor
 FSoftObjectPath::FSoftObjectPath(FName InAssetPathName, FString InSubPathString)
 {
@@ -935,8 +939,7 @@ bool FSoftObjectPathThreadContext::GetSerializationOptions(FName& OutPackageName
 #if WITH_EDITOR
 		bEditorOnly = Archive->IsEditorOnlyPropertyOnTheStack();
 
-		static FName UntrackedName = TEXT("Untracked");
-		if (CurrentProperty && CurrentProperty->GetOwnerProperty()->HasMetaData(UntrackedName))
+		if (CurrentProperty && CurrentProperty->GetOwnerProperty()->HasMetaData(FSoftObjectPath::NAME_Untracked))
 		{
 			// Property has the Untracked metadata, so set to never collect references if it's higher than NeverCollect
 			CurrentCollectType = FMath::Min(ESoftObjectPathCollectType::NeverCollect, CurrentCollectType);
diff --git a/Engine/Source/Runtime/CoreUObject/Public/UObject/SoftObjectPath.h b/Engine/Source/Runtime/CoreUObject/Public/UObject/SoftObjectPath.h
--- a/Engine/Source/Runtime/CoreUObject/Public/UObject/SoftObjectPath.h
+++ b/Engine/Source/Runtime/CoreUObject/Public/UObject/SoftObjectPath.h
@@ -369,6 +369,9 @@ PRAGMA_ENABLE_DEPRECATION_WARNINGS;
 	/** Disables special PIE path handling, call when PIE finishes to clear list */
 	static COREUOBJECT_API void ClearPIEPackageNames();
 
+#if WITH_EDITOR
+	static COREUOBJECT_API FName NAME_Untracked;
+#endif
 private:
 	/** Asset path, patch to a top level object in a package. This is /package/path.assetname */
 	FTopLevelAssetPath AssetPath;
