SoftObjectPath: Add comment explaining NAME_Untracked.
#rnx

diff --git a/Engine/Source/Runtime/CoreUObject/Public/UObject/SoftObjectPath.h b/Engine/Source/Runtime/CoreUObject/Public/UObject/SoftObjectPath.h
--- a/Engine/Source/Runtime/CoreUObject/Public/UObject/SoftObjectPath.h
+++ b/Engine/Source/Runtime/CoreUObject/Public/UObject/SoftObjectPath.h
@@ -370,6 +370,10 @@ PRAGMA_ENABLE_DEPRECATION_WARNINGS;
 	static COREUOBJECT_API void ClearPIEPackageNames();
 
 #if WITH_EDITOR
+	/**
+	 * Text of the Untracked parameter in a UPROPERTY's meta field, if this parameter is present it marks that
+	 * the SoftObjectPath reference should not cause the target package to be added to the cook.
+	 */
 	static COREUOBJECT_API FName NAME_Untracked;
 #endif
 private:
