Remove unused debug class
* Added as accident

#rnx

diff --git a/Engine/Source/Runtime/CoreUObject/Private/UObject/InstanceDataObjectUtils.cpp b/Engine/Source/Runtime/CoreUObject/Private/UObject/InstanceDataObjectUtils.cpp
--- a/Engine/Source/Runtime/CoreUObject/Private/UObject/InstanceDataObjectUtils.cpp
+++ b/Engine/Source/Runtime/CoreUObject/Private/UObject/InstanceDataObjectUtils.cpp
@@ -996,20 +996,6 @@ namespace UE
 
 	using FPropertyInitializedAnnotationStore = FUObjectAnnotationSparse<FPropertyInitializedAnnotation, /*bAutoRemove*/ true>;
 
-	struct TestOnDelete
-	{
-		TestOnDelete()
-		{
-			volatile int i = 0;
-			++i;
-		}
-		~TestOnDelete()
-		{
-			volatile int i = 0;
-			++i;
-		}
-	};
-
 	UE_AUTORTFM_ALWAYS_OPEN
 	FPropertyInitializedAnnotationStore& GetPropertyInitializedAnnotations()
 	{
