Fix for calling Remote UFUNCTIONs from other Remote UFUNCTIONs.

diff --git a/Engine/Source/Programs/Shared/EpicGames.UHT/Exporters/CodeGen/UhtHeaderCodeGeneratorCppFile.cs b/Engine/Source/Programs/Shared/EpicGames.UHT/Exporters/CodeGen/UhtHeaderCodeGeneratorCppFile.cs
--- a/Engine/Source/Programs/Shared/EpicGames.UHT/Exporters/CodeGen/UhtHeaderCodeGeneratorCppFile.cs
+++ b/Engine/Source/Programs/Shared/EpicGames.UHT/Exporters/CodeGen/UhtHeaderCodeGeneratorCppFile.cs
@@ -1181,6 +1181,7 @@ namespace EpicGames.UHT.Exporters.CodeGen
 			{
 				bool isBlueprintEvent = function.FunctionFlags.HasAnyFlags(EFunctionFlags.BlueprintEvent);
 				bool isNetEvent = function.FunctionFlags.HasAnyFlags(EFunctionFlags.Net);
+				bool isNetRemoteEvent = isNetEvent && !function.FunctionFlags.HasAnyFlags(EFunctionFlags.NetClient | EFunctionFlags.NetServer | EFunctionFlags.NetMulticast);
 				bool isCallInEditor = function.MetaData.ContainsKey(UhtNames.CallInEditor);
 				bool isEditorFunction = function.FunctionFlags.HasAnyFlags(EFunctionFlags.EditorOnly);
 				bool hasValidCppImpl = (function.StrippedFunctionName != function.CppImplName);
@@ -1217,7 +1218,13 @@ namespace EpicGames.UHT.Exporters.CodeGen
 						.Append("\t{\r\n");
 
 					AppendEventFunctionPrologue(builder, function, function.EngineName, /*tabs*/ 1, "\r\n", /*addEventParameterStruct=*/false,/*addEventParameterStruct*/ false);
-				}								
+				}
+
+				// For remote net functions add them to our tracking stack
+				if (isNetRemoteEvent)
+				{
+					builder.Append("\tUE::Net::Private::FScopedRemoteRPCMode CallingRemoteRPC(Func, UE::Net::Private::ERemoteFunctionMode::Sending);\r\n");
+				}
 
 				// Cast away const just in case, because ProcessEvent isn't const
 				builder.Append('\t');
diff --git a/Engine/Source/Runtime/CoreUObject/Private/UObject/CoreNetContext.cpp b/Engine/Source/Runtime/CoreUObject/Private/UObject/CoreNetContext.cpp
new file mode 100644
--- /dev/null
+++ b/Engine/Source/Runtime/CoreUObject/Private/UObject/CoreNetContext.cpp
@@ -0,0 +1,55 @@
+// Copyright Epic Games, Inc. All Rights Reserved.
+
+#include "UObject/CoreNetContext.h"
+#include "UObject/Class.h"
+
+namespace UE::Net::Private
+{
+
+FCoreNetContext* FCoreNetContext::Get()
+{
+	static FCoreNetContext* InstancePtr;
+
+	if (!InstancePtr)
+	{
+		// It's important that the initializer goes here to avoid the overhead of
+		// "magic static" initialization on every call (mostly an issue with MSVC
+		// because of their epoch-based initialization scheme which doesn't seem
+		// to make any real sense on x86)
+
+		static FCoreNetContext Instance;
+		InstancePtr = &Instance;
+	}
+
+	return InstancePtr;
+}
+
+ERemoteFunctionMode FCoreNetContext::GetCurrentRemoteFunctionMode()
+{
+	return RemoteFunctionStack.IsEmpty() ? ERemoteFunctionMode::None : RemoteFunctionStack.Top();
+}
+
+FScopedRemoteRPCMode::FScopedRemoteRPCMode(const UFunction* Function, ERemoteFunctionMode Mode)
+{
+	// Only add to stack if the function is specified as "Remote"
+	bAddToStack =
+		((Function->FunctionFlags & FUNC_Net) != 0) &&
+		((Function->FunctionFlags & (FUNC_NetClient | FUNC_NetServer | FUNC_NetMulticast)) == 0);
+
+	if (bAddToStack)
+	{
+		check(FCoreNetContext::Get()->RemoteFunctionStack.Num() < FCoreNetContext::MaxRemoteStackSize);
+		FCoreNetContext::Get()->RemoteFunctionStack.Push(Mode);
+	}
+}
+
+FScopedRemoteRPCMode::~FScopedRemoteRPCMode()
+{
+	if (bAddToStack)
+	{
+		check(!FCoreNetContext::Get()->RemoteFunctionStack.IsEmpty());
+		FCoreNetContext::Get()->RemoteFunctionStack.Pop();
+	}
+}
+
+}
\ No newline at end of file
diff --git a/Engine/Source/Runtime/CoreUObject/Public/UObject/CoreNetContext.h b/Engine/Source/Runtime/CoreUObject/Public/UObject/CoreNetContext.h
new file mode 100644
--- /dev/null
+++ b/Engine/Source/Runtime/CoreUObject/Public/UObject/CoreNetContext.h
@@ -0,0 +1,52 @@
+// Copyright Epic Games, Inc. All Rights Reserved.
+
+#pragma once
+
+#include "Containers/Array.h"
+
+class UFunction;
+
+namespace UE::Net::Private
+{
+	// Indicates whether a the current scope of a ProcessEvent is call is due to
+	// receiving or sending a "Remote"-specified RPC, or not processing a "Remote" RPC at all.
+	enum class ERemoteFunctionMode
+	{
+		None,
+		Receiving,
+		Sending
+	};
+
+	// Global state required by the replication system, that also needs to be
+	// accessible to generated code.
+	class FCoreNetContext
+	{
+	public:
+		static COREUOBJECT_API FCoreNetContext* Get();
+
+		// Returns the "Remote" function mode currently at the top of the stack.
+		// Used internally to determine whether a "Remote"-specified function
+		// should be run locally or sent over the network as an RPC instead.
+		COREUOBJECT_API ERemoteFunctionMode GetCurrentRemoteFunctionMode();
+
+	private:
+		friend class FScopedRemoteRPCMode;
+
+		// Since Remote-specified functions don't run locally and can't recurse in the same callstack,
+		// the stack count shouldn't grow beyond 2.
+		static constexpr uint32 MaxRemoteStackSize = 2;
+		TArray<ERemoteFunctionMode, TFixedAllocator<MaxRemoteStackSize>> RemoteFunctionStack;
+	};
+
+	// Used internally to indicate whether a ProcessEvent call within the scope is due
+	// to sending or receiving a "Remote"-specified UFUNCTION as an RPC.
+	class COREUOBJECT_API FScopedRemoteRPCMode
+	{
+	public:
+		explicit FScopedRemoteRPCMode(const UFunction* Function, ERemoteFunctionMode Mode);
+		~FScopedRemoteRPCMode();
+
+	private:
+		bool bAddToStack = false;
+	};
+}
\ No newline at end of file
diff --git a/Engine/Source/Runtime/CoreUObject/Public/UObject/GeneratedCppIncludes.h b/Engine/Source/Runtime/CoreUObject/Public/UObject/GeneratedCppIncludes.h
--- a/Engine/Source/Runtime/CoreUObject/Public/UObject/GeneratedCppIncludes.h
+++ b/Engine/Source/Runtime/CoreUObject/Public/UObject/GeneratedCppIncludes.h
@@ -5,6 +5,7 @@
 #include "UObject/Object.h"
 #include "UObject/UObjectGlobals.h"
 #include "UObject/CoreNative.h"
+#include "UObject/CoreNetContext.h"
 #include "UObject/Class.h"
 #include "UObject/MetaData.h"
 #include "UObject/Package.h"
diff --git a/Engine/Source/Runtime/Engine/Private/Actor.cpp b/Engine/Source/Runtime/Engine/Private/Actor.cpp
--- a/Engine/Source/Runtime/Engine/Private/Actor.cpp
+++ b/Engine/Source/Runtime/Engine/Private/Actor.cpp
@@ -18,6 +18,7 @@
 #include "Engine/Engine.h"
 #include "Net/Core/PropertyConditions/RepChangedPropertyTracker.h"
 #include "UObject/AssetRegistryTagsContext.h"
+#include "UObject/CoreNetContext.h"
 #include "UObject/ObjectSaveContext.h"
 #include "Physics/Experimental/PhysScene_Chaos.h"
 #include "UObject/UE5MainStreamObjectVersion.h"
@@ -5120,7 +5121,7 @@ int32 AActor::GetFunctionCallspace( UFunction* Function, FFrame* Stack )
 			return Callspace;
 		}
 	}
-	else if (UE::Net::FNetContext::IsInsideNetRPC())
+	else if (UE::Net::Private::FCoreNetContext::Get()->GetCurrentRemoteFunctionMode() == UE::Net::Private::ERemoteFunctionMode::Receiving)
 	{
 		// Received a function with FUNC_Net but not Client, Server, or Multicast (Remote specifier) - can be sent bi-directionally
 		DEBUG_CALLSPACE(TEXT("GetFunctionCallspace Received a Remote function, calling: %s %s"), *Function->GetName(), FunctionCallspace::ToString(Callspace));
diff --git a/Engine/Source/Runtime/Engine/Private/DataReplication.cpp b/Engine/Source/Runtime/Engine/Private/DataReplication.cpp
--- a/Engine/Source/Runtime/Engine/Private/DataReplication.cpp
+++ b/Engine/Source/Runtime/Engine/Private/DataReplication.cpp
@@ -25,6 +25,7 @@
 #include "HAL/LowLevelMemStats.h"
 #include "Net/Core/PushModel/Types/PushModelPerNetDriverState.h"
 #include "Net/RPCDoSDetection.h"
+#include "UObject/CoreNetContext.h"
 
 DECLARE_LLM_MEMORY_STAT(TEXT("NetObjReplicator"), STAT_NetObjReplicatorLLM, STATGROUP_LLMFULL);
 LLM_DEFINE_TAG(NetObjReplicator, NAME_None, TEXT("Networking"), GET_STATFNAME(STAT_NetObjReplicatorLLM), GET_STATFNAME(STAT_NetworkingSummaryLLM));
@@ -1381,6 +1382,8 @@ bool FObjectReplicator::ReceivedRPC(FNetBitReader& Reader, const FReplicationFla
 
 			{
 				UE::Net::FScopedNetContextRPC CallingRPC;
+				UE::Net::Private::FScopedRemoteRPCMode ReceivingRemoteRPC(Function, UE::Net::Private::ERemoteFunctionMode::Receiving);
+
 				// Call the function.
 				Object->ProcessEvent(Function, Parms);
 			}
diff --git a/Engine/Source/Runtime/Experimental/Iris/Core/Private/Iris/ReplicationSystem/NetBlob/NetRPC.cpp b/Engine/Source/Runtime/Experimental/Iris/Core/Private/Iris/ReplicationSystem/NetBlob/NetRPC.cpp
--- a/Engine/Source/Runtime/Experimental/Iris/Core/Private/Iris/ReplicationSystem/NetBlob/NetRPC.cpp
+++ b/Engine/Source/Runtime/Experimental/Iris/Core/Private/Iris/ReplicationSystem/NetBlob/NetRPC.cpp
@@ -25,6 +25,7 @@
 #include "Iris/Core/IrisProfiler.h"
 #include "Containers/ArrayView.h"
 #include "ProfilingDebugging/CsvProfiler.h"
+#include "UObject/CoreNetContext.h"
 
 DEFINE_LOG_CATEGORY_STATIC(LogIrisRpc, Log, All);
 
@@ -589,6 +590,7 @@ void FNetRPC::CallFunction(FNetRPCCallContext& CallContext)
 #endif
 
 		UE::Net::FScopedNetContextRPC CallingRPC;
+		UE::Net::Private::FScopedRemoteRPCMode ReceivingRemoteRPC(ActualFunction, UE::Net::Private::ERemoteFunctionMode::Receiving);
 		Object->ProcessEvent(const_cast<UFunction*>(ActualFunction), FunctionParameters);
 	}
 
