Adding delegates called when object/object data gets migrated to or from remote server

diff --git a/Engine/Source/Runtime/CoreUObject/Private/UObject/RemoteObjectSerialization.cpp b/Engine/Source/Runtime/CoreUObject/Private/UObject/RemoteObjectSerialization.cpp
--- a/Engine/Source/Runtime/CoreUObject/Private/UObject/RemoteObjectSerialization.cpp
+++ b/Engine/Source/Runtime/CoreUObject/Private/UObject/RemoteObjectSerialization.cpp
@@ -73,7 +73,7 @@ void LoadObjectFromDisk(FRemoteObjectId ObjectId, FRemoteServerId CurrentOwnerSe
 	*FileReader << ObjectData->Bytes;
 	FileReader->Close();
 
-	UE::RemoteObject::Transfer::OnObjectDataReceived(ObjectId, ObjectData);
+	UE::RemoteObject::Transfer::OnObjectDataReceived(ObjectId, CurrentOwnerServerId, ObjectData);
 }
 
 void SaveObjectToDisk(FRemoteObjectId ObjectId, FRemoteServerId ServerId, FRemoteObjectData& ObjectData)
diff --git a/Engine/Source/Runtime/CoreUObject/Private/UObject/RemoteObjectTransfer.cpp b/Engine/Source/Runtime/CoreUObject/Private/UObject/RemoteObjectTransfer.cpp
--- a/Engine/Source/Runtime/CoreUObject/Private/UObject/RemoteObjectTransfer.cpp
+++ b/Engine/Source/Runtime/CoreUObject/Private/UObject/RemoteObjectTransfer.cpp
@@ -20,6 +20,10 @@ namespace UE::RemoteObject::Transfer
 TDelegate<void(FRemoteObjectId ObjectId, FRemoteServerId CurrentOwnerServerId, FRemoteServerId DestinationServerId)> RequestRemoteObjectDelegate;
 TDelegate<void(FRemoteObjectId ObjectId, FRemoteServerId DestinationServerId, FRemoteObjectData& ObjectData)> RemoteObjectTransferDelegate;
 TDelegate<void()> FlushRemoteObjectRequestsDelegate;
+TMulticastDelegate<void(const FRemoteObjectData& ObjectData, FRemoteServerId RemoteServerId)> OnObjectDataSentDelegate;
+TMulticastDelegate<void(const TSet<UObject*>& MigratedObjects, FRemoteServerId RemoteServerId)> OnObjectsSentDelegate;
+TMulticastDelegate<void(const FRemoteObjectData& ObjectData, FRemoteServerId RemoteServerId)> OnObjectDataReceivedDelegate;
+TMulticastDelegate<void(const TArray<UObject*>& MigratedObjects, FRemoteServerId RemoteServerId)> OnObjectsReceivedDelegate;
 
 /**
  * Remote object transfer queue. 
@@ -357,10 +361,13 @@ void SendObjectData(UObject* Object, FRemoteServerId DestinationServerId, FRemot
 	GTransferQueue.SendRemoteObject(UE::CoreUObject::Private::FObjectHandleUtils::GetRemoteId(Object), DestinationServerId, ObjectData);
 }
 
-UObject* MigrateObjectFromObjectData(FRemoteObjectData& ObjectData, UObject* DestinationOuter)
+UObject* MigrateObjectFromObjectData(FRemoteObjectData& ObjectData, FRemoteServerId RemoteServerId, UObject* DestinationOuter)
 {
 	using namespace UE::RemoteObject::Serialization;
 
+	// Notify that object data has been migrated from a remote server
+	OnObjectDataReceivedDelegate.Broadcast(ObjectData, RemoteServerId);
+
 	TArray<UObject*> ReceivedObjects;
 	DeserializeObjectData(ObjectData, DestinationOuter, ReceivedObjects);
 
@@ -373,6 +380,9 @@ UObject* MigrateObjectFromObjectData(FRemoteObjectData& ObjectData, UObject* Des
 	}
 #endif
 
+	// Notify that objects have been migrated from a remote server
+	OnObjectsReceivedDelegate.Broadcast(ReceivedObjects, RemoteServerId);
+
 	return ReceivedObjects[0];
 } 
 
@@ -387,7 +397,7 @@ UObject* MigrateObjectFromRemoteServer(FRemoteObjectId ObjectId, FRemoteServerId
 
 	// The below code should probably be moved to FulfillReceiveRequest
 	FRemoteObjectData& ObjectData = *ObjectFuture.Get();
-	UObject* Result = MigrateObjectFromObjectData(ObjectData, DestinationOuter);
+	UObject* Result = MigrateObjectFromObjectData(ObjectData, CurrentOwnerServerId, DestinationOuter);
 	checkf(UE::CoreUObject::Private::FObjectHandleUtils::GetRemoteId(Result) == ObjectId, TEXT("PeerId:%s requested object %s migration from PeerId:%s but received %s"),
 		*GetGlobalServerId().ToString(), *ObjectId.ToString(), *CurrentOwnerServerId.ToString(), *UE::CoreUObject::Private::FObjectHandleUtils::GetRemoteId(Result).ToString());
 
@@ -401,9 +411,16 @@ void MigrateObjectToRemoteServer(UObject* Object, FRemoteServerId DestinationSer
 
 	TSet<UObject*> SentObjects;	
 	FRemoteObjectData* ObjectData = SerializeObjectData(Object, SentObjects);
+
+	// Notify that object data has been migrated (actually will be but SendObjectData does a MoveTemp on the data)
+	OnObjectDataSentDelegate.Broadcast(*ObjectData, DestinationServerId);
+
 	SendObjectData(Object, DestinationServerId, ObjectData);
 
-	checkf(SentObjects.Num() > 0, TEXT("PeerId:%s unable to migrate object %s to PeerId:%s"), *GetGlobalServerId().ToString(), *UE::CoreUObject::Private::FObjectHandleUtils::GetRemoteId(Object).ToString(), *DestinationServerId.ToString());
+	checkf(SentObjects.Num() > 0, TEXT("PeerId:%s unable to migrate object %s to PeerId:%s"), *GetGlobalServerId().ToString(), *UE::CoreUObject::Private::FObjectHandleUtils::GetRemoteId(Object).ToString(), *DestinationServerId.ToString());	
+
+	// Notify that objects have been migrated to a remote server
+	OnObjectsSentDelegate.Broadcast(SentObjects, DestinationServerId);
 
 	for (UObject* SentObject : SentObjects)
 	{
@@ -437,11 +454,11 @@ void MigrateObjectToRemoteServer(FRemoteObjectId ObjectId, FRemoteServerId Desti
 	}
 }
 
-void OnObjectDataReceived(FRemoteObjectId ObjectId, FRemoteObjectData* Data)
+void OnObjectDataReceived(FRemoteObjectId ObjectId, FRemoteServerId RemoteServerId, FRemoteObjectData* Data)
 {
 	if (!GTransferQueue.FulfillReceiveRequest(ObjectId, Data))
 	{
-		MigrateObjectFromObjectData(*Data, nullptr);
+		MigrateObjectFromObjectData(*Data, RemoteServerId, nullptr);
 	}
 }
 
diff --git a/Engine/Source/Runtime/CoreUObject/Public/UObject/RemoteObjectTransfer.h b/Engine/Source/Runtime/CoreUObject/Public/UObject/RemoteObjectTransfer.h
--- a/Engine/Source/Runtime/CoreUObject/Public/UObject/RemoteObjectTransfer.h
+++ b/Engine/Source/Runtime/CoreUObject/Public/UObject/RemoteObjectTransfer.h
@@ -15,9 +15,10 @@ namespace UE::RemoteObject::Transfer
 	/**
 	* Called when remote object data has been received from a remote server
 	* @param ObjectId Remote object id
+	* @param RemoteServerId Server that sent object data
 	* @param Data Remote object data. Data ownership is transferred to the remote object queue
 	*/
-	COREUOBJECT_API void OnObjectDataReceived(FRemoteObjectId ObjectId, FRemoteObjectData* Data);
+	COREUOBJECT_API void OnObjectDataReceived(FRemoteObjectId ObjectId, FRemoteServerId RemoteServerId, FRemoteObjectData* Data);
 
 	/**
 	* Requests remote object data to be transferred from Current Owner Server to Destination Server
@@ -72,5 +73,24 @@ namespace UE::RemoteObject::Transfer
 	/** Delegate that flushes remote object requests */
 	extern COREUOBJECT_API TDelegate<void()> FlushRemoteObjectRequestsDelegate;
 
+	/** 
+	 * Delegate executed when objects have been migrated from another server. 
+	 */
+	extern COREUOBJECT_API TMulticastDelegate<void(const FRemoteObjectData& /*ObjectData*/, FRemoteServerId /*RemoteServerId*/)> OnObjectDataReceivedDelegate;
+
+	/**
+	 * Delegate executed when objects have been migrated from another server.
+	 */
+	extern COREUOBJECT_API TMulticastDelegate<void(const TArray<UObject*>& /*MigratedObjects*/, FRemoteServerId /*RemoteServerId*/)> OnObjectsReceivedDelegate;
+
+	/** Delegate executed when object data has been migrated to another server */
+	extern COREUOBJECT_API TMulticastDelegate<void(const FRemoteObjectData& /*ObjectData*/, FRemoteServerId /*RemoteServerId*/)> OnObjectDataSentDelegate;
+
+	/** 
+	 * Delegate executed when objects have been migrated to another server. Migrated objects can still be resolved with weak pointers but at this point 
+	 * any changes to their internal state won't be migrated across to the remote server 
+	 */
+	extern COREUOBJECT_API TMulticastDelegate<void(const TSet<UObject*>& /*MigratedObjects*/, FRemoteServerId /*RemoteServerId*/)> OnObjectsSentDelegate;
+
 } // namespace UE::RemoveObject::Transfer
 
