Updated the comment for BlueprintPure to reflect the correct definition of "pure".

BlueprintPure was originally intended to mean functional pure, but that never materialized. It just means that a Blueprint function node will hide its exec pins by default, which can be toggled at the callsite.

diff --git a/Engine/Source/Runtime/CoreUObject/Public/UObject/ObjectMacros.h b/Engine/Source/Runtime/CoreUObject/Public/UObject/ObjectMacros.h
--- a/Engine/Source/Runtime/CoreUObject/Public/UObject/ObjectMacros.h
+++ b/Engine/Source/Runtime/CoreUObject/Public/UObject/ObjectMacros.h
@@ -946,7 +946,7 @@ namespace UF
 		/// Only valid when used in conjunction with Client, Server, or NetMulticast
 		Unreliable,
 
-		/// This function fulfills a contract of producing no side effects, and additionally implies BlueprintCallable.
+		/// This function will not have visible exec pins by default. Implies BlueprintCallable.
 		BlueprintPure,
 
 		/// This function can be called from blueprint code and should be exposed to the user of blueprint editing tools.
