[UPackage] Add CVAR to enable dumping a stacktrace when a package is dirtied.

diff --git a/Engine/Source/Runtime/CoreUObject/Private/UObject/Package.cpp b/Engine/Source/Runtime/CoreUObject/Private/UObject/Package.cpp
--- a/Engine/Source/Runtime/CoreUObject/Private/UObject/Package.cpp
+++ b/Engine/Source/Runtime/CoreUObject/Private/UObject/Package.cpp
@@ -6,6 +6,7 @@
 #include "Cooker/CookDependency.h"
 #include "HAL/FileManager.h"
 #include "HAL/PlatformMath.h"
+#include "HAL/IConsoleManager.h"
 #include "Misc/AssetRegistryInterface.h"
 #include "Misc/ITransaction.h"
 #include "Misc/PackageName.h"
@@ -22,6 +23,50 @@
 	UPackage.
 -----------------------------------------------------------------------------*/
 
+#if WITH_EDITOR && !UE_BUILD_SHIPPING
+static FAutoConsoleVariable DumpStackTraceOnPackageDirty(
+	TEXT("Package.DumpStackTraceOnDirty"),
+	0,
+	TEXT("Dumps a stack trace every time a package is dirtied. Executes the number of times specified, -1 = infinite."),
+	FConsoleVariableDelegate::CreateLambda([](IConsoleVariable* Variable)
+	{
+		if (!IsInGameThread())
+		{
+			return;
+		}
+		
+		static FDelegateHandle PackageDirtyStateChangedEventHandle;
+		int32 MaxRunCount = Variable->GetInt();
+		static int32 RunsSoFar = 0;
+
+		if (PackageDirtyStateChangedEventHandle.IsValid())
+		{
+			UPackage::PackageDirtyStateChangedEvent.Remove(PackageDirtyStateChangedEventHandle);
+			PackageDirtyStateChangedEventHandle.Reset();
+			RunsSoFar = 0;
+		}
+
+		if (MaxRunCount != 0)
+		{
+			PackageDirtyStateChangedEventHandle = UPackage::PackageDirtyStateChangedEvent.AddLambda([MaxRunCount](const UPackage* InPackage)
+			{
+				if(MaxRunCount != -1 && RunsSoFar >= MaxRunCount)
+				{
+					return;
+				}
+				if (InPackage && InPackage->IsDirty())
+				{
+					const SIZE_T StackTraceSize = 65536;
+					ANSICHAR StackTrace[StackTraceSize] = { 0 };
+					FPlatformStackWalk::StackWalkAndDump(StackTrace, StackTraceSize, 1);
+					UE_LOG(LogObj, Warning, TEXT("***** Package %s Marked Dirty ******\n%s"), *InPackage->GetPathName(), ANSI_TO_TCHAR(StackTrace));
+					RunsSoFar++;
+				}
+			});
+		}
+	}));
+#endif // WITH_EDITOR && !UE_BUILD_SHIPPING
+
 PRAGMA_DISABLE_DEPRECATION_WARNINGS;
 UPackage::FPreSavePackage UPackage::PreSavePackageEvent;
 UPackage::FOnPackageSaved UPackage::PackageSavedEvent;
