Added support for leaking objects to disk when they are potentially referenced by a remote server

diff --git a/Engine/Source/Runtime/CoreUObject/Private/UObject/GarbageCollection.cpp b/Engine/Source/Runtime/CoreUObject/Private/UObject/GarbageCollection.cpp
--- a/Engine/Source/Runtime/CoreUObject/Private/UObject/GarbageCollection.cpp
+++ b/Engine/Source/Runtime/CoreUObject/Private/UObject/GarbageCollection.cpp
@@ -45,6 +45,7 @@
 #include "UObject/ExpandingChunkedList.h"
 #include "UObject/ReachabilityAnalysisState.h"
 #include "UObject/ReachabilityAnalysis.h"
+#include "UObject/RemoteObjectPrivate.h"
 
 #include <atomic>
 
@@ -5951,6 +5952,12 @@ bool UnhashUnreachableObjects(bool bUseTimeLimit, double TimeLimit)
 				UObject* Object = static_cast<UObject*>(ObjectItem->Object);
 				FScopedCBDProfile Profile(Object);
 				// Begin the object's asynchronous destruction.
+#if UE_WITH_REMOTE_OBJECT_HANDLE
+				if (!GExitPurge && ObjectItem->HasAnyFlags(EInternalObjectFlags::RemoteReference) && !ObjectItem->HasAnyFlags(EInternalObjectFlags::Remote))
+				{
+					UE::RemoteObject::Transfer::Private::StoreObjectToDatabase(Object);
+				}
+#endif // UE_WITH_REMOTE_OBJECT_HANDLE
 				Object->ConditionalBeginDestroy();
 			}
 
diff --git a/Engine/Source/Runtime/CoreUObject/Private/UObject/RemoteObject.cpp b/Engine/Source/Runtime/CoreUObject/Private/UObject/RemoteObject.cpp
--- a/Engine/Source/Runtime/CoreUObject/Private/UObject/RemoteObject.cpp
+++ b/Engine/Source/Runtime/CoreUObject/Private/UObject/RemoteObject.cpp
@@ -7,6 +7,7 @@
 #include "Serialization/ArchiveUObject.h"
 #include "Serialization/MemoryWriter.h"
 #include "Serialization/MemoryReader.h"
+#include "UObject/GarbageCollection.h"
 #include "UObject/RemoteObjectTransfer.h"
 #include "UObject/RemoteObjectSerialization.h"
 #include "UObject/RemoteObjectPrivate.h"
@@ -126,6 +127,7 @@ void InitRemoteObjects()
 	}
 	checkf(ServerIdNumber != 0, TEXT("Remote Server Id=%u is not valid"), ServerIdNumber);
 	RemoteServerId = FRemoteServerId(ServerIdNumber);
+	checkf(RemoteServerId != UE::RemoteObject::Transfer::DatabaseId, TEXT("Remote Server Id=%s is reserved for database"), *RemoteServerId.ToString());
 
 	if (!UE::RemoteObject::Transfer::RemoteObjectTransferDelegate.IsBound())
 	{
@@ -135,8 +137,16 @@ void InitRemoteObjects()
 	{
 		UE::RemoteObject::Transfer::RequestRemoteObjectDelegate.BindStatic(&UE::RemoteObject::Serialization::Disk::LoadObjectFromDisk);
 	}
+	if (!UE::RemoteObject::Transfer::StoreRemoteObjectDataDelegate.IsBound())
+	{
+		UE::RemoteObject::Transfer::StoreRemoteObjectDataDelegate.BindStatic(&UE::RemoteObject::Serialization::Disk::SaveObjectToDisk);
+	}
+	if (!UE::RemoteObject::Transfer::RestoreRemoteObjectDataDelegate.IsBound())
+	{
+		UE::RemoteObject::Transfer::RestoreRemoteObjectDataDelegate.BindStatic(&UE::RemoteObject::Serialization::Disk::LoadObjectFromDisk);
+	}
 
-	FCoreDelegates::OnEndFrame.AddStatic(&TickRemoteObjectRequests);
+	FCoreDelegates::OnEndFrame.AddStatic(&UE::RemoteObject::Transfer::Private::TickRemoteObjectRequests);
 }
 
 void MarkAsRemote(UObject* Object, FRemoteServerId DestinationServerId)
@@ -144,7 +154,7 @@ void MarkAsRemote(UObject* Object, FRemoteServerId DestinationServerId)
 	static_assert(sizeof(FObjectHandle) == sizeof(UObject*));
 	FUObjectItem* ObjectItem = GUObjectArray.ObjectToObjectItem(Object);
 	ObjectItem->SetFlags(EInternalObjectFlags::Remote);
-	ObjectItem->ClearFlags(EInternalObjectFlags_RootFlags);
+	ObjectItem->ClearFlags(EInternalObjectFlags_RootFlags | EInternalObjectFlags::RemoteReference);
 	ObjectMaps->FindOrAddRemoteObjectStub(Object, DestinationServerId);
 }
 
@@ -152,6 +162,7 @@ void MarkAsLocal(UObject* Object)
 {
 	FUObjectItem* ObjectItem = GUObjectArray.ObjectToObjectItem(Object);
 	ObjectItem->ClearFlags(EInternalObjectFlags::Remote);
+	ObjectItem->SetFlags(EInternalObjectFlags::RemoteReference);
 }
 
 void MapObject(UObject* Object)
@@ -216,6 +227,8 @@ UObject* ResolveObject(const FRemoteObjectStub* Stub)
 	using namespace UE::RemoteObject::Transfer;
 	using namespace UE::RemoteObject::Private;
 
+	checkf(!GIsGarbageCollecting, TEXT("Resolving remote objects while collecting garbage is not allowed (trying to resolve object %s (%s)"), *Stub->Id.ToString(), *Stub->Name.ToString());
+
 	UObject* Object = ObjectMaps->FindObject(Stub->Id);
 	bool bRemoteObject = false;
 	if (Object)
@@ -260,6 +273,19 @@ namespace UE::CoreUObject::Private
 		MapObject((UObject*)Object);
 	}
 
+	FRemoteObjectId FRemoteObjectHandlePrivate::GetRemoteId() const
+	{
+		if ((PointerOrHandle & UPTRINT(1)))
+		{
+			return ToStub()->Id;
+		}
+		else if (PointerOrHandle)
+		{
+			return FRemoteObjectId(reinterpret_cast<const UObjectBase*>(PointerOrHandle));
+		}
+		return FRemoteObjectId();
+	}
+
 	FRemoteObjectHandlePrivate FRemoteObjectHandlePrivate::ConvertToRemoteHandle(UObject* Object)
 	{
 		using namespace UE::RemoteObject::Handle;
diff --git a/Engine/Source/Runtime/CoreUObject/Private/UObject/RemoteObjectPrivate.h b/Engine/Source/Runtime/CoreUObject/Private/UObject/RemoteObjectPrivate.h
--- a/Engine/Source/Runtime/CoreUObject/Private/UObject/RemoteObjectPrivate.h
+++ b/Engine/Source/Runtime/CoreUObject/Private/UObject/RemoteObjectPrivate.h
@@ -43,13 +43,21 @@ namespace UE::RemoteObject::Private
 	*/
 	UObject* FindObject(FRemoteObjectId ObjectId);
 
-	/**
-	* Processes remote object requests (send and receive requests)
-	*/
-	void TickRemoteObjectRequests();
-
 	/**
 	* Finds remote object stub
 	*/
 	UE::RemoteObject::Handle::FRemoteObjectStub* FindRemoteObjectStub(FRemoteObjectId ObjectId);
 }
+
+namespace UE::RemoteObject::Transfer::Private
+{
+	/**
+	* Stores unreachable objects to database
+	*/
+	void StoreObjectToDatabase(UObject* Object);
+
+	/**
+	* Processes remote object requests (send and receive requests)
+	*/
+	void TickRemoteObjectRequests();
+}
\ No newline at end of file
diff --git a/Engine/Source/Runtime/CoreUObject/Private/UObject/RemoteObjectSerialization.cpp b/Engine/Source/Runtime/CoreUObject/Private/UObject/RemoteObjectSerialization.cpp
--- a/Engine/Source/Runtime/CoreUObject/Private/UObject/RemoteObjectSerialization.cpp
+++ b/Engine/Source/Runtime/CoreUObject/Private/UObject/RemoteObjectSerialization.cpp
@@ -58,27 +58,29 @@ void DeserializeNameTables(FArchive& Ar, FRemoteObjectData& OutObjectData)
 	Ar << OutObjectData.PathNames;
 }
 
-FString GenerateRemoteObjectFilename(FRemoteObjectId ObjectId)
+FString GenerateRemoteObjectFilename(FRemoteObjectId ObjectId, FRemoteServerId DestinationServerId)
 {
-	return FPaths::Combine(*FPaths::ProjectSavedDir(), *(ObjectId.ToString() + TEXT(".tmp")));
+	return FPaths::Combine(*FPaths::ProjectSavedDir(), *(FString::Printf(TEXT("%s_%s.remote"), *ObjectId.ToString(), *DestinationServerId.ToString())));
 }
 
 void LoadObjectFromDisk(FRemoteObjectId ObjectId, FRemoteServerId CurrentOwnerServerId, FRemoteServerId DestinationServerId)
 {
-	TUniquePtr<FArchive> FileReader{ IFileManager::Get().CreateFileReader(*GenerateRemoteObjectFilename(ObjectId)) };
+	const FString Filename(GenerateRemoteObjectFilename(ObjectId, CurrentOwnerServerId));
+	TUniquePtr<FArchive> FileReader{ IFileManager::Get().CreateFileReader(*Filename) };
 	checkf(FileReader.IsValid(), TEXT("Unable to create file reader for remote object %s"), *ObjectId.ToString());
 	FRemoteObjectData* ObjectData = new FRemoteObjectData();
 	
 	DeserializeNameTables(*FileReader, *ObjectData);
 	*FileReader << ObjectData->Bytes;
 	FileReader->Close();
+	IFileManager::Get().Delete(*Filename, false, true, true);
 
 	UE::RemoteObject::Transfer::OnObjectDataReceived(ObjectId, CurrentOwnerServerId, ObjectData);
 }
 
-void SaveObjectToDisk(FRemoteObjectId ObjectId, FRemoteServerId ServerId, FRemoteObjectData& ObjectData)
+void SaveObjectToDisk(FRemoteObjectId ObjectId, FRemoteServerId DestinationServerId, FRemoteObjectData& ObjectData)
 {
-	TUniquePtr<FArchive> FileWriter{ IFileManager::Get().CreateFileWriter(*GenerateRemoteObjectFilename(ObjectId)) };
+	TUniquePtr<FArchive> FileWriter{ IFileManager::Get().CreateFileWriter(*GenerateRemoteObjectFilename(ObjectId, DestinationServerId)) };
 	checkf(FileWriter.IsValid(), TEXT("Unable to create file writer for remote object %s"), *ObjectId.ToString());
 	
 	SerializeNameTables(*FileWriter, ObjectData);
@@ -140,24 +142,38 @@ class FArchiveRemoteObjectWriter : public FMemoryWriter
 	TMap<FString, TMap<FProperty*, FSerializeStats>> ObjectPropertyStats;
 #endif
 
-	ERemoteReferenceType WriteObjectReference(UObject*& Obj)
+	void WriteObjectReference(const FObjectPtr& ObjPtr)
 	{
 		using namespace UE::CoreUObject::Private;
 
 		ERemoteReferenceType Type = ERemoteReferenceType::None;
-		FRemoteObjectId Id;
+		FRemoteObjectId Id = ObjPtr.GetRemoteId();
+		UObject* Obj = nullptr;
 
-		if (Obj)
+		if (ObjPtr)
 		{
-			Id = FObjectHandleUtils::GetRemoteId(Obj);
-
-			if (Obj == RootObject || UE::RemoteObject::Handle::IsRemote(Obj) || Obj->IsIn(RootObject) || (Id.GetServerId().IsValid() && Id.GetServerId() != UE::RemoteObject::GetGlobalServerId()))
+			if (ObjPtr.IsResolved())
 			{
-				Type = ERemoteReferenceType::IdOnly;
+				Obj = ObjPtr.Get();
+				checkf(Obj, TEXT("Attempting to serialize resolved Object pointer but the pointer resolves to null (%s)"), *Id.ToString());
+				Obj->SetInternalFlags(EInternalObjectFlags::RemoteReference);
+				check(!UE::RemoteObject::Handle::IsRemote(Obj));
+
+				Id = FObjectHandleUtils::GetRemoteId(Obj);
+
+				if (Obj == RootObject || Obj->IsIn(RootObject) || (Id.GetServerId().IsValid() && Id.GetServerId() != UE::RemoteObject::GetGlobalServerId()))
+				{
+					Type = ERemoteReferenceType::IdOnly;
+				}
+				else
+				{
+					
+					Type = ERemoteReferenceType::PathName;
+				}
 			}
 			else
 			{
-				Type = ERemoteReferenceType::PathName;
+				Type = ERemoteReferenceType::IdOnly;
 			}
 		}
 
@@ -168,13 +184,28 @@ class FArchiveRemoteObjectWriter : public FMemoryWriter
 			*this << Id;
 		}
 
-		if (Type == ERemoteReferenceType::PathName)
+		if (Type == ERemoteReferenceType::IdOnly)
+		{
+			// Add subobjects of the root object to the list of objects to serialize
+			if (Obj && Obj->IsIn(RootObject))
+			{
+				ObjectsToSerialize.Add(Obj);
+			}
+		}
+		else if (Type == ERemoteReferenceType::PathName)
 		{
 			FNameIndexType PathNameIndex = AddPathNameToNameMap(Obj);
 			*this << PathNameIndex;
 		}
 
-		return Type;
+		if (Obj)
+		{
+			// Map all references we encounter during serialization, especially if they're objects that are not going to be serialized
+			// because we may end up needing to resolve them later if a remote server requests them with their unique id.
+			// Note: we could just hash all objects with their FRemoteObjectId in UObjectHash.cpp but for now we can limit the memory impact by hashing 
+			// only those we encounter when serializing objects that are actually going to be migrated
+			UE::RemoteObject::Private::MapObject(Obj);
+		}
 	}
 
 public:
@@ -312,23 +343,7 @@ public:
 
 	virtual FArchive& operator<<(UObject*& Obj) override
 	{
-		ERemoteReferenceType RefType = WriteObjectReference(Obj);
-		if (RefType == ERemoteReferenceType::IdOnly)
-		{
-			if (Obj && Obj->IsIn(RootObject))
-			{
-				ObjectsToSerialize.Add(Obj);
-			}			
-		}
-		if (Obj)
-		{
-			// Map all references we encounter during serialization, especially if they're objects that are not going to be serialized
-			// because we may end up needing to resolve them later if a remote server requests them with their unique id.
-			// Note: we could just hash all objects with their FRemoteObjectId in UObjectHash.cpp but for now we can limit the memory impact by hashing 
-			// only those we encounter when serializing objects that are actually going to be migrated
-			UE::RemoteObject::Private::MapObject(Obj);
-		}
-
+		WriteObjectReference(FObjectPtr(Obj));
 		return *this;
 	}
 
@@ -342,7 +357,8 @@ public:
 	virtual FArchive& operator<<(FLazyObjectPtr& Value) override { return FArchiveUObject::SerializeLazyObjectPtr(*this, Value); }
 	virtual FArchive& operator<<(FObjectPtr& Value) override 
 	{ 
-		return FArchiveUObject::SerializeObjectPtr(*this, Value); 
+		WriteObjectReference(Value);
+		return *this;
 	}
 	virtual FArchive& operator<<(FSoftObjectPtr& Value) override { return FArchiveUObject::SerializeSoftObjectPtr(*this, Value); }
 	virtual FArchive& operator<<(FSoftObjectPath& Value) override { return FArchiveUObject::SerializeSoftObjectPath(*this, Value); }
@@ -432,6 +448,7 @@ class FArchiveRemoteObjectReader : public FMemoryReader
 				UObject* Obj = ResolvedPathNameObjects[PathNameIndex];
 				if (Obj)
 				{
+					Obj->SetInternalFlags(EInternalObjectFlags::RemoteReference);
 					Value = FObjectPtr(Obj);
 					bNeedsResolvingWithId = false;
 				}
diff --git a/Engine/Source/Runtime/CoreUObject/Private/UObject/RemoteObjectTransfer.cpp b/Engine/Source/Runtime/CoreUObject/Private/UObject/RemoteObjectTransfer.cpp
--- a/Engine/Source/Runtime/CoreUObject/Private/UObject/RemoteObjectTransfer.cpp
+++ b/Engine/Source/Runtime/CoreUObject/Private/UObject/RemoteObjectTransfer.cpp
@@ -17,14 +17,20 @@
 namespace UE::RemoteObject::Transfer
 {
 
+const FRemoteServerId DatabaseId(0x3ff);
+
 TDelegate<void(FRemoteObjectId ObjectId, FRemoteServerId CurrentOwnerServerId, FRemoteServerId DestinationServerId)> RequestRemoteObjectDelegate;
 TDelegate<void(FRemoteObjectId ObjectId, FRemoteServerId DestinationServerId, FRemoteObjectData& ObjectData)> RemoteObjectTransferDelegate;
 TDelegate<void()> FlushRemoteObjectRequestsDelegate;
+
 TMulticastDelegate<void(const FRemoteObjectData& ObjectData, FRemoteServerId RemoteServerId)> OnObjectDataSentDelegate;
 TMulticastDelegate<void(const TSet<UObject*>& MigratedObjects, FRemoteServerId RemoteServerId)> OnObjectsSentDelegate;
 TMulticastDelegate<void(const FRemoteObjectData& ObjectData, FRemoteServerId RemoteServerId)> OnObjectDataReceivedDelegate;
 TMulticastDelegate<void(const TArray<UObject*>& MigratedObjects, FRemoteServerId RemoteServerId)> OnObjectsReceivedDelegate;
 
+TDelegate<void(FRemoteObjectId ObjectId, FRemoteServerId DatabaseId, FRemoteObjectData& ObjectData)> StoreRemoteObjectDataDelegate;
+TDelegate<void(FRemoteObjectId ObjectId, FRemoteServerId DatabaseId, FRemoteServerId DestinationServerId)> RestoreRemoteObjectDataDelegate;
+
 /**
  * Remote object transfer queue. 
  * Queued requests (both send and receive) are processed on the game thread since some systems (like RPCs) that are used for transfering object data are GT-only
@@ -80,8 +86,15 @@ class FRemoteObjectTransferQueue
 		}
 		virtual void Process() override
 		{
-			// Request an object from remote server
-			RequestRemoteObjectDelegate.Execute(ObjectId, /*CurrentOwner*/ ServerId, GetGlobalServerId());
+			if (ServerId != DatabaseId)
+			{
+				// Request an object from remote server
+				RequestRemoteObjectDelegate.Execute(ObjectId, /*CurrentOwner*/ ServerId, GetGlobalServerId());
+			}
+			else
+			{
+				RestoreRemoteObjectDataDelegate.Execute(ObjectId, DatabaseId, GetGlobalServerId());
+			}
 		}
 	};
 
@@ -460,14 +473,62 @@ void OnObjectDataReceived(FRemoteObjectId ObjectId, FRemoteServerId RemoteServer
 	{
 		MigrateObjectFromObjectData(*Data, RemoteServerId, nullptr);
 	}
-}
+} 
+
+} // // namespace UE::RemoteObject::Transfer
 
-} // namespace UE::RemoteObject::Transfer
+namespace UE::RemoteObject::Transfer::Private
+{
+
+UObject* GetOutermostUnreachableRemoteReference(UObject* Object)
+{
+	UObject* OutermostUnreachableRemoteReference = Object;
+	for (UObject* Outer = Object->GetOuter(); Outer; Outer = Outer->GetOuter())
+	{
+		FUObjectItem* OuterItem = GUObjectArray.ObjectToObjectItem(Outer);
+		if (OuterItem->HasAllFlags(EInternalObjectFlags::RemoteReference | EInternalObjectFlags::Unreachable))
+		{
+			OutermostUnreachableRemoteReference = Outer;
+		}
+		else
+		{
+			break;
+		}
+	}
+	return OutermostUnreachableRemoteReference;
+}
 
-namespace UE::RemoteObject::Private
+void StoreObjectToDatabase(UObject* Object)
 {
-	void TickRemoteObjectRequests()
+	using namespace UE::RemoteObject::Serialization;
+	using namespace UE::RemoteObject::Private;
+
+	// Object could a subobject of its outer that is going to be destroyed (and have this function invoked on) later so to make
+	// storing to database (disk) consistent with how objects are migrated (root object with its subobjects in one transfer), 
+	// find this object's outer that is also unreachable and marked as remote reference
+	UObject* OutermostUnreachableRemoteReference = GetOutermostUnreachableRemoteReference(Object);
+
+	TSet<UObject*> SentObjects;
+	TUniquePtr<FRemoteObjectData> ObjectData(SerializeObjectData(OutermostUnreachableRemoteReference, SentObjects));
+	
+	UE::RemoteObject::Transfer::StoreRemoteObjectDataDelegate.Execute(FRemoteObjectId(OutermostUnreachableRemoteReference), DatabaseId, *ObjectData);
+
+	checkf(SentObjects.Num() > 0, TEXT("PeerId:%s unable to store object %s to DatabaseId:%s"), *GetGlobalServerId().ToString(), *FRemoteObjectId(Object).ToString(), *DatabaseId.ToString());
+
+	// At this point the entire leaked object tree is remote with the exception that its current owner server is the database
+	for (UObject* SentObject : SentObjects)
 	{
-		UE::RemoteObject::Transfer::GTransferQueue.Tick();
+		MarkAsRemote(SentObject, DatabaseId);
 	}
-}
\ No newline at end of file
+
+	checkf(!Object->HasAnyInternalFlags(EInternalObjectFlags::RemoteReference), 
+		TEXT("PeerId:%s has object %s still marked as remote reference after storing to database%s"), *FRemoteObjectId(Object).ToString(), *Object->GetPathName(),
+		OutermostUnreachableRemoteReference != Object ? *FString::Printf(TEXT(" when storing its outermost remote referenced object %s"), *OutermostUnreachableRemoteReference->GetPathName()) : TEXT("."));
+}
+
+void TickRemoteObjectRequests()
+{
+	UE::RemoteObject::Transfer::GTransferQueue.Tick();
+}
+
+} // namespace UE::RemoteObject::Transfer::Private
\ No newline at end of file
diff --git a/Engine/Source/Runtime/CoreUObject/Public/UObject/ObjectHandle.h b/Engine/Source/Runtime/CoreUObject/Public/UObject/ObjectHandle.h
--- a/Engine/Source/Runtime/CoreUObject/Public/UObject/ObjectHandle.h
+++ b/Engine/Source/Runtime/CoreUObject/Public/UObject/ObjectHandle.h
@@ -95,6 +95,8 @@ namespace UE::CoreUObject::Private
 			return reinterpret_cast<UE::RemoteObject::Handle::FRemoteObjectStub*>(PointerOrHandle & ~UPTRINT(1));
 		}
 
+		COREUOBJECT_API FRemoteObjectId GetRemoteId() const;
+
 		COREUOBJECT_API static FRemoteObjectHandlePrivate ConvertToRemoteHandle(UObject* Object);
 		COREUOBJECT_API static FRemoteObjectHandlePrivate FromIdNoResolve(FRemoteObjectId ObjectId);
 	};
diff --git a/Engine/Source/Runtime/CoreUObject/Public/UObject/ObjectMacros.h b/Engine/Source/Runtime/CoreUObject/Public/UObject/ObjectMacros.h
--- a/Engine/Source/Runtime/CoreUObject/Public/UObject/ObjectMacros.h
+++ b/Engine/Source/Runtime/CoreUObject/Public/UObject/ObjectMacros.h
@@ -618,7 +618,8 @@ enum class EInternalObjectFlags : int32
 	ReachabilityFlag1 = 1 << 1, ///< One of the flags used by Garbage Collector to determine UObject's reachability state
 	ReachabilityFlag2 = 1 << 2, ///< One of the flags used by Garbage Collector to determine UObject's reachability state
 
-	Remote = 1 << 19, //< Object is no longer owned by this process
+	Remote = 1 << 18, //< Object is no longer owned by this process
+	RemoteReference = 1 << 19, //< Object referenced by remote process
 	LoaderImport = 1 << 20, ///< Object is ready to be imported by another package during loading
 	Garbage = 1 << 21, ///< Garbage from logical point of view and should not be referenced. This flag is mirrored in EObjectFlags as RF_Garbage for performance
 	AsyncLoadingPhase1 = 1 << 22, ///< Object is being asynchronously loaded.
@@ -641,7 +642,7 @@ ENUM_CLASS_FLAGS(EInternalObjectFlags);
 
 //~ Make sure these macros are up to date!
 #define EInternalObjectFlags_GarbageCollectionKeepFlags (EInternalObjectFlags::Native | EInternalObjectFlags::Async | EInternalObjectFlags::AsyncLoadingPhase1 | EInternalObjectFlags::AsyncLoadingPhase2 | EInternalObjectFlags::LoaderImport | EInternalObjectFlags::RefCounted)
-#define EInternalObjectFlags_AllFlags (EInternalObjectFlags::ReachabilityFlag0 | EInternalObjectFlags::ReachabilityFlag1 | EInternalObjectFlags::ReachabilityFlag2 | EInternalObjectFlags::Remote | EInternalObjectFlags::LoaderImport | EInternalObjectFlags::Garbage | EInternalObjectFlags::ReachableInCluster | EInternalObjectFlags::ClusterRoot | EInternalObjectFlags::Native | EInternalObjectFlags::RefCounted | EInternalObjectFlags::Async | EInternalObjectFlags::AsyncLoadingPhase1 | EInternalObjectFlags::AsyncLoadingPhase2 | EInternalObjectFlags::RootSet | EInternalObjectFlags::PendingConstruction | (EInternalObjectFlags::Unreachable))
+#define EInternalObjectFlags_AllFlags (EInternalObjectFlags::ReachabilityFlag0 | EInternalObjectFlags::ReachabilityFlag1 | EInternalObjectFlags::ReachabilityFlag2 | EInternalObjectFlags::Remote | EInternalObjectFlags::RemoteReference | EInternalObjectFlags::LoaderImport | EInternalObjectFlags::Garbage | EInternalObjectFlags::ReachableInCluster | EInternalObjectFlags::ClusterRoot | EInternalObjectFlags::Native | EInternalObjectFlags::RefCounted | EInternalObjectFlags::Async | EInternalObjectFlags::AsyncLoadingPhase1 | EInternalObjectFlags::AsyncLoadingPhase2 | EInternalObjectFlags::RootSet | EInternalObjectFlags::PendingConstruction | (EInternalObjectFlags::Unreachable))
 #define EInternalObjectFlags_RootFlags (EInternalObjectFlags::RootSet | EInternalObjectFlags_GarbageCollectionKeepFlags)
 #define EInternalObjectFlags_ReachabilityFlags (EInternalObjectFlags::ReachabilityFlag0 | EInternalObjectFlags::ReachabilityFlag1 | EInternalObjectFlags::ReachabilityFlag2 | EInternalObjectFlags::Unreachable)
 #define EInternalObjectFlags_AsyncLoading (EInternalObjectFlags::AsyncLoadingPhase1 | EInternalObjectFlags::AsyncLoadingPhase2)
diff --git a/Engine/Source/Runtime/CoreUObject/Public/UObject/ObjectPtr.h b/Engine/Source/Runtime/CoreUObject/Public/UObject/ObjectPtr.h
--- a/Engine/Source/Runtime/CoreUObject/Public/UObject/ObjectPtr.h
+++ b/Engine/Source/Runtime/CoreUObject/Public/UObject/ObjectPtr.h
@@ -215,6 +215,11 @@ public:
 
 	/** Returns the full name for the object in the form: Class ObjectPath */
 	COREUOBJECT_API FString GetFullName(EObjectFullNameFlags Flags = EObjectFullNameFlags::None) const;
+
+	FRemoteObjectId GetRemoteId() const
+	{
+		return FRemoteObjectId();
+	}
 #else
 	FString GetPathName() const
 	{
@@ -253,6 +258,15 @@ public:
 		// UObjectBaseUtility::GetFullName is safe to call on null objects.
 		return UE::CoreUObject::Private::GetFullName(Get(), nullptr, Flags);
 	}
+
+	FRemoteObjectId GetRemoteId() const
+	{
+#if UE_WITH_REMOTE_OBJECT_HANDLE
+		return GetHandle().GetRemoteId();
+#else
+		return FRemoteObjectId();
+#endif
+	}
 #endif
 
 	FORCEINLINE FObjectHandle GetHandle() const { return Handle; }
diff --git a/Engine/Source/Runtime/CoreUObject/Public/UObject/RemoteObjectTransfer.h b/Engine/Source/Runtime/CoreUObject/Public/UObject/RemoteObjectTransfer.h
--- a/Engine/Source/Runtime/CoreUObject/Public/UObject/RemoteObjectTransfer.h
+++ b/Engine/Source/Runtime/CoreUObject/Public/UObject/RemoteObjectTransfer.h
@@ -64,6 +64,8 @@ namespace UE::RemoteObject::Transfer
 #endif
 	}
 
+	extern COREUOBJECT_API const FRemoteServerId DatabaseId;
+
 	/** Delegate that transfers object data to another server */
 	extern COREUOBJECT_API TDelegate<void(FRemoteObjectId /*ObjectId*/, FRemoteServerId /*DestinationServerId*/, FRemoteObjectData& /*ObjectData*/)> RemoteObjectTransferDelegate;
 
@@ -92,5 +94,11 @@ namespace UE::RemoteObject::Transfer
 	 */
 	extern COREUOBJECT_API TMulticastDelegate<void(const TSet<UObject*>& /*MigratedObjects*/, FRemoteServerId /*RemoteServerId*/)> OnObjectsSentDelegate;
 
+	/** Delegate that stores locally unrachable object data into a database. */
+	extern COREUOBJECT_API TDelegate<void(FRemoteObjectId /*ObjectId*/, FRemoteServerId /*DatabaseId*/, FRemoteObjectData& /*ObjectData*/)> StoreRemoteObjectDataDelegate;
+
+	/** Delegate that restores object data from a database. */
+	extern COREUOBJECT_API TDelegate<void(FRemoteObjectId /*ObjectId*/, FRemoteServerId /*DatabaseId*/, FRemoteServerId /*DestinationServerId*/)> RestoreRemoteObjectDataDelegate;
+
 } // namespace UE::RemoveObject::Transfer
 
