Exposing TickRemoteObjectRequests() and FlushkRemoteObjectRequests() as public API

diff --git a/Engine/Source/Runtime/CoreUObject/Private/UObject/RemoteObject.cpp b/Engine/Source/Runtime/CoreUObject/Private/UObject/RemoteObject.cpp
--- a/Engine/Source/Runtime/CoreUObject/Private/UObject/RemoteObject.cpp
+++ b/Engine/Source/Runtime/CoreUObject/Private/UObject/RemoteObject.cpp
@@ -146,7 +146,7 @@ void InitRemoteObjects()
 		UE::RemoteObject::Transfer::RestoreRemoteObjectDataDelegate.BindStatic(&UE::RemoteObject::Serialization::Disk::LoadObjectFromDisk);
 	}
 
-	FCoreDelegates::OnEndFrame.AddStatic(&UE::RemoteObject::Transfer::Private::TickRemoteObjectRequests);
+	FCoreDelegates::OnEndFrame.AddStatic(&UE::RemoteObject::Transfer::TickRemoteObjectRequests);
 }
 
 void MarkAsRemote(UObject* Object, FRemoteServerId DestinationServerId)
diff --git a/Engine/Source/Runtime/CoreUObject/Private/UObject/RemoteObjectPrivate.h b/Engine/Source/Runtime/CoreUObject/Private/UObject/RemoteObjectPrivate.h
--- a/Engine/Source/Runtime/CoreUObject/Private/UObject/RemoteObjectPrivate.h
+++ b/Engine/Source/Runtime/CoreUObject/Private/UObject/RemoteObjectPrivate.h
@@ -55,9 +55,4 @@ namespace UE::RemoteObject::Transfer::Private
 	* Stores unreachable objects to database
 	*/
 	void StoreObjectToDatabase(UObject* Object);
-
-	/**
-	* Processes remote object requests (send and receive requests)
-	*/
-	void TickRemoteObjectRequests();
 }
\ No newline at end of file
diff --git a/Engine/Source/Runtime/CoreUObject/Private/UObject/RemoteObjectTransfer.cpp b/Engine/Source/Runtime/CoreUObject/Private/UObject/RemoteObjectTransfer.cpp
--- a/Engine/Source/Runtime/CoreUObject/Private/UObject/RemoteObjectTransfer.cpp
+++ b/Engine/Source/Runtime/CoreUObject/Private/UObject/RemoteObjectTransfer.cpp
@@ -312,8 +312,9 @@ public:
 
 	/**
 	* Processes exiting requests
+	* @return true if all requests have been processed and the queue is empty
 	*/
-	FORCENOINLINE void Tick()
+	FORCENOINLINE bool Tick()
 	{
 		FScopeLock RequestLock(&RequestCritical);
 		FRemoteObjectRequestBase* Prev = nullptr;
@@ -335,6 +336,7 @@ public:
 				Request = Request->Next;
 			}
 		}
+		return !Head;
 	}
 
 	/**
@@ -475,6 +477,20 @@ void OnObjectDataReceived(FRemoteObjectId ObjectId, FRemoteServerId RemoteServer
 	}
 } 
 
+void TickRemoteObjectRequests()
+{
+	UE::RemoteObject::Transfer::GTransferQueue.Tick();
+}
+
+void FlushRemoteObjectRequests()
+{
+	do
+	{
+		FlushRemoteObjectRequestsDelegate.ExecuteIfBound();
+	}
+	while (!UE::RemoteObject::Transfer::GTransferQueue.Tick());
+}
+
 } // // namespace UE::RemoteObject::Transfer
 
 namespace UE::RemoteObject::Transfer::Private
@@ -526,9 +542,4 @@ void StoreObjectToDatabase(UObject* Object)
 		OutermostUnreachableRemoteReference != Object ? *FString::Printf(TEXT(" when storing its outermost remote referenced object %s"), *OutermostUnreachableRemoteReference->GetPathName()) : TEXT("."));
 }
 
-void TickRemoteObjectRequests()
-{
-	UE::RemoteObject::Transfer::GTransferQueue.Tick();
-}
-
 } // namespace UE::RemoteObject::Transfer::Private
\ No newline at end of file
diff --git a/Engine/Source/Runtime/CoreUObject/Public/UObject/RemoteObjectTransfer.h b/Engine/Source/Runtime/CoreUObject/Public/UObject/RemoteObjectTransfer.h
--- a/Engine/Source/Runtime/CoreUObject/Public/UObject/RemoteObjectTransfer.h
+++ b/Engine/Source/Runtime/CoreUObject/Public/UObject/RemoteObjectTransfer.h
@@ -64,6 +64,16 @@ namespace UE::RemoteObject::Transfer
 #endif
 	}
 
+	/**
+	* Ticks remote object requests (send and receive requests)
+	*/
+	COREUOBJECT_API void TickRemoteObjectRequests();
+
+	/**
+	* Ticks remote object requests until all have been processed (send and receive requests)
+	*/
+	COREUOBJECT_API void FlushRemoteObjectRequests();
+
 	extern COREUOBJECT_API const FRemoteServerId DatabaseId;
 
 	/** Delegate that transfers object data to another server */
