When collecting garbage don't resolve remote objects if the object's memory has not yet been freed.

diff --git a/Engine/Source/Runtime/CoreUObject/Private/UObject/RemoteObject.cpp b/Engine/Source/Runtime/CoreUObject/Private/UObject/RemoteObject.cpp
--- a/Engine/Source/Runtime/CoreUObject/Private/UObject/RemoteObject.cpp
+++ b/Engine/Source/Runtime/CoreUObject/Private/UObject/RemoteObject.cpp
@@ -251,9 +251,16 @@ UObject* ResolveObject(UObject* Object)
 {
 	using namespace UE::RemoteObject::Private;
 
-	FRemoteObjectStub* Stub = ObjectMaps->FindRemoteObjectStub(UE::CoreUObject::Private::FObjectHandleUtils::GetRemoteId(Object));
-	checkf(Stub, TEXT("Failed to find remote object stub for %s"), *GetPathNameSafe(Object));
-	return ResolveObject(Stub);
+	// Begin/FinishDestroy overrides may attempt to access subobjects of objects that have been migrated in which case we
+	// don't want to accidentally migrate them back mid-purge and if the Object memory is still valid (but has EInternalObjectFlags::Remote flag)
+	// we can just return it and let the owner finish its cleanup.
+	if (!GIsGarbageCollecting) // Note: this is not a thread-safe check
+	{
+		FRemoteObjectStub* Stub = ObjectMaps->FindRemoteObjectStub(UE::CoreUObject::Private::FObjectHandleUtils::GetRemoteId(Object));
+		checkf(Stub, TEXT("Failed to find remote object stub for %s"), *GetPathNameSafe(Object));
+		return ResolveObject(Stub);
+	}
+	return Object;
 }
 
 } // namespace UE::RemoteObject::Handle
