[AutoRTFM] Disable memory validation for open scopes that fail validation

These fail validation for unknown reasons, and are tracked in the associated bugs.

diff --git a/Engine/Source/Runtime/CoreUObject/Private/UObject/UObjectArchetype.cpp b/Engine/Source/Runtime/CoreUObject/Private/UObject/UObjectArchetype.cpp
--- a/Engine/Source/Runtime/CoreUObject/Private/UObject/UObjectArchetype.cpp
+++ b/Engine/Source/Runtime/CoreUObject/Private/UObject/UObjectArchetype.cpp
@@ -4,6 +4,7 @@
 	UObjectArchetype.cpp: Unreal object archetype relationship management
 =============================================================================*/
 
+#include "AutoRTFM/AutoRTFM.h"
 #include "UObject/UObjectArchetypeInternal.h"
 #include "CoreMinimal.h"
 #include "UObject/UObjectHash.h"
@@ -94,7 +95,8 @@ UObject* GetArchetypeFromRequiredInfoImpl(const UClass* Class, const UObject* Ou
 		if (Outer
 			&& Outer->GetClass() != UPackage::StaticClass()) // packages cannot have subobjects
 		{
-			AutoRTFM::Open([&]()
+			// Validation temporarily disabled until SOL-7251 is investigated
+			UE_AUTORTFM_OPEN_NO_VALIDATION
 			{
 				// Get a lock on the UObject hash tables for the duration of the GetArchetype operation
 				FScopedUObjectHashTablesLock HashTablesLock;
@@ -158,7 +160,7 @@ UObject* GetArchetypeFromRequiredInfoImpl(const UClass* Class, const UObject* Ou
 
 					Result = ArchetypeToSearch->GetClass()->FindArchetype(Class, Name);
 				}
-			});
+			};
 		}
 
 		if (!Result)
diff --git a/Engine/Source/Runtime/CoreUObject/Private/UObject/UObjectGlobals.cpp b/Engine/Source/Runtime/CoreUObject/Private/UObject/UObjectGlobals.cpp
--- a/Engine/Source/Runtime/CoreUObject/Private/UObject/UObjectGlobals.cpp
+++ b/Engine/Source/Runtime/CoreUObject/Private/UObject/UObjectGlobals.cpp
@@ -3630,7 +3630,8 @@ UObject* StaticAllocateObject
 
 		// perform the UObjectBase construction in the open - we expect the GC to invoke
 		// the destructor in the case of transaction abort
-		UE_AUTORTFM_OPEN
+		// Validation temporarily disabled until SOL-7253 is investigated
+		UE_AUTORTFM_OPEN_NO_VALIDATION
 		{
 			FMemory::Memzero((void *)Obj, TotalSize);
 			new ((void *)Obj) UObjectBase(const_cast<UClass*>(InClass), InFlags|RF_NeedInitialization, InternalSetFlags, InOuter, InName, OldIndex, OldSerialNumber == 0 ? SerialNumber : OldSerialNumber, RemoteId);
diff --git a/Engine/Source/Runtime/CoreUObject/Public/UObject/DynamicallyTypedValue.h b/Engine/Source/Runtime/CoreUObject/Public/UObject/DynamicallyTypedValue.h
--- a/Engine/Source/Runtime/CoreUObject/Public/UObject/DynamicallyTypedValue.h
+++ b/Engine/Source/Runtime/CoreUObject/Public/UObject/DynamicallyTypedValue.h
@@ -63,7 +63,11 @@ namespace UE
 	{
 		static COREUOBJECT_API FDynamicallyTypedValueType& NullType();
 
-		FDynamicallyTypedValue() { UE_AUTORTFM_OPEN{ InitializeToNull(); }; }
+		FDynamicallyTypedValue()
+		{
+			// Validation temporarily disabled until SOL-7252 is investigated
+			UE_AUTORTFM_OPEN_NO_VALIDATION{ InitializeToNull(); };
+		}
 		FDynamicallyTypedValue(const FDynamicallyTypedValue& Copyee) { UE_AUTORTFM_OPEN{ InitializeFromCopy(Copyee); }; }
 		FDynamicallyTypedValue(FDynamicallyTypedValue&& Movee) { UE_AUTORTFM_OPEN{ InitializeFromMove(MoveTemp(Movee)); }; }
 
