Maybe fixing CIS

diff --git a/Engine/Source/Runtime/CoreUObject/Public/UObject/ObjectPtr.h b/Engine/Source/Runtime/CoreUObject/Public/UObject/ObjectPtr.h
--- a/Engine/Source/Runtime/CoreUObject/Public/UObject/ObjectPtr.h
+++ b/Engine/Source/Runtime/CoreUObject/Public/UObject/ObjectPtr.h
@@ -215,11 +215,6 @@ public:
 
 	/** Returns the full name for the object in the form: Class ObjectPath */
 	COREUOBJECT_API FString GetFullName(EObjectFullNameFlags Flags = EObjectFullNameFlags::None) const;
-
-	FRemoteObjectId GetRemoteId() const
-	{
-		return FRemoteObjectId();
-	}
 #else
 	FString GetPathName() const
 	{
@@ -258,15 +253,6 @@ public:
 		// UObjectBaseUtility::GetFullName is safe to call on null objects.
 		return UE::CoreUObject::Private::GetFullName(Get(), nullptr, Flags);
 	}
-
-	FRemoteObjectId GetRemoteId() const
-	{
-#if UE_WITH_REMOTE_OBJECT_HANDLE
-		return GetHandle().GetRemoteId();
-#else
-		return FRemoteObjectId();
-#endif
-	}
 #endif
 
 	FORCEINLINE FObjectHandle GetHandle() const { return Handle; }
@@ -280,6 +266,15 @@ public:
 		return IsA(T::StaticClass());
 	}
 
+	FORCEINLINE FRemoteObjectId GetRemoteId() const
+	{
+#if UE_WITH_REMOTE_OBJECT_HANDLE
+		return GetHandle().GetRemoteId();
+#else
+		return FRemoteObjectId();
+#endif
+	}
+
 private:
 	friend FORCEINLINE uint32 GetTypeHash(const FObjectPtr& Object)
 	{
