[Backout] - CL37111721 - Main Build Error
Original CL Desc
-----------------------------------------------------------------
Added support for leaking objects to disk when they are potentially referenced by a remote server

diff --git a/Engine/Source/Runtime/CoreUObject/Private/UObject/GarbageCollection.cpp b/Engine/Source/Runtime/CoreUObject/Private/UObject/GarbageCollection.cpp
--- a/Engine/Source/Runtime/CoreUObject/Private/UObject/GarbageCollection.cpp
+++ b/Engine/Source/Runtime/CoreUObject/Private/UObject/GarbageCollection.cpp
@@ -45,7 +45,6 @@
 #include "UObject/ExpandingChunkedList.h"
 #include "UObject/ReachabilityAnalysisState.h"
 #include "UObject/ReachabilityAnalysis.h"
-#include "UObject/RemoteObjectPrivate.h"
 
 #include <atomic>
 
@@ -5952,12 +5951,6 @@ bool UnhashUnreachableObjects(bool bUseTimeLimit, double TimeLimit)
 				UObject* Object = static_cast<UObject*>(ObjectItem->Object);
 				FScopedCBDProfile Profile(Object);
 				// Begin the object's asynchronous destruction.
-#if UE_WITH_REMOTE_OBJECT_HANDLE
-				if (!GExitPurge && ObjectItem->HasAnyFlags(EInternalObjectFlags::RemoteReference) && !ObjectItem->HasAnyFlags(EInternalObjectFlags::Remote))
-				{
-					UE::RemoteObject::Transfer::Private::StoreObjectToDatabase(Object);
-				}
-#endif // UE_WITH_REMOTE_OBJECT_HANDLE
 				Object->ConditionalBeginDestroy();
 			}
 
diff --git a/Engine/Source/Runtime/CoreUObject/Private/UObject/RemoteObjectPrivate.h b/Engine/Source/Runtime/CoreUObject/Private/UObject/RemoteObjectPrivate.h
--- a/Engine/Source/Runtime/CoreUObject/Private/UObject/RemoteObjectPrivate.h
+++ b/Engine/Source/Runtime/CoreUObject/Private/UObject/RemoteObjectPrivate.h
@@ -44,15 +44,12 @@ namespace UE::RemoteObject::Private
 	UObject* FindObject(FRemoteObjectId ObjectId);
 
 	/**
-	* Finds remote object stub
+	* Processes remote object requests (send and receive requests)
 	*/
-	UE::RemoteObject::Handle::FRemoteObjectStub* FindRemoteObjectStub(FRemoteObjectId ObjectId);
-}
+	void TickRemoteObjectRequests();
 
-namespace UE::RemoteObject::Transfer::Private
-{
 	/**
-	* Stores unreachable objects to database
+	* Finds remote object stub
 	*/
-	void StoreObjectToDatabase(UObject* Object);
-}
\ No newline at end of file
+	UE::RemoteObject::Handle::FRemoteObjectStub* FindRemoteObjectStub(FRemoteObjectId ObjectId);
+}
diff --git a/Engine/Source/Runtime/CoreUObject/Private/UObject/RemoteObjectSerialization.cpp b/Engine/Source/Runtime/CoreUObject/Private/UObject/RemoteObjectSerialization.cpp
--- a/Engine/Source/Runtime/CoreUObject/Private/UObject/RemoteObjectSerialization.cpp
+++ b/Engine/Source/Runtime/CoreUObject/Private/UObject/RemoteObjectSerialization.cpp
@@ -58,29 +58,27 @@ void DeserializeNameTables(FArchive& Ar, FRemoteObjectData& OutObjectData)
 	Ar << OutObjectData.PathNames;
 }
 
-FString GenerateRemoteObjectFilename(FRemoteObjectId ObjectId, FRemoteServerId DestinationServerId)
+FString GenerateRemoteObjectFilename(FRemoteObjectId ObjectId)
 {
-	return FPaths::Combine(*FPaths::ProjectSavedDir(), *(FString::Printf(TEXT("%s_%s.remote"), *ObjectId.ToString(), *DestinationServerId.ToString())));
+	return FPaths::Combine(*FPaths::ProjectSavedDir(), *(ObjectId.ToString() + TEXT(".tmp")));
 }
 
 void LoadObjectFromDisk(FRemoteObjectId ObjectId, FRemoteServerId CurrentOwnerServerId, FRemoteServerId DestinationServerId)
 {
-	const FString Filename(GenerateRemoteObjectFilename(ObjectId, CurrentOwnerServerId));
-	TUniquePtr<FArchive> FileReader{ IFileManager::Get().CreateFileReader(*Filename) };
+	TUniquePtr<FArchive> FileReader{ IFileManager::Get().CreateFileReader(*GenerateRemoteObjectFilename(ObjectId)) };
 	checkf(FileReader.IsValid(), TEXT("Unable to create file reader for remote object %s"), *ObjectId.ToString());
 	FRemoteObjectData* ObjectData = new FRemoteObjectData();
 	
 	DeserializeNameTables(*FileReader, *ObjectData);
 	*FileReader << ObjectData->Bytes;
 	FileReader->Close();
-	IFileManager::Get().Delete(*Filename, false, true, true);
 
 	UE::RemoteObject::Transfer::OnObjectDataReceived(ObjectId, CurrentOwnerServerId, ObjectData);
 }
 
-void SaveObjectToDisk(FRemoteObjectId ObjectId, FRemoteServerId DestinationServerId, FRemoteObjectData& ObjectData)
+void SaveObjectToDisk(FRemoteObjectId ObjectId, FRemoteServerId ServerId, FRemoteObjectData& ObjectData)
 {
-	TUniquePtr<FArchive> FileWriter{ IFileManager::Get().CreateFileWriter(*GenerateRemoteObjectFilename(ObjectId, DestinationServerId)) };
+	TUniquePtr<FArchive> FileWriter{ IFileManager::Get().CreateFileWriter(*GenerateRemoteObjectFilename(ObjectId)) };
 	checkf(FileWriter.IsValid(), TEXT("Unable to create file writer for remote object %s"), *ObjectId.ToString());
 	
 	SerializeNameTables(*FileWriter, ObjectData);
@@ -142,38 +140,24 @@ class FArchiveRemoteObjectWriter : public FMemoryWriter
 	TMap<FString, TMap<FProperty*, FSerializeStats>> ObjectPropertyStats;
 #endif
 
-	void WriteObjectReference(const FObjectPtr& ObjPtr)
+	ERemoteReferenceType WriteObjectReference(UObject*& Obj)
 	{
 		using namespace UE::CoreUObject::Private;
 
 		ERemoteReferenceType Type = ERemoteReferenceType::None;
-		FRemoteObjectId Id = ObjPtr.GetRemoteId();
-		UObject* Obj = nullptr;
+		FRemoteObjectId Id;
 
-		if (ObjPtr)
+		if (Obj)
 		{
-			if (ObjPtr.IsResolved())
-			{
-				Obj = ObjPtr.Get();
-				checkf(Obj, TEXT("Attempting to serialize resolved Object pointer but the pointer resolves to null (%s)"), *Id.ToString());
-				Obj->SetInternalFlags(EInternalObjectFlags::RemoteReference);
-				check(!UE::RemoteObject::Handle::IsRemote(Obj));
-
-				Id = FObjectHandleUtils::GetRemoteId(Obj);
+			Id = FObjectHandleUtils::GetRemoteId(Obj);
 
-				if (Obj == RootObject || Obj->IsIn(RootObject) || (Id.GetServerId().IsValid() && Id.GetServerId() != UE::RemoteObject::GetGlobalServerId()))
-				{
-					Type = ERemoteReferenceType::IdOnly;
-				}
-				else
-				{
-					
-					Type = ERemoteReferenceType::PathName;
-				}
+			if (Obj == RootObject || UE::RemoteObject::Handle::IsRemote(Obj) || Obj->IsIn(RootObject) || (Id.GetServerId().IsValid() && Id.GetServerId() != UE::RemoteObject::GetGlobalServerId()))
+			{
+				Type = ERemoteReferenceType::IdOnly;
 			}
 			else
 			{
-				Type = ERemoteReferenceType::IdOnly;
+				Type = ERemoteReferenceType::PathName;
 			}
 		}
 
@@ -184,28 +168,13 @@ class FArchiveRemoteObjectWriter : public FMemoryWriter
 			*this << Id;
 		}
 
-		if (Type == ERemoteReferenceType::IdOnly)
-		{
-			// Add subobjects of the root object to the list of objects to serialize
-			if (Obj && Obj->IsIn(RootObject))
-			{
-				ObjectsToSerialize.Add(Obj);
-			}
-		}
-		else if (Type == ERemoteReferenceType::PathName)
+		if (Type == ERemoteReferenceType::PathName)
 		{
 			FNameIndexType PathNameIndex = AddPathNameToNameMap(Obj);
 			*this << PathNameIndex;
 		}
 
-		if (Obj)
-		{
-			// Map all references we encounter during serialization, especially if they're objects that are not going to be serialized
-			// because we may end up needing to resolve them later if a remote server requests them with their unique id.
-			// Note: we could just hash all objects with their FRemoteObjectId in UObjectHash.cpp but for now we can limit the memory impact by hashing 
-			// only those we encounter when serializing objects that are actually going to be migrated
-			UE::RemoteObject::Private::MapObject(Obj);
-		}
+		return Type;
 	}
 
 public:
@@ -343,7 +312,23 @@ public:
 
 	virtual FArchive& operator<<(UObject*& Obj) override
 	{
-		WriteObjectReference(FObjectPtr(Obj));
+		ERemoteReferenceType RefType = WriteObjectReference(Obj);
+		if (RefType == ERemoteReferenceType::IdOnly)
+		{
+			if (Obj && Obj->IsIn(RootObject))
+			{
+				ObjectsToSerialize.Add(Obj);
+			}			
+		}
+		if (Obj)
+		{
+			// Map all references we encounter during serialization, especially if they're objects that are not going to be serialized
+			// because we may end up needing to resolve them later if a remote server requests them with their unique id.
+			// Note: we could just hash all objects with their FRemoteObjectId in UObjectHash.cpp but for now we can limit the memory impact by hashing 
+			// only those we encounter when serializing objects that are actually going to be migrated
+			UE::RemoteObject::Private::MapObject(Obj);
+		}
+
 		return *this;
 	}
 
@@ -357,8 +342,7 @@ public:
 	virtual FArchive& operator<<(FLazyObjectPtr& Value) override { return FArchiveUObject::SerializeLazyObjectPtr(*this, Value); }
 	virtual FArchive& operator<<(FObjectPtr& Value) override 
 	{ 
-		WriteObjectReference(Value);
-		return *this;
+		return FArchiveUObject::SerializeObjectPtr(*this, Value); 
 	}
 	virtual FArchive& operator<<(FSoftObjectPtr& Value) override { return FArchiveUObject::SerializeSoftObjectPtr(*this, Value); }
 	virtual FArchive& operator<<(FSoftObjectPath& Value) override { return FArchiveUObject::SerializeSoftObjectPath(*this, Value); }
@@ -448,7 +432,6 @@ class FArchiveRemoteObjectReader : public FMemoryReader
 				UObject* Obj = ResolvedPathNameObjects[PathNameIndex];
 				if (Obj)
 				{
-					Obj->SetInternalFlags(EInternalObjectFlags::RemoteReference);
 					Value = FObjectPtr(Obj);
 					bNeedsResolvingWithId = false;
 				}
diff --git a/Engine/Source/Runtime/CoreUObject/Public/UObject/ObjectHandle.h b/Engine/Source/Runtime/CoreUObject/Public/UObject/ObjectHandle.h
--- a/Engine/Source/Runtime/CoreUObject/Public/UObject/ObjectHandle.h
+++ b/Engine/Source/Runtime/CoreUObject/Public/UObject/ObjectHandle.h
@@ -95,8 +95,6 @@ namespace UE::CoreUObject::Private
 			return reinterpret_cast<UE::RemoteObject::Handle::FRemoteObjectStub*>(PointerOrHandle & ~UPTRINT(1));
 		}
 
-		COREUOBJECT_API FRemoteObjectId GetRemoteId() const;
-
 		COREUOBJECT_API static FRemoteObjectHandlePrivate ConvertToRemoteHandle(UObject* Object);
 		COREUOBJECT_API static FRemoteObjectHandlePrivate FromIdNoResolve(FRemoteObjectId ObjectId);
 	};
diff --git a/Engine/Source/Runtime/CoreUObject/Public/UObject/ObjectMacros.h b/Engine/Source/Runtime/CoreUObject/Public/UObject/ObjectMacros.h
--- a/Engine/Source/Runtime/CoreUObject/Public/UObject/ObjectMacros.h
+++ b/Engine/Source/Runtime/CoreUObject/Public/UObject/ObjectMacros.h
@@ -618,8 +618,7 @@ enum class EInternalObjectFlags : int32
 	ReachabilityFlag1 = 1 << 1, ///< One of the flags used by Garbage Collector to determine UObject's reachability state
 	ReachabilityFlag2 = 1 << 2, ///< One of the flags used by Garbage Collector to determine UObject's reachability state
 
-	Remote = 1 << 18, //< Object is no longer owned by this process
-	RemoteReference = 1 << 19, //< Object referenced by remote process
+	Remote = 1 << 19, //< Object is no longer owned by this process
 	LoaderImport = 1 << 20, ///< Object is ready to be imported by another package during loading
 	Garbage = 1 << 21, ///< Garbage from logical point of view and should not be referenced. This flag is mirrored in EObjectFlags as RF_Garbage for performance
 	AsyncLoadingPhase1 = 1 << 22, ///< Object is being asynchronously loaded.
@@ -642,7 +641,7 @@ ENUM_CLASS_FLAGS(EInternalObjectFlags);
 
 //~ Make sure these macros are up to date!
 #define EInternalObjectFlags_GarbageCollectionKeepFlags (EInternalObjectFlags::Native | EInternalObjectFlags::Async | EInternalObjectFlags::AsyncLoadingPhase1 | EInternalObjectFlags::AsyncLoadingPhase2 | EInternalObjectFlags::LoaderImport | EInternalObjectFlags::RefCounted)
-#define EInternalObjectFlags_AllFlags (EInternalObjectFlags::ReachabilityFlag0 | EInternalObjectFlags::ReachabilityFlag1 | EInternalObjectFlags::ReachabilityFlag2 | EInternalObjectFlags::Remote | EInternalObjectFlags::RemoteReference | EInternalObjectFlags::LoaderImport | EInternalObjectFlags::Garbage | EInternalObjectFlags::ReachableInCluster | EInternalObjectFlags::ClusterRoot | EInternalObjectFlags::Native | EInternalObjectFlags::RefCounted | EInternalObjectFlags::Async | EInternalObjectFlags::AsyncLoadingPhase1 | EInternalObjectFlags::AsyncLoadingPhase2 | EInternalObjectFlags::RootSet | EInternalObjectFlags::PendingConstruction | (EInternalObjectFlags::Unreachable))
+#define EInternalObjectFlags_AllFlags (EInternalObjectFlags::ReachabilityFlag0 | EInternalObjectFlags::ReachabilityFlag1 | EInternalObjectFlags::ReachabilityFlag2 | EInternalObjectFlags::Remote | EInternalObjectFlags::LoaderImport | EInternalObjectFlags::Garbage | EInternalObjectFlags::ReachableInCluster | EInternalObjectFlags::ClusterRoot | EInternalObjectFlags::Native | EInternalObjectFlags::RefCounted | EInternalObjectFlags::Async | EInternalObjectFlags::AsyncLoadingPhase1 | EInternalObjectFlags::AsyncLoadingPhase2 | EInternalObjectFlags::RootSet | EInternalObjectFlags::PendingConstruction | (EInternalObjectFlags::Unreachable))
 #define EInternalObjectFlags_RootFlags (EInternalObjectFlags::RootSet | EInternalObjectFlags_GarbageCollectionKeepFlags)
 #define EInternalObjectFlags_ReachabilityFlags (EInternalObjectFlags::ReachabilityFlag0 | EInternalObjectFlags::ReachabilityFlag1 | EInternalObjectFlags::ReachabilityFlag2 | EInternalObjectFlags::Unreachable)
 #define EInternalObjectFlags_AsyncLoading (EInternalObjectFlags::AsyncLoadingPhase1 | EInternalObjectFlags::AsyncLoadingPhase2)
