[Backout] - CL37082464 - Main Build Error
Original CL Desc
-----------------------------------------------------------------
Adding delegates called when object/object data gets migrated to or from remote server

diff --git a/Engine/Source/Runtime/CoreUObject/Private/UObject/RemoteObjectSerialization.cpp b/Engine/Source/Runtime/CoreUObject/Private/UObject/RemoteObjectSerialization.cpp
--- a/Engine/Source/Runtime/CoreUObject/Private/UObject/RemoteObjectSerialization.cpp
+++ b/Engine/Source/Runtime/CoreUObject/Private/UObject/RemoteObjectSerialization.cpp
@@ -73,7 +73,7 @@ void LoadObjectFromDisk(FRemoteObjectId ObjectId, FRemoteServerId CurrentOwnerSe
 	*FileReader << ObjectData->Bytes;
 	FileReader->Close();
 
-	UE::RemoteObject::Transfer::OnObjectDataReceived(ObjectId, CurrentOwnerServerId, ObjectData);
+	UE::RemoteObject::Transfer::OnObjectDataReceived(ObjectId, ObjectData);
 }
 
 void SaveObjectToDisk(FRemoteObjectId ObjectId, FRemoteServerId ServerId, FRemoteObjectData& ObjectData)
