Basic implementation to support tracking property initialization state from copy/paste with verse objects
* Note support does not extend to elements within containers such as arrays as support is not complete for the property initialization state tracking infrastructure
* Enabled for UVerseClass or IDOs only

diff --git a/Engine/Source/Runtime/CoreUObject/Internal/UObject/InstanceDataObjectUtils.h b/Engine/Source/Runtime/CoreUObject/Internal/UObject/InstanceDataObjectUtils.h
--- a/Engine/Source/Runtime/CoreUObject/Internal/UObject/InstanceDataObjectUtils.h
+++ b/Engine/Source/Runtime/CoreUObject/Internal/UObject/InstanceDataObjectUtils.h
@@ -50,6 +50,8 @@ void CopyPropertyValueSerializedData(const FFieldVariant& OldField, void* OldDat
 
 
 UE_API bool SupportsPropertyInitStateTracking();
+
+UE_API bool SupportsPropertyInitStateTracking(const UStruct* Struct, void* StructData);
 	
 /** Query whether the property value is initialized for a property in the struct. */
 UE_API bool IsPropertyValueInitialized(const UStruct* Struct, void* StructData, const FProperty* Property, int32 ArrayIndex = 0);
diff --git a/Engine/Source/Runtime/CoreUObject/Private/UObject/InstanceDataObjectUtils.cpp b/Engine/Source/Runtime/CoreUObject/Private/UObject/InstanceDataObjectUtils.cpp
--- a/Engine/Source/Runtime/CoreUObject/Private/UObject/InstanceDataObjectUtils.cpp
+++ b/Engine/Source/Runtime/CoreUObject/Private/UObject/InstanceDataObjectUtils.cpp
@@ -931,6 +931,25 @@ namespace UE
 		return bUseAnnotationsForPropertyStateInitTracking || bEnableIDOSupport;
 	}
 
+	bool SupportsPropertyInitStateTracking(const UStruct* Struct, void* StructData)
+	{
+		if (bUseAnnotationsForPropertyStateInitTracking)
+		{
+			if (Cast<UVerseClass>(Struct) != nullptr)
+			{
+				return true;
+			}
+		}
+		else if (bEnableIDOSupport)
+		{
+			if (IsClassOfInstanceDataObjectClass(Struct))
+			{
+				return true;
+			}
+		}
+		return false;
+	}
+
 	static bool IsIdoPropertyValueInitialized(const UStruct* Struct, void* StructData, const FProperty* Property, int32 ArrayIndex, const FByteProperty* InitializedValuesProperty)
 	{
 		if (!SupportsPropertyInitStateTracking())
diff --git a/Engine/Source/Runtime/CoreUObject/Private/UObject/Property.cpp b/Engine/Source/Runtime/CoreUObject/Private/UObject/Property.cpp
--- a/Engine/Source/Runtime/CoreUObject/Private/UObject/Property.cpp
+++ b/Engine/Source/Runtime/CoreUObject/Private/UObject/Property.cpp
@@ -24,6 +24,7 @@
 #include "UObject/Class.h"
 #include "UObject/CoreNetTypes.h"
 #include "UObject/CoreRedirects.h"
+#include "UObject/InstanceDataObjectUtils.h"
 #include "UObject/OverridableManager.h"
 #include "UObject/Package.h"
 #include "UObject/PropertyHelper.h"
@@ -2403,6 +2404,14 @@ const TCHAR* FProperty::ImportSingleProperty( const TCHAR* Str, void* DestData,
 
 
 					const TCHAR* Result = ImportText(Str, Property, Property->ContainerPtrToValuePtr<void>(DestData, Index));
+
+#if WITH_EDITORONLY_DATA
+					if (UE::SupportsPropertyInitStateTracking(ObjectStruct, DestData))
+					{
+						UE::SetPropertyValueInitialized(ObjectStruct, DestData, Property);
+					}
+#endif
+					
 					// in the failure case, don't return NULL so the caller can potentially skip less and get values further in the string
 					if (Result != nullptr)
 					{
