Fixing CIS for real (temp fix)

diff --git a/Engine/Source/Runtime/CoreUObject/Private/UObject/RemoteObjectSerialization.cpp b/Engine/Source/Runtime/CoreUObject/Private/UObject/RemoteObjectSerialization.cpp
--- a/Engine/Source/Runtime/CoreUObject/Private/UObject/RemoteObjectSerialization.cpp
+++ b/Engine/Source/Runtime/CoreUObject/Private/UObject/RemoteObjectSerialization.cpp
@@ -147,7 +147,11 @@ class FArchiveRemoteObjectWriter : public FMemoryWriter
 		using namespace UE::CoreUObject::Private;
 
 		ERemoteReferenceType Type = ERemoteReferenceType::None;
+#if UE_WITH_REMOTE_OBJECT_HANDLE
 		FRemoteObjectId Id = ObjPtr.GetRemoteId();
+#else
+		FRemoteObjectId Id;
+#endif
 		UObject* Obj = nullptr;
 
 		if (ObjPtr)
