Fixed issue where IDO properties would not get marked as initialized.

#rnx

diff --git a/Engine/Source/Runtime/CoreUObject/Private/UObject/InstanceDataObjectUtils.cpp b/Engine/Source/Runtime/CoreUObject/Private/UObject/InstanceDataObjectUtils.cpp
--- a/Engine/Source/Runtime/CoreUObject/Private/UObject/InstanceDataObjectUtils.cpp
+++ b/Engine/Source/Runtime/CoreUObject/Private/UObject/InstanceDataObjectUtils.cpp
@@ -734,6 +734,8 @@ namespace UE
 		Writer.ArNoDelta = true;
 		Source->GetClass()->SerializeTaggedProperties(Writer, (uint8*)Source, Source->GetClass(), nullptr);
 
+		TGuardValue<bool> ScopedTrackInitializedProperties(SerializeContext->bTrackInitializedProperties, true);
+
 		FObjectReader Reader(Buffer);
 		FSerializingDefaultsScope ReaderDefaultsScope(Reader, Dest);
 		Reader.ArMergeOverrides = true;
diff --git a/Engine/Source/Runtime/CoreUObject/Tests/UObject/InstanceDataObjectUtilsTest.cpp b/Engine/Source/Runtime/CoreUObject/Tests/UObject/InstanceDataObjectUtilsTest.cpp
--- a/Engine/Source/Runtime/CoreUObject/Tests/UObject/InstanceDataObjectUtilsTest.cpp
+++ b/Engine/Source/Runtime/CoreUObject/Tests/UObject/InstanceDataObjectUtilsTest.cpp
@@ -77,6 +77,7 @@ TEST_CASE_NAMED(FInstanceDataObjectUtilsTest, "CoreUObject::Serialization::Insta
 	UObject* Owner = NewObject<UObject>(GetTransientPackage(), TestClass, TestObjectName);
 	void* StructData = StructProperty->ContainerPtrToValuePtr<void>(Owner);
 	ResetPropertyValueInitialized(TestClass, Owner);
+	ResetPropertyValueInitialized(StructProperty->Struct, StructData);
 
 	CHECK_FALSE(WasPropertyValueSerialized(TestClass, Owner, StructProperty));
 	MarkPropertyValueSerialized(TestClass, Owner, StructProperty);
