Fixed unneccesary resolving when assigning FSoftObjectPtr with a TObjectPtr.

diff --git a/Engine/Source/Runtime/AIModule/Classes/BrainComponent.h b/Engine/Source/Runtime/AIModule/Classes/BrainComponent.h
--- a/Engine/Source/Runtime/AIModule/Classes/BrainComponent.h
+++ b/Engine/Source/Runtime/AIModule/Classes/BrainComponent.h
@@ -46,7 +46,7 @@ struct FAIMessage
 	/** message param: custom flags */
 	uint8 MessageFlags;
 
-	FAIMessage() : MessageName(NAME_None), Sender(NULL), RequestID(0), Status(FAIMessage::Success), MessageFlags(0) {}
+	FAIMessage() : MessageName(NAME_None), Sender(nullptr), RequestID(0), Status(FAIMessage::Success), MessageFlags(0) {}
 	FAIMessage(FName InMessage, UObject* InSender) : MessageName(InMessage), Sender(InSender), RequestID(0), Status(FAIMessage::Success), MessageFlags(0) {}
 	FAIMessage(FName InMessage, UObject* InSender, FAIRequestID InID, EStatus InStatus) : MessageName(InMessage), Sender(InSender), RequestID(InID), Status(InStatus), MessageFlags(0) {}
 	FAIMessage(FName InMessage, UObject* InSender, FAIRequestID InID, bool bSuccess) : MessageName(InMessage), Sender(InSender), RequestID(InID), Status(bSuccess ? Success : Failure), MessageFlags(0) {}
diff --git a/Engine/Source/Runtime/Core/Public/Templates/NonNullPointer.h b/Engine/Source/Runtime/Core/Public/Templates/NonNullPointer.h
--- a/Engine/Source/Runtime/Core/Public/Templates/NonNullPointer.h
+++ b/Engine/Source/Runtime/Core/Public/Templates/NonNullPointer.h
@@ -244,6 +244,10 @@ public:
 	//////////////////////////////////////////////////
 
 private:
+	friend FORCEINLINE uint32 GetTypeHash(const TNonNullPtr& InPtr)
+	{
+		return PointerHash(InPtr.Object);
+	}
 
 	/** The object we're holding a reference to. */
 	ObjectType* Object;
diff --git a/Engine/Source/Runtime/CoreUObject/Private/UObject/LazyObjectPtr.cpp b/Engine/Source/Runtime/CoreUObject/Private/UObject/LazyObjectPtr.cpp
--- a/Engine/Source/Runtime/CoreUObject/Private/UObject/LazyObjectPtr.cpp
+++ b/Engine/Source/Runtime/CoreUObject/Private/UObject/LazyObjectPtr.cpp
@@ -80,8 +80,10 @@ void FUniqueObjectGuid::FromString(const FString& From)
 	}
 }
 
-FUniqueObjectGuid FUniqueObjectGuid::GetOrCreateIDForObject(const class UObject *Object)
+FUniqueObjectGuid FUniqueObjectGuid::GetOrCreateIDForObject(FObjectPtr ObjectPtr)
 {
+	UObject* Object = ObjectPtr.Get();
+
 	check(Object);
 	checkSlow(IsInGameThread());
 	FUniqueObjectGuid ObjectGuid(Object);
diff --git a/Engine/Source/Runtime/CoreUObject/Private/UObject/SoftObjectPath.cpp b/Engine/Source/Runtime/CoreUObject/Private/UObject/SoftObjectPath.cpp
--- a/Engine/Source/Runtime/CoreUObject/Private/UObject/SoftObjectPath.cpp
+++ b/Engine/Source/Runtime/CoreUObject/Private/UObject/SoftObjectPath.cpp
@@ -757,7 +757,7 @@ UObject* FSoftObjectPath::ResolveObjectInternal(const TCHAR* PathString) const
 	return FoundObject;
 }
 
-FSoftObjectPath FSoftObjectPath::GetOrCreateIDForObject(const class UObject *Object)
+FSoftObjectPath FSoftObjectPath::GetOrCreateIDForObject(FObjectPtr Object)
 {
 	check(Object);
 	return FSoftObjectPath(Object);
diff --git a/Engine/Source/Runtime/CoreUObject/Private/UObject/WeakObjectPtr.cpp b/Engine/Source/Runtime/CoreUObject/Private/UObject/WeakObjectPtr.cpp
--- a/Engine/Source/Runtime/CoreUObject/Private/UObject/WeakObjectPtr.cpp
+++ b/Engine/Source/Runtime/CoreUObject/Private/UObject/WeakObjectPtr.cpp
@@ -25,11 +25,12 @@ DEFINE_LOG_CATEGORY_STATIC(LogWeakObjectPtr, Log, All);
  * Copy from an object pointer
  * @param Object object to create a weak pointer to
 **/
-void FWeakObjectPtr::operator=(const class UObject *Object)
+void FWeakObjectPtr::operator=(FObjectPtr ObjectPtr)
 {
-	if (Object // && UObjectInitialized() we might need this at some point, but it is a speed hit we would prefer to avoid
+	if (ObjectPtr // && UObjectInitialized() we might need this at some point, but it is a speed hit we would prefer to avoid
 		)
 	{
+		const UObject* Object = ObjectPtr.Get();
 		ObjectIndex = GUObjectArray.ObjectToIndex((UObjectBase*)Object);
 		ObjectSerialNumber = GUObjectArray.AllocateSerialNumber(ObjectIndex);
 #if UE_WITH_REMOTE_OBJECT_HANDLE
@@ -43,10 +44,6 @@ void FWeakObjectPtr::operator=(const class UObject *Object)
 		Reset();
 	}
 }
-void FWeakObjectPtr::operator=(TObjectPtr<UObject> Object)
-{
-	*this = Object.Get();
-}
 
 bool FWeakObjectPtr::IsValid(bool bEvenIfGarbage, bool bThreadsafeTest) const
 {
diff --git a/Engine/Source/Runtime/CoreUObject/Public/UObject/LazyObjectPtr.h b/Engine/Source/Runtime/CoreUObject/Public/UObject/LazyObjectPtr.h
--- a/Engine/Source/Runtime/CoreUObject/Public/UObject/LazyObjectPtr.h
+++ b/Engine/Source/Runtime/CoreUObject/Public/UObject/LazyObjectPtr.h
@@ -21,6 +21,8 @@
 
 #include <type_traits>
 
+class UObject;
+
 template <typename T> struct TIsPODType;
 template <typename T> struct TIsWeakPointerType;
 
@@ -44,7 +46,7 @@ struct FUniqueObjectGuid
 	}
 
 	/** Construct from an existing object */
-	COREUOBJECT_API explicit FUniqueObjectGuid(const class UObject* InObject);
+	COREUOBJECT_API explicit FUniqueObjectGuid(const UObject* InObject);
 
 	/** Converts into a string */
 	COREUOBJECT_API FString ToString() const;
@@ -60,7 +62,7 @@ struct FUniqueObjectGuid
 	 *
 	 * @return Found UObject, or nullptr if not currently loaded
 	 */
-	COREUOBJECT_API class UObject *ResolveObject() const;
+	COREUOBJECT_API UObject* ResolveObject() const;
 
 	/** Test if this can ever point to a live UObject */
 	FORCEINLINE bool IsValid() const
@@ -117,7 +119,16 @@ struct FUniqueObjectGuid
 		return 0;
 	}
 
-	static COREUOBJECT_API FUniqueObjectGuid GetOrCreateIDForObject(const class UObject *Object);
+	static COREUOBJECT_API FUniqueObjectGuid GetOrCreateIDForObject(FObjectPtr Object);
+	static FORCEINLINE FUniqueObjectGuid GetOrCreateIDForObject(const UObject* Object)
+	{
+		return GetOrCreateIDForObject(FObjectPtr(const_cast<UObject*>(Object)));
+	}
+	template <typename T>
+	static FORCEINLINE FUniqueObjectGuid GetOrCreateIDForObject(TObjectPtr<T> Object)
+	{
+		return GetOrCreateIDForObject(FObjectPtr(Object));
+	}
 
 private:
 	/** Guid representing the object, should be unique */
@@ -138,28 +149,60 @@ template<> struct TIsPODType<FUniqueObjectGuid> { enum { Value = true }; };
  */
 struct FLazyObjectPtr : public TPersistentObjectPtr<FUniqueObjectGuid>
 {
-public:	
+private:
+	using Super = TPersistentObjectPtr<FUniqueObjectGuid>;
+
+public:
 	/** Default constructor, sets to null */
 	FORCEINLINE FLazyObjectPtr()
 	{
 	}
 
 	/** Construct from object already in memory */
-	explicit FORCEINLINE FLazyObjectPtr(const UObject* Object)
+	explicit FORCEINLINE FLazyObjectPtr(FObjectPtr Object)
 	{
 		(*this)=Object;
 	}
-	
+	explicit FORCEINLINE FLazyObjectPtr(const UObject* Object)
+		: FLazyObjectPtr(FObjectPtr(const_cast<UObject*>(Object)))
+	{
+	}
+	template <typename T>
+	explicit FORCEINLINE FLazyObjectPtr(TObjectPtr<T> Object)
+		: FLazyObjectPtr(FObjectPtr(Object))
+	{
+		// This needs to be a template instead of TObjectPtr<const UObject> because C++ does derived-to-base
+		// pointer conversions ('standard conversion sequences') in more cases than TSmartPtr<Derived>-to-TSmartPtr<Base>
+		// conversions ('user-defined conversions'), meaning it doesn't auto-convert in many real use cases.
+		//
+		// https://en.cppreference.com/w/cpp/language/implicit_conversion
+	}
+
 	/** Copy from an object already in memory */
-	FORCEINLINE void operator=(const UObject *Object)
+	FORCEINLINE void operator=(FObjectPtr Object)
 	{
-		TPersistentObjectPtr<FUniqueObjectGuid>::operator=(Object);
+		Super::operator=(Object);
+	}
+	FORCEINLINE void operator=(const UObject* Object)
+	{
+		*this = FObjectPtr(const_cast<UObject*>(Object));
+	}
+	template <typename T>
+	FORCEINLINE void operator=(TObjectPtr<T> Object)
+	{
+		// This needs to be a template instead of TObjectPtr<const UObject> because C++ does derived-to-base
+		// pointer conversions ('standard conversion sequences') in more cases than TSmartPtr<Derived>-to-TSmartPtr<Base>
+		// conversions ('user-defined conversions'), meaning it doesn't auto-convert in many real use cases.
+		//
+		// https://en.cppreference.com/w/cpp/language/implicit_conversion
+
+		*this = FObjectPtr(Object);
 	}
 
 	/** Copy from a unique object identifier */
 	FORCEINLINE void operator=(const FUniqueObjectGuid& InObjectID)
 	{
-		TPersistentObjectPtr<FUniqueObjectGuid>::operator=(InObjectID);
+		Super::operator=(InObjectID);
 	}
 
 	/** Fixes up this FLazyObjectPtr to target the right UID as set in PIEGuidMap, this only works for directly serialized pointers */
diff --git a/Engine/Source/Runtime/CoreUObject/Public/UObject/ObjectPtr.h b/Engine/Source/Runtime/CoreUObject/Public/UObject/ObjectPtr.h
--- a/Engine/Source/Runtime/CoreUObject/Public/UObject/ObjectPtr.h
+++ b/Engine/Source/Runtime/CoreUObject/Public/UObject/ObjectPtr.h
@@ -1833,6 +1833,10 @@ public:
 	explicit operator bool() const = delete;
 
 private:
+	friend FORCEINLINE uint32 GetTypeHash(const TNonNullPtr& InPtr)
+	{
+		return GetTypeHash(InPtr.Object);
+	}
 
 	/** The object we're holding a reference to. */
 	TObjectPtr<ObjectType> Object;
diff --git a/Engine/Source/Runtime/CoreUObject/Public/UObject/PersistentObjectPtr.h b/Engine/Source/Runtime/CoreUObject/Public/UObject/PersistentObjectPtr.h
--- a/Engine/Source/Runtime/CoreUObject/Public/UObject/PersistentObjectPtr.h
+++ b/Engine/Source/Runtime/CoreUObject/Public/UObject/PersistentObjectPtr.h
@@ -9,6 +9,8 @@
 #include "UObject/Object.h"
 #include "UObject/WeakObjectPtr.h"
 
+class UObject;
+
 /**
  * TPersistentObjectPtr is a template base class for FLazyObjectPtr and FSoftObjectPtr
  */
@@ -51,7 +53,7 @@ struct TPersistentObjectPtr
 	}
 
 	/** Copy from an object pointer */
-	FORCEINLINE void operator=(const class UObject* Object)
+	void operator=(FObjectPtr Object)
 	{
 		if (Object)
 		{
@@ -70,7 +72,22 @@ struct TPersistentObjectPtr
 			Reset();
 		}
 	}
-	
+	FORCEINLINE void operator=(const UObject* Object)
+	{
+		*this = FObjectPtr(const_cast<UObject*>(Object));
+	}
+	template <typename T>
+	FORCEINLINE void operator=(TObjectPtr<T> Object)
+	{
+		// This needs to be a template instead of TObjectPtr<const UObject> because C++ does derived-to-base
+		// pointer conversions ('standard conversion sequences') in more cases than TSmartPtr<Derived>-to-TSmartPtr<Base>
+		// conversions ('user-defined conversions'), meaning it doesn't auto-convert in many real use cases.
+		//
+		// https://en.cppreference.com/w/cpp/language/implicit_conversion
+
+		*this = FObjectPtr(Object);
+	}
+
 	/** Copy from an existing weak pointer, reserve IDs if required */
 	FORCEINLINE void operator=(const FWeakObjectPtr& Other)
 	{
@@ -109,7 +126,7 @@ struct TPersistentObjectPtr
 		if (!Object && ObjectID.IsValid())
 		{
 			Object = ObjectID.ResolveObject();
-			if (CanCacheObjectPointer(Object))
+			if (CanCacheObjectPointer(FObjectPtr(Object)))
 			{
 				WeakPtr = Object;
 			}
@@ -137,7 +154,7 @@ struct TPersistentObjectPtr
 			Object = ObjectID.ResolveObject();
 			FWeakObjectPtr LocalWeakPtr = Object;
 
-			if (CanCacheObjectPointer(Object))
+			if (CanCacheObjectPointer(FObjectPtr(Object)))
 			{
 				WeakPtr = Object;
 			}
@@ -148,13 +165,13 @@ struct TPersistentObjectPtr
 	}
 
 	/** Dereference the pointer */
-	FORCEINLINE class UObject& operator*() const
+	FORCEINLINE UObject& operator*() const
 	{
 		return *Get();
 	}
 
 	/** Dereference the pointer */
-	FORCEINLINE class UObject* operator->() const
+	FORCEINLINE UObject* operator->() const
 	{
 		return Get();
 	}
@@ -240,7 +257,7 @@ private:
 
 	// Returns whether the object pointer can be stored in WeakPtr for later retrieval
 	// For example, objects that are in the process of being async loaded may not be cached
-	inline bool CanCacheObjectPointer(const UObject* Ptr) const
+	inline bool CanCacheObjectPointer(FObjectPtr Ptr) const
 	{
 		if (IsInAsyncLoadingThread() && Ptr && Ptr->HasAnyInternalFlags(EInternalObjectFlags_AsyncLoading))
 		{
diff --git a/Engine/Source/Runtime/CoreUObject/Public/UObject/SoftObjectPath.h b/Engine/Source/Runtime/CoreUObject/Public/UObject/SoftObjectPath.h
--- a/Engine/Source/Runtime/CoreUObject/Public/UObject/SoftObjectPath.h
+++ b/Engine/Source/Runtime/CoreUObject/Public/UObject/SoftObjectPath.h
@@ -361,8 +361,23 @@ PRAGMA_ENABLE_DEPRECATION_WARNINGS;
 	{
 		return 0;
 	}
-	static COREUOBJECT_API FSoftObjectPath GetOrCreateIDForObject(const UObject* Object);
-	
+	static COREUOBJECT_API FSoftObjectPath GetOrCreateIDForObject(FObjectPtr Object);
+	static FORCEINLINE FSoftObjectPath GetOrCreateIDForObject(const UObject* Object)
+	{
+		return GetOrCreateIDForObject(FObjectPtr(const_cast<UObject*>(Object)));
+	}
+	template <typename T>
+	static FORCEINLINE FSoftObjectPath GetOrCreateIDForObject(TObjectPtr<T> Object)
+	{
+		// This needs to be a template instead of TObjectPtr<const UObject> because C++ does derived-to-base
+		// pointer conversions ('standard conversion sequences') in more cases than TSmartPtr<Derived>-to-TSmartPtr<Base>
+		// conversions ('user-defined conversions'), meaning it doesn't auto-convert in many real use cases.
+		//
+		// https://en.cppreference.com/w/cpp/language/implicit_conversion
+
+		return GetOrCreateIDForObject(FObjectPtr(Object));
+	}
+
 	/** Adds list of packages names that have been created specifically for PIE, this is used for editor fixup */
 	static COREUOBJECT_API void AddPIEPackageName(FName NewPIEPackageName);
 	
diff --git a/Engine/Source/Runtime/CoreUObject/Public/UObject/SoftObjectPtr.h b/Engine/Source/Runtime/CoreUObject/Public/UObject/SoftObjectPtr.h
--- a/Engine/Source/Runtime/CoreUObject/Public/UObject/SoftObjectPtr.h
+++ b/Engine/Source/Runtime/CoreUObject/Public/UObject/SoftObjectPtr.h
@@ -43,6 +43,9 @@ struct TIsSoftObjectPointerType
  */
 struct FSoftObjectPtr : public TPersistentObjectPtr<FSoftObjectPath>
 {
+private:
+	using Super = TPersistentObjectPtr<FSoftObjectPath>;
+
 public:	
 	FORCEINLINE FSoftObjectPtr() = default;
 	FORCEINLINE FSoftObjectPtr(const FSoftObjectPtr& Other) = default;
@@ -52,18 +55,27 @@ public:
 	FORCEINLINE FSoftObjectPtr& operator=(FSoftObjectPtr&& Other) = default;
 
 	explicit FORCEINLINE FSoftObjectPtr(const FSoftObjectPath& ObjectPath)
-		: TPersistentObjectPtr<FSoftObjectPath>(ObjectPath)
+		: Super(ObjectPath)
 	{
 	}
 
-	explicit FORCEINLINE FSoftObjectPtr(const UObject* Object)
+	explicit FORCEINLINE FSoftObjectPtr(FObjectPtr Object)
 	{
 		(*this)=Object;
 	}
-
-	explicit FORCEINLINE FSoftObjectPtr(TObjectPtr<UObject> Object)
+	explicit FORCEINLINE FSoftObjectPtr(const UObject* Object)
+		: FSoftObjectPtr(FObjectPtr(const_cast<UObject*>(Object)))
 	{
-		(*this)=Object;
+	}
+	template <typename T>
+	explicit FORCEINLINE FSoftObjectPtr(TObjectPtr<T> Object)
+		: FSoftObjectPtr(FObjectPtr(Object))
+	{
+		// This needs to be a template instead of TObjectPtr<const UObject> because C++ does derived-to-base
+		// pointer conversions ('standard conversion sequences') in situations that TSmartPtr<Derived>-to-TSmartPtr<Base>
+		// conversions ('user-defined conversions') doesn't, meaning it won't auto-convert in many real use cases.
+		//
+		// https://en.cppreference.com/w/cpp/language/implicit_conversion
 	}
 
 	/** Synchronously load (if necessary) and return the asset object represented by this asset ptr */
@@ -121,18 +133,33 @@ public:
 			// If this object is pending kill or otherwise invalid, this will return nullptr just like TPersistentObjectPtr<FSoftObjectPath>::Get()
 			return Result.Get();
 		}
-		return TPersistentObjectPtr<FSoftObjectPath>::Get();
+		return Super::Get();
 	}
 #endif
 
 	// Implicit conversion from UObject* via assignment shouldn't really be allowed if the constructor is explicit
 	using TPersistentObjectPtr<FSoftObjectPath>::operator=;
-	FSoftObjectPtr& operator=(TObjectPtr<UObject> Ptr)
+	FSoftObjectPtr& operator=(FObjectPtr Ptr)
 	{
-		// It should be possible to do this without resolving
-		(*this) = Ptr.Get();
+		Super* SuperThis = this;
+		(*SuperThis) = Ptr;
 		return *this;
 	}
+	FORCEINLINE FSoftObjectPtr& operator=(const UObject* Ptr)
+	{
+		return *this = FObjectPtr(const_cast<UObject*>(Ptr));
+	}
+	template <typename T>
+	FORCEINLINE FSoftObjectPtr& operator=(TObjectPtr<T> Ptr)
+	{
+		// This needs to be a template instead of TObjectPtr<const UObject> because C++ does derived-to-base
+		// pointer conversions ('standard conversion sequences') in situations that TSmartPtr<Derived>-to-TSmartPtr<Base>
+		// conversions ('user-defined conversions') doesn't, meaning it won't auto-convert in many real use cases.
+		//
+		// https://en.cppreference.com/w/cpp/language/implicit_conversion
+
+		return *this = FObjectPtr(Ptr);
+	}
 };
 
 template <> struct TIsPODType<FSoftObjectPtr> { enum { Value = TIsPODType<TPersistentObjectPtr<FSoftObjectPath> >::Value }; };
diff --git a/Engine/Source/Runtime/CoreUObject/Public/UObject/WeakObjectPtr.h b/Engine/Source/Runtime/CoreUObject/Public/UObject/WeakObjectPtr.h
--- a/Engine/Source/Runtime/CoreUObject/Public/UObject/WeakObjectPtr.h
+++ b/Engine/Source/Runtime/CoreUObject/Public/UObject/WeakObjectPtr.h
@@ -73,17 +73,33 @@ public:
 	}
 #endif
 
+	UE_DEPRECATED(5.6, "Constructing a FWeakObjectPtr from NULL has been deprecated - please use nullptr instead.")
+	FORCEINLINE FWeakObjectPtr(int)
+	{
+		(*this) = nullptr;
+	}
+
 	/**  
 	 * Construct from an object pointer or something that can be implicitly converted to an object pointer
 	 * @param Object object to create a weak pointer to
 	 */
-	template <
-		typename U
-		UE_REQUIRES(std::is_convertible_v<U, const UObject*>)
-	>
-	FORCEINLINE FWeakObjectPtr(U&& Object)
+	FORCEINLINE FWeakObjectPtr(FObjectPtr Object)
 	{
-		(*this) = ImplicitConv<const UObject*>(Object);
+		*this = Object;
+	}
+	FORCEINLINE FWeakObjectPtr(const UObject* Object)
+		: FWeakObjectPtr(FObjectPtr(const_cast<UObject*>(Object)))
+	{
+	}
+	template <typename T>
+	FORCEINLINE FWeakObjectPtr(TObjectPtr<T> Object)
+		: FWeakObjectPtr(FObjectPtr(const_cast<UObject*>(ImplicitConv<const UObject*>(Object.Get()))))
+	{
+		// This needs to be a template instead of TObjectPtr<const UObject> because C++ does derived-to-base
+		// pointer conversions ('standard conversion sequences') in situations that TSmartPtr<Derived>-to-TSmartPtr<Base>
+		// conversions ('user-defined conversions') doesn't, meaning it won't auto-convert in many real use cases.
+		//
+		// https://en.cppreference.com/w/cpp/language/implicit_conversion
 	}
 
 	/**  
@@ -110,8 +126,22 @@ public:
 	 * Copy from an object pointer
 	 * @param Object object to create a weak pointer to
 	 */
-	COREUOBJECT_API void operator=(const class UObject* Object);
-	COREUOBJECT_API void operator=(TObjectPtr<UObject> Object);
+	COREUOBJECT_API void operator=(FObjectPtr Object);
+	FORCEINLINE void operator=(const UObject* Object)
+	{
+		*this = FObjectPtr(const_cast<UObject*>(Object));
+	}
+	template <typename T>
+	FORCEINLINE void operator=(TObjectPtr<T> Object)
+	{
+		// This needs to be a template instead of TObjectPtr<const UObject> because C++ does derived-to-base
+		// pointer conversions ('standard conversion sequences') in situations that TSmartPtr<Derived>-to-TSmartPtr<Base>
+		// conversions ('user-defined conversions') doesn't, meaning it won't auto-convert in many real use cases.
+		//
+		// https://en.cppreference.com/w/cpp/language/implicit_conversion
+
+		*this = FObjectPtr(Object);
+	}
 
 	/**  
 	 * Construct from another weak pointer
