IncrementalCook: Change TObjectPtr to not report a resolve for its internal resolves when those internal resolves are not behaviorally necessary but instead are used only for better performance of gathering the data.
Add TObjectPtr::IsIn to provide that same skip of internal resolve for IsIn queries.
#rnx

diff --git a/Engine/Source/Runtime/CoreUObject/Private/UObject/ObjectFwd.cpp b/Engine/Source/Runtime/CoreUObject/Private/UObject/ObjectFwd.cpp
--- a/Engine/Source/Runtime/CoreUObject/Private/UObject/ObjectFwd.cpp
+++ b/Engine/Source/Runtime/CoreUObject/Private/UObject/ObjectFwd.cpp
@@ -20,6 +20,11 @@ UObject* GetOuter(const UObject* Obj)
 	return Obj->GetOuter();
 }
 
+bool IsIn(const UObject* Obj, const UObject* SomeOuter)
+{
+	return Obj->IsIn(SomeOuter);
+}
+
 FString GetPathName( const UObject* Obj, const UObject* StopOuter )
 {
 	return Obj->GetPathName(StopOuter);
diff --git a/Engine/Source/Runtime/CoreUObject/Private/UObject/ObjectPtr.cpp b/Engine/Source/Runtime/CoreUObject/Private/UObject/ObjectPtr.cpp
--- a/Engine/Source/Runtime/CoreUObject/Private/UObject/ObjectPtr.cpp
+++ b/Engine/Source/Runtime/CoreUObject/Private/UObject/ObjectPtr.cpp
@@ -16,7 +16,7 @@ FString FObjectPtr::GetPathName() const
 	FObjectHandle LocalHandle = Handle;
 	if (IsObjectHandleResolved(LocalHandle) && !IsObjectHandleNull(LocalHandle))
 	{
-		return ResolveObjectHandle(LocalHandle)->GetPathName();
+		return ResolveObjectHandleNoRead(LocalHandle)->GetPathName();
 	}
 	else
 	{
@@ -30,7 +30,7 @@ FName FObjectPtr::GetFName() const
 	FObjectHandle LocalHandle = Handle;
 	if (IsObjectHandleResolved(LocalHandle) && !IsObjectHandleNull(LocalHandle))
 	{
-		return ResolveObjectHandle(LocalHandle)->GetFName();
+		return ResolveObjectHandleNoRead(LocalHandle)->GetFName();
 	}
 	else
 	{
@@ -44,7 +44,7 @@ FString FObjectPtr::GetFullName(EObjectFullNameFlags Flags) const
 	FObjectHandle LocalHandle = Handle;
 	if (IsObjectHandleResolved(LocalHandle) && !IsObjectHandleNull(LocalHandle))
 	{
-		return ResolveObjectHandle(LocalHandle)->GetFullName(nullptr, Flags);
+		return ResolveObjectHandleNoRead(LocalHandle)->GetFullName(nullptr, Flags);
 	}
 	else
 	{
@@ -58,7 +58,7 @@ FObjectPtr FObjectPtr::GetOuter() const
 	FObjectHandle LocalHandle = Handle;
 	if (IsObjectHandleResolved(LocalHandle) && !IsObjectHandleNull(LocalHandle))
 	{
-		return FObjectPtr(ResolveObjectHandle(LocalHandle)->GetOuter());
+		return FObjectPtr(ResolveObjectHandleNoRead(LocalHandle)->GetOuter());
 	}
 	UE::CoreUObject::Private::FPackedObjectRef PackedRef = UE::CoreUObject::Private::GetOuter(ReadObjectHandlePackedObjectRefNoCheck(LocalHandle));
 	FObjectPtr Ptr({ PackedRef.EncodedRef });
@@ -70,12 +70,25 @@ FObjectPtr FObjectPtr::GetPackage() const
 	FObjectHandle LocalHandle = Handle;
 	if (IsObjectHandleResolved(LocalHandle) && !IsObjectHandleNull(LocalHandle))
 	{
-		return FObjectPtr(ResolveObjectHandle(LocalHandle)->GetPackage());
+		return FObjectPtr(ResolveObjectHandleNoRead(LocalHandle)->GetPackage());
 	}
 	UE::CoreUObject::Private::FPackedObjectRef PackedRef = UE::CoreUObject::Private::GetPackage(ReadObjectHandlePackedObjectRefNoCheck(LocalHandle));
 	FObjectPtr Ptr({ PackedRef.EncodedRef });
 	return Ptr;
 }
+
+bool FObjectPtr::IsIn(FObjectPtr SomeOuter) const
+{
+	FObjectHandle SomeOuterHandle = SomeOuter.Handle;
+	FObjectHandle LocalHandle = Handle;
+	if (IsObjectHandleNull(LocalHandle) || IsObjectHandleNull(SomeOuterHandle))
+	{
+		return false;
+	}
+	// TODO: Handle without resolving. Need to decide how ObjectPtrs handle objects in external packages.
+	return ResolveObjectHandleNoRead(LocalHandle)->IsIn(ResolveObjectHandleNoRead(SomeOuterHandle));
+}
+
 #endif
 
 bool FObjectPtr::IsA(const UClass* SomeBase) const
diff --git a/Engine/Source/Runtime/CoreUObject/Public/UObject/ObjectFwd.h b/Engine/Source/Runtime/CoreUObject/Public/UObject/ObjectFwd.h
--- a/Engine/Source/Runtime/CoreUObject/Public/UObject/ObjectFwd.h
+++ b/Engine/Source/Runtime/CoreUObject/Public/UObject/ObjectFwd.h
@@ -40,6 +40,7 @@ COREUOBJECT_API FString GetPathName( const UObject* Obj, const UObject* StopOute
 COREUOBJECT_API UClass* GetClass(UObject* Obj);
 COREUOBJECT_API FName GetFName(const UObject* Obj);
 COREUOBJECT_API UObject* GetOuter(const UObject* Obj);
+COREUOBJECT_API bool IsIn(const UObject* Obj, const UObject* SomeOuter);
 
 COREUOBJECT_API bool HasAnyFlags(const UObject* Obj, int32 FlagsToCheck);
 }
diff --git a/Engine/Source/Runtime/CoreUObject/Public/UObject/ObjectPtr.h b/Engine/Source/Runtime/CoreUObject/Public/UObject/ObjectPtr.h
--- a/Engine/Source/Runtime/CoreUObject/Public/UObject/ObjectPtr.h
+++ b/Engine/Source/Runtime/CoreUObject/Public/UObject/ObjectPtr.h
@@ -51,7 +51,7 @@
 #endif
 
 class UClass;
-	
+
 template <typename T>
 struct TObjectPtr;
 
@@ -205,14 +205,15 @@ public:
 		return GetFName().ToString();
 	}
 
-	// Gets the Outer UObject for this Object
-	// returns unresolved if this object ptr is unresolved
+	// Gets the Outer UObject for this Object without resolving the object reference.
 	COREUOBJECT_API FObjectPtr GetOuter() const;
 
-	// Gets the package for this Object
-	// returns unresolved if this object ptr is unresolved
+	// Gets the package for this Object without resolving the object reference.
 	COREUOBJECT_API FObjectPtr GetPackage() const;
 
+	// Reports whether this Object is within an Outer or Package. In some cases may need to resolve the object reference.
+	COREUOBJECT_API bool IsIn(FObjectPtr SomeOuter) const;
+
 	/** Returns the full name for the object in the form: Class ObjectPath */
 	COREUOBJECT_API FString GetFullName(EObjectFullNameFlags Flags = EObjectFullNameFlags::None) const;
 #else
@@ -244,6 +245,15 @@ public:
 		const UObject* ResolvedObject = Get();
 		return ResolvedObject ? FObjectPtr(UE::CoreUObject::Private::GetPackage(ResolvedObject)) : FObjectPtr(nullptr);
 	}
+
+	bool IsIn(FObjectPtr SomeOuter) const
+	{
+		const UObject* ResolvedObject = Get();
+		const UObject* ResolvedSomeOuter = SomeOuter.Get();
+		return (ResolvedObject && ResolvedSomeOuter) 
+			? UE::CoreUObject::Private::IsIn(ResolvedObject, ResolvedSomeOuter) : false;
+	}
+
 	/**
 	 * Returns the fully qualified pathname for this object as well as the name of the class, in the format:
 	 * 'ClassName Outermost.[Outer:]Name'.
@@ -667,7 +677,18 @@ public:
 		FObjectPtr Ptr = ObjectPtr.GetPackage();
 		return TObjectPtr<UPackage>(Ptr);
 	}
-	
+
+	FORCEINLINE bool IsIn(FObjectPtr SomeOuter) const
+	{
+		return ObjectPtr.IsIn(SomeOuter);
+	}
+
+	template <typename U>
+	FORCEINLINE bool IsIn(TObjectPtr<U> SomeOuter) const
+	{
+		return ObjectPtr.IsIn(SomeOuter.ObjectPtr);
+	}
+
 	FORCEINLINE operator T* () const { return Get(); }
 	template <typename U>
 	UE_OBJPTR_DEPRECATED(5.0, "Explicit cast to other raw pointer types is deprecated.  Please use the Cast API or get the raw pointer with ToRawPtr and cast that instead.")
