[RegressionFix] 35924957 broke IDOs and caused complete data loss during reinstancing

diff --git a/Engine/Source/Runtime/CoreUObject/Private/UObject/PropertyBagRepository.cpp b/Engine/Source/Runtime/CoreUObject/Private/UObject/PropertyBagRepository.cpp
--- a/Engine/Source/Runtime/CoreUObject/Private/UObject/PropertyBagRepository.cpp
+++ b/Engine/Source/Runtime/CoreUObject/Private/UObject/PropertyBagRepository.cpp
@@ -1043,7 +1043,7 @@ FScopedIDOSerializationContext::FScopedIDOSerializationContext(UObject* InObject
 	FUObjectSerializeContext* SerializeContext = FUObjectThreadContext::Get().GetSerializeContext();
 	SaveSerializeContext(SerializeContext);
 
-	SerializeContext->bImpersonateProperties = bImpersonate && IsInstanceDataObjectSupportEnabled(Object);
+	SerializeContext->bImpersonateProperties = bImpersonate && IsInstanceDataObjectSupportEnabled(InObject);
 }
 
 FScopedIDOSerializationContext::FScopedIDOSerializationContext(bool bImpersonate)
