Allow undo/redo transaction deserialization to set non-nullable object properties to null

* Low level utility functions FObjectPropertyBase::SetObjectPtrPropertyValue** prevent NonNullable properties from being set to null.
* Undo/Redo from transaction buffer relies on deserialization to set back the value of a property (which could be NonNullable) to the Null state prior to it being manually set to not null.
* The low level behaviour is currently relied upon by existing data, so it cannot be removed.
* The workaround for now adds a TLS that can be set based on the archive state flags to opt out the low level behaviour from disallowing null property sets on NonNullable properties.
* This CL can be disabled with CVar `CoreUObject.EnableAllowSetNullToNonNullableOverride=0`

diff --git a/Engine/Source/Runtime/CoreUObject/Private/UObject/NonNullPropertyUtils.cpp b/Engine/Source/Runtime/CoreUObject/Private/UObject/NonNullPropertyUtils.cpp
new file mode 100644
--- /dev/null
+++ b/Engine/Source/Runtime/CoreUObject/Private/UObject/NonNullPropertyUtils.cpp
@@ -0,0 +1,62 @@
+﻿// Copyright Epic Games, Inc. All Rights Reserved.
+
+#include "NonNullPropertyUtils.h"
+
+#include "HAL/IConsoleManager.h"
+
+namespace UE::CoreUObject::Private
+{
+#if WITH_EDITORONLY_DATA
+	TAutoConsoleVariable<bool> CVarEnableNullDeserializeToNonNullableOverride(
+		TEXT("CoreUObject.EnableAllowSetNullToNonNullableOverride"),
+		true,
+		TEXT("Allows scoped overrides to enable null to be set on a NonNullable property")
+	);
+		
+	thread_local FScopedSetNullOnNonNullable* TLSScopedSetNullOnNonNullableBehaviorStackHead = nullptr;
+	
+	FScopedSetNullOnNonNullable::FScopedSetNullOnNonNullable(TOptional<EAllowSetNullOnNonNullableBehavior> InBehavior)
+		: Previous(nullptr)
+	{
+		if (InBehavior.IsSet() && Private::CVarEnableNullDeserializeToNonNullableOverride.GetValueOnAnyThread())
+		{
+			bIsBehaviorSet = true;
+			Behavior =  *InBehavior;
+
+			// Push the stack object to the list
+			Previous = Private::TLSScopedSetNullOnNonNullableBehaviorStackHead;
+			Private::TLSScopedSetNullOnNonNullableBehaviorStackHead = this;
+		}
+		else
+		{
+			bIsBehaviorSet = false;
+			Behavior = EAllowSetNullOnNonNullableBehavior::Disabled; // Set it to something
+		}
+	}
+
+	FScopedSetNullOnNonNullable::~FScopedSetNullOnNonNullable()
+	{
+		if (bIsBehaviorSet)
+		{	
+			// Pop us from the stack and set the previous to head
+			check(Private::TLSScopedSetNullOnNonNullableBehaviorStackHead == this);
+			Private::TLSScopedSetNullOnNonNullableBehaviorStackHead = Previous;
+		}
+	}
+
+	EAllowSetNullOnNonNullableBehavior FScopedSetNullOnNonNullable::GetBehavior() const
+	{
+		check(bIsBehaviorSet);
+		return Behavior;
+	}
+	
+	EAllowSetNullOnNonNullableBehavior AllowSetNullOnNonNullableBehavior()
+	{
+		if (TLSScopedSetNullOnNonNullableBehaviorStackHead)
+		{
+			return TLSScopedSetNullOnNonNullableBehaviorStackHead->GetBehavior();
+		}
+		return  EAllowSetNullOnNonNullableBehavior::Disabled;
+	}
+#endif
+}
diff --git a/Engine/Source/Runtime/CoreUObject/Private/UObject/NonNullPropertyUtils.h b/Engine/Source/Runtime/CoreUObject/Private/UObject/NonNullPropertyUtils.h
new file mode 100644
--- /dev/null
+++ b/Engine/Source/Runtime/CoreUObject/Private/UObject/NonNullPropertyUtils.h
@@ -0,0 +1,49 @@
+﻿// Copyright Epic Games, Inc. All Rights Reserved.
+
+#pragma once
+
+#include "CoreTypes.h"
+#include "Misc/Optional.h"
+
+namespace UE::CoreUObject::Private
+{
+	enum class EAllowSetNullOnNonNullableBehavior : uint8
+	{
+		Disabled,
+		Enabled
+	};
+	
+#if WITH_EDITORONLY_DATA
+	class FScopedSetNullOnNonNullable
+	{
+	public:
+		// If set, InBehavior will adjust later code in stack to adopt the behavior
+		// If not set, prevailing behavior is preserved
+		explicit FScopedSetNullOnNonNullable(TOptional<EAllowSetNullOnNonNullableBehavior> InBehavior);
+		~FScopedSetNullOnNonNullable();
+	
+		EAllowSetNullOnNonNullableBehavior GetBehavior() const;		
+	
+	private:
+		FScopedSetNullOnNonNullable(const FScopedSetNullOnNonNullable&) = delete;
+		FScopedSetNullOnNonNullable(FScopedSetNullOnNonNullable&&) = delete;
+		FScopedSetNullOnNonNullable& operator=(const FScopedSetNullOnNonNullable&) = delete;
+		FScopedSetNullOnNonNullable& operator=(FScopedSetNullOnNonNullable&&) = delete;
+		
+		FScopedSetNullOnNonNullable* Previous;
+		EAllowSetNullOnNonNullableBehavior Behavior;
+		bool bIsBehaviorSet;
+	};
+
+	// Thread local, set with FScopedSetNullOnNonNullable
+	// If true, will allow setting null on a NonNullable property
+	// Some cases in editor deserialization requires null to be set on NonNull properties
+	// Example: deserializing a previously null value to NonNullable property from the transaction buffer for Undo/Redo
+	COREUOBJECT_API EAllowSetNullOnNonNullableBehavior AllowSetNullOnNonNullableBehavior();
+#else
+	inline EAllowSetNullOnNonNullableBehavior AllowSetNullOnNonNullableBehavior()
+	{
+		return  EAllowSetNullOnNonNullableBehavior::Disabled;
+	}
+#endif
+}
diff --git a/Engine/Source/Runtime/CoreUObject/Private/UObject/Obj.cpp b/Engine/Source/Runtime/CoreUObject/Private/UObject/Obj.cpp
--- a/Engine/Source/Runtime/CoreUObject/Private/UObject/Obj.cpp
+++ b/Engine/Source/Runtime/CoreUObject/Private/UObject/Obj.cpp
@@ -5,6 +5,7 @@
 =============================================================================*/
 
 #include "CoreMinimal.h"
+#include "NonNullPropertyUtils.h"
 #include "AssetRegistry/AssetData.h"
 #include "Misc/CoreMisc.h"
 #include "Misc/CommandLine.h"
@@ -1674,6 +1675,20 @@ void UObject::Serialize(FStructuredArchive::FRecord Record)
 			}
 		}
 
+#if WITH_EDITORONLY_DATA
+		// Allow Null to be deserialized into a NonNull property
+		// Enabled for Undo/Redo purposes
+		UE::CoreUObject::Private::FScopedSetNullOnNonNullable ScopedAllowNullPtrDeserialization = [&UnderlyingArchive]()
+		{
+			TOptional<UE::CoreUObject::Private::EAllowSetNullOnNonNullableBehavior> Behaviour;
+			if (UnderlyingArchive.IsTransacting() && UnderlyingArchive.IsLoading())
+			{
+				Behaviour.Emplace(UE::CoreUObject::Private::EAllowSetNullOnNonNullableBehavior::Enabled);
+			}
+			return UE::CoreUObject::Private::FScopedSetNullOnNonNullable(Behaviour);
+		}();
+#endif
+		
 		// Special info.
 		if ((!UnderlyingArchive.IsLoading() && !UnderlyingArchive.IsSaving() && !UnderlyingArchive.IsObjectReferenceCollector()))
 		{
diff --git a/Engine/Source/Runtime/CoreUObject/Private/UObject/PropertyBaseObject.cpp b/Engine/Source/Runtime/CoreUObject/Private/UObject/PropertyBaseObject.cpp
--- a/Engine/Source/Runtime/CoreUObject/Private/UObject/PropertyBaseObject.cpp
+++ b/Engine/Source/Runtime/CoreUObject/Private/UObject/PropertyBaseObject.cpp
@@ -2,6 +2,7 @@
 
 #include "CoreMinimal.h"
 #include "HAL/IConsoleManager.h"
+#include "NonNullPropertyUtils.h"
 #include "UObject/ObjectMacros.h"
 #include "Templates/Casts.h"
 #include "UObject/Package.h"
@@ -666,7 +667,8 @@ void FObjectPropertyBase::SetObjectPtrPropertyValueUnchecked_InContainer(void* C
 
 void FObjectPropertyBase::SetObjectPropertyValue(void* PropertyValueAddress, UObject* Value) const
 {
-	if (Value || !HasAnyPropertyFlags(CPF_NonNullable))
+	using namespace UE::CoreUObject;
+	if (Value || !HasAnyPropertyFlags(CPF_NonNullable) || (Private::AllowSetNullOnNonNullableBehavior() == Private::EAllowSetNullOnNonNullableBehavior::Enabled))
 	{
 		SetObjectPropertyValueUnchecked(PropertyValueAddress, Value);
 	}
@@ -678,7 +680,8 @@ void FObjectPropertyBase::SetObjectPropertyValue(void* PropertyValueAddress, UOb
 
 void FObjectPropertyBase::SetObjectPtrPropertyValue(void* PropertyValueAddress, TObjectPtr<UObject> Ptr) const
 {
-	if (Ptr || !HasAnyPropertyFlags(CPF_NonNullable))
+	using namespace UE::CoreUObject;
+	if (Ptr || !HasAnyPropertyFlags(CPF_NonNullable) || (Private::AllowSetNullOnNonNullableBehavior() == Private::EAllowSetNullOnNonNullableBehavior::Enabled))
 	{
 		SetObjectPtrPropertyValueUnchecked(PropertyValueAddress, Ptr);
 	}
@@ -690,7 +693,8 @@ void FObjectPropertyBase::SetObjectPtrPropertyValue(void* PropertyValueAddress,
 
 void FObjectPropertyBase::SetObjectPropertyValue_InContainer(void* ContainerAddress, UObject* Value, int32 ArrayIndex) const
 {
-	if (Value || !HasAnyPropertyFlags(CPF_NonNullable))
+	using namespace UE::CoreUObject;
+	if (Value || !HasAnyPropertyFlags(CPF_NonNullable) || (Private::AllowSetNullOnNonNullableBehavior() == Private::EAllowSetNullOnNonNullableBehavior::Enabled))
 	{
 		SetObjectPropertyValueUnchecked_InContainer(ContainerAddress, Value, ArrayIndex);
 	}
@@ -702,7 +706,8 @@ void FObjectPropertyBase::SetObjectPropertyValue_InContainer(void* ContainerAddr
 
 void FObjectPropertyBase::SetObjectPtrPropertyValue_InContainer(void* ContainerAddress, TObjectPtr<UObject> Ptr, int32 ArrayIndex) const
 {
-	if (Ptr || !HasAnyPropertyFlags(CPF_NonNullable))
+	using namespace UE::CoreUObject;
+	if (Ptr || !HasAnyPropertyFlags(CPF_NonNullable) || (Private::AllowSetNullOnNonNullableBehavior() == Private::EAllowSetNullOnNonNullableBehavior::Enabled))
 	{
 		SetObjectPtrPropertyValueUnchecked_InContainer(ContainerAddress, Ptr, ArrayIndex);
 	}
