StructView: Fix issues with derived type uses

Fix template argument to Get<>() when attempting to cast to a derived type.
Fix struct constructor when using a derived type so that the correct UScriptStruct is identified by the view (previously it always used the base type, so derived type information was lost).

diff --git a/Engine/Source/Runtime/CoreUObject/Public/StructUtils/StructView.h b/Engine/Source/Runtime/CoreUObject/Public/StructUtils/StructView.h
--- a/Engine/Source/Runtime/CoreUObject/Public/StructUtils/StructView.h
+++ b/Engine/Source/Runtime/CoreUObject/Public/StructUtils/StructView.h
@@ -143,8 +143,8 @@ public:
 	explicit TStructView() = default;
 
 	template<typename T = BaseStructT, typename = std::enable_if_t<std::is_base_of_v<BaseStructT, std::decay_t<T>>>>
-	TStructView(BaseStructT& InStruct)
-		: FStructView(TBaseStructure<BaseStructT>::Get(), reinterpret_cast<uint8*>(&InStruct))
+	TStructView(T& InStruct)
+		: FStructView(TBaseStructure<T>::Get(), reinterpret_cast<uint8*>(&InStruct))
 	{}
 
 	template<typename T = BaseStructT, typename = std::enable_if_t<std::is_base_of_v<BaseStructT, std::decay_t<T>>>>
@@ -166,7 +166,7 @@ public:
 	template<typename T = BaseStructT, typename = std::enable_if_t<std::is_base_of_v<BaseStructT, std::decay_t<T>>>>
 	T& Get() const
 	{
-		return UE::StructUtils::GetStructRef<BaseStructT>(ScriptStruct, StructMemory);
+		return UE::StructUtils::GetStructRef<T>(ScriptStruct, StructMemory);
 	}
 
 	/** Returns mutable pointer to the struct, or nullptr if cast is not valid. */
