Prevent bulkdata objects being saved with duplicate guids when saving a level created from a template.

- This issue was originally fixed in CL 23787156 but since then the editor has changed it's loading path and the 'LOAD_RegenerateBulkDataGuids' flag was no longer being passed to all of the packages being loaded as part of the map template.
- Taking a step back, the problem here is that the .umap package is correctly marked as being instanced but it's associated _BuiltData.uasset is not, so this time we will try fixing that instead by adding info about the _BuiltData.uasset package to the FLinkerInstancingContext being passed into the ::LoadPackage request. FEditorBulkData should then be made aware that it needs to generate new guids for the textures within that package.

- Deprecated LOAD_RegenerateBulkDataGuids as it is no longer needed as this was our only usecase.
- Deprecated FLinkerLoad::ShouldRegenerateGuids and removed the code using it from FEditorBulkData.

diff --git a/Engine/Source/Editor/UnrealEd/Private/EditorServer.cpp b/Engine/Source/Editor/UnrealEd/Private/EditorServer.cpp
--- a/Engine/Source/Editor/UnrealEd/Private/EditorServer.cpp
+++ b/Engine/Source/Editor/UnrealEd/Private/EditorServer.cpp
@@ -2571,7 +2571,11 @@ bool UEditorEngine::Map_Load(const TCHAR* Str, FOutputDevice& Ar)
 						FSoftObjectPath(*WriteToString<256>(WorldPackage->GetName(), TEXT("."), ShortWorldPackageName))
 					);
 
-					LoadFlags |= LOAD_RegenerateBulkDataGuids;
+					// Add remapping for the _BuiltData.uasset package file as well (assuming that the template has one)
+					FName BuiltDataOriginal = FName(*WriteToString<256>(LongTempFname, TEXT("_BuiltData")));
+					FName BuiltDataRemapped = FName(*WriteToString<256>(WorldPackageFName, TEXT("_BuiltData")));
+					
+					WorldPackageInstancingContext.AddPackageMapping(BuiltDataOriginal, BuiltDataRemapped);
 
 					WorldPackage = LoadPackage( WorldPackage, *LongTempFname, LoadFlags, nullptr /* InReaderOverride */, &WorldPackageInstancingContext);
 					WorldPackage->SetPackageFlags(PKG_NewlyCreated);
diff --git a/Engine/Source/Runtime/CoreUObject/Private/Blueprint/BlueprintSupport.cpp b/Engine/Source/Runtime/CoreUObject/Private/Blueprint/BlueprintSupport.cpp
--- a/Engine/Source/Runtime/CoreUObject/Private/Blueprint/BlueprintSupport.cpp
+++ b/Engine/Source/Runtime/CoreUObject/Private/Blueprint/BlueprintSupport.cpp
@@ -1577,7 +1577,7 @@ void FLinkerLoad::ResolveDeferredDependencies(UStruct* LoadStruct)
 			// fully loading... we need to make sure metadata gets loaded, etc.)
 			if ((SourcePackage != nullptr) && !SourcePackage->HasAnyFlags(RF_WasLoaded | RF_WillBeLoaded))
 			{
-				uint32 InternalLoadFlags = LoadFlags & (LOAD_NoVerify | LOAD_NoWarn | LOAD_Quiet | LOAD_RegenerateBulkDataGuids);
+				uint32 InternalLoadFlags = LoadFlags & (LOAD_NoVerify | LOAD_NoWarn | LOAD_Quiet);
 				// make sure LoadAllObjects() is called for this package
 				LoadPackageInternal(/*Outer =*/nullptr, SourceLinker->GetPackagePath(), InternalLoadFlags, this, nullptr/*InReaderOverride*/, nullptr/*InstancingContext*/, nullptr /* DiffPackagePath */); //-V595
 			}
diff --git a/Engine/Source/Runtime/CoreUObject/Private/Serialization/EditorBulkData.cpp b/Engine/Source/Runtime/CoreUObject/Private/Serialization/EditorBulkData.cpp
--- a/Engine/Source/Runtime/CoreUObject/Private/Serialization/EditorBulkData.cpp
+++ b/Engine/Source/Runtime/CoreUObject/Private/Serialization/EditorBulkData.cpp
@@ -234,22 +234,13 @@ FString GetDebugNameFromOwner(UObject* Owner)
 
 /** 
  * Utility to check if we need to generate a new unique identifier for the editor bulkdata or not.
- * At the moment the instancing context (if any) can request that we generate new guids. At the time
- * of writing the use case for this is creating a new map from a template, and the bulkdata is being
- * loaded from the template as we don't want every map created from the template to have the same 
- * identifiers.
- * The second use case is if we are loading the editor bulkdata to a transient package, which can 
+ * The current use case is if we are loading the editor bulkdata to a transient package, which can 
  * occur when some assets duplicate a number of UObjects from a template package that will eventually
  * be added to themselves. The duplication will load the editor bulkdata to a transient package before
  * it is re-parented.
  */
 bool ShouldGenerateNewIdentifier(FLinkerLoad* LinkerLoad, UObject* Owner)
 {
-	if (LinkerLoad && LinkerLoad->ShouldRegenerateGuids())
-	{
-		return true;
-	}
-
 	UPackage* Package = Owner != nullptr ? Owner->GetPackage() : nullptr;
 	if (Package != nullptr && Package->HasAnyFlags(RF_Transient))
 	{
diff --git a/Engine/Source/Runtime/CoreUObject/Private/UObject/LinkerLoad.cpp b/Engine/Source/Runtime/CoreUObject/Private/UObject/LinkerLoad.cpp
--- a/Engine/Source/Runtime/CoreUObject/Private/UObject/LinkerLoad.cpp
+++ b/Engine/Source/Runtime/CoreUObject/Private/UObject/LinkerLoad.cpp
@@ -3446,7 +3446,7 @@ bool FLinkerLoad::VerifyImportInner(const int32 ImportIndex, FString& WarningSuf
 		check(Import.ClassName == NAME_Package || Import.HasPackageName());
 
 		UPackage* Package = nullptr;
-		uint32 InternalLoadFlags = LoadFlags & (LOAD_NoVerify | LOAD_NoWarn | LOAD_Quiet | LOAD_RegenerateBulkDataGuids);
+		uint32 InternalLoadFlags = LoadFlags & (LOAD_NoVerify | LOAD_NoWarn | LOAD_Quiet);
 		FUObjectSerializeContext* SerializeContext = FUObjectThreadContext::Get().GetSerializeContext();
 
 		// Resolve the package name for the import, potentially remapping it, if instancing
diff --git a/Engine/Source/Runtime/CoreUObject/Public/UObject/LinkerLoad.h b/Engine/Source/Runtime/CoreUObject/Public/UObject/LinkerLoad.h
--- a/Engine/Source/Runtime/CoreUObject/Public/UObject/LinkerLoad.h
+++ b/Engine/Source/Runtime/CoreUObject/Public/UObject/LinkerLoad.h
@@ -985,10 +985,10 @@ public:
 	 */
 	COREUOBJECT_API void DetachExports();
 
-	/** Should bulkdata identifiers should be regenerated as they are loaded or not */
+	UE_DEPRECATED(5.5, "LOAD_RegenerateBulkDataGuids is now obsolete so this method no longer has any purpose.")
 	bool ShouldRegenerateGuids() const
 	{
-		return (LoadFlags & LOAD_RegenerateBulkDataGuids) != 0;
+		return false;
 	}
 
 private:
diff --git a/Engine/Source/Runtime/CoreUObject/Public/UObject/ObjectMacros.h b/Engine/Source/Runtime/CoreUObject/Public/UObject/ObjectMacros.h
--- a/Engine/Source/Runtime/CoreUObject/Public/UObject/ObjectMacros.h
+++ b/Engine/Source/Runtime/CoreUObject/Public/UObject/ObjectMacros.h
@@ -65,8 +65,8 @@ enum ELoadFlags
 //	LOAD_Unused						= 0x00000040
 	LOAD_NoVerify					= 0x00000080,	///< Don't verify imports yet.
 	LOAD_IsVerifying				= 0x00000100,	///< Is verifying imports
-	LOAD_SkipLoadImportedPackages	= 0x00000200,	///< Assume that all import packages are already loaded and don't call LoadPackage when creating imports 
-	LOAD_RegenerateBulkDataGuids	= 0x00000400,	///< BulkData identifiers should be regenerated as they are loaded 
+	LOAD_SkipLoadImportedPackages	= 0x00000200,	///< Assume that all import packages are already loaded and don't call LoadPackage when creating imports
+	LOAD_RegenerateBulkDataGuids UE_DEPRECATED(5.5, "This flag is now obsolete as the functionality is unnecessary.") = 0x00000400, ///< BulkData identifiers should be regenerated as they are loaded
 //	LOAD_Unused						= 0x00000800,
 	LOAD_DisableDependencyPreloading = 0x00001000,	///< Bypass dependency preloading system
 	LOAD_Quiet						= 0x00002000,	///< No log warnings.
