Removed deprecated TAssetPtr and TAssetSubclassOf.

#rnx

diff --git a/Engine/Source/Runtime/CoreUObject/Public/UObject/SoftObjectPtr.h b/Engine/Source/Runtime/CoreUObject/Public/UObject/SoftObjectPtr.h
--- a/Engine/Source/Runtime/CoreUObject/Public/UObject/SoftObjectPtr.h
+++ b/Engine/Source/Runtime/CoreUObject/Public/UObject/SoftObjectPtr.h
@@ -979,12 +979,6 @@ struct FSoftObjectPtrLexicalLess : private FSoftObjectPathLexicalLess
 	}
 };
 
-template<class T=UObject>
-using TAssetPtr UE_DEPRECATED(5.0, "TAssetPtr was renamed to TSoftObjectPtr as it is not necessarily an asset") = TSoftObjectPtr<T>;
-
-template<class TClass = UObject>
-using TAssetSubclassOf UE_DEPRECATED(5.0, "TAssetSubclassOf was renamed to TSoftClassPtr") = TSoftClassPtr<TClass>;
-
 /** Not directly inlined on purpose so compiler have the option of not inlining it. (and it also works with extern template) */
 template<class T>
 T* TSoftObjectPtr<T>::Get() const
