Change SavePackage to opt out (in some of its code) of transitive build dependencies for its handling of imports. If the inclusion/exclusion of a instance of a class depends on more dependencies than the package bytes and class code, then the object's class will need to override IsEditorOnly and manually add the necessary dependencies.
Where possible in SavePackage code, change the code to avoid the resolve of TObjectPtr, and do not add the opt-out scope. Avoiding opt-out scopes is for futureproofing: it will keep the code robust (but worse performance) if future changes add in new dependencies.
Implementation detail: Pass through the FObjectStatus held by SaveContext for each UObject it encounters into low level functions that might need to resolve the TObjectPtr, so we can add references on its first resolve but skip their add in later resolves for performance.
#rnx

diff --git a/Engine/Source/Runtime/CoreUObject/Private/Serialization/ArchiveObjectCrc32.cpp b/Engine/Source/Runtime/CoreUObject/Private/Serialization/ArchiveObjectCrc32.cpp
--- a/Engine/Source/Runtime/CoreUObject/Private/Serialization/ArchiveObjectCrc32.cpp
+++ b/Engine/Source/Runtime/CoreUObject/Private/Serialization/ArchiveObjectCrc32.cpp
@@ -48,35 +48,28 @@ FArchive& FArchiveObjectCrc32::operator<<(class FName& Name)
 
 FArchive& FArchiveObjectCrc32::operator<<(class UObject*& Object)
 {
-	FArchive& Ar = *this;
-
-	if (!RootObject || !Object || !Object->IsIn(RootObject))
-	{
-		auto UniqueName = GetPathNameSafe(Object);
-		Ar << UniqueName;
-	}
-	else
-	{
-		ObjectsToSerialize.Enqueue(Object);
-	}
-
-	return Ar;
+	FObjectPtr ObjectPtr(Object);
+	return *this << ObjectPtr;
 }
 
 FArchive& FArchiveObjectCrc32::operator<<(FObjectPtr& ObjectPtr)
 {
-	FArchive& Ar = *this;
-	if (!ObjectPtr.IsResolved())
+	if (!ObjectPtr)
+	{
+		FString UniqueName(TEXT("None"));
+		*this << UniqueName;
+	}
+	else if (!RootObject || !ObjectPtr.IsIn(FObjectPtr(const_cast<UObject*>(RootObject))))
 	{
-		auto UniqueName = ObjectPtr.GetPathName();
-		Ar << UniqueName;
+		FString UniqueName = ObjectPtr.GetPathName();
+		*this << UniqueName;
 	}
 	else
 	{
-		UObject* Object = ObjectPtr.Get();
-		Ar << Object;
+		ObjectsToSerialize.Enqueue(ObjectPtr.Get());
 	}
-	return Ar;
+
+	return *this;
 }
 
 uint32 FArchiveObjectCrc32::Crc32(UObject* Object, UObject* Root, uint32 CRC)
diff --git a/Engine/Source/Runtime/CoreUObject/Private/UObject/LinkerSave.cpp b/Engine/Source/Runtime/CoreUObject/Private/UObject/LinkerSave.cpp
--- a/Engine/Source/Runtime/CoreUObject/Private/UObject/LinkerSave.cpp
+++ b/Engine/Source/Runtime/CoreUObject/Private/UObject/LinkerSave.cpp
@@ -417,14 +417,26 @@ FArchive& FLinkerSave::operator<<( FName& InName )
 	return Ar << Save << Number;
 }
 
-FArchive& FLinkerSave::operator<<( UObject*& Obj )
+FArchive& FLinkerSave::operator<<(UObject*& Obj)
+{
+	SerializeObjectPointer(FObjectPtr(Obj));
+	return *this;
+}
+
+FArchive& FLinkerSave::operator<<(FObjectPtr& Value)
+{
+	SerializeObjectPointer(Value);
+	return *this;
+}
+
+void FLinkerSave::SerializeObjectPointer(const FObjectPtr& Obj)
 {
 	FPackageIndex Save;
 	if (Obj)
 	{
-		Save = MapObject(Obj);
+		Save = MapObject(TObjectPtr<const UObject>(Obj));
 	}
-	return *this << Save;
+	*this << Save;
 }
 
 FArchive& FLinkerSave::operator<<(FSoftObjectPath& SoftObjectPath)
diff --git a/Engine/Source/Runtime/CoreUObject/Private/UObject/SavePackage/PackageHarvester.cpp b/Engine/Source/Runtime/CoreUObject/Private/UObject/SavePackage/PackageHarvester.cpp
--- a/Engine/Source/Runtime/CoreUObject/Private/UObject/SavePackage/PackageHarvester.cpp
+++ b/Engine/Source/Runtime/CoreUObject/Private/UObject/SavePackage/PackageHarvester.cpp
@@ -3,7 +3,9 @@
 
 #include "UObject/SavePackage/PackageHarvester.h"
 
+#include "Cooker/CookDependency.h"
 #include "Interfaces/ITargetPlatform.h"
+#include "Misc/PackageAccessTracking.h"
 #include "UObject/OverridableManager.h"
 #include "UObject/SavePackage/SaveContext.h"
 #include "UObject/SavePackage/SavePackageUtilities.h"
@@ -76,29 +78,35 @@ EObjectMark GenerateMarksForObject(const UObject* InObject, FSaveContext& SaveCo
 	return Marks;
 }
 
-bool ConditionallyExcludeObjectForRealm(FSaveContext& SaveContext, TObjectPtr<UObject> Obj, ESaveRealm HarvestingContext)
+bool ConditionallyExcludeObjectForRealm(FSaveContext& SaveContext, TObjectPtr<UObject> ObjPtr,
+	UE::SavePackageUtilities::FObjectStatus& ObjectStatus, ESaveRealm HarvestingContext)
 {
-	if (!Obj || Obj.GetPackage().GetFName() == GLongCoreUObjectPackageName)
+	using namespace UE::SavePackageUtilities;
+
+	if (!ObjPtr || ObjPtr.GetPackage().GetFName() == GLongCoreUObjectPackageName)
 	{
 		// No object or in CoreUObject, don't exclude
 		return false;
 	}
 	FHarvestedRealm& RealmBeingChecked = SaveContext.GetHarvestedRealm(HarvestingContext);
-	if (RealmBeingChecked.IsExcluded(Obj))
+	if (RealmBeingChecked.IsExcluded(ObjPtr))
 	{
 		return true;
 	}
-	if (RealmBeingChecked.IsIncluded(Obj))
+	if (RealmBeingChecked.IsIncluded(ObjPtr))
 	{
 		return false;
 	}
-	if (RealmBeingChecked.IsNotExcluded(Obj))
+	if (RealmBeingChecked.IsNotExcluded(ObjPtr))
 	{
 		return false;
 	}
 
 	const EObjectMark ExcludedObjectMarks = SaveContext.GetExcludedObjectMarks(HarvestingContext);
 	const ITargetPlatform* TargetPlatform = SaveContext.GetTargetPlatform();
+
+	UObject* Obj = SaveContext.ResolveForSave(ObjPtr, ObjectStatus);
+
 	EObjectMark ObjectMarks = GenerateMarksForObject(Obj, SaveContext);
 	if (!!(ObjectMarks & ExcludedObjectMarks))
 	{
@@ -108,16 +116,23 @@ bool ConditionallyExcludeObjectForRealm(FSaveContext& SaveContext, TObjectPtr<UO
 
 	// If the object class is excluded, the object must be excluded too
 	bool bApplyHasNonEditorOnlyReferences = EnumHasAnyFlags(SaveContext.GetEditorOnlyObjectFlags(),
-		UE::SavePackageUtilities::EEditorOnlyObjectFlags::ApplyHasNonEditorOnlyReferences);
+		EEditorOnlyObjectFlags::ApplyHasNonEditorOnlyReferences);
 	bool bIgnoreEditorOnlyClass = bApplyHasNonEditorOnlyReferences && Obj->HasNonEditorOnlyReferences();
-	if (!bIgnoreEditorOnlyClass && ConditionallyExcludeObjectForRealm(SaveContext, Obj->GetClass(), HarvestingContext))
+	if (!bIgnoreEditorOnlyClass)
 	{
-		RealmBeingChecked.AddExcluded(Obj);
-		return true;
+		TObjectPtr<UObject> Class = Obj->GetClass();
+		FObjectStatus& ClassStatus = SaveContext.GetCachedObjectStatus(Class);
+		if (ConditionallyExcludeObjectForRealm(SaveContext, Class, ClassStatus, HarvestingContext))
+		{
+			RealmBeingChecked.AddExcluded(Obj);
+			return true;
+		}
 	}
 
 	// If the object outer is excluded, the object must be excluded too
-	if (ConditionallyExcludeObjectForRealm(SaveContext, Obj.GetOuter(), HarvestingContext))
+	TObjectPtr<UObject> Outer = Obj->GetOuter();
+	FObjectStatus& OuterStatus = SaveContext.GetCachedObjectStatus(Outer);
+	if (ConditionallyExcludeObjectForRealm(SaveContext, Outer, OuterStatus, HarvestingContext))
 	{
 		RealmBeingChecked.AddExcluded(Obj);
 		return true;
@@ -130,19 +145,27 @@ bool ConditionallyExcludeObjectForRealm(FSaveContext& SaveContext, TObjectPtr<UO
 		if (ThisStruct)
 		{
 			UObject* SuperStruct = ThisStruct->GetSuperStruct();
-			if (SuperStruct && ConditionallyExcludeObjectForRealm(SaveContext, SuperStruct, HarvestingContext))
+			if (SuperStruct)
 			{
-				RealmBeingChecked.AddExcluded(Obj);
-				return true;
+				FObjectStatus& SuperStructStatus = SaveContext.GetCachedObjectStatus(SuperStruct);
+				if (ConditionallyExcludeObjectForRealm(SaveContext, SuperStruct, SuperStructStatus, HarvestingContext))
+				{
+					RealmBeingChecked.AddExcluded(Obj);
+					return true;
+				}
 			}
 		}
 
 		// Check archetype, this may not have been covered in the case of components
 		UObject* Archetype = Obj->GetArchetype();
-		if (Archetype && ConditionallyExcludeObjectForRealm(SaveContext, Archetype, HarvestingContext))
+		if (Archetype)
 		{
-			RealmBeingChecked.AddExcluded(Obj);
-			return true;
+			FObjectStatus& ArchetypeStatus = SaveContext.GetCachedObjectStatus(Archetype);
+			if (ConditionallyExcludeObjectForRealm(SaveContext, Archetype, ArchetypeStatus, HarvestingContext))
+			{
+				RealmBeingChecked.AddExcluded(Obj);
+				return true;
+			}
 		}
 	}
 
@@ -327,12 +350,13 @@ FPackageHarvester::FHarvestScope FPackageHarvester::EnterNewExportOnlyScope(UObj
 	return Scope;
 }
 
-FPackageHarvester::FHarvestScope FPackageHarvester::EnterNotExcludedScope(TObjectPtr<UObject> Object)
+FPackageHarvester::FHarvestScope FPackageHarvester::EnterNotExcludedScope(TObjectPtr<UObject> Object,
+	UE::SavePackageUtilities::FObjectStatus& ObjectStatus)
 {
 	FHarvestScope Scope(*this);
-	CurrentExportHarvestingRealms.RemoveAllSwap([this, Object](ESaveRealm HarvestingRealm)
+	CurrentExportHarvestingRealms.RemoveAllSwap([this, Object, &ObjectStatus](ESaveRealm HarvestingRealm)
 		{
-			return ConditionallyExcludeObjectForRealm(SaveContext, Object, HarvestingRealm);
+			return ConditionallyExcludeObjectForRealm(SaveContext, Object, ObjectStatus, HarvestingRealm);
 		}, EAllowShrinking::No);
 	return Scope;
 }
@@ -375,7 +399,7 @@ void FPackageHarvester::GetPreviouslyIncludedRealms(TObjectPtr<UObject> Object,
 	}
 }
 
-bool FPackageHarvester::IsObjNative(TObjectPtr<UObject> InObj)
+bool FPackageHarvester::IsObjNative(UObject* InObj)
 {
 	bool bIsNative = InObj->IsNative();
 	UObject* Outer = InObj->GetOuter();
@@ -571,7 +595,7 @@ void FPackageHarvester::ProcessExport(const FExportWithContext& InProcessContext
 	}
 }
 
-void FPackageHarvester::TryHarvestExport(UObject* InObject)
+void FPackageHarvester::TryHarvestExport(UObject* InObject, UE::SavePackageUtilities::FObjectStatus& ObjectStatus)
 {
 	check(!HasAnyExportHarvestingRealms());
 
@@ -600,13 +624,15 @@ void FPackageHarvester::TryHarvestExport(UObject* InObject)
 	}
 	CurrentExportHarvestingRealms = HarvestingRealms;
 
-	TryHarvestExportInternal(InObject);
+	TryHarvestExportInternal(InObject, ObjectStatus);
 }
 
-void FPackageHarvester::TryHarvestExportInternal(UObject* InObject)
+void FPackageHarvester::TryHarvestExportInternal(UObject* InObject, UE::SavePackageUtilities::FObjectStatus& ObjectStatus)
 {
+	using namespace UE::SavePackageUtilities;
+
 	// Those should have been already validated
-	check(InObject && InObject->IsInPackage(SaveContext.GetPackage()));
+	check(InObject && ObjectStatus.IsInSavePackage(InObject, SaveContext.GetPackage()));
 
 	// Switch to the Optional Realm if the Object is optional
 	FHarvestScope OptionalObjectScope = EnterConditionalOptionalObjectScope(InObject);
@@ -631,14 +657,14 @@ void FPackageHarvester::TryHarvestExportInternal(UObject* InObject)
 	}
 
 	// Check whether the object is unsaveable and skip adding it as an export to any realm if so
-	SaveContext.GetCachedObjectStatus(InObject).bAttemptedExport = true;
-	if (SaveContext.IsUnsaveable(InObject))
+	ObjectStatus.bAttemptedExport = true;
+	if (SaveContext.IsUnsaveable(InObject, ObjectStatus))
 	{
 		return;
 	}
 
 	// Filter out any realms in which the export is excluded
-	FHarvestScope NotExcludedScope = EnterNotExcludedScope(InObject);
+	FHarvestScope NotExcludedScope = EnterNotExcludedScope(InObject, ObjectStatus);
 	if (NotExcludedScope.IsEmpty())
 	{
 		return;
@@ -659,29 +685,34 @@ void FPackageHarvester::TryHarvestExportInternal(UObject* InObject)
 			}
 
 			UObject* ObjValue = ObjProp->GetObjectPropertyValue(ObjPropIt.Value());
-			if (!ObjValue
-				|| !ObjValue->IsInOuter(InObject)
-				|| !ObjValue->IsInPackage(SaveContext.GetPackage()))
+			if (!ObjValue || !ObjValue->IsInOuter(InObject))
+			{
+				continue;
+			}
+			FObjectStatus& ObjValueStatus = SaveContext.GetCachedObjectStatus(ObjValue);
+			if (!ObjValueStatus.IsInSavePackage(ObjValue, SaveContext.GetPackage()))
 			{
 				continue;
 			}
 
-			TryHarvestExportInternal(ObjValue);
+			TryHarvestExportInternal(ObjValue, ObjValueStatus);
 		}
 	}
 }
 
-void FPackageHarvester::TryHarvestImport(TObjectPtr<UObject> InObject)
+void FPackageHarvester::TryHarvestImport(TObjectPtr<UObject> InObject, UE::SavePackageUtilities::FObjectStatus& ObjectStatus)
 {
+	UE_COOK_RESULTPROJECTION_SCOPED(UE::Cook::ResultProjection::PackageAndClass);
+
 	// Those should have been already validated
-	check(InObject && !InObject->IsInPackage(SaveContext.GetPackage()));
+	check(InObject && !ObjectStatus.IsInSavePackage(InObject, SaveContext.GetPackage()))
 	if (InObject == nullptr)
 	{
 		return;
 	}
 
 	// Do not add unsaveable imports to any realm
-	if (SaveContext.IsUnsaveable(InObject))
+	if (SaveContext.IsUnsaveable(InObject, ObjectStatus))
 	{
 		return;
 	}
@@ -696,25 +727,27 @@ void FPackageHarvester::TryHarvestImport(TObjectPtr<UObject> InObject)
 	}
 
 	// Filter out any realms in which the import is excluded
-	FHarvestScope NotExcludedScope = EnterNotExcludedScope(InObject);
+	FHarvestScope NotExcludedScope = EnterNotExcludedScope(InObject, ObjectStatus);
 	if (NotExcludedScope.IsEmpty())
 	{
 		return;
 	}
 
 	HarvestImport(InObject);
-	ProcessImport(InObject);
+	ProcessImport(InObject, ObjectStatus);
 }
 
-void FPackageHarvester::ProcessImport(TObjectPtr<UObject> InObject)
+void FPackageHarvester::ProcessImport(TObjectPtr<UObject> InObjPtr,
+	UE::SavePackageUtilities::FObjectStatus& ObjectStatus)
 {
 	++CurrentExportDependencies.ProcessImportDepth;
 	ON_SCOPE_EXIT{ --CurrentExportDependencies.ProcessImportDepth; };
 
+	UObject* InObject = SaveContext.ResolveForSave(InObjPtr, ObjectStatus);
 	bool bIsNative = IsObjNative(InObject);
-	TObjectPtr<UObject> ObjOuter = InObject.GetOuter();
-	UClass* ObjClass = InObject.GetClass();
-	FName ObjName = InObject.GetFName();
+	TObjectPtr<UObject> ObjOuter = InObject->GetOuter();
+	UClass* ObjClass = InObject->GetClass();
+	FName ObjName = InObject->GetFName();
 	if (SaveContext.IsCooking())
 	{
 		// The ignore dependencies check is necessary not to have infinite recursive calls
@@ -767,9 +800,9 @@ void FPackageHarvester::ProcessImport(TObjectPtr<UObject> InObject)
 		// "Missing import package name...". Log it here instead with an explanation for why it's not included.
 		// It's possible the package will be marked as editoronly by not including it in the game realm, so don't check
 		// all CurrentExportHarvestingRealms, just check the SaveContext's HarvestingRealm.
-		if (SaveContext.IsIncluded(InObject) && !SaveContext.IsIncluded(Package))
+		if (SaveContext.IsIncluded(InObjPtr) && !SaveContext.IsIncluded(Package))
 		{
-			SaveContext.RecordIllegalReference(CurrentExportDependencies.CurrentExport, InObject,
+			SaveContext.RecordIllegalReference(CurrentExportDependencies.CurrentExport, InObjPtr,
 				EIllegalRefReason::ExternalPackage, Package->GetName());
 		}
 	}
@@ -815,12 +848,19 @@ void FPackageHarvester::MarkSearchableName(const TObjectPtr<const UObject>& Type
 }
 
 FArchive& FPackageHarvester::operator<<(UObject*& Obj)
+{
+	FObjectPtr ObjectPtr(Obj);
+	return *this << ObjectPtr;
+}
+
+FArchive& FPackageHarvester::operator<<(FObjectPtr& ObjBase)
 {
 	// if the object is null, skip the harvest
-	if (!Obj)
+	if (!ObjBase)
 	{
 		return *this;
 	}
+	TObjectPtr<UObject> Obj(ObjBase);
 
 #if WITH_EDITORONLY_DATA
 	FHarvestScope EditorOnlyScope = EnterConditionalEditorOnlyScope(IsEditorOnlyPropertyOnTheStack());
@@ -833,7 +873,7 @@ FArchive& FPackageHarvester::operator<<(UObject*& Obj)
 		return *this;
 	}
 
-	// if the package we are saving is referenced, just harvest its name
+	// if the referenced object is the package we are saving, just harvest its name
 	if (Obj == SaveContext.GetPackage())
 	{
 		HarvestPackageHeaderName(Obj->GetFName());
@@ -848,18 +888,19 @@ FArchive& FPackageHarvester::operator<<(UObject*& Obj)
 	FExportingRealmsArray NewRealms;
 	GetPreviouslyIncludedRealms(Obj, PreviouslyIncludedRealms, NewRealms);
 
+	UE::SavePackageUtilities::FObjectStatus& ObjectStatus = SaveContext.GetCachedObjectStatus(Obj);
 	if (!NewRealms.IsEmpty())
 	{
 		// if the object is in the save context package, try to tag it as export for the new realms.
 		// Otherwise visit the import for the new realms.
 		FHarvestScope NewRealmsScope = EnterRealmsArrayScope(NewRealms);
-		if (Obj->IsInPackage(SaveContext.GetPackage()))
+		if (ObjectStatus.IsInSavePackage(Obj, SaveContext.GetPackage()))
 		{
-			TryHarvestExportInternal(Obj);
+			TryHarvestExportInternal(Obj, ObjectStatus);
 		}
 		else
 		{
-			TryHarvestImport(Obj);
+			TryHarvestImport(Obj, ObjectStatus);
 		}
 	}
 
@@ -868,8 +909,10 @@ FArchive& FPackageHarvester::operator<<(UObject*& Obj)
 	FHarvestScope ObjIncludedScope = EnterIncludedScope(Obj);
 	if (!ObjIncludedScope.IsEmpty())
 	{
+		UObject* ResolvedObject = SaveContext.ResolveForSave(Obj, ObjectStatus);
+
 		// Add a dependency from the current referencing export to Obj
-		HarvestDependency(Obj, IsObjNative(Obj));
+		HarvestDependency(Obj, IsObjNative(ResolvedObject));
 
 		// Validate the reference is allowed (this validation is a side-effect of EnterConditionalOptionalObjectScope)
 		EnterConditionalOptionalObjectScope(Obj);
@@ -1149,7 +1192,8 @@ ESaveableStatus FPackageHarvester::GetSaveableStatusForRealm(UObject* Obj, ESave
 		return ESaveableStatus::Success;
 	}
 
-	UE::SavePackageUtilities::FObjectStatus& ObjectStatus = SaveContext.UpdateSaveableStatus(Obj);
+	UE::SavePackageUtilities::FObjectStatus& ObjectStatus = SaveContext.GetCachedObjectStatus(Obj);
+	SaveContext.UpdateSaveableStatus(Obj, ObjectStatus);
 	if (ObjectStatus.SaveableStatus != ESaveableStatus::Success)
 	{
 		if (ObjectStatus.SaveableStatus == ESaveableStatus::OuterUnsaveable)
diff --git a/Engine/Source/Runtime/CoreUObject/Private/UObject/SavePackage/PackageHarvester.h b/Engine/Source/Runtime/CoreUObject/Private/UObject/SavePackage/PackageHarvester.h
--- a/Engine/Source/Runtime/CoreUObject/Private/UObject/SavePackage/PackageHarvester.h
+++ b/Engine/Source/Runtime/CoreUObject/Private/UObject/SavePackage/PackageHarvester.h
@@ -14,6 +14,7 @@ enum class ESaveRealm : uint32;
 enum class EIllegalRefReason : uint8;
 class FSaveContext;
 class UObject;
+namespace UE::SavePackageUtilities { struct FObjectStatus; }
 struct FWeakObjectPtr;
 struct FLazyObjectPtr;
 struct FSoftObjectPath;
@@ -61,7 +62,7 @@ public:
 
 	FExportWithContext PopExportToProcess();
 
-	void TryHarvestExport(UObject* InObject);
+	void TryHarvestExport(UObject* InObject, UE::SavePackageUtilities::FObjectStatus& ObjectStatus);
 	/** Add objects/names/others that are referenced when an object is being saved */
 	void ProcessExport(const FExportWithContext& InProcessContext);
 	FHarvestScope EnterRootReferencesScope();
@@ -82,11 +83,12 @@ public:
 	virtual FArchive& operator<<(FLazyObjectPtr& LazyObjectPtr) override;
 	virtual FArchive& operator<<(FSoftObjectPath& Value) override;
 	virtual FArchive& operator<<(FName& Name) override;
+	virtual FArchive& operator<<(FObjectPtr& Value) override;
 
 	virtual bool ShouldSkipProperty(const FProperty* InProperty) const override;
 
 	/** Return whether the Object and all its Outers are native. */
-	static bool IsObjNative(TObjectPtr<UObject> InObj);
+	static bool IsObjNative(UObject* InObj);
 	/** Report whether the object should be put in the optional realm even if harvesting realm is non-optional. */
 	static bool ShouldObjectBeHarvestedInOptionalRealm(TObjectPtr<UObject> InObj, FSaveContext& InSaveContext);
 
@@ -136,12 +138,12 @@ private:
 
 	void ResolveOverrides();
 
-	void TryHarvestExportInternal(UObject* InObject);
+	void TryHarvestExportInternal(UObject* InObject, UE::SavePackageUtilities::FObjectStatus& ObjectStatus);
 	void HarvestExport(UObject* InObject);
-	void TryHarvestImport(TObjectPtr<UObject> InObject);
+	void TryHarvestImport(TObjectPtr<UObject> InObject, UE::SavePackageUtilities::FObjectStatus& ObjectStatus);
 	void HarvestImport(TObjectPtr<UObject> InObject);
 	/** Add objects/names/others that are referenced even when an object is just being imported */
-	void ProcessImport(TObjectPtr<UObject> InObject);
+	void ProcessImport(TObjectPtr<UObject> InObject, UE::SavePackageUtilities::FObjectStatus& ObjectStatus);
 
 	void AppendCurrentExportDependencies();
 	FString GetUnsaveableReason(UObject* Required, ESaveRealm RealmInWhichItIsUnsaveable);
@@ -157,7 +159,8 @@ private:
 	/** Remove CurrentExportHarvestingRealms that already include Export. */
 	FHarvestScope EnterNewExportOnlyScope(UObject* Export);
 	/** Remove CurrentExportHarvestingRealms that return false for ConditionallyExcludeObjectForRealm(Object). */
-	FHarvestScope EnterNotExcludedScope(TObjectPtr<UObject> Object);
+	FHarvestScope EnterNotExcludedScope(TObjectPtr<UObject> Object,
+		UE::SavePackageUtilities::FObjectStatus& ObjectStatus);
 	/** Remove CurrentExportHarvestingRealms that return true for IsExcluded(Object). */
 	FHarvestScope EnterNotPreviouslyExcludedScope(TObjectPtr<UObject> Object);
 	/** Remove CurrentExportHarvestingRealms that return false for IsIncluded(Object). */
diff --git a/Engine/Source/Runtime/CoreUObject/Private/UObject/SavePackage/SaveContext.cpp b/Engine/Source/Runtime/CoreUObject/Private/UObject/SavePackage/SaveContext.cpp
--- a/Engine/Source/Runtime/CoreUObject/Private/UObject/SavePackage/SaveContext.cpp
+++ b/Engine/Source/Runtime/CoreUObject/Private/UObject/SavePackage/SaveContext.cpp
@@ -7,6 +7,7 @@
 #include "Cooker/CookDependency.h"
 #include "Misc/ConfigCacheIni.h"
 #include "Serialization/PackageWriter.h"
+#include "UObject/ObjectSaveContext.h"
 #include "UObject/UObjectGlobals.h"
 
 TArray<ESaveRealm> FSaveContext::GetHarvestedRealmsToSave()
@@ -27,13 +28,14 @@ TArray<ESaveRealm> FSaveContext::GetHarvestedRealmsToSave()
 	return HarvestedContextsToSave;
 }
 
-bool FSaveContext::IsUnsaveable(TObjectPtr<UObject> InObject, bool bEmitWarning)
+bool FSaveContext::IsUnsaveable(TObjectPtr<UObject> InObject, UE::SavePackageUtilities::FObjectStatus& ObjectStatus,
+	bool bEmitWarning)
 {
 	if (!InObject)
 	{
 		return false;
 	}
-	UE::SavePackageUtilities::FObjectStatus& ObjectStatus = UpdateSaveableStatus(InObject);
+	UpdateSaveableStatus(InObject, ObjectStatus);
 	check(ObjectStatus.bSaveableStatusValid);
 
 	if (bEmitWarning && ObjectStatus.SaveableStatus != ESaveableStatus::Success)
@@ -61,76 +63,65 @@ bool FSaveContext::IsUnsaveable(TObjectPtr<UObject> InObject, bool bEmitWarning)
 	return ObjectStatus.SaveableStatus != ESaveableStatus::Success;
 }
 
-UE::SavePackageUtilities::FObjectStatus& FSaveContext::UpdateSaveableStatus(TObjectPtr<UObject> InObject)
+void FSaveContext::UpdateSaveableStatus(TObjectPtr<UObject> InObject, UE::SavePackageUtilities::FObjectStatus& ObjectStatus)
 {
-	UE::SavePackageUtilities::FObjectStatus* ObjectStatus = &ObjectStatusCache.FindOrAdd(InObject);
-	if (ObjectStatus->bSaveableStatusValid)
+	if (ObjectStatus.bSaveableStatusValid)
 	{
-		return *ObjectStatus;
+		return;
 	}
 
-	ObjectStatus->bSaveableStatusValid = true;
-	ObjectStatus->SaveableStatus = ESaveableStatus::Success;
+	ObjectStatus.bSaveableStatusValid = true;
+	ObjectStatus.SaveableStatus = ESaveableStatus::Success;
 
-	ESaveableStatus StatusNoOuter = GetSaveableStatusNoOuter(InObject, *ObjectStatus);
+	ESaveableStatus StatusNoOuter = GetSaveableStatusNoOuter(InObject, ObjectStatus);
 	if (StatusNoOuter != ESaveableStatus::Success)
 	{
 		check(StatusNoOuter != ESaveableStatus::OuterUnsaveable &&
 			StatusNoOuter != ESaveableStatus::ClassUnsaveable);
-		ObjectStatus->SaveableStatus = StatusNoOuter;
-		return *ObjectStatus;
-	}
-
-	if (!InObject.IsResolved())
-	{
-		// We do not test the saveability of the outer of unresolved objects because we cannot get their
-		// outer without resolving them.
-		return *ObjectStatus;
+		ObjectStatus.SaveableStatus = StatusNoOuter;
+		return;
 	}
 
-	UObject* Outer = InObject->GetOuter();
+	TObjectPtr<UObject> Outer = InObject.GetOuter();
 	if (Outer)
 	{
-		UE::SavePackageUtilities::FObjectStatus& OuterStatus = UpdateSaveableStatus(Outer);
-		// Calling UpdateSaveableStatus on the outer might have modified ObjectStatusCache, so
-		// look up our pointer to ObjectStatus again.
-		ObjectStatus = ObjectStatusCache.Find(InObject);
-		check(ObjectStatus);
-
-		ObjectStatus->bSaveableStatusValid = true;
-		ObjectStatus->SaveableStatus = ESaveableStatus::Success;
+		UE::SavePackageUtilities::FObjectStatus& OuterStatus = GetCachedObjectStatus(Outer);
+		UpdateSaveableStatus(Outer, OuterStatus);
+
+		ObjectStatus.bSaveableStatusValid = true;
+		ObjectStatus.SaveableStatus = ESaveableStatus::Success;
 		if (OuterStatus.SaveableStatus != ESaveableStatus::Success)
 		{
-			ObjectStatus->SaveableStatus = ESaveableStatus::OuterUnsaveable;
+			ObjectStatus.SaveableStatus = ESaveableStatus::OuterUnsaveable;
 			if (OuterStatus.SaveableStatus == ESaveableStatus::OuterUnsaveable)
 			{
 				check(OuterStatus.SaveableStatusCulprit);
 				check(OuterStatus.SaveableStatusCulpritStatus != ESaveableStatus::Success);
-				ObjectStatus->SaveableStatusCulprit = OuterStatus.SaveableStatusCulprit;
-				ObjectStatus->SaveableStatusCulpritStatus = OuterStatus.SaveableStatusCulpritStatus;
+				ObjectStatus.SaveableStatusCulprit = OuterStatus.SaveableStatusCulprit;
+				ObjectStatus.SaveableStatusCulpritStatus = OuterStatus.SaveableStatusCulpritStatus;
 			}
 			else
 			{
-				ObjectStatus->SaveableStatusCulprit = Outer;
-				ObjectStatus->SaveableStatusCulpritStatus = OuterStatus.SaveableStatus;
+				ObjectStatus.SaveableStatusCulprit = Outer;
+				ObjectStatus.SaveableStatusCulpritStatus = OuterStatus.SaveableStatus;
 			}
 		}
 	}
-
-	return *ObjectStatus;
 }
 
-ESaveableStatus FSaveContext::GetSaveableStatusNoOuter(TObjectPtr<UObject> Obj,
-	UE::SavePackageUtilities::FObjectStatus& ObjectStatus) const
+ESaveableStatus FSaveContext::GetSaveableStatusNoOuter(TObjectPtr<UObject> ObjPtr,
+	UE::SavePackageUtilities::FObjectStatus& ObjectStatus)
 {
+	UObject* Obj = ResolveForSave(ObjPtr, ObjectStatus);
+
 	// pending kill objects are unsaveable
-	if (Obj.IsResolved() && !IsValidChecked(Obj))
+	if (!IsValidChecked(Obj))
 	{
 		return ESaveableStatus::PendingKill;
 	}
 
 	// transient objects are unsaveable if non-native
-	if (Obj.IsResolved() && !Obj->IsNative())
+	if (!Obj->IsNative())
 	{
 		if (ObjectStatus.HasTransientFlag(Obj))
 		{
@@ -142,7 +133,7 @@ ESaveableStatus FSaveContext::GetSaveableStatusNoOuter(TObjectPtr<UObject> Obj,
 		}
 	}
 
-	UClass* Class = Obj.GetClass();
+	UClass* Class = Obj->GetClass();
 	// if the object class is abstract, has been marked as deprecated, there is a newer version that exist, or the class is marked transient, then the object is unsaveable
 	// @note: Although object instances of a transient class should definitely be unsaveable, it results in discrepancies with the old save algorithm and currently load problems
 	if (Class->HasAnyClassFlags(CLASS_Abstract | CLASS_Deprecated | CLASS_NewerVersionExists /*| CLASS_Transient*/)
@@ -158,19 +149,54 @@ ESaveableStatus FSaveContext::GetSaveableStatusNoOuter(TObjectPtr<UObject> Obj,
 	return ESaveableStatus::Success;
 }
 
+UObject* FSaveContext::ResolveForSave(TObjectPtr<UObject> InObject, UE::SavePackageUtilities::FObjectStatus& Status)
+{
+	using namespace UE::Cook;
+
+	if (!Status.IsInSavePackage(InObject, GetPackage()))
+	{
+		UObject* Object = nullptr;
+		{
+			// Disable the automated tracking of the resolve with the cooker.
+			UE_COOK_RESULTPROJECTION_SCOPED(UE::Cook::ResultProjection::None);
+			Object = InObject.Get();
+		}
+		if (!Status.bDeclaredSaveBuildDependency && Object)
+		{
+			// Declare the build dependencies necessary for our use of the object.
+#if WITH_EDITOR
+			ObjectSaveContext.CookBuildDependencies.Add(FCookDependency::Package(Object->GetPackage()->GetFName()));
+			ObjectSaveContext.CookBuildDependencies.Add(FCookDependency::NativeClass(Object->GetClass()));
+#endif
+			Status.bDeclaredSaveBuildDependency = true;
+		}
+		return Object;
+	}
+	else
+	{
+		// No need to consider build dependencies when resolving an export
+		return InObject.Get();
+	}
+}
+
 bool FSaveContext::IsTransient(TObjectPtr<UObject> InObject)
+{
+	return IsTransient(InObject, GetCachedObjectStatus(InObject));
+}
+
+bool FSaveContext::IsTransient(TObjectPtr<UObject> InObject, UE::SavePackageUtilities::FObjectStatus& Status)
 {
 	if (!InObject)
 	{
 		return false;
 	}
 
-	if (InObject->HasAnyFlags(RF_Transient))
+	UObject* ObjectPtr = ResolveForSave(InObject, Status);
+	if (ObjectPtr->HasAnyFlags(RF_Transient))
 	{
 		return true;
 	}
 
-	UE::SavePackageUtilities::FObjectStatus& Status = GetCachedObjectStatus(InObject);
 	if (Status.bSaveOverrideForcedTransient)
 	{
 		return true;
diff --git a/Engine/Source/Runtime/CoreUObject/Private/UObject/SavePackage/SaveContext.h b/Engine/Source/Runtime/CoreUObject/Private/UObject/SavePackage/SaveContext.h
--- a/Engine/Source/Runtime/CoreUObject/Private/UObject/SavePackage/SaveContext.h
+++ b/Engine/Source/Runtime/CoreUObject/Private/UObject/SavePackage/SaveContext.h
@@ -33,7 +33,7 @@
  */
 struct FTaggedExport
 {
-	UObject* Obj;
+	TObjectPtr<UObject> Obj;
 	uint32 bNotAlwaysLoadedForEditorGame : 1;
 	/**
 	 * Indicate that this export should have a public hash even if it isn't marked as RF_Public
@@ -53,7 +53,7 @@ struct FTaggedExport
 		, bFromOptionalReference(false)
 	{}
 
-	FTaggedExport(UObject* InObj, bool bInNotAlwaysLoadedForEditorGame = true, bool bInFromOptionalReference = false)
+	FTaggedExport(TObjectPtr<UObject> InObj, bool bInNotAlwaysLoadedForEditorGame = true, bool bInFromOptionalReference = false)
 		: Obj(InObj)
 		, bNotAlwaysLoadedForEditorGame(bInNotAlwaysLoadedForEditorGame)
 		, bGeneratePublicHash(false)
@@ -128,6 +128,8 @@ struct FObjectStatus
 	FObjectStatus()
 		: bSaveOverrideForcedTransient(false)
 		, bSaveableStatusValid(false), bEditorOnlyValid(false), bEditorOnly(false), bAttemptedExport(false)
+		, bIsInSavePackageValid(false), bIsInSavePackage(false)
+		, bDeclaredSaveBuildDependency(false)
 	{
 	}
 	FObjectStatus(const FObjectStatus&) = default;
@@ -145,6 +147,23 @@ struct FObjectStatus
 		bSaveableStatusValid = false;
 		SaveableStatusCulprit = nullptr;
 	}
+	void SetIsInSavePackage(bool bInIsInSavePackage)
+	{
+		bIsInSavePackageValid = true;
+		bIsInSavePackage = bInIsInSavePackage;
+	}
+	TOptional<bool> IsInSavePackage() const
+	{
+		return bIsInSavePackageValid ? TOptional<bool>(bIsInSavePackage) : TOptional<bool>();
+	}
+	bool IsInSavePackage(TObjectPtr<UObject> InObject, UPackage* SavePackage)
+	{
+		if (!bIsInSavePackageValid)
+		{
+			SetIsInSavePackage(InObject.GetPackage() == SavePackage);
+		}
+		return bIsInSavePackage;
+	}
 
 	UObject* SaveableStatusCulprit = nullptr;
 	ESaveableStatus SaveableStatus = ESaveableStatus::Success;
@@ -154,6 +173,9 @@ struct FObjectStatus
 	bool bEditorOnlyValid : 1;
 	bool bEditorOnly : 1;
 	bool bAttemptedExport : 1;
+	bool bIsInSavePackageValid : 1;
+	bool bIsInSavePackage : 1;
+	bool bDeclaredSaveBuildDependency : 1;
 };
 
 }
@@ -205,14 +227,14 @@ struct FHarvestedRealm
 		return Imports.Contains(InObject);
 	}
 
-	bool IsExport(UObject* InObject) const
+	bool IsExport(TObjectPtr<UObject> InObject) const
 	{
 		return Exports.Contains(InObject);
 	}
 
 	bool IsIncluded(TObjectPtr<UObject> InObject) const
 	{
-		return IsImport(InObject) || (InObject.IsResolved() && IsExport(InObject));
+		return IsImport(InObject) || IsExport(InObject);
 	}
 
 	/**
@@ -811,12 +833,28 @@ public:
 	/** Returns which save context should be saved. */
 	TArray<ESaveRealm> GetHarvestedRealmsToSave();
 
+	/**
+	 * Resolve a TObjectPtr into a raw UObject* for use during SavePackage. If the object is an import,
+	 * this function avoids the normal automatic behavior of marking the object's package as a generic build
+	 * dependency, and instead manually records a dependency on the PackageAndClass cook result projection,
+	 * which are the only dependencies necessary for generic saving of imports.
+	 */
+	UObject* ResolveForSave(TObjectPtr<UObject> InObject, UE::SavePackageUtilities::FObjectStatus& Status);
+
 	bool IsTransient(TObjectPtr<UObject> InObject);
-	bool IsUnsaveable(TObjectPtr<UObject> InObject, bool bEmitWarning = true);
-	UE::SavePackageUtilities::FObjectStatus& UpdateSaveableStatus(TObjectPtr<UObject> InObject);
+	bool IsTransient(TObjectPtr<UObject> InObject, UE::SavePackageUtilities::FObjectStatus& ObjectStatus);
+	bool IsUnsaveable(TObjectPtr<UObject> InObject, UE::SavePackageUtilities::FObjectStatus& ObjectStatus,
+		bool bEmitWarning = true);
+	void UpdateSaveableStatus(TObjectPtr<UObject> InObject, UE::SavePackageUtilities::FObjectStatus& ObjectStatus);
 	UE::SavePackageUtilities::FObjectStatus& GetCachedObjectStatus(TObjectPtr<UObject> InObject)
 	{
-		return ObjectStatusCache.FindOrAdd(InObject);
+		using namespace UE::SavePackageUtilities;
+		TUniquePtr<FObjectStatus>& ObjectStatus = ObjectStatusCache.FindOrAdd(InObject);
+		if (!ObjectStatus)
+		{
+			ObjectStatus = MakeUnique<FObjectStatus>();
+		}
+		return *ObjectStatus;
 	}
 
 	void RecordIllegalReference(UObject* InFrom, UObject* InTo, EIllegalRefReason InReason, FString&& InOptionalReasonText = FString())
@@ -1154,9 +1192,9 @@ public:
 
 	void ClearSaveableCache()
 	{
-		for (TPair<TObjectPtr<UObject>, UE::SavePackageUtilities::FObjectStatus>& Pair : ObjectStatusCache)
+		for (TPair<TObjectPtr<UObject>, TUniquePtr<UE::SavePackageUtilities::FObjectStatus>>& Pair : ObjectStatusCache)
 		{
-			Pair.Value.ClearSaveableStatus();
+			Pair.Value->ClearSaveableStatus();
 		}
 	}
 
@@ -1166,7 +1204,7 @@ public:
 		using namespace UE::SavePackageUtilities;
 		return [this](const UObject* Obj)
 			{
-				FObjectStatus& Status = ObjectStatusCache.FindOrAdd(const_cast<UObject*>(Obj));
+				FObjectStatus& Status = GetCachedObjectStatus(const_cast<UObject*>(Obj));
 				return !Status.bEditorOnlyValid ? EEditorOnlyObjectResult::Uninitialized :
 					Status.bEditorOnly ? EEditorOnlyObjectResult::EditorOnly :
 					EEditorOnlyObjectResult::NonEditorOnly;
@@ -1177,7 +1215,7 @@ public:
 		using namespace UE::SavePackageUtilities;
 		return [this](const UObject* Obj, bool bEditorOnly)
 			{
-				FObjectStatus& Status = ObjectStatusCache.FindOrAdd(const_cast<UObject*>(Obj));
+				FObjectStatus& Status = GetCachedObjectStatus(const_cast<UObject*>(Obj));
 				Status.bEditorOnlyValid = true;
 				Status.bEditorOnly = bEditorOnly;
 			};
@@ -1205,7 +1243,7 @@ private:
 	// Create the harvesting contexts and automatic optional context gathering options
 	void SetupHarvestingRealms();
 	ESaveableStatus GetSaveableStatusNoOuter(TObjectPtr<UObject> Obj,
-		UE::SavePackageUtilities::FObjectStatus& ObjectStatus) const;
+		UE::SavePackageUtilities::FObjectStatus& ObjectStatus);
 	static EObjectMark GetExcludedObjectMarksForGameRealm(const ITargetPlatform* TargetPlatform);
 
 	friend class FPackageHarvester;
@@ -1266,7 +1304,7 @@ private:
 	TMap<UObject*, TSet<FProperty*>> TransientPropertyOverrides;
 
 	// Cache of FObjectStatus for every object encountered during the save
-	TMap<TObjectPtr<UObject>, UE::SavePackageUtilities::FObjectStatus> ObjectStatusCache;
+	TMap<TObjectPtr<UObject>, TUniquePtr<UE::SavePackageUtilities::FObjectStatus>> ObjectStatusCache;
 
 	// List of package build dependencies reported from PreSave or Serialize functions
 	TArray<FName> PackageBuildDependencies;
diff --git a/Engine/Source/Runtime/CoreUObject/Private/UObject/SavePackage2.cpp b/Engine/Source/Runtime/CoreUObject/Private/UObject/SavePackage2.cpp
--- a/Engine/Source/Runtime/CoreUObject/Private/UObject/SavePackage2.cpp
+++ b/Engine/Source/Runtime/CoreUObject/Private/UObject/SavePackage2.cpp
@@ -9,6 +9,7 @@
 #include "AssetRegistry/AssetData.h"
 #include "Async/ParallelFor.h"
 #include "Blueprint/BlueprintSupport.h"
+#include "Cooker/CookDependency.h"
 #include "HAL/FileManager.h"
 #include "Internationalization/TextPackageNamespaceUtil.h"
 #include "Internationalization/PackageLocalizationManager.h"
@@ -274,7 +275,7 @@ ESavePackageResult RoutePresave(FSaveContext& SaveContext)
 	{
 		// Do not emit unsaveable warning while routing presave. 
 		// this is to prevent warning on objects which won't be harvested later since they are unreferenced
-		if (!SaveContext.IsUnsaveable(Object, false/*bEmitWarning*/))
+		if (!SaveContext.IsUnsaveable(Object, SaveContext.GetCachedObjectStatus(Object), false/*bEmitWarning*/))
 		{
 			FObjectSaveContextData& ObjectSaveContext = SaveContext.GetObjectSaveContext();
 			ObjectSaveContext.Object = Object;
@@ -338,18 +339,19 @@ ESavePackageResult HarvestPackage(FSaveContext& SaveContext)
 	{
 		// ForEachObjectWithPackage will filter out Transient objects but not objects with the bForceTransient Save Override
 		// Perform one more check to cover this case
-		if (SaveContext.IsTransient(InRoot))
+		UE::SavePackageUtilities::FObjectStatus& ObjectStatus = SaveContext.GetCachedObjectStatus(InRoot);
+		if (SaveContext.IsTransient(InRoot, ObjectStatus))
 		{
 			return;
 		}
 
-		Harvester.TryHarvestExport(InRoot);
+		Harvester.TryHarvestExport(InRoot, ObjectStatus);
 		// if we are automatically generating an optional package output, re-harvest objects with that realm as the default
 		if (SaveContext.IsSaveAutoOptional())
 		{
 			//@todo: FH: the top level might have to be altered to use the editor one when auto harvesting optional data
 			FSaveContext::FSetSaveRealmToSaveScope RealmScope(SaveContext, ESaveRealm::Optional);
-			Harvester.TryHarvestExport(InRoot);
+			Harvester.TryHarvestExport(InRoot, ObjectStatus);
 		}
 	};
 
@@ -894,6 +896,7 @@ ESavePackageResult ValidateImports(FSaveContext& SaveContext)
 	TArray<UObject*> ObjectsInOtherMaps;
 	const TSet<TObjectPtr<UObject>>& Imports = SaveContext.GetImports();
 	const TSet<TObjectPtr<UObject>>& DirectImports = SaveContext.GetDirectImports();
+	UE_COOK_RESULTPROJECTION_SCOPED(UE::Cook::ResultProjection::PackageAndClass);
 	for (const TObjectPtr<UObject>& Import : Imports)
 	{
 		TObjectPtr<UPackage> ImportPackage = Import.GetPackage();
@@ -1304,6 +1307,7 @@ ESavePackageResult BuildLinker(FSaveContext& SaveContext)
 	// Build ImportMap
 	{
 		SCOPED_SAVETIMER(UPackage_Save_BuildImportMap);
+		UE_COOK_RESULTPROJECTION_SCOPED(UE::Cook::ResultProjection::PackageAndClass);
 
 		for (TObjectPtr<UObject> Import : SaveContext.GetImports())
 		{
@@ -1552,7 +1556,7 @@ void SavePreloadDependencies(FStructuredArchive::FRecord& StructuredArchiveRoot,
 					AddTo.Add(Index);
 					return;
 				}
-				else if (!SaveContext.IsUnsaveable(ToTest))
+				else if (!SaveContext.IsUnsaveable(ToTest, SaveContext.GetCachedObjectStatus(ToTest)))
 				{
 					UE_CLOG(Outermost->HasAnyPackageFlags(PKG_CompiledIn), LogSavePackage, Verbose, TEXT("A compiled in dependency '%s' of '%s' was not actually in the linker tables and so will be ignored (%d)."), *ToTest->GetFullName(), *ForObj->GetFullName(), CallSite);
 					UE_CLOG(!Outermost->HasAnyPackageFlags(PKG_CompiledIn), LogSavePackage, Fatal, TEXT("A dependency '%s' of '%s' was not actually in the linker tables and so will be ignored (%d)."), *ToTest->GetFullName(), *ForObj->GetFullName(), CallSite);
@@ -2710,7 +2714,10 @@ ESavePackageResult BeginCachePlatformCookedData(FSaveContext& SaveContext)
 	TArray<UObject*> ObjectsInPackage;
 	UPackage* Package = SaveContext.GetPackage();
 	GetObjectsWithPackage(Package, ObjectsInPackage);
-	ObjectsInPackage.RemoveAllSwap([&SaveContext](UObject* Object) { return SaveContext.IsUnsaveable(Object); });
+	ObjectsInPackage.RemoveAllSwap([&SaveContext](UObject* Object)
+		{
+			return SaveContext.IsUnsaveable(Object, SaveContext.GetCachedObjectStatus(Object));
+		});
 
 	// Call the PackageWriter to dispatch the BeginCache calls
 	ICookedPackageWriter::FBeginCacheForCookedPlatformDataInfo Info{
diff --git a/Engine/Source/Runtime/CoreUObject/Public/UObject/LinkerSave.h b/Engine/Source/Runtime/CoreUObject/Public/UObject/LinkerSave.h
--- a/Engine/Source/Runtime/CoreUObject/Public/UObject/LinkerSave.h
+++ b/Engine/Source/Runtime/CoreUObject/Public/UObject/LinkerSave.h
@@ -168,10 +168,11 @@ public:
 
 	// FArchive interface.
 	using FArchiveUObject::operator<<; // For visibility of the overloads we don't override
-	FArchive& operator<<( FName& InName );
-	FArchive& operator<<( UObject*& Obj );
-	FArchive& operator<<(FSoftObjectPath& SoftObjectPath);
-	FArchive& operator<<( FLazyObjectPtr& LazyObjectPtr );
+	virtual FArchive& operator<<(FName& InName) override;
+	virtual FArchive& operator<<(UObject*& Obj) override;
+	virtual FArchive& operator<<(FObjectPtr& Value) override;
+	virtual FArchive& operator<<(FSoftObjectPath& SoftObjectPath) override;
+	virtual FArchive& operator<<(FLazyObjectPtr& LazyObjectPtr) override;
 	virtual bool ShouldSkipProperty(const FProperty* InProperty) const override;
 	UE_DEPRECATED(5.5, "GetSerializeContext is not supported. Use FUObjectThreadContext::Get().GetSerializeContext().")
 	FUObjectSerializeContext* GetSerializeContext() override;
@@ -321,6 +322,8 @@ protected:
 
 	const TMap<FBulkDataCookedIndex, TUniquePtr<FFileRegionMemoryWriter>>& GetArchives(EBulkDataPayloadType Type) const;
 
+	void SerializeObjectPointer(const FObjectPtr& Value);
+
 private:
 	/** Optional log output to bubble errors back up. */
 	FOutputDevice* LogOutput = nullptr;
