InstancedStruct: Handle struct redirection when importing from text

diff --git a/Engine/Source/Runtime/CoreUObject/Private/StructUtils/InstancedStruct.cpp b/Engine/Source/Runtime/CoreUObject/Private/StructUtils/InstancedStruct.cpp
--- a/Engine/Source/Runtime/CoreUObject/Private/StructUtils/InstancedStruct.cpp
+++ b/Engine/Source/Runtime/CoreUObject/Private/StructUtils/InstancedStruct.cpp
@@ -3,6 +3,7 @@
 #include "StructUtils/StructView.h"
 #include "Serialization/CustomVersion.h"
 #include "UObject/Linker.h"
+#include "UObject/CoreRedirects.h"
 
 #if WITH_EDITOR
 #include "Serialization/MemoryReader.h"
@@ -267,6 +268,15 @@ bool FInstancedStruct::ImportTextItem(const TCHAR*& Buffer, int32 PortFlags, UOb
 	}
 	else
 	{
+		// Redirect the struct name if required
+		FCoreRedirectObjectName OldName = FCoreRedirectObjectName(StructPathName.ToString());
+		FCoreRedirectObjectName NewName = FCoreRedirects::GetRedirectedName(ECoreRedirectFlags::Type_Struct, OldName, ECoreRedirectMatchFlags::AllowPartialMatch);
+		if(OldName != NewName)
+		{
+			StructPathName.Reset();
+			StructPathName.Append(NewName.ToString());
+		}
+
 		// Make sure the struct is actually loaded before trying to import the text (this boils down to FindObject if the struct is already loaded).
 		// This is needed for user defined structs, BP pin values, config, copy/paste, where there's no guarantee that the referenced struct has actually been loaded yet.
 		UScriptStruct* StructTypePtr = LoadObject<UScriptStruct>(nullptr, StructPathName.ToString());
