Fix loading loose cooked file that contains no preload dependencies

diff --git a/Engine/Source/Runtime/CoreUObject/Private/UObject/LinkerLoad.cpp b/Engine/Source/Runtime/CoreUObject/Private/UObject/LinkerLoad.cpp
--- a/Engine/Source/Runtime/CoreUObject/Private/UObject/LinkerLoad.cpp
+++ b/Engine/Source/Runtime/CoreUObject/Private/UObject/LinkerLoad.cpp
@@ -1499,9 +1499,9 @@ FLinkerLoad::ELinkerStatus FLinkerLoad::SerializePackageFileSummaryInternal()
 		}
 	}
 
-	if (!FPlatformProperties::RequiresCookedData() &&
+	if (!FPlatformProperties::RequiresCookedData()
 		// We can't check the post tag if the file is an EDL cooked package
-		!((Summary.GetPackageFlags() & PKG_FilterEditorOnly) && Summary.PreloadDependencyCount > 0 && Summary.PreloadDependencyOffset > 0)
+		&& !bIsCooked
 		&& !IsTextFormat() && bLoaderNeedsEngineVersionChecks)
 	{
 		// check if this package version stored the 4-byte magic post tag
