Virtualized tracking of initialized and serialized state for property values and moved the API out of IDOs

#rnx

diff --git a/Engine/Source/Runtime/CoreUObject/Internal/UObject/InstanceDataObjectUtils.h b/Engine/Source/Runtime/CoreUObject/Internal/UObject/InstanceDataObjectUtils.h
--- a/Engine/Source/Runtime/CoreUObject/Internal/UObject/InstanceDataObjectUtils.h
+++ b/Engine/Source/Runtime/CoreUObject/Internal/UObject/InstanceDataObjectUtils.h
@@ -6,6 +6,8 @@
 
 #if WITH_EDITORONLY_DATA
 
+#include "UObject/PropertyStateTracking.h" // Temporary...
+
 #define UE_API COREUOBJECT_API
 
 class FFieldVariant;
@@ -41,28 +43,12 @@ UE_API bool IsClassOfInstanceDataObjectClass(const UStruct* Class);
 /** Generate a UClass that contains the union of the properties of PropertyTree and OwnerClass. */
 UClass* CreateInstanceDataObjectClass(const FPropertyPathNameTree* PropertyTree, const FUnknownEnumNames* EnumNames, UClass* OwnerClass, UObject* Outer);
 
-/** Notify that a property in an struct was set when the struct was deserialized. */
-void MarkPropertyValueSerialized(const UStruct* Struct, void* StructData, const FProperty* Property, int32 ArrayIndex = 0);
-/** Query whether a property in the struct was set when the struct was deserialized. */
-bool WasPropertyValueSerialized(const UStruct* Struct, const void* StructData, const FProperty* Property, int32 ArrayIndex = 0);
 /** Copy whether each property was set by serialization from one IDO to another. */
 void CopyPropertyValueSerializedData(const FFieldVariant& OldField, void* OldDataPtr, const FFieldVariant& NewField, void* NewDataPtr);
 
-
 UE_API bool SupportsPropertyInitStateTracking();
 
 UE_API bool SupportsPropertyInitStateTracking(const UStruct* Struct, void* StructData);
-	
-/** Query whether the property value is initialized for a property in the struct. */
-UE_API bool IsPropertyValueInitialized(const UStruct* Struct, void* StructData, const FProperty* Property, int32 ArrayIndex = 0);
-/** Set the flag that the property value is initialized for a property in the struct. */
-UE_API void SetPropertyValueInitialized(const UStruct* Struct, void* StructData, const FProperty* Property, int32 ArrayIndex = 0);
-/** Clear the flag that the property value is initialized for a property in the struct. */
-UE_API void ClearPropertyValueInitialized(const UStruct* Struct, void* StructData, const FProperty* Property, int32 ArrayIndex = 0);
-/** Reset the property value initialized flags for every property in the struct. */
-UE_API void ResetPropertyValueInitialized(const UStruct* Struct, void* StructData);
-		
-UE_API void SerializePropertyValueInitializationState(UObject& Object, FStructuredArchiveRecord& ObjectRecord);
 
 } // UE
 
diff --git a/Engine/Source/Runtime/CoreUObject/Internal/UObject/PropertyStateTracking.h b/Engine/Source/Runtime/CoreUObject/Internal/UObject/PropertyStateTracking.h
new file mode 100644
--- /dev/null
+++ b/Engine/Source/Runtime/CoreUObject/Internal/UObject/PropertyStateTracking.h
@@ -0,0 +1,55 @@
+// Copyright Epic Games, Inc. All Rights Reserved.
+
+#pragma once
+
+#include "CoreTypes.h"
+
+#if WITH_EDITORONLY_DATA
+
+class FProperty;
+class FStructuredArchiveRecord;
+class UStruct;
+
+#define UE_API COREUOBJECT_API
+
+namespace UE
+{
+
+// A property is initialized if it has a value set on it.
+// A property may be uninitialized if it has no default and has never been set.
+// A property may be uninitialized if it is an unknown property in an InstanceDataObject (IDO)
+// and has no value on a particular instance or has had its value unset while being fixed up.
+
+/** Try to activate tracking of property value initialization for the struct. */
+UE_API bool ActivateTrackingPropertyValueInitialized(const UStruct* Struct, void* Data);
+/** Query whether property value initialization is being tracked for the struct. */
+UE_API bool IsTrackingPropertyValueInitialized(const UStruct* Struct, const void* Data);
+/** Query whether the property value is initialized for a property in the struct. */
+UE_API bool IsPropertyValueInitialized(const UStruct* Struct, const void* Data, const FProperty* Property, int32 ArrayIndex = 0);
+/** Set the flag that the property value is initialized for a property in the struct. */
+UE_API void SetPropertyValueInitialized(const UStruct* Struct, void* Data, const FProperty* Property, int32 ArrayIndex = 0);
+/** Clear the flag that the property value is initialized for a property in the struct. */
+UE_API void ClearPropertyValueInitialized(const UStruct* Struct, void* Data, const FProperty* Property, int32 ArrayIndex = 0);
+/** Reset the property value initialized flags for every property in the struct. */
+UE_API void ResetPropertyValueInitialized(const UStruct* Struct, void* Data);
+/** Serialize the property value initialized flags for every property in the struct. */
+UE_API void SerializePropertyValueInitialized(const UStruct* Struct, void* Data, FStructuredArchiveRecord Record);
+
+// A property is serialized if a value was loaded into it during serialization.
+// A property may be initialized without having been serialized because a template object provided a default.
+// A property may be serialized without being initialized if its value was removed following serialization.
+
+/** Try to activate tracking of property value serialization for the struct. */
+bool ActivateTrackingPropertyValueSerialized(const UStruct* Struct, void* Data);
+/** Query whether property value serialization is being tracked for the struct. */
+bool IsTrackingPropertyValueSerialized(const UStruct* Struct, const void* Data);
+/** Set the flag that the property value was deserialized for a property in the struct. */
+void MarkPropertyValueSerialized(const UStruct* Struct, void* Data, const FProperty* Property, int32 ArrayIndex = 0);
+/** Query whether the property value was deserialized for a property in a struct. */
+bool WasPropertyValueSerialized(const UStruct* Struct, const void* Data, const FProperty* Property, int32 ArrayIndex = 0);
+
+} // UE
+
+#undef UE_API
+
+#endif // WITH_EDITORONLY_DATA
diff --git a/Engine/Source/Runtime/CoreUObject/Private/UObject/Class.cpp b/Engine/Source/Runtime/CoreUObject/Private/UObject/Class.cpp
--- a/Engine/Source/Runtime/CoreUObject/Private/UObject/Class.cpp
+++ b/Engine/Source/Runtime/CoreUObject/Private/UObject/Class.cpp
@@ -40,6 +40,7 @@
 #include "UObject/PropertyBagRepository.h"
 #include "UObject/PropertyOptional.h"
 #include "UObject/PropertyPathNameTree.h"
+#include "UObject/PropertyStateTracking.h"
 #include "UObject/StructOnScope.h"
 #include "UObject/StructScriptLoader.h"
 #include "UObject/PropertyHelper.h"
@@ -1391,18 +1392,15 @@ void UStruct::LoadTaggedPropertiesFromText(FStructuredArchive::FSlot Slot, uint8
 				Tag.ArrayIndex = ItemIndex;
 				Tag.Name = PropertyName;
 
-				#if WITH_EDITORONLY_DATA
-				const auto SetInitialized = [bTrackInitializedProperties = SerializeContext->bTrackInitializedProperties](const UStruct* Struct, void* Data, const FProperty* Property, int32 ArrayIndex)
+				const auto SetInitialized = [SerializeContext](const UStruct* Struct, void* Data, const FProperty* Property, int32 ArrayIndex)
 				{
-				
-					if (bTrackInitializedProperties)
+				#if WITH_EDITORONLY_DATA
+					if (SerializeContext->bTrackInitializedProperties)
 					{
 						UE::SetPropertyValueInitialized(Struct, Data, Property, ArrayIndex);
 					}
-				};
-				#else
-				const auto SetInitialized = [](const UStruct* Struct, void* Data, const FProperty* Property, int32 ArrayIndex) {};
 				#endif
+				};
 
 				if (Tag.SerializeType == EPropertyTagSerializeType::Skipped)
 				{
@@ -2582,6 +2580,37 @@ const UStruct* UStruct::HasMetaDataHierarchical(const FName& Key) const
 
 #endif // WITH_METADATA
 
+#if WITH_EDITORONLY_DATA
+
+bool UStruct::ActivateTrackingPropertyValueFlag(EPropertyValueFlags Flags, void* Data) const
+{
+	return false;
+}
+
+bool UStruct::IsTrackingPropertyValueFlag(EPropertyValueFlags Flags, const void* Data) const
+{
+	return false;
+}
+
+bool UStruct::HasPropertyValueFlag(EPropertyValueFlags Flags, const void* Data, const FProperty* Property, int32 ArrayIndex) const
+{
+	return true;
+}
+
+void UStruct::SetPropertyValueFlag(EPropertyValueFlags Flags, bool bValue, void* Data, const FProperty* Property, int32 ArrayIndex) const
+{
+}
+
+void UStruct::ResetPropertyValueFlags(EPropertyValueFlags Flags, void* Data) const
+{
+}
+
+void UStruct::SerializePropertyValueFlags(EPropertyValueFlags Flags, void* Data, FStructuredArchiveRecord Record, FArchiveFieldName Name) const
+{
+}
+
+#endif // WITH_EDITORONLY_DATA
+
 #if USE_CIRCULAR_DEPENDENCY_LOAD_DEFERRING
 	/**
 	 * If we're loading, then the value of the script's UObject* expression 
@@ -5681,6 +5710,144 @@ bool UClass::IsStructTrashed() const
 	return Children == nullptr && ChildProperties == nullptr && ClassDefaultObject == nullptr;
 }
 
+#if WITH_EDITORONLY_DATA
+
+int32 CalculatePropertyIndex(const UStruct* Struct, const FProperty* Property, int32 ArrayIndex)
+{
+	checkf(Struct->IsChildOf(Property->GetOwnerStruct()),
+		TEXT("Property %s in %s is not related to the provided type %s."),
+		*Property->GetAuthoredName(), *Property->GetOwnerStruct()->GetPathName(), *Struct->GetPathName());
+
+	const int32 IndexInOwner = Property->GetIndexInOwner();
+	checkf(IndexInOwner >= 0,
+		TEXT("Property %s in %s does not have an index within its owner."),
+		*Property->GetAuthoredName(), *Struct->GetPathName());
+
+	checkf(ArrayIndex < Property->ArrayDim,
+		TEXT("Property %s in %s has out of range array index %d with capacity %d."),
+		*Property->GetAuthoredName(), *Struct->GetPathName(), ArrayIndex, Property->ArrayDim);
+
+	return IndexInOwner + ArrayIndex;
+}
+
+struct FPropertyValueFlagsAnnotationData
+{
+	TBitArray<> FlagByIndex;
+};
+
+struct FPropertyValueFlagsAnnotation
+{
+	bool IsDefault() const
+	{
+		return !Data.IsValid();
+	}
+
+	TSharedPtr<FPropertyValueFlagsAnnotationData> Data;
+};
+
+using FPropertyValueFlagsAnnotationStore = FUObjectAnnotationSparse<FPropertyValueFlagsAnnotation, /*bAutoRemove*/ true>;
+
+template <EPropertyValueFlags Flags>
+UE_AUTORTFM_ALWAYS_OPEN FPropertyValueFlagsAnnotationStore& GetPropertyValueFlagsAnnotations()
+{
+	static FPropertyValueFlagsAnnotationStore Annotations;
+	return Annotations;
+}
+
+static FPropertyValueFlagsAnnotationStore& GetPropertyValueFlagsAnnotations(EPropertyValueFlags Flags)
+{
+	switch (Flags)
+	{
+	default:
+		checkNoEntry();
+		[[fallthrough]];
+	case EPropertyValueFlags::Initialized:
+		return GetPropertyValueFlagsAnnotations<EPropertyValueFlags::Initialized>();
+	case EPropertyValueFlags::Serialized:
+		return GetPropertyValueFlagsAnnotations<EPropertyValueFlags::Serialized>();
+	}
+}
+
+bool UClass::ActivateTrackingPropertyValueFlag(EPropertyValueFlags Flags, void* Data) const
+{
+	UObjectBase* Object = static_cast<UObjectBase*>(Data);
+	FPropertyValueFlagsAnnotationStore& Store = GetPropertyValueFlagsAnnotations(Flags);
+	TSharedPtr<FPropertyValueFlagsAnnotationData> AnnotationData = Store.GetAnnotation(Object).Data;
+	if (!AnnotationData)
+	{
+		AnnotationData = MakeShared<FPropertyValueFlagsAnnotationData>();
+		Store.AddAnnotation(Object, {AnnotationData});
+	}
+	return true;
+}
+
+bool UClass::IsTrackingPropertyValueFlag(EPropertyValueFlags Flags, const void* Data) const
+{
+	const UObjectBase* Object = static_cast<const UObjectBase*>(Data);
+	return GetPropertyValueFlagsAnnotations(Flags).GetAnnotation(Object).Data.IsValid();
+}
+
+bool UClass::HasPropertyValueFlag(EPropertyValueFlags Flags, const void* Data, const FProperty* Property, int32 ArrayIndex) const
+{
+	const UObjectBase* Object = static_cast<const UObjectBase*>(Data);
+	if (TSharedPtr<FPropertyValueFlagsAnnotationData> AnnotationData = GetPropertyValueFlagsAnnotations(Flags).GetAnnotation(Object).Data)
+	{
+		const int32 PropertyIndex = CalculatePropertyIndex(this, Property, ArrayIndex);
+		if (AnnotationData->FlagByIndex.IsValidIndex(PropertyIndex))
+		{
+			return AnnotationData->FlagByIndex[PropertyIndex];
+		}
+		// Default to uninitialized when tracking is active and a flag has not been stored.
+		return false;
+	}
+	// Default to initialized when tracking is inactive.
+	return true;
+}
+
+void UClass::SetPropertyValueFlag(EPropertyValueFlags Flags, bool bValue, void* Data, const FProperty* Property, int32 ArrayIndex) const
+{
+	UObjectBase* Object = static_cast<UObjectBase*>(Data);
+	if (TSharedPtr<FPropertyValueFlagsAnnotationData> AnnotationData = GetPropertyValueFlagsAnnotations(Flags).GetAnnotation(Object).Data)
+	{
+		const int32 PropertyIndex = CalculatePropertyIndex(this, Property, ArrayIndex);
+		if (AnnotationData->FlagByIndex.Num() <= PropertyIndex)
+		{
+			checkf(PropertyIndex < TotalFieldCount,
+				TEXT("Property %s in %s has out of range index %d with capacity for %d."),
+				*Property->GetAuthoredName(), *GetPathName(), PropertyIndex, TotalFieldCount);
+			AnnotationData->FlagByIndex.SetNum(TotalFieldCount, /*bValue*/ false);
+		}
+		AnnotationData->FlagByIndex[PropertyIndex] = bValue;
+	}
+}
+
+void UClass::ResetPropertyValueFlags(EPropertyValueFlags Flags, void* Data) const
+{
+	UObjectBase* Object = static_cast<UObjectBase*>(Data);
+	if (TSharedPtr<FPropertyValueFlagsAnnotationData> AnnotationData = GetPropertyValueFlagsAnnotations(Flags).GetAnnotation(Object).Data)
+	{
+		AnnotationData->FlagByIndex.SetRange(0, AnnotationData->FlagByIndex.Num(), /*bValue*/ false);
+	}
+}
+
+void UClass::SerializePropertyValueFlags(EPropertyValueFlags Flags, void* Data, FStructuredArchiveRecord Record, FArchiveFieldName Name) const
+{
+	UObjectBase* Object = static_cast<UObjectBase*>(Data);
+	FPropertyValueFlagsAnnotationStore& Store = GetPropertyValueFlagsAnnotations(Flags);
+	TSharedPtr<FPropertyValueFlagsAnnotationData> AnnotationData = Store.GetAnnotation(Object).Data;
+	if (TOptional<FStructuredArchiveSlot> Slot = Record.TryEnterField(Name, AnnotationData.IsValid()))
+	{
+		if (!AnnotationData)
+		{
+			AnnotationData = MakeShared<FPropertyValueFlagsAnnotationData>();
+			Store.AddAnnotation(Object, {AnnotationData});
+		}
+		*Slot << AnnotationData->FlagByIndex;
+	}
+}
+
+#endif // WITH_EDITORONLY_DATA
+
 void UClass::Serialize( FArchive& Ar )
 {
 	if ( Ar.IsLoading() || (Ar.IsModifyingWeakAndStrongReferences() && !Ar.IsSaving()) )
diff --git a/Engine/Source/Runtime/CoreUObject/Private/UObject/InstanceDataObjectUtils.cpp b/Engine/Source/Runtime/CoreUObject/Private/UObject/InstanceDataObjectUtils.cpp
--- a/Engine/Source/Runtime/CoreUObject/Private/UObject/InstanceDataObjectUtils.cpp
+++ b/Engine/Source/Runtime/CoreUObject/Private/UObject/InstanceDataObjectUtils.cpp
@@ -22,34 +22,134 @@
 #include "UObject/UnrealType.h"
 #include "UObject/UObjectAnnotation.h"
 
+// Implemented in Class.cpp
+int32 CalculatePropertyIndex(const UStruct* Struct, const FProperty* Property, int32 ArrayIndex);
+
 static const FName NAME_InitializedValues(ANSITEXTVIEW("_InitializedValues"));
 static const FName NAME_SerializedValues(ANSITEXTVIEW("_SerializedValues"));
 
-/** Type used for InstanceDataObject classes. */
-class UInstanceDataObjectClass final : public UClass
+template <typename SuperType>
+class TInstanceDataObjectPropertyValueFlags : public SuperType
 {
 public:
-	DECLARE_CASTED_CLASS_INTRINSIC(UInstanceDataObjectClass, UClass, CLASS_Transient, TEXT("/Script/CoreUObject"), CASTCLASS_UClass)
+	using SuperType::SuperType;
+
+	bool ActivateTrackingPropertyValueFlag(EPropertyValueFlags Flags, void* Data) const final
+	{
+		// Nothing to activate because tracking is either always on or always off.
+		return IsTrackingPropertyValueFlag(Flags, Data);
+	}
+
+	bool IsTrackingPropertyValueFlag(EPropertyValueFlags Flags, const void* Data) const final
+	{
+		return !!GetPropertyValueFlagsProperty(Flags);
+	}
+
+	bool HasPropertyValueFlag(EPropertyValueFlags Flags, const void* Data, const FProperty* Property, int32 ArrayIndex) const final
+	{
+		if (const FProperty* FlagsProperty = GetPropertyValueFlagsProperty(Flags))
+		{
+			const int32 PropertyIndex = CalculatePropertyIndex(this, Property, ArrayIndex);
+			const int32 ByteIndex = PropertyIndex / 8;
+			const int32 BitOffset = PropertyIndex % 8;
+			if (ensureMsgf(ByteIndex < FlagsProperty->ArrayDim,
+				TEXT("Property %s in %s has out of range index %d with capacity for %d."),
+				*Property->GetAuthoredName(), *SuperType::GetPathName(), PropertyIndex, FlagsProperty->ArrayDim * 8))
+			{
+				const uint8* FlagsData = FlagsProperty->ContainerPtrToValuePtr<uint8>(Data, ByteIndex);
+				return (*FlagsData & (1 << BitOffset)) != 0;
+			}
+		}
+		// Default to initialized when tracking is inactive.
+		return true;
+	}
+
+	void SetPropertyValueFlag(EPropertyValueFlags Flags, bool bValue, void* Data, const FProperty* Property, int32 ArrayIndex) const final
+	{
+		if (const FProperty* FlagsProperty = GetPropertyValueFlagsProperty(Flags))
+		{
+			const int32 PropertyIndex = CalculatePropertyIndex(this, Property, ArrayIndex);
+			const int32 ByteIndex = PropertyIndex / 8;
+			const int32 BitOffset = PropertyIndex % 8;
+			if (ensureMsgf(ByteIndex < FlagsProperty->ArrayDim,
+				TEXT("Property %s in %s has out of range index %d with capacity for %d."),
+				*Property->GetAuthoredName(), *SuperType::GetPathName(), PropertyIndex, FlagsProperty->ArrayDim * 8))
+			{
+				uint8* FlagsData = FlagsProperty->ContainerPtrToValuePtr<uint8>(Data, ByteIndex);
+				if (bValue)
+				{
+					*FlagsData |= (1 << BitOffset);
+				}
+				else
+				{
+					*FlagsData &= ~(1 << BitOffset);
+				}
+			}
+		}
+	}
+
+	void ResetPropertyValueFlags(EPropertyValueFlags Flags, void* Data) const final
+	{
+		if (const FProperty* FlagsProperty = GetPropertyValueFlagsProperty(Flags))
+		{
+			uint8* FlagsData = FlagsProperty->ContainerPtrToValuePtr<uint8>(Data);
+			FMemory::Memzero(FlagsData, FlagsProperty->ArrayDim);
+		}
+	}
+
+	void SerializePropertyValueFlags(EPropertyValueFlags Flags, void* Data, FStructuredArchiveRecord Record, FArchiveFieldName Name) const final
+	{
+		const FProperty* FlagsProperty = GetPropertyValueFlagsProperty(Flags);
+		if (TOptional<FStructuredArchiveSlot> Slot = Record.TryEnterField(Name, !!FlagsProperty))
+		{
+			checkf(FlagsProperty, TEXT("Type %s is missing a property that is needed to serialize property value flags."), *SuperType::GetPathName());
+			uint8* FlagsData = FlagsProperty->ContainerPtrToValuePtr<uint8>(Data);
+			Slot->Serialize(FlagsData, FlagsProperty->ArrayDim);
+		}
+	}
+
+	const FProperty* GetPropertyValueFlagsProperty(EPropertyValueFlags Flags) const
+	{
+		switch (Flags)
+		{
+		case EPropertyValueFlags::Initialized:
+			return InitializedValuesProperty;
+		case EPropertyValueFlags::Serialized:
+			return SerializedValuesProperty;
+		default:
+			checkNoEntry();
+			return nullptr;
+		}
+	}
 
 	FByteProperty* InitializedValuesProperty = nullptr;
 	FByteProperty* SerializedValuesProperty = nullptr;
 };
 
+/** Type used for InstanceDataObject classes. */
+class UInstanceDataObjectClass final : public TInstanceDataObjectPropertyValueFlags<UClass>
+{
+	using SuperType = TInstanceDataObjectPropertyValueFlags<UClass>;
+
+public:
+	DECLARE_CASTED_CLASS_INTRINSIC(UInstanceDataObjectClass, SuperType, CLASS_Transient, TEXT("/Script/CoreUObject"), CASTCLASS_UClass)
+};
+
 IMPLEMENT_CORE_INTRINSIC_CLASS(UInstanceDataObjectClass, UClass,
 {
 });
 
 /** Type used for InstanceDataObject structs to provide support for hashing and custom guids. */
-class UInstanceDataObjectStruct final : public UScriptStruct
+class UInstanceDataObjectStruct final : public TInstanceDataObjectPropertyValueFlags<UScriptStruct>
 {
+	using SuperType = TInstanceDataObjectPropertyValueFlags<UScriptStruct>;
+
 public:
-	DECLARE_CASTED_CLASS_INTRINSIC(UInstanceDataObjectStruct, UScriptStruct, CLASS_Transient, TEXT("/Script/CoreUObject"), CASTCLASS_UScriptStruct)
+	DECLARE_CASTED_CLASS_INTRINSIC(UInstanceDataObjectStruct, SuperType, CLASS_Transient, TEXT("/Script/CoreUObject"), CASTCLASS_UScriptStruct)
 
 	uint32 GetStructTypeHash(const void* Src) const final;
 	FGuid GetCustomGuid() const final { return Guid; }
 
-	FByteProperty* InitializedValuesProperty = nullptr;
-	FByteProperty* SerializedValuesProperty = nullptr;
 	FGuid Guid;
 };
 
@@ -785,37 +885,6 @@ namespace UE
 		return CastField<FByteProperty>(Struct->FindPropertyByName(NAME_SerializedValues));
 	}
 
-	void MarkPropertyValueSerialized(const UStruct* Struct, void* StructData, const FProperty* Property, int32 ArrayIndex)
-	{
-		if (const FByteProperty* SerializedValuesProperty = FindSerializedValuesProperty(Struct))
-		{
-			const int32 PropertyIndex = Property->GetIndexInOwner() + ArrayIndex;
-			const int32 ByteIndex = PropertyIndex / 8;
-			const int32 BitOffset = PropertyIndex % 8;
-			if (ByteIndex < SerializedValuesProperty->ArrayDim)
-			{
-				uint8* PropertyDataPtr = SerializedValuesProperty->ContainerPtrToValuePtr<uint8>(StructData, ByteIndex);
-				*PropertyDataPtr |= (1 << BitOffset);
-			}
-		}
-	}
-
-	bool WasPropertyValueSerialized(const UStruct* Struct, const void* StructData, const FProperty* Property, int32 ArrayIndex)
-	{
-		if (const FByteProperty* SerializedValuesProperty = FindSerializedValuesProperty(Struct))
-		{
-			const int32 PropertyIndex = Property->GetIndexInOwner() + ArrayIndex;
-			const int32 ByteIndex = PropertyIndex / 8;
-			const int32 BitOffset = PropertyIndex % 8;
-			if (ByteIndex < SerializedValuesProperty->ArrayDim)
-			{
-				const uint8* PropertyDataPtr = SerializedValuesProperty->ContainerPtrToValuePtr<uint8>(StructData, ByteIndex);
-				return (*PropertyDataPtr & (1 << BitOffset)) != 0;
-			}
-		}
-		return false;
-	}
-
 	void CopyPropertyValueSerializedData(const FFieldVariant& OldField, void* OldDataPtr, const FFieldVariant& NewField, void* NewDataPtr)
 	{
 		if (const FStructProperty* OldAsStructProperty = OldField.Get<FStructProperty>())
@@ -935,298 +1004,13 @@ namespace UE
 
 	bool SupportsPropertyInitStateTracking(const UStruct* Struct, void* StructData)
 	{
-		if (bUseAnnotationsForPropertyStateInitTracking)
-		{
-			if (Cast<UVerseClass>(Struct) != nullptr)
-			{
-				return true;
-			}
-		}
-		else if (bEnableIDOSupport)
+		if (bUseAnnotationsForPropertyStateInitTracking || bEnableIDOSupport)
 		{
-			if (IsClassOfInstanceDataObjectClass(Struct))
-			{
-				return true;
-			}
+			return IsTrackingPropertyValueInitialized(Struct, StructData);
 		}
 		return false;
 	}
 
-	static bool IsIdoPropertyValueInitialized(const UStruct* Struct, void* StructData, const FProperty* Property, int32 ArrayIndex, const FByteProperty* InitializedValuesProperty)
-	{
-		if (!SupportsPropertyInitStateTracking())
-		{
-			return true;
-		}
-		
-		const int32 PropertyIndex = Property->GetIndexInOwner() + ArrayIndex;
-		const int32 ByteIndex = PropertyIndex / 8;
-		const int32 BitOffset = PropertyIndex % 8;
-		if (ensureMsgf(ByteIndex < InitializedValuesProperty->ArrayDim,
-			TEXT("Property %s in %s has out of range index %d with capacity for %d."),
-			*Property->GetAuthoredName(), *Struct->GetPathName(), PropertyIndex, InitializedValuesProperty->ArrayDim * 8))
-		{
-			const uint8* PropertyDataPtr = InitializedValuesProperty->ContainerPtrToValuePtr<uint8>(StructData, ByteIndex);
-			return (*PropertyDataPtr & (1 << BitOffset)) != 0;
-		}
-		// True is a safe fallback when this function is misused.
-		return true;
-	}
-
-	static void SetIdoPropertyValueInitializedFlag(const UStruct* Struct, void* StructData, const FProperty* Property, int32 ArrayIndex, bool bValue, const FByteProperty* InitializedValuesProperty)
-	{
-		const int32 PropertyIndex = Property->GetIndexInOwner() + ArrayIndex;
-		const int32 ByteIndex = PropertyIndex / 8;
-		const int32 BitOffset = PropertyIndex % 8;
-		if (ensureMsgf(ByteIndex < InitializedValuesProperty->ArrayDim,
-			TEXT("Property %s in %s has out of range index %d with capacity for %d."),
-			*Property->GetAuthoredName(), *Struct->GetPathName(), PropertyIndex, InitializedValuesProperty->ArrayDim * 8))
-		{
-			uint8* PropertyDataPtr = InitializedValuesProperty->ContainerPtrToValuePtr<uint8>(StructData, ByteIndex);
-			if (bValue)
-			{
-				*PropertyDataPtr |= (1 << BitOffset);
-			}
-			else
-			{
-				*PropertyDataPtr &= ~(1 << BitOffset);
-			}
-		}
-	}
-
-	static void ResetIdoPropertyValueInitialized(const UStruct* Struct, void* StructData, const FByteProperty* InitializedValuesProperty)
-	{
-		uint8* PropertyDataPtr = InitializedValuesProperty->ContainerPtrToValuePtr<uint8>(StructData);
-		FMemory::Memzero(PropertyDataPtr, InitializedValuesProperty->ArrayDim);
-	}
-
-	struct FPropertyInitializedAnnotationData
-	{
-		TBitArray<> PropertyInitialized;
-	};
-
-	struct FPropertyInitializedAnnotation
-	{
-		bool IsDefault() const
-		{
-			return !Data.IsValid();
-		}
-
-		TSharedPtr<FPropertyInitializedAnnotationData> Data;
-	};
-
-	using FPropertyInitializedAnnotationStore = FUObjectAnnotationSparse<FPropertyInitializedAnnotation, /*bAutoRemove*/ true>;
-
-	UE_AUTORTFM_ALWAYS_OPEN
-	FPropertyInitializedAnnotationStore& GetPropertyInitializedAnnotations()
-	{
-		static FPropertyInitializedAnnotationStore Annotations;
-		return Annotations;
-	}
-
-	static bool IsPropertyValueInitializedByAnnotation(const UStruct* Struct, void* StructData, const FProperty* Property, int32 ArrayIndex)
-	{
-		if (!bUseAnnotationsForPropertyStateInitTracking)
-		{
-			return true;
-		}
-		UObjectBase* Object = static_cast<UObjectBase*>(StructData);
-		FPropertyInitializedAnnotationStore& As = GetPropertyInitializedAnnotations();
-		FPropertyInitializedAnnotation Annotation = As.GetAnnotation(Object);
-		
-		if (TSharedPtr<FPropertyInitializedAnnotationData>& Data = Annotation.Data)
-		{
-			const int32 PropertyIndex = Property->GetIndexInOwner() + ArrayIndex;
-			if (Data->PropertyInitialized.IsValidIndex(PropertyIndex))
-			{
-				return Data->PropertyInitialized[PropertyIndex];
-			}
-		}
-		// Default to uninitialized when a flag has not been stored.
-		return false;
-	}
-
-	static void SetPropertyValueInitializedFlagByAnnotation(const UStruct* Struct, void* StructData, const FProperty* Property, int32 ArrayIndex, bool bValue)
-	{
-		if (!bUseAnnotationsForPropertyStateInitTracking)
-		{
-			return;
-		}		
-		UObjectBase* Object = static_cast<UObjectBase*>(StructData);
-		FPropertyInitializedAnnotationStore& Store = GetPropertyInitializedAnnotations();
-		TSharedPtr<FPropertyInitializedAnnotationData> Data = Store.GetAnnotation(Object).Data;
-		if (!Data)
-		{
-			Data = MakeShared<FPropertyInitializedAnnotationData>();
-			Store.AddAnnotation(Object, {Data});
-		}
-
-		const int32 PropertyIndex = Property->GetIndexInOwner() + ArrayIndex;
-		if (Data->PropertyInitialized.Num() <= PropertyIndex)
-		{
-			checkf(PropertyIndex < Struct->TotalFieldCount,
-				TEXT("Property %s in %s has out of range index %d with capacity for %d."),
-				*Property->GetAuthoredName(), *Struct->GetPathName(), PropertyIndex, Struct->TotalFieldCount);
-			Data->PropertyInitialized.SetNum(Struct->TotalFieldCount, /*bValue*/ false);
-		}
-		Data->PropertyInitialized[PropertyIndex] = bValue;
-	}
-
-	static void ResetPropertyValueInitializedByAnnotation(const UStruct* Struct, void* StructData)
-	{
-		if (!bUseAnnotationsForPropertyStateInitTracking)
-		{
-			return;
-		}
-		UObjectBase* Object = static_cast<UObjectBase*>(StructData);
-		if (TSharedPtr<FPropertyInitializedAnnotationData> Data = GetPropertyInitializedAnnotations().GetAnnotation(Object).Data)
-		{
-			Data->PropertyInitialized.SetRange(0, Data->PropertyInitialized.Num(), /*bValue*/ false);
-		}
-	}
-
-	bool IsPropertyValueInitialized(const UStruct* Struct, void* StructData, const FProperty* Property, int32 ArrayIndex)
-	{
-		if (const UInstanceDataObjectClass* IdoClass = Cast<UInstanceDataObjectClass>(Struct))
-		{
-			return IsIdoPropertyValueInitialized(Struct, StructData, Property, ArrayIndex, IdoClass->InitializedValuesProperty);
-		}
-		if (const UInstanceDataObjectStruct* IdoStruct = Cast<UInstanceDataObjectStruct>(Struct))
-		{
-			return IsIdoPropertyValueInitialized(Struct, StructData, Property, ArrayIndex, IdoStruct->InitializedValuesProperty);
-		}
-		if (const UClass* Class = Cast<UClass>(Struct))
-		{
-			return IsPropertyValueInitializedByAnnotation(Struct, StructData, Property, ArrayIndex);
-		}
-		return true;
-	}
-
-	static void SetPropertyValueInitializedFlag(const UStruct* Struct, void* StructData, const FProperty* Property, int32 ArrayIndex, bool bValue)
-	{
-		if (const UInstanceDataObjectClass* IdoClass = Cast<UInstanceDataObjectClass>(Struct))
-		{
-			return SetIdoPropertyValueInitializedFlag(Struct, StructData, Property, ArrayIndex, bValue, IdoClass->InitializedValuesProperty);
-		}
-		if (const UInstanceDataObjectStruct* IdoStruct = Cast<UInstanceDataObjectStruct>(Struct))
-		{
-			return SetIdoPropertyValueInitializedFlag(Struct, StructData, Property, ArrayIndex, bValue, IdoStruct->InitializedValuesProperty);
-		}
-		if (const UClass* Class = Cast<UClass>(Struct))
-		{
-			return SetPropertyValueInitializedFlagByAnnotation(Struct, StructData, Property, ArrayIndex, bValue);
-		}
-	}
-
-	void SetPropertyValueInitialized(const UStruct* Struct, void* StructData, const FProperty* Property, int32 ArrayIndex)
-	{
-		SetPropertyValueInitializedFlag(Struct, StructData, Property, ArrayIndex, /*bValue*/ true);
-	}
-
-	void ClearPropertyValueInitialized(const UStruct* Struct, void* StructData, const FProperty* Property, int32 ArrayIndex)
-	{
-		SetPropertyValueInitializedFlag(Struct, StructData, Property, ArrayIndex, /*bValue*/ false);
-	}
-
-	void ResetPropertyValueInitialized(const UStruct* Struct, void* StructData)
-	{
-		if (const UInstanceDataObjectClass* IdoClass = Cast<UInstanceDataObjectClass>(Struct))
-		{
-			return ResetIdoPropertyValueInitialized(Struct, StructData, IdoClass->InitializedValuesProperty);
-		}
-		if (const UInstanceDataObjectStruct* IdoStruct = Cast<UInstanceDataObjectStruct>(Struct))
-		{
-			return ResetIdoPropertyValueInitialized(Struct, StructData, IdoStruct->InitializedValuesProperty);
-		}
-		if (const UClass* Class = Cast<UClass>(Struct))
-		{
-			return ResetPropertyValueInitializedByAnnotation(Struct, StructData);
-		}
-	}
-
-	static void CopyPropertyInitializationState(const FFieldVariant& Destination, void* DestinationDataPtr, const FFieldVariant& Source, void* SourceDataPtr)
-	{
-		if (const FStructProperty* SourceAsStructProperty = Source.Get<FStructProperty>())
-		{
-			const FStructProperty* DestAsStructProperty = Destination.Get<FStructProperty>();
-			checkf(DestAsStructProperty, TEXT("Type mismatch between Source field and Destingation field. Expected FStructProperty"));
-			CopyPropertyValueSerializedData(DestAsStructProperty->Struct, DestinationDataPtr, SourceAsStructProperty->Struct, SourceDataPtr);
-		}
-		else if (const FArrayProperty* SourceAsArrayProperty = Source.Get<FArrayProperty>())
-		{
-			// Any properties within Arrays are assumed initialized
-		}
-		else if (const FSetProperty* SourceAsSetProperty = Source.Get<FSetProperty>())
-		{
-			// Any properties within Sets are assumed initialized
-		}
-		else if (const FMapProperty* SourceAsMapProperty = Source.Get<FMapProperty>())
-		{
-			// Any properties within Maps are assumed initialized
-		}
-		else if (UStruct* SourceAsStruct = Source.Get<UStruct>())
-		{
-			UStruct* DestAsStruct = Destination.Get<UStruct>();
-
-			// Integrity check - both should be same type
-			// Can use the same properties then
-			check(SourceAsStruct->GetClass() == DestAsStruct->GetClass());			
-			
-			for (const FProperty* Property : TFieldRange<FProperty>(SourceAsStruct))
-			{
-				for (int32 ArrayIndex = 0; ArrayIndex < Property->ArrayDim; ++ArrayIndex)
-				{
-					const bool bIsInitialized = IsPropertyValueInitializedByAnnotation(SourceAsStruct->GetClass(), SourceAsStruct, Property, ArrayIndex);
-					SetPropertyValueInitializedFlagByAnnotation(SourceAsStruct->GetClass(), SourceAsStruct, Property, ArrayIndex, bIsInitialized);
-				}
-			}
-		}
-	}
-
-	void SerializePropertyValueInitializationState(UObject& Object, FStructuredArchiveRecord& ObjectRecord)
-	{
-		// Serialization of the property init state when IDOs enabled are done via a hidden property on the IDOs
-		if (!bUseAnnotationsForPropertyStateInitTracking)
-		{
-			return;
-		}
-		const FArchiveState& ArchiveState = ObjectRecord.GetArchiveState();
-		
-		bool bHasPropertyValueInitializationState = false;
-		TSharedPtr<FPropertyInitializedAnnotationData> Data;
-		if (ArchiveState.IsSaving())
-		{
-			Data = GetPropertyInitializedAnnotations().GetAnnotation(&Object).Data;
-
-			if (Data.IsValid())
-			{
-				bHasPropertyValueInitializationState = true;
-			}
-		}
-
-		TOptional<FStructuredArchiveSlot> Slot = ObjectRecord.TryEnterField(TEXT("PropertyValueInitializationState"), bHasPropertyValueInitializationState);
-		if (Slot.IsSet())
-		{
-			*Slot << SA_ATTRIBUTE(TEXT("PropertyValueInitializationState"), bHasPropertyValueInitializationState);
-
-			if (bHasPropertyValueInitializationState)
-			{
-				if (ArchiveState.IsLoading())
-				{
-					Data = MakeShared<FPropertyInitializedAnnotationData>();
-				}
-
-				*Slot << Data->PropertyInitialized;
-				
-				if (ArchiveState.IsLoading())
-				{
-					FPropertyInitializedAnnotation Annotation;
-					Annotation.Data = Data;
-					GetPropertyInitializedAnnotations().AddAnnotation(&Object, MoveTemp(Annotation));
-				}
-			}
-		}
-	}
 } // UE
 
 #endif // WITH_EDITORONLY_DATA
diff --git a/Engine/Source/Runtime/CoreUObject/Private/UObject/Obj.cpp b/Engine/Source/Runtime/CoreUObject/Private/UObject/Obj.cpp
--- a/Engine/Source/Runtime/CoreUObject/Private/UObject/Obj.cpp
+++ b/Engine/Source/Runtime/CoreUObject/Private/UObject/Obj.cpp
@@ -39,6 +39,7 @@
 #include "UObject/SoftObjectPtr.h"
 #include "UObject/PropertyBagRepository.h"
 #include "UObject/PropertyPortFlags.h"
+#include "UObject/PropertyStateTracking.h"
 #include "UObject/UnrealType.h"
 #include "UObject/ObjectRedirector.h"
 #include "UObject/ObjectSaveContext.h"
@@ -473,54 +474,20 @@ void UObject::PostEditChangeProperty(FPropertyChangedEvent& PropertyChangedEvent
 #if WITH_EDITORONLY_DATA
 	if (const FProperty* Property = PropertyChangedEvent.MemberProperty)
 	{
-		if (UE::SupportsPropertyInitStateTracking())
+		switch (PropertyChangedEvent.ChangeType)
 		{
-			// Set property initialized state when a property is changed
-
-			auto SetPropertyValueInitializedHelper = [&PropertyChangedEvent](UObject* Object, const FProperty* P)
-			{
-				int32 ChangedIndex = 0;
-				switch (PropertyChangedEvent.ChangeType)
-				{
-				case EPropertyChangeType::ArrayAdd:
-				case EPropertyChangeType::ArrayRemove:
-				case EPropertyChangeType::ArrayMove:
-					// No tracking for elements of container.  Any change counts as initialization
-					UE::SetPropertyValueInitialized(Object->GetClass(), Object, P);
-					break;
-				case EPropertyChangeType::ArrayClear: // Clearing the entire array. No need for an index. Clear also counts as initialization of the property
-				case EPropertyChangeType::ValueSet:
-				case EPropertyChangeType::Interactive:
-					UE::SetPropertyValueInitialized(Object->GetClass(), Object, P);
-					break;
-				default:
-					// No op for now
-						break;
-				}
-			};
-	
-			if (!UE::IsInstanceDataObjectSupportEnabled())
-			{
-				// IDOs as a general system are disabled
-				UVerseClass* VerseClass = Cast<UVerseClass>(GetClass());
-				if (VerseClass)
-				{
-					SetPropertyValueInitializedHelper(this, Property);
-				}
-			}
-			else if (UE::IsInstanceDataObjectSupportEnabled(this))
-			{
-				// IDOs as a general system are enabled AND they are enabled on this object
-				// IDO Enabled case
-				UE::FPropertyBagRepository& Repository = UE::FPropertyBagRepository::Get();
-		
-				UObject* ChangedIDO = UE::IsClassOfInstanceDataObjectClass(this->GetClass()) ?
-										this : Repository.FindInstanceDataObject(this);
-				if (ChangedIDO)
-				{
-					SetPropertyValueInitializedHelper(ChangedIDO, Property);
-				}
-			}
+		case EPropertyChangeType::ArrayAdd:
+		case EPropertyChangeType::ArrayRemove:
+		case EPropertyChangeType::ArrayClear: // Clearing the entire array. No need for an index. Clear also counts as initialization of the property.
+		case EPropertyChangeType::ArrayMove:
+		case EPropertyChangeType::ValueSet:
+		case EPropertyChangeType::Interactive:
+			// No tracking for elements of container. Any change counts as initialization.
+			UE::SetPropertyValueInitialized(GetClass(), this, Property);
+			break;
+		default:
+			// No op for now
+			break;
 		}
 	}
 #endif
@@ -1766,7 +1733,7 @@ void UObject::Serialize(FStructuredArchive::FRecord Record)
 			{
 				FOverridableManager::Get().SerializeOverriddenProperties(*this, Record);
 #if WITH_EDITORONLY_DATA
-				UE::SerializePropertyValueInitializationState(*this, Record);
+				UE::SerializePropertyValueInitialized(GetClass(), this, Record);
 #endif
 			}
 
diff --git a/Engine/Source/Runtime/CoreUObject/Private/UObject/PropertyStateTracking.cpp b/Engine/Source/Runtime/CoreUObject/Private/UObject/PropertyStateTracking.cpp
new file mode 100644
--- /dev/null
+++ b/Engine/Source/Runtime/CoreUObject/Private/UObject/PropertyStateTracking.cpp
@@ -0,0 +1,70 @@
+// Copyright Epic Games, Inc. All Rights Reserved.
+
+#include "UObject/PropertyStateTracking.h"
+
+#if WITH_EDITORONLY_DATA
+
+#include "Serialization/StructuredArchive.h"
+#include "UObject/Class.h"
+
+namespace UE
+{
+
+bool ActivateTrackingPropertyValueInitialized(const UStruct* Struct, void* Data)
+{
+	return Struct->ActivateTrackingPropertyValueFlag(EPropertyValueFlags::Initialized, Data);
+}
+
+bool IsTrackingPropertyValueInitialized(const UStruct* Struct, const void* Data)
+{
+	return Struct->IsTrackingPropertyValueFlag(EPropertyValueFlags::Initialized, Data);
+}
+
+bool IsPropertyValueInitialized(const UStruct* Struct, const void* Data, const FProperty* Property, int32 ArrayIndex)
+{
+	return Struct->HasPropertyValueFlag(EPropertyValueFlags::Initialized, Data, Property, ArrayIndex);
+}
+
+void SetPropertyValueInitialized(const UStruct* Struct, void* Data, const FProperty* Property, int32 ArrayIndex)
+{
+	Struct->SetPropertyValueFlag(EPropertyValueFlags::Initialized, /*bValue*/ true, Data, Property, ArrayIndex);
+}
+
+void ClearPropertyValueInitialized(const UStruct* Struct, void* Data, const FProperty* Property, int32 ArrayIndex)
+{
+	Struct->SetPropertyValueFlag(EPropertyValueFlags::Initialized, /*bValue*/ false, Data, Property, ArrayIndex);
+}
+
+void ResetPropertyValueInitialized(const UStruct* Struct, void* Data)
+{
+	Struct->ResetPropertyValueFlags(EPropertyValueFlags::Initialized, Data);
+}
+
+void SerializePropertyValueInitialized(const UStruct* Struct, void* Data, FStructuredArchiveRecord Record)
+{
+	Struct->SerializePropertyValueFlags(EPropertyValueFlags::Initialized, Data, Record, TEXT("PropertyValueInitializedFlags"));
+}
+
+bool ActivateTrackingPropertyValueSerialized(const UStruct* Struct, void* Data)
+{
+	return Struct->ActivateTrackingPropertyValueFlag(EPropertyValueFlags::Serialized, Data);
+}
+
+bool IsTrackingPropertyValueSerialized(const UStruct* Struct, const void* Data)
+{
+	return Struct->IsTrackingPropertyValueFlag(EPropertyValueFlags::Serialized, Data);
+}
+
+void MarkPropertyValueSerialized(const UStruct* Struct, void* Data, const FProperty* Property, int32 ArrayIndex)
+{
+	Struct->SetPropertyValueFlag(EPropertyValueFlags::Serialized, /*bValue*/ true, Data, Property, ArrayIndex);
+}
+
+bool WasPropertyValueSerialized(const UStruct* Struct, const void* Data, const FProperty* Property, int32 ArrayIndex)
+{
+	return Struct->HasPropertyValueFlag(EPropertyValueFlags::Serialized, Data, Property, ArrayIndex);
+}
+
+} // UE
+
+#endif // WITH_EDITORONLY_DATA
diff --git a/Engine/Source/Runtime/CoreUObject/Private/UObject/UObjectGlobals.cpp b/Engine/Source/Runtime/CoreUObject/Private/UObject/UObjectGlobals.cpp
--- a/Engine/Source/Runtime/CoreUObject/Private/UObject/UObjectGlobals.cpp
+++ b/Engine/Source/Runtime/CoreUObject/Private/UObject/UObjectGlobals.cpp
@@ -36,6 +36,7 @@
 #include "UObject/FastReferenceCollector.h"
 #include "UObject/InstanceDataObjectUtils.h"
 #include "UObject/OverridableManager.h"
+#include "UObject/PropertyStateTracking.h"
 #include "UObject/UObjectIterator.h"
 #include "UObject/Package.h"
 #include "Templates/Casts.h"
@@ -4045,6 +4046,15 @@ void FObjectInitializer::PostConstructInit()
 	}
 #endif // USE_CIRCULAR_DEPENDENCY_LOAD_DEFERRING
 
+#if WITH_EDITORONLY_DATA
+	// Activate tracking of initialized property state for Verse objects.
+	// The Class of the Class may not be available during very early initialization of CoreUObject.
+	if (UClass* ClassClass = Class->GetClass(); ClassClass && ClassClass->GetClass()->IsChildOf<UVerseClass>())
+	{
+		UE::ActivateTrackingPropertyValueInitialized(Class, Obj);
+	}
+#endif
+
 	if (bShouldInitializePropsFromArchetype)
 	{
 		UClass* BaseClass = (bIsCDO && !GIsDuplicatingClassForReinstancing) ? SuperClass : Class;
@@ -4277,35 +4287,19 @@ void FObjectInitializer::InitProperties(UObject* Obj, UClass* DefaultsClass, UOb
 		bCopyTransientsFromClassDefaults = false;
 	}
 
-	auto SetupPropertyValueInitializedState = [](UObject* TargetObject, UObject* Default, FProperty* Property)
-	{
 #if WITH_EDITORONLY_DATA
-		if (!UE::SupportsPropertyInitStateTracking())
-		{
-			return;
-		}
-		UClass* TargetClass = TargetObject->GetClass();
-		UClass* DefaultsClass = Default->GetClass();
-		
-		// PropertyValueInitialized state tracking only possible for UVerseClass types at this time
-		if (Cast<UVerseClass>(TargetClass) != nullptr)
+	auto SetUpPropertyInitializedFlag = [Obj, Class, bTrackInitializedProperties = UE::IsTrackingPropertyValueInitialized(Class, Obj)](UObject* Defaults, FProperty* Property)
+	{
+		if (bTrackInitializedProperties)
 		{
-			// For RequiredParm properties, check if the template's property is initialized.
-			// if it is, then data just copied in was initialized data.
-			const bool bHasRequiredParam = Property->HasAllPropertyFlags(CPF_RequiredParm);
-			bool bIsTemplateCounterpartInitialized = UE::IsPropertyValueInitialized(DefaultsClass, Default, Property);
-			
-			if (!bHasRequiredParam || bIsTemplateCounterpartInitialized)
-			{
-				UE::SetPropertyValueInitialized(TargetClass, TargetObject, Property);
-			}
-			else
-			{
-				UE::ClearPropertyValueInitialized(TargetClass, TargetObject, Property);
-			}
+			// A required property is only initialized if the copied default is initialized.
+			const bool bInitialized =
+				!Property->HasAllPropertyFlags(CPF_RequiredParm) ||
+				Defaults->GetClass()->HasPropertyValueFlag(EPropertyValueFlags::Initialized, Defaults, Property);
+			Class->SetPropertyValueFlag(EPropertyValueFlags::Initialized, bInitialized, Obj, Property);
 		}
-#endif
 	};
+#endif // WITH_EDITORONLY_DATA
 
 	if (!bNeedInitialize && bCanUsePostConstructLink)
 	{
@@ -4339,7 +4333,9 @@ void FObjectInitializer::InitProperties(UObject* Obj, UClass* DefaultsClass, UOb
 						if (P->IsInContainer(DefaultsClass))
 						{
 							P->CopyCompleteValue_InContainer(Obj, DefaultData);
-							SetupPropertyValueInitializedState(Obj, DefaultData, P);
+						#if WITH_EDITORONLY_DATA
+							SetUpPropertyInitializedFlag(DefaultData, P);
+						#endif
 						}
 					}
 				}
@@ -4357,7 +4353,7 @@ void FObjectInitializer::InitProperties(UObject* Obj, UClass* DefaultsClass, UOb
 		for (FProperty* P = bCanUsePostConstructLink ? Class->PostConstructLink : Class->PropertyLink; P; P = bCanUsePostConstructLink ? P->PostConstructLinkNext : P->PropertyLinkNext)
 		{
 			if (bNeedInitialize)
-			{		
+			{
 				bNeedInitialize = InitNonNativeProperty(P, Obj);
 			}
 
@@ -4369,12 +4365,16 @@ void FObjectInitializer::InitProperties(UObject* Obj, UClass* DefaultsClass, UOb
 					// This is a duplicate. The value for all transient or non-duplicatable properties should be copied
 					// from the source class's defaults.
 					P->CopyCompleteValue_InContainer(Obj, ClassDefaults);
-					SetupPropertyValueInitializedState(Obj, ClassDefaults, P);
+				#if WITH_EDITORONLY_DATA
+					SetUpPropertyInitializedFlag(ClassDefaults, P);
+				#endif
 				}
 				else if (P->IsInContainer(DefaultsClass))
 				{
 					P->CopyCompleteValue_InContainer(Obj, DefaultData);
-					SetupPropertyValueInitializedState(Obj, DefaultData, P);					
+				#if WITH_EDITORONLY_DATA
+					SetUpPropertyInitializedFlag(DefaultData, P);
+				#endif
 				}
 			}
 		}
diff --git a/Engine/Source/Runtime/CoreUObject/Public/UObject/Class.h b/Engine/Source/Runtime/CoreUObject/Public/UObject/Class.h
--- a/Engine/Source/Runtime/CoreUObject/Public/UObject/Class.h
+++ b/Engine/Source/Runtime/CoreUObject/Public/UObject/Class.h
@@ -86,8 +86,10 @@ class FField;
 class FOutputDevice;
 class FProperty;
 class FStructProperty;
+class FStructuredArchiveRecord;
 class UPropertyWrapper;
 struct CGetTypeHashable;
+struct FArchiveFieldName;
 struct FAssetData;
 struct FBlake3Hash;
 struct FColor;
@@ -351,6 +353,19 @@ class UField : public UObject
 	using FLinkedListBuilder = TLinkedListBuilder<UField>;
 };
 
+#if WITH_EDITORONLY_DATA
+/**
+ * Flags that are optionally tracked per property by struct and class instances.
+ */
+enum class EPropertyValueFlags : uint8
+{
+	/** Whether the property value is initialized. */
+	Initialized,
+	/** Whether the property value was written during serialization. */
+	Serialized,
+};
+#endif // WITH_EDITORONLY_DATA
+
 #if USTRUCT_FAST_ISCHILDOF_IMPL == USTRUCT_ISCHILDOF_STRUCTARRAY
 class FStructBaseChain
 {
@@ -688,6 +703,28 @@ public:
 #endif // WITH_METADATA
 
 #if WITH_EDITORONLY_DATA
+	/**
+	 * Try to activate tracking of this type of flag for the instance.
+	 *
+	 * @return Whether the instance is tracking this type of flag.
+	 */
+	UE_INTERNAL COREUOBJECT_API virtual bool ActivateTrackingPropertyValueFlag(EPropertyValueFlags Flags, void* Data) const;
+
+	/** Query whether the instance is tracking this type of flag. */
+	UE_INTERNAL COREUOBJECT_API virtual bool IsTrackingPropertyValueFlag(EPropertyValueFlags Flags, const void* Data) const;
+
+	/** Query whether the instance has the flag set for the property. True if tracking is inactive. */
+	UE_INTERNAL COREUOBJECT_API virtual bool HasPropertyValueFlag(EPropertyValueFlags Flags, const void* Data, const FProperty* Property, int32 ArrayIndex = 0) const;
+
+	/** Set the value of the flag for the property on the instance. Ignored if tracking is inactive. */
+	UE_INTERNAL COREUOBJECT_API virtual void SetPropertyValueFlag(EPropertyValueFlags Flags, bool bValue, void* Data, const FProperty* Property, int32 ArrayIndex = 0) const;
+
+	/** Reset the flag to false for every property on the instance. */
+	UE_INTERNAL COREUOBJECT_API virtual void ResetPropertyValueFlags(EPropertyValueFlags Flags, void* Data) const;
+
+	/** Serialize the flag for every property on the instance. NOT FOR PERSISTENT SERIALIZATION! */
+	UE_INTERNAL COREUOBJECT_API virtual void SerializePropertyValueFlags(EPropertyValueFlags Flags, void* Data, FStructuredArchiveRecord Record, FArchiveFieldName Name) const;
+
 	/* Returns true if this struct has Asset Registry searchable properties */
 	FORCEINLINE bool HasAssetRegistrySearchableProperties() const
 	{
@@ -3246,6 +3283,15 @@ public:
 	COREUOBJECT_API virtual bool IsStructTrashed() const override;
 	// End of UStruct interface.
 
+#if WITH_EDITORONLY_DATA
+	UE_INTERNAL COREUOBJECT_API virtual bool ActivateTrackingPropertyValueFlag(EPropertyValueFlags Flags, void* Data) const override;
+	UE_INTERNAL COREUOBJECT_API virtual bool IsTrackingPropertyValueFlag(EPropertyValueFlags Flags, const void* Data) const override;
+	UE_INTERNAL COREUOBJECT_API virtual bool HasPropertyValueFlag(EPropertyValueFlags Flags, const void* Data, const FProperty* Property, int32 ArrayIndex = 0) const override;
+	UE_INTERNAL COREUOBJECT_API virtual void SetPropertyValueFlag(EPropertyValueFlags Flags, bool bValue, void* Data, const FProperty* Property, int32 ArrayIndex = 0) const override;
+	UE_INTERNAL COREUOBJECT_API virtual void ResetPropertyValueFlags(EPropertyValueFlags Flags, void* Data) const override;
+	UE_INTERNAL COREUOBJECT_API virtual void SerializePropertyValueFlags(EPropertyValueFlags Flags, void* Data, FStructuredArchiveRecord Record, FArchiveFieldName Name) const override;
+#endif // WITH_EDITORONLY_DATA
+
 	/**
 	 * Returns class path name as a package + class FName pair
 	 */
