Remove unnecessary forward declaration of FIoChunkId as we are already including the header for it.

#rnx

diff --git a/Engine/Source/Runtime/CoreUObject/Public/Serialization/BulkData.h b/Engine/Source/Runtime/CoreUObject/Public/Serialization/BulkData.h
--- a/Engine/Source/Runtime/CoreUObject/Public/Serialization/BulkData.h
+++ b/Engine/Source/Runtime/CoreUObject/Public/Serialization/BulkData.h
@@ -37,7 +37,6 @@
 #include "UObject/NameTypes.h"
 
 class FIoBuffer;
-class FIoChunkId;
 class FLinkerLoad;
 class FOutputDevice;
 class IAsyncReadFileHandle;
