When serializing remote objects Weak Object Pointers wil be serialized using the same path as strong references with the exception that they will not be force resolved on load.
TMulticastScriptDelegate will no longer attempt to compact the invocation list when deserializing remote objects which would result in resolving all weak object pointers.
FObjectPropertyBase::CheckValidObject will no longer attempt to resolve remote objects

diff --git a/Engine/Source/Runtime/Core/Public/UObject/ScriptDelegates.h b/Engine/Source/Runtime/Core/Public/UObject/ScriptDelegates.h
--- a/Engine/Source/Runtime/Core/Public/UObject/ScriptDelegates.h
+++ b/Engine/Source/Runtime/Core/Public/UObject/ScriptDelegates.h
@@ -844,7 +844,9 @@ public:
 			FWriteAccessScope WriteScope = D.GetWriteAccessScope();
 			Ar << D.InvocationList;
 			// After loading the delegate, clean up the list to make sure there are no bad object references
-			if (Ar.IsLoading())
+			// unless the archive is used to migrate remote objects in which case we don't want to resolve weak object pointers
+			// which may be pointing to objects that haven't been migrated yet
+			if (Ar.IsLoading() && !Ar.IsMigratingRemoteObjects())
 			{
 				D.CompactInvocationList();
 			}
diff --git a/Engine/Source/Runtime/CoreUObject/Private/UObject/PropertyBaseObject.cpp b/Engine/Source/Runtime/CoreUObject/Private/UObject/PropertyBaseObject.cpp
--- a/Engine/Source/Runtime/CoreUObject/Private/UObject/PropertyBaseObject.cpp
+++ b/Engine/Source/Runtime/CoreUObject/Private/UObject/PropertyBaseObject.cpp
@@ -773,6 +773,14 @@ void FObjectPropertyBase::CheckValidObject(void* ValueAddress, TObjectPtr<UObjec
 	{
 		return;
 	}
+#if UE_WITH_REMOTE_OBJECT_HANDLE
+	// Verifying the object class will result in resolving remote objects and we don't want that.
+	// However, since remote objects are guarnteed to have compatible binaries it's not required to validate objects.
+	if (!Object.IsResolved())
+	{
+		return;
+	}
+#endif
 	//
 	// here we want to make sure the the object value still matches the 
 	// object type expected by the property...
diff --git a/Engine/Source/Runtime/CoreUObject/Private/UObject/RemoteObjectSerialization.cpp b/Engine/Source/Runtime/CoreUObject/Private/UObject/RemoteObjectSerialization.cpp
--- a/Engine/Source/Runtime/CoreUObject/Private/UObject/RemoteObjectSerialization.cpp
+++ b/Engine/Source/Runtime/CoreUObject/Private/UObject/RemoteObjectSerialization.cpp
@@ -225,6 +225,7 @@ public:
 
 		SetIsPersistent(false);
 		SetUseUnversionedPropertySerialization(true);
+		SetMigratingRemoteObjects(true);
 	}
 
 	virtual ~FArchiveRemoteObjectWriter()
@@ -366,7 +367,13 @@ public:
 	}
 	virtual FArchive& operator<<(FSoftObjectPtr& Value) override { return FArchiveUObject::SerializeSoftObjectPtr(*this, Value); }
 	virtual FArchive& operator<<(FSoftObjectPath& Value) override { return FArchiveUObject::SerializeSoftObjectPath(*this, Value); }
-	virtual FArchive& operator<<(FWeakObjectPtr& Value) override { return FArchiveUObject::SerializeWeakObjectPtr(*this, Value); }
+	virtual FArchive& operator<<(FWeakObjectPtr& Value) override
+	{
+		FObjectPtr Ptr(Value.Get(/** bEvenIfGarbage */ true));
+		WriteObjectReference(Ptr);
+
+		return *this;
+	}
 };
 
 struct FRemoteObjectWriterScope
@@ -483,6 +490,7 @@ public:
 	{
 		SetIsPersistent(false);
 		SetUseUnversionedPropertySerialization(true);
+		SetMigratingRemoteObjects(true);
 	}
 
 	virtual FArchive& operator<<(UObject*& Obj) override
@@ -511,7 +519,24 @@ public:
 
 	virtual FArchive& operator<<(FSoftObjectPtr& Value) override { return FArchiveUObject::SerializeSoftObjectPtr(*this, Value); }
 	virtual FArchive& operator<<(FSoftObjectPath& Value) override { return FArchiveUObject::SerializeSoftObjectPath(*this, Value); }
-	virtual FArchive& operator<<(FWeakObjectPtr& Value) override { return FArchiveUObject::SerializeWeakObjectPtr(*this, Value); }
+	virtual FArchive& operator<<(FWeakObjectPtr& Value) override 
+	{
+		FObjectPtr Ptr;
+		ReadObjectReference(Ptr);
+
+		if (Ptr.IsResolved())
+		{
+			Value = Ptr.Get();
+		}
+#if UE_WITH_REMOTE_OBJECT_HANDLE
+		else
+		{
+			Value = FWeakObjectPtr(Ptr.GetRemoteId());
+		}
+#endif // UE_WITH_REMOTE_OBJECT_HANDLE
+
+		return *this;
+	}
 };
 
 /**
diff --git a/Engine/Source/Runtime/CoreUObject/Private/UObject/WeakObjectPtr.cpp b/Engine/Source/Runtime/CoreUObject/Private/UObject/WeakObjectPtr.cpp
--- a/Engine/Source/Runtime/CoreUObject/Private/UObject/WeakObjectPtr.cpp
+++ b/Engine/Source/Runtime/CoreUObject/Private/UObject/WeakObjectPtr.cpp
@@ -179,8 +179,6 @@ FUObjectItem* FWeakObjectPtr::TryResolveRemoteObject() const
 			FUObjectItem* ObjectItem = GUObjectArray.IndexToObject(NewIndex);
 			checkf(ObjectItem->GetRemoteId() == ObjectRemoteId, TEXT("Remote object %s has a different id (%s) than the weak object pointer that resolved it (%s)"),
 				*ResolvedObject->GetPathName(), *ObjectItem->GetRemoteId().ToString(), *ObjectRemoteId.ToString());
-			checkf(ObjectItem->GetSerialNumber() == ObjectSerialNumber, TEXT("Resolved remote object %s has a different serial number (%d) than the weak object pointer that resolved it (%d)"),
-				*ResolvedObject->GetPathName(), ObjectItem->GetSerialNumber(), ObjectSerialNumber);
 			if (ObjectIndex != NewIndex)
 			{
 				// Currently we can't just rely on object id to resolve weak object pointers and we don't want to heep hitting the resolve path
@@ -190,6 +188,12 @@ FUObjectItem* FWeakObjectPtr::TryResolveRemoteObject() const
 				// With RemoteId neither ObjectIndex nor SerialNumber is used for FWeakObjectPtr comparisons/hash value calculation
 				FPlatformAtomics::AtomicStore(&const_cast<FWeakObjectPtr*>(this)->ObjectIndex, NewIndex);
 			}
+			if (ObjectSerialNumber != ObjectItem->GetSerialNumber())
+			{
+				checkf(ObjectSerialNumber == 0, TEXT("Attempting to change existing and valid serial number %d to %d when resolving remote object %s (%s)"),
+					ObjectSerialNumber, ObjectItem->GetSerialNumber(), *ResolvedObject->GetPathName(), *ObjectRemoteId.ToString());
+				FPlatformAtomics::AtomicStore(&const_cast<FWeakObjectPtr*>(this)->ObjectSerialNumber, ObjectItem->GetSerialNumber());
+			}
 			return ObjectItem;
 		}
 	}
diff --git a/Engine/Source/Runtime/CoreUObject/Public/UObject/WeakObjectPtr.h b/Engine/Source/Runtime/CoreUObject/Public/UObject/WeakObjectPtr.h
--- a/Engine/Source/Runtime/CoreUObject/Public/UObject/WeakObjectPtr.h
+++ b/Engine/Source/Runtime/CoreUObject/Public/UObject/WeakObjectPtr.h
@@ -108,6 +108,18 @@ public:
 	 */
 	FWeakObjectPtr(const FWeakObjectPtr& Other) = default;
 
+#if UE_WITH_REMOTE_OBJECT_HANDLE
+	explicit FWeakObjectPtr(FRemoteObjectId RemoteId)
+	{
+		// When using remote object handles both ObjectIndex and SerialNumber are not used when comparing weak object pointers
+		// ObjectIndex is used for performance only and will be updated when this weak pointer is accessed
+		// SerialNumber will also get updated when this weak pointer is resolved
+		ObjectIndex = 0;		
+		ObjectSerialNumber = 0;
+		ObjectRemoteId = RemoteId;
+	}
+#endif // UE_WITH_REMOTE_OBJECT_HANDLE
+
 	/**
 	 * Reset the weak pointer back to the null state
 	 */
