UE::RemoteObject::Handle::IsRemote(FRemoteObjectId) will now properly check if an id represents a remote object

diff --git a/Engine/Source/Runtime/CoreUObject/Private/UObject/RemoteObject.cpp b/Engine/Source/Runtime/CoreUObject/Private/UObject/RemoteObject.cpp
--- a/Engine/Source/Runtime/CoreUObject/Private/UObject/RemoteObject.cpp
+++ b/Engine/Source/Runtime/CoreUObject/Private/UObject/RemoteObject.cpp
@@ -32,11 +32,10 @@ class FRemoteObjectMaps
 
 public:
 
-	bool IsRemote(FRemoteObjectId ObjectId) const
+	bool IsMapped(FRemoteObjectId ObjectId) const
 	{
 		FScopeLock MapLock(&ObjectMapCritical);
-		bool bResult = !IdToObjectMap.Contains(ObjectId);
-		return bResult;
+		return IdToObjectMap.Contains(ObjectId);
 	}
 
 	void MapObject(UObject* Object)
@@ -212,14 +211,27 @@ FRemoteObjectStub::FRemoteObjectStub(UObject* Object)
 
 bool IsRemote(FRemoteObjectId ObjectId)
 {
-	return UE::RemoteObject::Private::ObjectMaps->IsRemote(ObjectId);
+	using namespace UE::RemoteObject::Private;
+
+	if (!ObjectId.IsValid())
+	{
+		return false;
+	}
+	if (UObject* Object = ObjectMaps->FindObject(ObjectId))
+	{
+		return IsRemote(Object);
+	}
+	else
+	{
+		return !!FindRemoteObjectStub(ObjectId);
+	}
 }
 
-bool IsRemote(UObject* Object)
+bool IsRemote(const UObject* Object)
 {
 	const int32 InternalIndex = GUObjectArray.ObjectToIndex(Object);
-	bool bResult = InternalIndex >= 0 && GUObjectArray.IndexToObject(InternalIndex)->HasAnyFlags(EInternalObjectFlags::Remote);
-	return bResult;
+	const bool bIsRemote = InternalIndex >= 0 && GUObjectArray.IndexToObject(InternalIndex)->HasAnyFlags(EInternalObjectFlags::Remote);
+	return bIsRemote;
 }
 
 UObject* ResolveObject(const FRemoteObjectStub* Stub)
@@ -227,20 +239,13 @@ UObject* ResolveObject(const FRemoteObjectStub* Stub)
 	using namespace UE::RemoteObject::Transfer;
 	using namespace UE::RemoteObject::Private;
 
-	checkf(!GIsGarbageCollecting, TEXT("Resolving remote objects while collecting garbage is not allowed (trying to resolve object %s (%s)"), *Stub->Id.ToString(), *Stub->Name.ToString());
-
+	// This is a slightly faster version of IsRemote(FRemoteObjectId) because we already know a Stub exists and we are going to re-use the Object pointer
 	UObject* Object = ObjectMaps->FindObject(Stub->Id);
-	bool bRemoteObject = false;
-	if (Object)
-	{
-		bRemoteObject = IsRemote(Object);
-	}
-	else
-	{
-		bRemoteObject = IsRemote(Stub->Id);
-	}
+	const bool bRemoteObject = !Object || IsRemote(Object);
 	if (bRemoteObject)
 	{
+		checkf(!GIsGarbageCollecting, TEXT("Resolving remote objects while collecting garbage is not allowed (trying to resolve object %s (%s)"), *Stub->Id.ToString(), *Stub->Name.ToString());
+
 		UObject* Outer = Object ? Object->GetOuter() : ObjectMaps->FindObject(Stub->OuterId);
 		Object = MigrateObjectFromRemoteServer(Stub->Id, Stub->RemoteServerId, Outer);
 	}
diff --git a/Engine/Source/Runtime/CoreUObject/Public/UObject/RemoteObject.h b/Engine/Source/Runtime/CoreUObject/Public/UObject/RemoteObject.h
--- a/Engine/Source/Runtime/CoreUObject/Public/UObject/RemoteObject.h
+++ b/Engine/Source/Runtime/CoreUObject/Public/UObject/RemoteObject.h
@@ -60,5 +60,5 @@ namespace UE::RemoteObject::Handle
 	* Checks if an object (memory that has not yet been GC'd) is remote
 	* @return True if the object is remote
 	*/
-	COREUOBJECT_API bool IsRemote(UObject* Object);
+	COREUOBJECT_API bool IsRemote(const UObject* Object);
 }
