Adding FWeakObjectPtr::GetRemoteId() (compiled out by default)

diff --git a/Engine/Source/Runtime/CoreUObject/Public/UObject/WeakObjectPtr.h b/Engine/Source/Runtime/CoreUObject/Public/UObject/WeakObjectPtr.h
--- a/Engine/Source/Runtime/CoreUObject/Public/UObject/WeakObjectPtr.h
+++ b/Engine/Source/Runtime/CoreUObject/Public/UObject/WeakObjectPtr.h
@@ -295,6 +295,16 @@ public:
 #endif // UE_WITH_REMOTE_OBJECT_HANDLE
 	}
 
+#if UE_WITH_REMOTE_OBJECT_HANDLE
+	/**
+	* Returns remote Id of the object this pointer is referencing
+	*/
+	FORCEINLINE FRemoteObjectId GetRemoteId() const
+	{
+		return ObjectRemoteId;
+	}
+#endif // UE_WITH_REMOTE_OBJECT_HANDLE
+
 	/** Hash function. */
 	FORCEINLINE uint32 GetTypeHash() const
 	{
