Fix FScriptMap/Set const correctness and adding FScriptMap/SetHelper constructor

#rnx

diff --git a/Engine/Source/Runtime/Core/Public/Containers/Map.h b/Engine/Source/Runtime/Core/Public/Containers/Map.h
--- a/Engine/Source/Runtime/Core/Public/Containers/Map.h
+++ b/Engine/Source/Runtime/Core/Public/Containers/Map.h
@@ -1910,7 +1910,7 @@ public:
 	}
 	
 	/** Finds the associated key, value from hash of Key, rather than linearly searching */
-	int32 FindPairIndex(const void* Key, const FScriptMapLayout& MapLayout, TFunctionRef<uint32(const void*)> GetKeyHash, TFunctionRef<bool(const void*, const void*)> KeyEqualityFn)
+	int32 FindPairIndex(const void* Key, const FScriptMapLayout& MapLayout, TFunctionRef<uint32(const void*)> GetKeyHash, TFunctionRef<bool(const void*, const void*)> KeyEqualityFn) const
 	{
 		if (Pairs.Num())
 		{
diff --git a/Engine/Source/Runtime/Core/Public/Containers/Set.h b/Engine/Source/Runtime/Core/Public/Containers/Set.h
--- a/Engine/Source/Runtime/Core/Public/Containers/Set.h
+++ b/Engine/Source/Runtime/Core/Public/Containers/Set.h
@@ -2048,7 +2048,7 @@ public:
 	}
 
 private:
-	int32 FindIndexImpl(const void* Element, const FScriptSetLayout& Layout, uint32 KeyHash, TFunctionRef<bool (const void*, const void*)> EqualityFn)
+	int32 FindIndexImpl(const void* Element, const FScriptSetLayout& Layout, uint32 KeyHash, TFunctionRef<bool (const void*, const void*)> EqualityFn) const
 	{
 		const int32  HashIndex = KeyHash & (HashSize - 1);
 
@@ -2068,7 +2068,7 @@ private:
 	}
 
 public:
-	int32 FindIndex(const void* Element, const FScriptSetLayout& Layout, TFunctionRef<uint32 (const void*)> GetKeyHash, TFunctionRef<bool (const void*, const void*)> EqualityFn)
+	int32 FindIndex(const void* Element, const FScriptSetLayout& Layout, TFunctionRef<uint32 (const void*)> GetKeyHash, TFunctionRef<bool (const void*, const void*)> EqualityFn) const
 	{
 		if (Elements.Num())
 		{
@@ -2079,7 +2079,7 @@ public:
 	}
 
 
-	int32 FindIndexByHash(const void* Element, const FScriptSetLayout& Layout, uint32 KeyHash, TFunctionRef<bool (const void*, const void*)> EqualityFn)
+	int32 FindIndexByHash(const void* Element, const FScriptSetLayout& Layout, uint32 KeyHash, TFunctionRef<bool (const void*, const void*)> EqualityFn) const
 	{
 		if (Elements.Num())
 		{
diff --git a/Engine/Source/Runtime/CoreUObject/Public/UObject/UnrealType.h b/Engine/Source/Runtime/CoreUObject/Public/UObject/UnrealType.h
--- a/Engine/Source/Runtime/CoreUObject/Public/UObject/UnrealType.h
+++ b/Engine/Source/Runtime/CoreUObject/Public/UObject/UnrealType.h
@@ -4429,8 +4429,13 @@ public:
 	 * @param  InProperty  The property associated with this memory
 	 * @param  InMap       Pointer to raw memory that corresponds to this map. This can be NULL, and sometimes is, but in that case almost all operations will crash.
 	 */
-	FORCEINLINE FScriptMapHelper(const FMapProperty* InProperty, const void* InMap)
-		: FScriptMapHelper(Internal, InProperty->KeyProp, InProperty->ValueProp, InMap, InProperty->MapLayout, InProperty->MapFlags)
+	FScriptMapHelper(const FMapProperty* InProperty, const void* InMap)
+	: FScriptMapHelper(Internal, InProperty->KeyProp, InProperty->ValueProp, InMap, InProperty->MapLayout, InProperty->MapFlags)
+	{
+	}
+
+	FScriptMapHelper(FProperty* InKeyProp, FProperty* InValueProp, const void* InMap, const FScriptMapLayout& InMapLayout, EMapPropertyFlags InMapFlags)
+	: FScriptMapHelper(Internal, InKeyProp, InValueProp, InMap, InMapLayout, InMapFlags)
 	{
 	}
 
@@ -5397,10 +5402,15 @@ public:
 	* @param  InProperty  The property associated with this memory
 	* @param  InSet       Pointer to raw memory that corresponds to this Set. This can be NULL, and sometimes is, but in that case almost all operations will crash.
 	*/
-	FORCEINLINE FScriptSetHelper(const FSetProperty* InProperty, const void* InSet)
-		: ElementProp(InProperty->ElementProp)
-		, Set((FScriptSet*)InSet)  //@todo, we are casting away the const here
-		, SetLayout(InProperty->SetLayout)
+	FScriptSetHelper(const FSetProperty* InProperty, const void* InSet)
+	: FScriptSetHelper(InProperty->ElementProp, InSet, InProperty->SetLayout)
+	{
+	}
+
+	FScriptSetHelper(FProperty* InElementProp, const void* InSet, const FScriptSetLayout& InLayout)
+	: ElementProp(InElementProp)
+	, Set((FScriptSet*)InSet)  //@todo, we are casting away the const here
+	, SetLayout(InLayout)
 	{
 		check(ElementProp);
 	}
