Fix for ActivateTrackingPropertyValueInitialized not getting called on Verse classes.

#rnx

diff --git a/Engine/Source/Runtime/CoreUObject/Private/UObject/UObjectGlobals.cpp b/Engine/Source/Runtime/CoreUObject/Private/UObject/UObjectGlobals.cpp
--- a/Engine/Source/Runtime/CoreUObject/Private/UObject/UObjectGlobals.cpp
+++ b/Engine/Source/Runtime/CoreUObject/Private/UObject/UObjectGlobals.cpp
@@ -4049,7 +4049,7 @@ void FObjectInitializer::PostConstructInit()
 #if WITH_EDITORONLY_DATA
 	// Activate tracking of initialized property state for Verse objects.
 	// The Class of the Class may not be available during very early initialization of CoreUObject.
-	if (UClass* ClassClass = Class->GetClass(); ClassClass && ClassClass->GetClass()->IsChildOf<UVerseClass>())
+	if (UClass* ClassClass = Class->GetClass(); ClassClass && ClassClass->IsChildOf<UVerseClass>())
 	{
 		UE::ActivateTrackingPropertyValueInitialized(Class, Obj);
 	}
