UObject::IsAsset will now return true for public objects that are not directly outered to a package if they are external.

diff --git a/Engine/Source/Runtime/CoreUObject/Private/UObject/Obj.cpp b/Engine/Source/Runtime/CoreUObject/Private/UObject/Obj.cpp
--- a/Engine/Source/Runtime/CoreUObject/Private/UObject/Obj.cpp
+++ b/Engine/Source/Runtime/CoreUObject/Private/UObject/Obj.cpp
@@ -2621,10 +2621,18 @@ bool UObject::IsAsset() const
 	if ( bHasValidObjectFlags && !GetClass()->HasAnyClassFlags(CLASS_Optional) )
 	{
 		// Don't count objects embedded in other objects (e.g. font textures, sequences, material expressions)
-		if ( UPackage* LocalOuterPackage = dynamic_cast<UPackage*>(GetOuter()) )
+		UPackage* ObjectPackage = Cast<UPackage>(GetOuter());
+
+		if ( !ObjectPackage )
+		{
+			// Check is the object is directly stored in an external package, in such case treat it as an asset also
+			ObjectPackage = GetExternalPackage();
+		}
+
+		if ( ObjectPackage )
 		{
 			// Also exclude any objects found in the transient package, in a package that is transient or in a play in editor package.
-			return LocalOuterPackage != GetTransientPackage() && !LocalOuterPackage->HasAnyFlags(RF_Transient) && !LocalOuterPackage->HasAnyPackageFlags(PKG_PlayInEditor);
+			return ObjectPackage != GetTransientPackage() && !ObjectPackage->HasAnyFlags(RF_Transient) && !ObjectPackage->HasAnyPackageFlags(PKG_PlayInEditor);
 		}
 	}
 
