UObjectBaseUtility::MarkPackageDirty early-out if the object it's called on is transient, but wasn't taking into account if the object was part of a transient object in its outer chain.

diff --git a/Engine/Source/Runtime/CoreUObject/Private/UObject/UObjectBaseUtility.cpp b/Engine/Source/Runtime/CoreUObject/Private/UObject/UObjectBaseUtility.cpp
--- a/Engine/Source/Runtime/CoreUObject/Private/UObject/UObjectBaseUtility.cpp
+++ b/Engine/Source/Runtime/CoreUObject/Private/UObject/UObjectBaseUtility.cpp
@@ -229,44 +229,50 @@ UPackage* UObjectBaseUtility::GetOutermost() const
 bool UObjectBaseUtility::MarkPackageDirty() const
 {
 	// since transient objects will never be saved into a package, there is no need to mark a package dirty
-	// if we're transient
-	if ( !HasAnyFlags(RF_Transient) )
+	// if we're transient along the outer chain up to the associated package.
 	{
-		UPackage* Package = GetOutermost();
-
-		if( Package != NULL	)
+		const UObjectBaseUtility* Object = this;
+		while( Object )
 		{
-			// It is against policy to dirty a map or package during load/undo/redo in the Editor, to enforce this policy
-			// we explicitly disable the ability to dirty a package or map during load/undo/redo.  Commandlets can still
-			// set the dirty state on load.
-			if( IsRunningCommandlet() || 
-				(!IsInAsyncLoadingThread() && GIsEditor && !UE::GetIsEditorLoadingPackage() && !GIsCookerLoadingPackage && !GIsPlayInEditorWorld && !IsReloadActive()
-#if WITH_EDITORONLY_DATA
-				&& !GIsTransacting
-				&& !Package->bIsCookedForEditor // Cooked packages can't be modified nor marked as dirty
-#endif
-				))
+			if( Object->HasAnyFlags(RF_Transient) )
 			{
-				const bool bIsDirty = Package->IsDirty();
-
-				// We prevent needless re-dirtying as this can be an expensive operation.
-				if( !bIsDirty )
-				{
-					Package->SetDirtyFlag(true);
-				}
-
-				// Always call PackageMarkedDirtyEvent, even when the package is already dirty
-				Package->PackageMarkedDirtyEvent.Broadcast(Package, bIsDirty);
-
 				return true;
 			}
-			else
+			Object = Object->GetOuter();
+		}
+	}
+
+	if( UPackage* Package = GetPackage() )
+	{
+		// It is against policy to dirty a map or package during load/undo/redo in the Editor, to enforce this policy
+		// we explicitly disable the ability to dirty a package or map during load/undo/redo.  Commandlets can still
+		// set the dirty state on load.
+		if( IsRunningCommandlet() || 
+			(!IsInAsyncLoadingThread() && GIsEditor && !UE::GetIsEditorLoadingPackage() && !GIsCookerLoadingPackage && !GIsPlayInEditorWorld && !IsReloadActive()
+#if WITH_EDITORONLY_DATA
+			&& !GIsTransacting
+			&& !Package->bIsCookedForEditor // Cooked packages can't be modified nor marked as dirty
+#endif
+			))
+		{
+			const bool bIsDirty = Package->IsDirty();
+
+			// We prevent needless re-dirtying as this can be an expensive operation.
+			if( !bIsDirty )
 			{
-				// notify the caller that the request to mark the package as dirty was suppressed
-				return false;
+				Package->SetDirtyFlag(true);
 			}
+
+			// Always call PackageMarkedDirtyEvent, even when the package is already dirty
+			Package->PackageMarkedDirtyEvent.Broadcast(Package, bIsDirty);
+
+			return true;
 		}
+
+		// notify the caller that the request to mark the package as dirty was suppressed
+		return false;
 	}
+
 	return true;
 }
 
@@ -683,6 +689,4 @@ void FScopeCycleCounterUObject::ReportHitch()
 	}
 	UE_LOG(LogCore, Error, TEXT("Leaving UObject scope on hitch (+%8.2fms) [%s] %s"), Delta, *ThreadString, *StackString);
 }
-#endif
-
-
+#endif
\ No newline at end of file
