Added annotation-based storage for unknown property tracking and unknown enum tracking

#rnx

diff --git a/Engine/Source/Runtime/CoreUObject/Internal/UObject/PropertyBagRepository.h b/Engine/Source/Runtime/CoreUObject/Internal/UObject/PropertyBagRepository.h
--- a/Engine/Source/Runtime/CoreUObject/Internal/UObject/PropertyBagRepository.h
+++ b/Engine/Source/Runtime/CoreUObject/Internal/UObject/PropertyBagRepository.h
@@ -6,7 +6,6 @@
 #include "Containers/Map.h"
 #include "HAL/CriticalSection.h"
 #include "UObject/GCObject.h"
-#include "UObject/PropertyTypeName.h"
 #include "Templates/FunctionFwd.h"
 
 class UObject;
@@ -17,9 +16,6 @@ namespace UE
 
 #if WITH_EDITORONLY_DATA
 
-class FPropertyPathNameTree;
-class FUnknownEnumNames;
-
 // Singleton class tracking property bag association with objects
 class FPropertyBagRepository : public FGCObject
 {
@@ -27,8 +23,6 @@ class FPropertyBagRepository : public FGCObject
 	{
 		void Destroy();
 
-		FPropertyPathNameTree* Tree = nullptr;
-		FUnknownEnumNames* EnumNames = nullptr;
 		TObjectPtr<UObject> InstanceDataObject = nullptr;
 		bool bNeedsFixup = false;
 	};
@@ -62,41 +56,6 @@ public:
 	// Reclaim space - TODO: Hook up to GC.
 	void ShrinkMaps();
 
-	/**
-	 * Finds or creates a property path name tree to collect unknown property paths within the owner.
-	 */
-	FPropertyPathNameTree* FindOrCreateUnknownPropertyTree(const UObject* Owner);
-
-	/**
-	 * Adds an unknown enum name to the names tracked for an object.
-	 *
-	 * @param Owner			The owner to associate the unknown name with.
-	 * @param Enum			The enum associated with the property being serialized. May be null or the wrong type.
-	 * @param EnumTypeName	The type name of the enum containing the unknown name.
-	 * @param EnumValueName	The unknown name to track.
-	 */
-	void AddUnknownEnumName(const UObject* Owner, const UEnum* Enum, FPropertyTypeName EnumTypeName, FName EnumValueName);
-
-	/**
-	 * Finds tracked unknown enum names associated with the object.
-	 *
-	 * @param Owner			The owner associated with the unknown names to find.
-	 * @param EnumTypeName	The type name of the enum containing the unknown names to find.
-	 * @param OutNames		Array to assign the unknown names to. Empty on return if no names are found.
-	 * @param bOutHasFlags	Assigned to true if the enum is known to have flags, otherwise false.
-	 */
-	void FindUnknownEnumNames(const UObject* Owner, FPropertyTypeName EnumTypeName, TArray<FName>& OutNames, bool& bOutHasFlags) const;
-
-	/**
-	 * Finds tracked unknown enum names associated with the object, otherwise null.
-	 */
-	const FUnknownEnumNames* FindUnknownEnumNames(const UObject* Owner) const;
-
-	/**
-	 * Resets tracked unknown enum names associated with the object.
-	 */
-	void ResetUnknownEnumNames(const UObject* Owner);
-
 	// Future version for reworked InstanceDataObjects - track InstanceDataObject rather than bag (directly):
 	/**
 	 * Instantiate an InstanceDataObject object representing all fields within the bag, tracked against the owner object.
@@ -196,22 +155,6 @@ private:
 	void CreateInstanceDataObjectUnsafe(UObject* Owner, FPropertyBagAssociationData& BagData, FArchive* Archive = nullptr);
 };
 
-class FUnknownEnumNames
-{
-public:
-	void Add(const UEnum* Enum, FPropertyTypeName EnumTypeName, FName EnumValueName);
-	void Find(FPropertyTypeName EnumTypeName, TArray<FName>& OutNames, bool& bOutHasFlags) const;
-
-private:
-	struct FInfo
-	{
-		TSet<FName> Names;
-		bool bHasFlags = false;
-	};
-
-	TMap<FPropertyTypeName, FInfo> Enums;
-};
-
 #endif // WITH_EDITORONLY_DATA
 
 // construct this context in the same scope as an object is being serialized to support InstanceDataObjects.
diff --git a/Engine/Source/Runtime/CoreUObject/Internal/UObject/PropertyStateTracking.h b/Engine/Source/Runtime/CoreUObject/Internal/UObject/PropertyStateTracking.h
--- a/Engine/Source/Runtime/CoreUObject/Internal/UObject/PropertyStateTracking.h
+++ b/Engine/Source/Runtime/CoreUObject/Internal/UObject/PropertyStateTracking.h
@@ -6,9 +6,17 @@
 
 #if WITH_EDITORONLY_DATA
 
+#include "Containers/ContainersFwd.h"
+#include "Templates/SharedPointerFwd.h"
+
+class FName;
 class FProperty;
 class FStructuredArchiveRecord;
+class UEnum;
+class UObject;
 class UStruct;
+namespace UE { class FPropertyPathNameTree; }
+namespace UE { class FPropertyTypeName; }
 
 #define UE_API COREUOBJECT_API
 
@@ -48,6 +56,66 @@ void MarkPropertyValueSerialized(const UStruct* Struct, void* Data, const FPrope
 /** Query whether the property value was deserialized for a property in a struct. */
 bool WasPropertyValueSerialized(const UStruct* Struct, const void* Data, const FProperty* Property, int32 ArrayIndex = 0);
 
+/**
+ * A tree of unknown properties found during serialization. Used with bTrackUnknownProperties.
+ */
+class FUnknownPropertyTree
+{
+public:
+	/** Constructs an accessor for the unknown property tree of the owner. */
+	explicit FUnknownPropertyTree(const UObject* Owner);
+
+	/** Finds the existing unknown property path name tree for the owner. */
+	TSharedPtr<FPropertyPathNameTree> Find() const;
+
+	/** Finds the existing unknown property path name tree for the owner or creates one. */
+	TSharedPtr<FPropertyPathNameTree> FindOrCreate();
+
+	/** Destroys the unknown property path name tree for the owner if it has one. */
+	void Destroy();
+
+private:
+	const UObject* Owner = nullptr;
+};
+
+/**
+ * A record of unknown enumerator names found during serialization. Used with bTrackUnknownEnumNames.
+ */
+class FUnknownEnumNames
+{
+public:
+	/** Constructs an accessor for the unknown enum names of the owner. */
+	explicit FUnknownEnumNames(const UObject* Owner);
+
+	/**
+	 * Adds an unknown enumerator name within an enum type.
+	 *
+	 * @param Enum The enum associated with the unknown enumerator. May be null or the wrong type.
+	 * @param EnumTypeName The type name of the enum associated with the unknown enumerator.
+	 * @param EnumValueName The name of the unknown enumerator.
+	 * @note At least one of Enum and EnumTypeName must be valid.
+	 */
+	void Add(const UEnum* Enum, FPropertyTypeName EnumTypeName, FName EnumValueName);
+
+	/**
+	 * Finds unknown enumerator names associated with an enum type.
+	 *
+	 * @param EnumTypeName The type name of the enum containing the unknown names to find.
+	 * @param OutNames Array to assign the unknown names to. Empty on return if no names are found.
+	 * @param bOutHasFlags Assigned to true if the enum is known to have flags, otherwise false.
+	 */
+	void Find(FPropertyTypeName EnumTypeName, TArray<FName>& OutNames, bool& bOutHasFlags) const;
+
+	/** True if there are no unknown enumerator names tracked for the owner. */
+	bool IsEmpty() const;
+
+	/** Destroys the unknown property enum names for the owner if it has any. */
+	void Destroy();
+
+private:
+	const UObject* Owner = nullptr;
+};
+
 } // UE
 
 #undef UE_API
diff --git a/Engine/Source/Runtime/CoreUObject/Private/UObject/Class.cpp b/Engine/Source/Runtime/CoreUObject/Private/UObject/Class.cpp
--- a/Engine/Source/Runtime/CoreUObject/Private/UObject/Class.cpp
+++ b/Engine/Source/Runtime/CoreUObject/Private/UObject/Class.cpp
@@ -37,7 +37,6 @@
 #include "UObject/Interface.h"
 #include "UObject/LinkerPlaceholderClass.h"
 #include "UObject/LinkerPlaceholderFunction.h"
-#include "UObject/PropertyBagRepository.h"
 #include "UObject/PropertyOptional.h"
 #include "UObject/PropertyPathNameTree.h"
 #include "UObject/PropertyStateTracking.h"
@@ -1611,7 +1610,7 @@ void UStruct::SerializeVersionedTaggedProperties(FStructuredArchive::FSlot Slot,
 #endif // WITH_TEXT_ARCHIVE_SUPPORT
 		{
 			// Track whether the unknown property tree has been looked up.
-			FPropertyPathNameTree* UnknownPropertyTree = nullptr;
+			TSharedPtr<FPropertyPathNameTree> UnknownPropertyTree;
 			bool bSearchedForUnknownPropertyTree = false;
 
 			// Load tagged properties.
@@ -1870,7 +1869,7 @@ void UStruct::SerializeVersionedTaggedProperties(FStructuredArchive::FSlot Slot,
 						{
 							if (UObject* Object = SerializeContext->SerializedObject)
 							{
-								UnknownPropertyTree = FPropertyBagRepository::Get().FindOrCreateUnknownPropertyTree(Object);
+								UnknownPropertyTree = FUnknownPropertyTree(Object).FindOrCreate();
 							}
 						}
 					}
diff --git a/Engine/Source/Runtime/CoreUObject/Private/UObject/EnumProperty.cpp b/Engine/Source/Runtime/CoreUObject/Private/UObject/EnumProperty.cpp
--- a/Engine/Source/Runtime/CoreUObject/Private/UObject/EnumProperty.cpp
+++ b/Engine/Source/Runtime/CoreUObject/Private/UObject/EnumProperty.cpp
@@ -9,7 +9,7 @@
 #include "UObject/UObjectThreadContext.h"
 
 #if WITH_EDITORONLY_DATA
-#include "UObject/PropertyBagRepository.h"
+#include "UObject/PropertyStateTracking.h"
 #endif
 
 namespace UE::CoreUObject::Private
@@ -298,7 +298,7 @@ void FEnumProperty::SerializeItem(FStructuredArchive::FSlot Slot, void* Value, v
 			FUObjectSerializeContext* SerializeContext = FUObjectThreadContext::Get().GetSerializeContext();
 			if (UNLIKELY(SerializeContext->bTrackUnknownEnumNames))
 			{
-				UE::FPropertyBagRepository::Get().AddUnknownEnumName(SerializeContext->SerializedObject, Enum, UE::FindOriginalType(this), EnumValueName);
+				UE::FUnknownEnumNames(SerializeContext->SerializedObject).Add(Enum, UE::FindOriginalType(this), EnumValueName);
 			}
 		#endif
 		}
diff --git a/Engine/Source/Runtime/CoreUObject/Private/UObject/PropertyBagRepository.cpp b/Engine/Source/Runtime/CoreUObject/Private/UObject/PropertyBagRepository.cpp
--- a/Engine/Source/Runtime/CoreUObject/Private/UObject/PropertyBagRepository.cpp
+++ b/Engine/Source/Runtime/CoreUObject/Private/UObject/PropertyBagRepository.cpp
@@ -5,17 +5,15 @@
 #if WITH_EDITORONLY_DATA
 
 #include "Containers/Queue.h"
-#include "Misc/StringBuilder.h"
 #include "Serialization/ObjectReader.h"
 #include "Serialization/ObjectWriter.h"
-#include "String/ParseTokens.h"
 #include "UObject/GarbageCollection.h"
 #include "UObject/InstanceDataObjectUtils.h"
 #include "UObject/LinkerLoad.h"
 #include "UObject/Object.h"
 #include "UObject/Package.h"
 #include "UObject/PropertyOptional.h"
-#include "UObject/PropertyPathNameTree.h"
+#include "UObject/PropertyStateTracking.h"
 #include "UObject/UObjectGlobals.h"
 #include "UObject/UObjectHash.h"
 #include "UObject/UObjectThreadContext.h"
@@ -130,12 +128,6 @@ public:
 
 void FPropertyBagRepository::FPropertyBagAssociationData::Destroy()
 {
-	delete Tree;
-	Tree = nullptr;
-
-	delete EnumNames;
-	EnumNames = nullptr;
-
 	if (InstanceDataObject && InstanceDataObject->IsValidLowLevel())
 	{
 		InstanceDataObject = nullptr;
@@ -625,128 +617,6 @@ void FPropertyBagRepository::PostEditChangeChainProperty(const UObject* Object,
 #endif
 }
 
-// TODO: Create these by class on construction?
-FPropertyPathNameTree* FPropertyBagRepository::FindOrCreateUnknownPropertyTree(const UObject* Owner)
-{
-	FPropertyBagRepositoryLock LockRepo(this);
-	FPropertyBagAssociationData* BagData = AssociatedData.Find(Owner);
-	if (!BagData)
-	{
-		BagData = &AssociatedData.Emplace(Owner);
-	}
-	if (!BagData->Tree)
-	{
-		BagData->Tree = new FPropertyPathNameTree;
-	}
-	return BagData->Tree;
-}
-
-void FPropertyBagRepository::AddUnknownEnumName(const UObject* Owner, const UEnum* Enum, FPropertyTypeName EnumTypeName, FName EnumValueName)
-{
-	check(Owner);
-	checkf(Enum || !EnumTypeName.IsEmpty(), TEXT("AddUnknownEnumName requires an enum or its type name. Owner: %s"), *Owner->GetPathName());
-
-	FPropertyBagRepositoryLock LockRepo(this);
-	FPropertyBagAssociationData& BagData = AssociatedData.FindOrAdd(Owner);
-	if (!BagData.EnumNames)
-	{
-		BagData.EnumNames = new FUnknownEnumNames;
-	}
-
-	BagData.EnumNames->Add(Enum, EnumTypeName, EnumValueName);
-}
-
-void FUnknownEnumNames::Add(const UEnum* Enum, FPropertyTypeName EnumTypeName, FName EnumValueName)
-{
-	check(Enum || !EnumTypeName.IsEmpty());
-
-	if (EnumTypeName.IsEmpty())
-	{
-		FPropertyTypeNameBuilder Builder;
-		Builder.AddPath(Enum);
-		EnumTypeName = Builder.Build();
-	}
-
-	FUnknownEnumNames::FInfo& Info = Enums.FindOrAdd(EnumTypeName);
-
-	TStringBuilder<128> EnumValueString(InPlace, EnumValueName);
-	if (String::FindFirstChar(EnumValueString, TEXT('|')) == INDEX_NONE)
-	{
-		if (int32 ColonIndex = String::FindFirst(EnumValueString, TEXTVIEW("::")); ColonIndex != INDEX_NONE)
-		{
-			Info.Names.Emplace(EnumValueString.ToView().RightChop(ColonIndex + TEXTVIEW("::").Len()));
-		}
-		else
-		{
-			Info.Names.Add(EnumValueName);
-		}
-	}
-	else
-	{
-		Info.bHasFlags = true;
-		String::ParseTokens(EnumValueString, TEXT('|'), [&Info, Enum](FStringView Token)
-		{
-			FName Name(Token);
-			if (!Enum || Enum->GetIndexByName(Name) == INDEX_NONE)
-			{
-				Info.Names.Add(Name);
-			}
-		}, String::EParseTokensOptions::SkipEmpty | String::EParseTokensOptions::Trim);
-	}
-
-	if (!Info.bHasFlags && Enum && Enum->HasAnyEnumFlags(EEnumFlags::Flags))
-	{
-		Info.bHasFlags = true;
-	}
-}
-
-void FPropertyBagRepository::FindUnknownEnumNames(const UObject* Owner, FPropertyTypeName EnumTypeName, TArray<FName>& OutNames, bool& bOutHasFlags) const
-{
-	checkf(!EnumTypeName.IsEmpty(), TEXT("FindUnknownEnumNames requires an enum type name. Owner: %s"), *Owner->GetPathName());
-
-	OutNames.Empty();
-	bOutHasFlags = false;
-
-	if (const FUnknownEnumNames* EnumNames = FindUnknownEnumNames(Owner))
-	{
-		EnumNames->Find(EnumTypeName, OutNames, bOutHasFlags);
-	}
-}
-
-const FUnknownEnumNames* FPropertyBagRepository::FindUnknownEnumNames(const UObject* Owner) const
-{
-	check(Owner);
-
-	FPropertyBagRepositoryLock LockRepo(this);
-	const FPropertyBagAssociationData* BagData = AssociatedData.Find(Owner);
-	return BagData ? BagData->EnumNames : nullptr;
-}
-
-void FUnknownEnumNames::Find(FPropertyTypeName EnumTypeName, TArray<FName>& OutNames, bool& bOutHasFlags) const
-{
-	OutNames.Empty();
-	bOutHasFlags = false;
-
-	if (const FUnknownEnumNames::FInfo* Info = Enums.Find(EnumTypeName))
-	{
-		OutNames = Info->Names.Array();
-		bOutHasFlags = Info->bHasFlags;
-	}
-}
-
-void FPropertyBagRepository::ResetUnknownEnumNames(const UObject* Owner)
-{
-	check(Owner);
-
-	FPropertyBagRepositoryLock LockRepo(this);
-	FPropertyBagAssociationData* BagData = AssociatedData.Find(Owner);
-	if (BagData)
-	{
-		delete BagData->EnumNames;
-		BagData->EnumNames = nullptr;
-	}
-}
-
 UObject* FPropertyBagRepository::CreateInstanceDataObject(UObject* Owner, FArchive* Archive)
 {
 	FPropertyBagRepositoryLock LockRepo(this);
@@ -945,11 +815,12 @@ FString FPropertyBagRepository::GetReferencerName() const
 void FPropertyBagRepository::CreateInstanceDataObjectUnsafe(UObject* Owner, FPropertyBagAssociationData& BagData, FArchive* Archive)
 {
 	check(!BagData.InstanceDataObject);	// No repeated calls
-	const FPropertyPathNameTree* PropertyTree = BagData.Tree;
-	const FUnknownEnumNames* EnumNames = BagData.EnumNames;
+	TSharedPtr<FPropertyPathNameTree> PropertyTree = FUnknownPropertyTree(Owner).Find();
+	const FUnknownEnumNames EnumNames(Owner);
+	const FUnknownEnumNames* EnumNamesPtr = !EnumNames.IsEmpty() ? &EnumNames : nullptr;
 	// construct InstanceDataObject class
 	// TODO: should we put the InstanceDataObject or its class in a package?
-	const UClass* InstanceDataObjectClass = CreateInstanceDataObjectClass(PropertyTree, EnumNames, Owner->GetClass(), GetTransientPackage());
+	const UClass* InstanceDataObjectClass = CreateInstanceDataObjectClass(PropertyTree.Get(), EnumNamesPtr, Owner->GetClass(), GetTransientPackage());
 
 	BagData.bNeedsFixup = StructContainsLooseProperties(InstanceDataObjectClass);
 
@@ -996,7 +867,7 @@ void FPropertyBagRepository::CreateInstanceDataObjectUnsafe(UObject* Owner, FPro
 	}
 	else
 	{
-		ensureMsgf(BagData.Tree == nullptr,
+		ensureMsgf(PropertyTree == nullptr,
 			TEXT("Linker missing when generating IDO for an object with unknown properties. The unknown properties will be lost. Path: %s"),
 			*Owner->GetPathName());
 		// copy data from owner to IDO
diff --git a/Engine/Source/Runtime/CoreUObject/Private/UObject/PropertyByte.cpp b/Engine/Source/Runtime/CoreUObject/Private/UObject/PropertyByte.cpp
--- a/Engine/Source/Runtime/CoreUObject/Private/UObject/PropertyByte.cpp
+++ b/Engine/Source/Runtime/CoreUObject/Private/UObject/PropertyByte.cpp
@@ -9,7 +9,7 @@
 #include "UObject/UObjectThreadContext.h"
 
 #if WITH_EDITORONLY_DATA
-#include "UObject/PropertyBagRepository.h"
+#include "UObject/PropertyStateTracking.h"
 #endif
 
 // Implemented in EnumProperty.cpp
@@ -77,7 +77,7 @@ void FByteProperty::SerializeItem( FStructuredArchive::FSlot Slot, void* Value,
 			FUObjectSerializeContext* SerializeContext = FUObjectThreadContext::Get().GetSerializeContext();
 			if (UNLIKELY(SerializeContext->bTrackUnknownEnumNames))
 			{
-				UE::FPropertyBagRepository::Get().AddUnknownEnumName(SerializeContext->SerializedObject, Enum, UE::FindOriginalType(this), EnumValueName);
+				UE::FUnknownEnumNames(SerializeContext->SerializedObject).Add(Enum, UE::FindOriginalType(this), EnumValueName);
 			}
 		#endif
 		}
diff --git a/Engine/Source/Runtime/CoreUObject/Private/UObject/PropertyStateTracking.cpp b/Engine/Source/Runtime/CoreUObject/Private/UObject/PropertyStateTracking.cpp
--- a/Engine/Source/Runtime/CoreUObject/Private/UObject/PropertyStateTracking.cpp
+++ b/Engine/Source/Runtime/CoreUObject/Private/UObject/PropertyStateTracking.cpp
@@ -4,8 +4,17 @@
 
 #if WITH_EDITORONLY_DATA
 
+#include "Containers/Array.h"
+#include "Containers/Map.h"
+#include "Containers/Set.h"
 #include "Serialization/StructuredArchive.h"
+#include "String/ParseTokens.h"
+#include "Templates/SharedPointer.h"
 #include "UObject/Class.h"
+#include "UObject/NameTypes.h"
+#include "UObject/PropertyPathNameTree.h"
+#include "UObject/PropertyTypeName.h"
+#include "UObject/UObjectAnnotation.h"
 
 namespace UE
 {
@@ -45,6 +54,8 @@ void SerializePropertyValueInitialized(const UStruct* Struct, void* Data, FStruc
 	Struct->SerializePropertyValueFlags(EPropertyValueFlags::Initialized, Data, Record, TEXT("PropertyValueInitializedFlags"));
 }
 
+///////////////////////////////////////////////////////////////////////////////////////////////////////////////////////
+
 bool ActivateTrackingPropertyValueSerialized(const UStruct* Struct, void* Data)
 {
 	return Struct->ActivateTrackingPropertyValueFlag(EPropertyValueFlags::Serialized, Data);
@@ -65,6 +76,163 @@ bool WasPropertyValueSerialized(const UStruct* Struct, const void* Data, const F
 	return Struct->HasPropertyValueFlag(EPropertyValueFlags::Serialized, Data, Property, ArrayIndex);
 }
 
+///////////////////////////////////////////////////////////////////////////////////////////////////////////////////////
+
+struct FUnknownPropertyTreeAnnotation
+{
+	bool IsDefault() const
+	{
+		return !Tree.IsValid();
+	}
+
+	TSharedPtr<FPropertyPathNameTree> Tree;
+};
+
+using FUnknownPropertyTreeAnnotationStore = FUObjectAnnotationSparse<FUnknownPropertyTreeAnnotation, /*bAutoRemove*/ true>;
+
+UE_AUTORTFM_ALWAYS_OPEN FUnknownPropertyTreeAnnotationStore& GetUnknownPropertyTreeAnnotations()
+{
+	static FUnknownPropertyTreeAnnotationStore Annotations;
+	return Annotations;
+}
+
+FUnknownPropertyTree::FUnknownPropertyTree(const UObject* InOwner)
+	: Owner(InOwner)
+{
+}
+
+TSharedPtr<FPropertyPathNameTree> FUnknownPropertyTree::Find() const
+{
+	return GetUnknownPropertyTreeAnnotations().GetAnnotation(Owner).Tree;
+}
+
+TSharedPtr<FPropertyPathNameTree> FUnknownPropertyTree::FindOrCreate()
+{
+	FUnknownPropertyTreeAnnotationStore& Store = GetUnknownPropertyTreeAnnotations();
+	TSharedPtr<FPropertyPathNameTree> Tree = Store.GetAnnotation(Owner).Tree;
+	if (!Tree)
+	{
+		Tree = MakeShared<FPropertyPathNameTree>();
+		Store.AddAnnotation(Owner, {Tree});
+	}
+	return Tree;
+}
+
+void FUnknownPropertyTree::Destroy()
+{
+	GetUnknownPropertyTreeAnnotations().RemoveAnnotation(Owner);
+}
+
+///////////////////////////////////////////////////////////////////////////////////////////////////////////////////////
+
+struct FUnknownEnumNamesAnnotation
+{
+	bool IsDefault() const
+	{
+		return !Enums.IsValid();
+	}
+
+	struct FInfo
+	{
+		TSet<FName> Names;
+		bool bHasFlags = false;
+	};
+
+	TSharedPtr<TMap<FPropertyTypeName, FInfo>> Enums;
+};
+
+using FUnknownEnumNamesAnnotationStore = FUObjectAnnotationSparse<FUnknownEnumNamesAnnotation, /*bAutoRemove*/ true>;
+
+UE_AUTORTFM_ALWAYS_OPEN FUnknownEnumNamesAnnotationStore& GetUnknownEnumNamesAnnotations()
+{
+	static FUnknownEnumNamesAnnotationStore Annotations;
+	return Annotations;
+}
+
+FUnknownEnumNames::FUnknownEnumNames(const UObject* InOwner)
+	: Owner(InOwner)
+{
+}
+
+void FUnknownEnumNames::Add(const UEnum* Enum, FPropertyTypeName EnumTypeName, FName EnumValueName)
+{
+	checkf(Enum || !EnumTypeName.IsEmpty(), TEXT("FUnknownEnumNames::Add() requires an enum or its type name. Owner: %s"), *Owner->GetPathName());
+
+	if (EnumTypeName.IsEmpty())
+	{
+		FPropertyTypeNameBuilder Builder;
+		Builder.AddPath(Enum);
+		EnumTypeName = Builder.Build();
+	}
+
+	FUnknownEnumNamesAnnotationStore& Store = GetUnknownEnumNamesAnnotations();
+	TSharedPtr<TMap<FPropertyTypeName, FUnknownEnumNamesAnnotation::FInfo>> Enums = Store.GetAnnotation(Owner).Enums;
+	if (!Enums)
+	{
+		Enums = MakeShared<TMap<FPropertyTypeName, FUnknownEnumNamesAnnotation::FInfo>>();
+		Store.AddAnnotation(Owner, {Enums});
+	}
+
+	FUnknownEnumNamesAnnotation::FInfo& Info = Enums->FindOrAdd(EnumTypeName);
+
+	TStringBuilder<128> EnumValueString(InPlace, EnumValueName);
+	if (String::FindFirstChar(EnumValueString, TEXT('|')) == INDEX_NONE)
+	{
+		if (int32 ColonIndex = String::FindFirst(EnumValueString, TEXTVIEW("::")); ColonIndex != INDEX_NONE)
+		{
+			Info.Names.Emplace(EnumValueString.ToView().RightChop(ColonIndex + TEXTVIEW("::").Len()));
+		}
+		else
+		{
+			Info.Names.Add(EnumValueName);
+		}
+	}
+	else
+	{
+		Info.bHasFlags = true;
+		String::ParseTokens(EnumValueString, TEXT('|'), [&Info, Enum](FStringView Token)
+		{
+			FName Name(Token);
+			if (!Enum || Enum->GetIndexByName(Name) == INDEX_NONE)
+			{
+				Info.Names.Add(Name);
+			}
+		}, String::EParseTokensOptions::SkipEmpty | String::EParseTokensOptions::Trim);
+	}
+
+	if (!Info.bHasFlags && Enum && Enum->HasAnyEnumFlags(EEnumFlags::Flags))
+	{
+		Info.bHasFlags = true;
+	}
+}
+
+void FUnknownEnumNames::Find(FPropertyTypeName EnumTypeName, TArray<FName>& OutNames, bool& bOutHasFlags) const
+{
+	checkf(!EnumTypeName.IsEmpty(), TEXT("FUnknownEnumNames::Find() requires an enum type name. Owner: %s"), *Owner->GetPathName());
+
+	OutNames.Empty();
+	bOutHasFlags = false;
+
+	if (TSharedPtr<TMap<FPropertyTypeName, FUnknownEnumNamesAnnotation::FInfo>> Enums = GetUnknownEnumNamesAnnotations().GetAnnotation(Owner).Enums)
+	{
+		if (const FUnknownEnumNamesAnnotation::FInfo* Info = Enums->Find(EnumTypeName))
+		{
+			OutNames = Info->Names.Array();
+			bOutHasFlags = Info->bHasFlags;
+		}
+	}
+}
+
+bool FUnknownEnumNames::IsEmpty() const
+{
+	return GetUnknownEnumNamesAnnotations().GetAnnotation(Owner).IsDefault();
+}
+
+void FUnknownEnumNames::Destroy()
+{
+	GetUnknownEnumNamesAnnotations().RemoveAnnotation(Owner);
+}
+
 } // UE
 
 #endif // WITH_EDITORONLY_DATA
diff --git a/Engine/Source/Runtime/CoreUObject/Tests/UObject/InstanceDataObjectUtilsTest.cpp b/Engine/Source/Runtime/CoreUObject/Tests/UObject/InstanceDataObjectUtilsTest.cpp
deleted file mode 100644
diff --git a/Engine/Source/Runtime/CoreUObject/Tests/UObject/InstanceDataObjectUtilsTest.h b/Engine/Source/Runtime/CoreUObject/Tests/UObject/InstanceDataObjectUtilsTest.h
deleted file mode 100644
diff --git a/Engine/Source/Runtime/CoreUObject/Tests/UObject/PropertyStateTrackingTest.cpp b/Engine/Source/Runtime/CoreUObject/Tests/UObject/PropertyStateTrackingTest.cpp
new file mode 100644
--- /dev/null
+++ b/Engine/Source/Runtime/CoreUObject/Tests/UObject/PropertyStateTrackingTest.cpp
@@ -0,0 +1,896 @@
+// Copyright Epic Games, Inc. All Rights Reserved.
+
+#include "PropertyStateTrackingTest.h"
+
+#if WITH_TESTS && WITH_EDITORONLY_DATA
+
+#include "Logging/LogScopedVerbosityOverride.h"
+#include "Serialization/Formatters/BinaryArchiveFormatter.h"
+#include "Serialization/Formatters/JsonArchiveInputFormatter.h"
+#include "Serialization/Formatters/JsonArchiveOutputFormatter.h"
+#include "Serialization/MemoryReader.h"
+#include "Serialization/MemoryWriter.h"
+#include "Serialization/StructuredArchive.h"
+#include "Tests/TestHarnessAdapter.h"
+#include "UObject/InstanceDataObjectUtils.h"
+#include "UObject/Package.h"
+#include "UObject/PropertyHelper.h"
+#include "UObject/PropertyPathNameTree.h"
+#include "UObject/PropertyStateTracking.h"
+#include "UObject/UObjectThreadContext.h"
+
+namespace UE
+{
+
+TEST_CASE_NAMED(FPropertyValueFlagsTest, "CoreUObject::Serialization::PropertyValueFlags", "[CoreUObject][EngineFilter]")
+{
+	const auto SavePropertyValueInitializedFlags = [](const UStruct* Struct, void* Data, TArray<uint8>& OutFlags)
+	{
+		FMemoryWriter Ar(OutFlags, /*bIsPersistent*/ true);
+		FBinaryArchiveFormatter Formatter(Ar);
+		FStructuredArchive StructuredAr(Formatter);
+		SerializePropertyValueInitialized(Struct, Data, StructuredAr.Open().EnterRecord());
+	};
+
+	const auto LoadPropertyValueInitializedFlags = [](const UStruct* Struct, void* Data, TArray<uint8>& OutFlags)
+	{
+		FMemoryReader Ar(OutFlags, /*bIsPersistent*/ true);
+		FBinaryArchiveFormatter Formatter(Ar);
+		FStructuredArchive StructuredAr(Formatter);
+		SerializePropertyValueInitialized(Struct, Data, StructuredAr.Open().EnterRecord());
+	};
+
+	UTestInstanceDataObjectClass* BaseObject = NewObject<UTestInstanceDataObjectClass>();
+
+	// Test BaseObject w/o IDOs
+	{
+		UObject* TestObject = BaseObject;
+		UClass* TestClass = UTestInstanceDataObjectClass::StaticClass();
+
+		CHECK_FALSE(IsTrackingPropertyValueInitialized(TestClass, TestObject));
+
+		// Test serialization of empty flags.
+		{
+			TArray<uint8> FlagsData;
+			SavePropertyValueInitializedFlags(TestClass, TestObject, FlagsData);
+			LoadPropertyValueInitializedFlags(TestClass, TestObject, FlagsData);
+			CHECK_FALSE(IsTrackingPropertyValueInitialized(TestClass, TestObject));
+		}
+
+		CHECK(ActivateTrackingPropertyValueInitialized(TestClass, TestObject));
+		CHECK(IsTrackingPropertyValueInitialized(TestClass, TestObject));
+
+		FIntProperty* AProperty = FindFProperty<FIntProperty>(TestClass, TEXT("A"));
+		FFloatProperty* BProperty = FindFProperty<FFloatProperty>(TestClass, TEXT("B"));
+		FInt64Property* CProperty = FindFProperty<FInt64Property>(TestClass, TEXT("C"));
+		FIntProperty* DProperty = FindFProperty<FIntProperty>(TestClass, TEXT("D"));
+		FIntProperty* EProperty = FindFProperty<FIntProperty>(TestClass, TEXT("E"));
+		REQUIRE(AProperty);
+		REQUIRE(BProperty);
+		REQUIRE(CProperty);
+		REQUIRE(DProperty);
+		REQUIRE(EProperty);
+
+		CHECK_FALSE(IsPropertyValueInitialized(TestClass, TestObject, AProperty, 0));
+		CHECK_FALSE(IsPropertyValueInitialized(TestClass, TestObject, AProperty, 1));
+		CHECK_FALSE(IsPropertyValueInitialized(TestClass, TestObject, AProperty, 2));
+		CHECK_FALSE(IsPropertyValueInitialized(TestClass, TestObject, AProperty, 3));
+		CHECK_FALSE(IsPropertyValueInitialized(TestClass, TestObject, BProperty));
+		CHECK_FALSE(IsPropertyValueInitialized(TestClass, TestObject, CProperty));
+		CHECK_FALSE(IsPropertyValueInitialized(TestClass, TestObject, DProperty));
+		CHECK_FALSE(IsPropertyValueInitialized(TestClass, TestObject, EProperty));
+
+		SetPropertyValueInitialized(TestClass, TestObject, AProperty, 1);
+		SetPropertyValueInitialized(TestClass, TestObject, AProperty, 3);
+		SetPropertyValueInitialized(TestClass, TestObject, CProperty);
+		SetPropertyValueInitialized(TestClass, TestObject, EProperty);
+
+		CHECK_FALSE(IsPropertyValueInitialized(TestClass, TestObject, AProperty, 0));
+		CHECK(IsPropertyValueInitialized(TestClass, TestObject, AProperty, 1));
+		CHECK_FALSE(IsPropertyValueInitialized(TestClass, TestObject, AProperty, 2));
+		CHECK(IsPropertyValueInitialized(TestClass, TestObject, AProperty, 3));
+		CHECK_FALSE(IsPropertyValueInitialized(TestClass, TestObject, BProperty));
+		CHECK(IsPropertyValueInitialized(TestClass, TestObject, CProperty));
+		CHECK_FALSE(IsPropertyValueInitialized(TestClass, TestObject, DProperty));
+		CHECK(IsPropertyValueInitialized(TestClass, TestObject, EProperty));
+
+		TArray<uint8> FlagsData;
+		SavePropertyValueInitializedFlags(TestClass, TestObject, FlagsData);
+
+		ResetPropertyValueInitialized(TestClass, TestObject);
+
+		CHECK_FALSE(IsPropertyValueInitialized(TestClass, TestObject, AProperty, 0));
+		CHECK_FALSE(IsPropertyValueInitialized(TestClass, TestObject, AProperty, 1));
+		CHECK_FALSE(IsPropertyValueInitialized(TestClass, TestObject, AProperty, 2));
+		CHECK_FALSE(IsPropertyValueInitialized(TestClass, TestObject, AProperty, 3));
+		CHECK_FALSE(IsPropertyValueInitialized(TestClass, TestObject, BProperty));
+		CHECK_FALSE(IsPropertyValueInitialized(TestClass, TestObject, CProperty));
+		CHECK_FALSE(IsPropertyValueInitialized(TestClass, TestObject, DProperty));
+		CHECK_FALSE(IsPropertyValueInitialized(TestClass, TestObject, EProperty));
+
+		LoadPropertyValueInitializedFlags(TestClass, TestObject, FlagsData);
+
+		CHECK_FALSE(IsPropertyValueInitialized(TestClass, TestObject, AProperty, 0));
+		CHECK(IsPropertyValueInitialized(TestClass, TestObject, AProperty, 1));
+		CHECK_FALSE(IsPropertyValueInitialized(TestClass, TestObject, AProperty, 2));
+		CHECK(IsPropertyValueInitialized(TestClass, TestObject, AProperty, 3));
+		CHECK_FALSE(IsPropertyValueInitialized(TestClass, TestObject, BProperty));
+		CHECK(IsPropertyValueInitialized(TestClass, TestObject, CProperty));
+		CHECK_FALSE(IsPropertyValueInitialized(TestClass, TestObject, DProperty));
+		CHECK(IsPropertyValueInitialized(TestClass, TestObject, EProperty));
+
+		UTestInstanceDataObjectClass* EmptyObject = NewObject<UTestInstanceDataObjectClass>();
+		LoadPropertyValueInitializedFlags(TestClass, EmptyObject, FlagsData);
+
+		CHECK_FALSE(IsPropertyValueInitialized(TestClass, EmptyObject, AProperty, 0));
+		CHECK(IsPropertyValueInitialized(TestClass, EmptyObject, AProperty, 1));
+		CHECK_FALSE(IsPropertyValueInitialized(TestClass, EmptyObject, AProperty, 2));
+		CHECK(IsPropertyValueInitialized(TestClass, EmptyObject, AProperty, 3));
+		CHECK_FALSE(IsPropertyValueInitialized(TestClass, EmptyObject, BProperty));
+		CHECK(IsPropertyValueInitialized(TestClass, EmptyObject, CProperty));
+		CHECK_FALSE(IsPropertyValueInitialized(TestClass, EmptyObject, DProperty));
+		CHECK(IsPropertyValueInitialized(TestClass, EmptyObject, EProperty));
+	}
+
+	// Create IDO for BaseObject
+
+	UClass* TestClass = CreateInstanceDataObjectClass(nullptr, nullptr, BaseObject->GetClass(), BaseObject->GetOuter());
+
+	FName TestObjectName = MakeUniqueObjectName(nullptr, TestClass, FName(WriteToString<128>(TestClass->GetFName(), TEXT("_Instance"))));
+	UObject* TestObject = NewObject<UObject>(GetTransientPackage(), TestClass, TestObjectName);
+
+	// Test BaseObject w/ IDOs
+	{
+		ResetPropertyValueInitialized(TestClass, TestObject);
+
+		CHECK(ActivateTrackingPropertyValueInitialized(TestClass, TestObject));
+		CHECK(IsTrackingPropertyValueInitialized(TestClass, TestObject));
+
+		FIntProperty* AProperty = FindFProperty<FIntProperty>(TestClass, TEXT("A"));
+		FFloatProperty* BProperty = FindFProperty<FFloatProperty>(TestClass, TEXT("B"));
+		FInt64Property* CProperty = FindFProperty<FInt64Property>(TestClass, TEXT("C"));
+		FIntProperty* DProperty = FindFProperty<FIntProperty>(TestClass, TEXT("D"));
+		FIntProperty* EProperty = FindFProperty<FIntProperty>(TestClass, TEXT("E"));
+		REQUIRE(AProperty);
+		REQUIRE(BProperty);
+		REQUIRE(CProperty);
+		REQUIRE(DProperty);
+		REQUIRE(EProperty);
+
+		CHECK_FALSE(IsPropertyValueInitialized(TestClass, TestObject, AProperty, 0));
+		CHECK_FALSE(IsPropertyValueInitialized(TestClass, TestObject, AProperty, 1));
+		CHECK_FALSE(IsPropertyValueInitialized(TestClass, TestObject, AProperty, 2));
+		CHECK_FALSE(IsPropertyValueInitialized(TestClass, TestObject, AProperty, 3));
+		CHECK_FALSE(IsPropertyValueInitialized(TestClass, TestObject, BProperty));
+		CHECK_FALSE(IsPropertyValueInitialized(TestClass, TestObject, CProperty));
+		CHECK_FALSE(IsPropertyValueInitialized(TestClass, TestObject, DProperty));
+		CHECK_FALSE(IsPropertyValueInitialized(TestClass, TestObject, EProperty));
+
+		SetPropertyValueInitialized(TestClass, TestObject, AProperty, 1);
+		SetPropertyValueInitialized(TestClass, TestObject, AProperty, 3);
+		SetPropertyValueInitialized(TestClass, TestObject, CProperty);
+		SetPropertyValueInitialized(TestClass, TestObject, EProperty);
+
+		CHECK_FALSE(IsPropertyValueInitialized(TestClass, TestObject, AProperty, 0));
+		CHECK(IsPropertyValueInitialized(TestClass, TestObject, AProperty, 1));
+		CHECK_FALSE(IsPropertyValueInitialized(TestClass, TestObject, AProperty, 2));
+		CHECK(IsPropertyValueInitialized(TestClass, TestObject, AProperty, 3));
+		CHECK_FALSE(IsPropertyValueInitialized(TestClass, TestObject, BProperty));
+		CHECK(IsPropertyValueInitialized(TestClass, TestObject, CProperty));
+		CHECK_FALSE(IsPropertyValueInitialized(TestClass, TestObject, DProperty));
+		CHECK(IsPropertyValueInitialized(TestClass, TestObject, EProperty));
+
+		TArray<uint8> FlagsData;
+		SavePropertyValueInitializedFlags(TestClass, TestObject, FlagsData);
+
+		ResetPropertyValueInitialized(TestClass, TestObject);
+
+		CHECK_FALSE(IsPropertyValueInitialized(TestClass, TestObject, AProperty, 0));
+		CHECK_FALSE(IsPropertyValueInitialized(TestClass, TestObject, AProperty, 1));
+		CHECK_FALSE(IsPropertyValueInitialized(TestClass, TestObject, AProperty, 2));
+		CHECK_FALSE(IsPropertyValueInitialized(TestClass, TestObject, AProperty, 3));
+		CHECK_FALSE(IsPropertyValueInitialized(TestClass, TestObject, BProperty));
+		CHECK_FALSE(IsPropertyValueInitialized(TestClass, TestObject, CProperty));
+		CHECK_FALSE(IsPropertyValueInitialized(TestClass, TestObject, DProperty));
+		CHECK_FALSE(IsPropertyValueInitialized(TestClass, TestObject, EProperty));
+
+		LoadPropertyValueInitializedFlags(TestClass, TestObject, FlagsData);
+
+		CHECK_FALSE(IsPropertyValueInitialized(TestClass, TestObject, AProperty, 0));
+		CHECK(IsPropertyValueInitialized(TestClass, TestObject, AProperty, 1));
+		CHECK_FALSE(IsPropertyValueInitialized(TestClass, TestObject, AProperty, 2));
+		CHECK(IsPropertyValueInitialized(TestClass, TestObject, AProperty, 3));
+		CHECK_FALSE(IsPropertyValueInitialized(TestClass, TestObject, BProperty));
+		CHECK(IsPropertyValueInitialized(TestClass, TestObject, CProperty));
+		CHECK_FALSE(IsPropertyValueInitialized(TestClass, TestObject, DProperty));
+		CHECK(IsPropertyValueInitialized(TestClass, TestObject, EProperty));
+	}
+
+	// Test BaseObject.Struct with IDOs
+	{
+		FStructProperty* StructProperty = FindFProperty<FStructProperty>(TestClass, TEXT("Struct"));
+		REQUIRE(StructProperty);
+		REQUIRE(StructProperty->Struct);
+		void* StructData = StructProperty->ContainerPtrToValuePtr<void>(TestObject);
+		ResetPropertyValueInitialized(StructProperty->Struct, StructData);
+
+		FIntProperty* AProperty = FindFProperty<FIntProperty>(StructProperty->Struct, TEXT("A"));
+		FIntProperty* BProperty = FindFProperty<FIntProperty>(StructProperty->Struct, TEXT("B"));
+		FIntProperty* CProperty = FindFProperty<FIntProperty>(StructProperty->Struct, TEXT("C"));
+		FIntProperty* DProperty = FindFProperty<FIntProperty>(StructProperty->Struct, TEXT("D"));
+		REQUIRE(AProperty);
+		REQUIRE(BProperty);
+		REQUIRE(CProperty);
+		REQUIRE(DProperty);
+
+		CHECK_FALSE(IsPropertyValueInitialized(StructProperty->Struct, StructData, AProperty));
+		CHECK_FALSE(IsPropertyValueInitialized(StructProperty->Struct, StructData, BProperty));
+		CHECK_FALSE(IsPropertyValueInitialized(StructProperty->Struct, StructData, CProperty));
+		CHECK_FALSE(IsPropertyValueInitialized(StructProperty->Struct, StructData, DProperty));
+
+		SetPropertyValueInitialized(StructProperty->Struct, StructData, AProperty);
+		SetPropertyValueInitialized(StructProperty->Struct, StructData, CProperty);
+		CHECK(IsPropertyValueInitialized(StructProperty->Struct, StructData, AProperty));
+		CHECK_FALSE(IsPropertyValueInitialized(StructProperty->Struct, StructData, BProperty));
+		CHECK(IsPropertyValueInitialized(StructProperty->Struct, StructData, CProperty));
+		CHECK_FALSE(IsPropertyValueInitialized(StructProperty->Struct, StructData, DProperty));
+
+		ClearPropertyValueInitialized(StructProperty->Struct, StructData, AProperty);
+		ClearPropertyValueInitialized(StructProperty->Struct, StructData, CProperty);
+		CHECK_FALSE(IsPropertyValueInitialized(StructProperty->Struct, StructData, AProperty));
+		CHECK_FALSE(IsPropertyValueInitialized(StructProperty->Struct, StructData, CProperty));
+
+		SetPropertyValueInitialized(StructProperty->Struct, StructData, BProperty);
+		SetPropertyValueInitialized(StructProperty->Struct, StructData, DProperty);
+
+		TArray<uint8> FlagsData;
+		SavePropertyValueInitializedFlags(StructProperty->Struct, StructData, FlagsData);
+
+		ResetPropertyValueInitialized(StructProperty->Struct, StructData);
+		CHECK_FALSE(IsPropertyValueInitialized(StructProperty->Struct, StructData, AProperty));
+		CHECK_FALSE(IsPropertyValueInitialized(StructProperty->Struct, StructData, BProperty));
+		CHECK_FALSE(IsPropertyValueInitialized(StructProperty->Struct, StructData, CProperty));
+		CHECK_FALSE(IsPropertyValueInitialized(StructProperty->Struct, StructData, DProperty));
+
+		LoadPropertyValueInitializedFlags(StructProperty->Struct, StructData, FlagsData);
+
+		CHECK_FALSE(IsPropertyValueInitialized(StructProperty->Struct, StructData, AProperty));
+		CHECK(IsPropertyValueInitialized(StructProperty->Struct, StructData, BProperty));
+		CHECK_FALSE(IsPropertyValueInitialized(StructProperty->Struct, StructData, CProperty));
+		CHECK(IsPropertyValueInitialized(StructProperty->Struct, StructData, DProperty));
+	}
+}
+
+TEST_CASE_NAMED(FInstanceDataObjectEnumsTest, "CoreUObject::Serialization::InstanceDataObjectEnums", "[CoreUObject][EngineFilter]")
+{
+	UTestInstanceDataObjectClass* BaseObject = NewObject<UTestInstanceDataObjectClass>();
+
+	FUnknownEnumNames UnknownEnumNames(BaseObject);
+	UnknownEnumNames.Add(StaticEnum<ETestInstanceDataObjectBird>(), {}, "TIDOB_Pigeon");
+	UnknownEnumNames.Add(StaticEnum<ETestInstanceDataObjectGrain::Type>(), {}, "Rye");
+	UnknownEnumNames.Add(StaticEnum<ETestInstanceDataObjectFruit>(), {}, "Cherry");
+	UnknownEnumNames.Add(StaticEnum<ETestInstanceDataObjectDirection>(), {}, "Up");
+	UnknownEnumNames.Add(StaticEnum<ETestInstanceDataObjectFullFlags>(), {}, "Flag3");
+	UnknownEnumNames.Add(StaticEnum<ETestInstanceDataObjectFullFlags>(), {}, "Flag8");
+	UnknownEnumNames.Add(StaticEnum<ETestInstanceDataObjectFullFlags>(), {}, "Flag9");
+
+	UClass* TestClass = CreateInstanceDataObjectClass(nullptr, &UnknownEnumNames, BaseObject->GetClass(), BaseObject->GetOuter());
+
+	FName TestObjectName = MakeUniqueObjectName(nullptr, TestClass, FName(WriteToString<128>(TestClass->GetFName(), TEXT("_Instance"))));
+	UObject* TestObject = NewObject<UObject>(GetTransientPackage(), TestClass, TestObjectName);
+
+	FStructProperty* StructProperty = FindFProperty<FStructProperty>(TestClass, TEXT("Struct"));
+	REQUIRE(StructProperty);
+	REQUIRE(StructProperty->Struct);
+
+	FByteProperty* BirdProperty = FindFProperty<FByteProperty>(StructProperty->Struct, TEXT("Bird"));
+	FByteProperty* GrainProperty = FindFProperty<FByteProperty>(StructProperty->Struct, TEXT("Grain"));
+	FEnumProperty* FruitProperty = FindFProperty<FEnumProperty>(StructProperty->Struct, TEXT("Fruit"));
+	FEnumProperty* DirectionProperty = FindFProperty<FEnumProperty>(StructProperty->Struct, TEXT("Direction"));
+	FEnumProperty* FullFlagsProperty = FindFProperty<FEnumProperty>(StructProperty->Struct, TEXT("FullFlags"));
+	REQUIRE(BirdProperty);
+	REQUIRE(GrainProperty);
+	REQUIRE(FruitProperty);
+	REQUIRE(DirectionProperty);
+	REQUIRE(FullFlagsProperty);
+
+	CHECK(BirdProperty->Enum->GetIndexByName("TIDOB_Pigeon") != INDEX_NONE);
+	CHECK(GrainProperty->Enum->GetIndexByName("ETestInstanceDataObjectGrain::Rye") != INDEX_NONE);
+	CHECK(FruitProperty->GetEnum()->GetIndexByName("ETestInstanceDataObjectFruit::Cherry") != INDEX_NONE);
+	CHECK(DirectionProperty->GetEnum()->GetIndexByName("ETestInstanceDataObjectDirection::Up") != INDEX_NONE);
+	CHECK(FullFlagsProperty->GetEnum()->GetIndexByName("ETestInstanceDataObjectFullFlags::Flag3") != INDEX_NONE);
+	CHECK(FullFlagsProperty->GetEnum()->GetIndexByName("ETestInstanceDataObjectFullFlags::Flag8") != INDEX_NONE);
+	CHECK(FullFlagsProperty->GetEnum()->GetIndexByName("ETestInstanceDataObjectFullFlags::Flag9") != INDEX_NONE);
+	CHECK(FullFlagsProperty->GetEnum()->GetMaxEnumValue() == 0b11'1111'1111);
+}
+
+TEST_CASE_NAMED(FTrackInitializedPropertiesTest, "CoreUObject::Serialization::TrackInitializedProperties", "[CoreUObject][EngineFilter]")
+{
+	UObject* TestObject = NewObject<UTestInstanceDataObjectClass>();
+	UClass* TestClass = TestObject->GetClass();
+	UObject* DefaultObject = TestClass->GetDefaultObject();
+	ActivateTrackingPropertyValueInitialized(TestClass, TestObject);
+
+	FStructProperty* StructProperty = FindFProperty<FStructProperty>(TestClass, TEXT("Struct"));
+	REQUIRE(StructProperty);
+	REQUIRE(TestClass);
+
+	FIntProperty* AProperty = FindFProperty<FIntProperty>(TestClass, TEXT("A"));
+	FFloatProperty* BProperty = FindFProperty<FFloatProperty>(TestClass, TEXT("B"));
+	FInt64Property* CProperty = FindFProperty<FInt64Property>(TestClass, TEXT("C"));
+	FIntProperty* DProperty = FindFProperty<FIntProperty>(TestClass, TEXT("D"));
+	REQUIRE(AProperty);
+	REQUIRE(BProperty);
+	REQUIRE(CProperty);
+	REQUIRE(DProperty);
+
+
+	AProperty->SetPropertyValue_InContainer(TestObject, 1, 0);
+	AProperty->SetPropertyValue_InContainer(TestObject, 2, 1);
+	AProperty->SetPropertyValue_InContainer(TestObject, 3, 2);
+	AProperty->SetPropertyValue_InContainer(TestObject, 4, 3);
+	BProperty->SetPropertyValue_InContainer(TestObject, 5);
+	CProperty->SetPropertyValue_InContainer(TestObject, 6);
+	DProperty->SetPropertyValue_InContainer(TestObject, -1);
+
+	SetPropertyValueInitialized(TestClass, TestObject, AProperty, 0);
+	SetPropertyValueInitialized(TestClass, TestObject, AProperty, 2);
+	SetPropertyValueInitialized(TestClass, TestObject, BProperty);
+	SetPropertyValueInitialized(TestClass, TestObject, DProperty);
+
+	FUObjectSerializeContext* SerializeContext = FUObjectThreadContext::Get().GetSerializeContext();
+	TGuardValue<bool> TrackInitializedPropertiesScope(SerializeContext->bTrackInitializedProperties, true);
+
+	TArray<uint8> BinaryData;
+	{
+		FMemoryWriter Ar(BinaryData, /*bIsPersistent*/ true);
+		FBinaryArchiveFormatter Formatter(Ar);
+		FStructuredArchive StructuredAr(Formatter);
+		TestClass->SerializeTaggedProperties(StructuredAr.Open(), (uint8*)TestObject, TestClass, (uint8*)DefaultObject);
+	}
+
+#if WITH_TEXT_ARCHIVE_SUPPORT
+	TArray<uint8> JsonData;
+	{
+		FMemoryWriter Ar(JsonData, /*bIsPersistent*/ true);
+		FJsonArchiveOutputFormatter Formatter(Ar);
+		FStructuredArchive StructuredAr(Formatter);
+		TestClass->SerializeTaggedProperties(StructuredAr.Open(), (uint8*)TestObject, TestClass, (uint8*)DefaultObject);
+	}
+#endif // WITH_TEXT_ARCHIVE_SUPPORT
+
+	AProperty->SetPropertyValue_InContainer(TestObject, 7, 0);
+	AProperty->SetPropertyValue_InContainer(TestObject, 7, 1);
+	AProperty->SetPropertyValue_InContainer(TestObject, 7, 2);
+	AProperty->SetPropertyValue_InContainer(TestObject, 7, 3);
+	BProperty->SetPropertyValue_InContainer(TestObject, 7);
+	CProperty->SetPropertyValue_InContainer(TestObject, 7);
+	DProperty->SetPropertyValue_InContainer(TestObject, 7);
+
+	ResetPropertyValueInitialized(TestClass, TestObject);
+
+	{
+		FMemoryReader Ar(BinaryData, /*bIsPersistent*/ true);
+		FBinaryArchiveFormatter Formatter(Ar);
+		FStructuredArchive StructuredAr(Formatter);
+		TestClass->SerializeTaggedProperties(StructuredAr.Open(), (uint8*)TestObject, TestClass, (uint8*)DefaultObject);
+	}
+
+	CHECK(AProperty->GetPropertyValue_InContainer(TestObject, 0) == 1);
+	CHECK(AProperty->GetPropertyValue_InContainer(TestObject, 1) == 7); // A[1] unchanged because it was not initialized
+	CHECK(AProperty->GetPropertyValue_InContainer(TestObject, 2) == 3);
+	CHECK(AProperty->GetPropertyValue_InContainer(TestObject, 3) == 7); // A[3] unchanged because it was not initialized
+	CHECK(BProperty->GetPropertyValue_InContainer(TestObject) == 5.0f);
+	CHECK(CProperty->GetPropertyValue_InContainer(TestObject) == 7); // C unchanged because it was not initialized
+	CHECK(DProperty->GetPropertyValue_InContainer(TestObject) == 7); // D unchanged because it was serialized without its value
+
+	CHECK(IsPropertyValueInitialized(TestClass, TestObject, AProperty, 0));
+	CHECK_FALSE(IsPropertyValueInitialized(TestClass, TestObject, AProperty, 1));
+	CHECK(IsPropertyValueInitialized(TestClass, TestObject, AProperty, 2));
+	CHECK_FALSE(IsPropertyValueInitialized(TestClass, TestObject, AProperty, 3));
+	CHECK(IsPropertyValueInitialized(TestClass, TestObject, BProperty));
+	CHECK_FALSE(IsPropertyValueInitialized(TestClass, TestObject, CProperty));
+	CHECK(IsPropertyValueInitialized(TestClass, TestObject, DProperty));
+
+#if WITH_TEXT_ARCHIVE_SUPPORT
+	AProperty->SetPropertyValue_InContainer(TestObject, 7, 0);
+	AProperty->SetPropertyValue_InContainer(TestObject, 7, 1);
+	AProperty->SetPropertyValue_InContainer(TestObject, 7, 2);
+	AProperty->SetPropertyValue_InContainer(TestObject, 7, 3);
+	BProperty->SetPropertyValue_InContainer(TestObject, 7);
+	CProperty->SetPropertyValue_InContainer(TestObject, 7);
+	DProperty->SetPropertyValue_InContainer(TestObject, 7);
+
+	ResetPropertyValueInitialized(TestClass, TestObject);
+
+	{
+		FMemoryReader Ar(JsonData, /*bIsPersistent*/ true);
+		FJsonArchiveInputFormatter Formatter(Ar);
+		FStructuredArchive StructuredAr(Formatter);
+		TestClass->SerializeTaggedProperties(StructuredAr.Open(), (uint8*)TestObject, TestClass, (uint8*)DefaultObject);
+	}
+
+	CHECK(AProperty->GetPropertyValue_InContainer(TestObject, 0) == 1);
+	// Disabled because text serialization behaves unexpectedly in these cases
+	//CHECK(AProperty->GetPropertyValue_InContainer(TestObject, 1) == 7); // A[1] unchanged because it was not initialized
+	//CHECK(AProperty->GetPropertyValue_InContainer(TestObject, 2) == 3);
+	//CHECK(AProperty->GetPropertyValue_InContainer(TestObject, 3) == 7); // A[3] unchanged because it was not initialized
+	CHECK(BProperty->GetPropertyValue_InContainer(TestObject) == 5.0f);
+	CHECK(CProperty->GetPropertyValue_InContainer(TestObject) == 7); // C unchanged because it was not initialized
+	CHECK(DProperty->GetPropertyValue_InContainer(TestObject) == 7); // D unchanged because it was serialized without its value
+
+	CHECK(IsPropertyValueInitialized(TestClass, TestObject, AProperty, 0));
+	//CHECK_FALSE(IsPropertyValueInitialized(TestClass, TestObject, AProperty, 1));
+	//CHECK(IsPropertyValueInitialized(TestClass, TestObject, AProperty, 2));
+	//CHECK_FALSE(IsPropertyValueInitialized(TestClass, TestObject, AProperty, 3));
+	CHECK(IsPropertyValueInitialized(TestClass, TestObject, BProperty));
+	CHECK_FALSE(IsPropertyValueInitialized(TestClass, TestObject, CProperty));
+	CHECK(IsPropertyValueInitialized(TestClass, TestObject, DProperty));
+#endif // WITH_TEXT_ARCHIVE_SUPPORT
+}
+
+TEST_CASE_NAMED(FTrackUnknownPropertiesTest, "CoreUObject::Serialization::TrackUnknownProperties", "[CoreUObject][EngineFilter]")
+{
+	const auto MakePropertyTypeName = [](FName Name)
+	{
+		FPropertyTypeNameBuilder Builder;
+		Builder.AddName(Name);
+		return Builder.Build();
+	};
+
+	const auto SavePropertyTypeName = [](const FProperty* Property)
+	{
+		FPropertyTypeNameBuilder Builder;
+		Property->SaveTypeName(Builder);
+		return Builder.Build();
+	};
+
+	UObject* Owner = NewObject<UTestInstanceDataObjectClass>();
+
+	FUObjectSerializeContext* SerializeContext = FUObjectThreadContext::Get().GetSerializeContext();
+	TGuardValue<UObject*> SerializedObjectScope(SerializeContext->SerializedObject, Owner);
+	TGuardValue<bool> TrackSerializedPropertyPathScope(SerializeContext->bTrackSerializedPropertyPath, true);
+	TGuardValue<bool> TrackUnknownPropertiesScope(SerializeContext->bTrackUnknownProperties, true);
+	TGuardValue<bool> TrackImpersonatePropertiesScope(SerializeContext->bImpersonateProperties, true);
+	FSerializedPropertyPathScope SerializedObjectPath(SerializeContext, {"Struct"});
+
+	FTestInstanceDataObjectStructAlternate AltStructData;
+	AltStructData.B = 2.5f;
+	AltStructData.C = 3;
+	AltStructData.D = 4;
+	AltStructData.E = 5;
+	AltStructData.Bird = TIDOB_Raven;
+	AltStructData.Grain = ETestInstanceDataObjectGrainAlternate::Corn;
+	AltStructData.Fruit = ETestInstanceDataObjectFruitAlternate::Orange;
+	AltStructData.Direction = ETestInstanceDataObjectDirectionAlternate::North | ETestInstanceDataObjectDirectionAlternate::West;
+	AltStructData.GrainTypeChange = ETestInstanceDataObjectGrainAlternate::Corn;
+	AltStructData.FruitTypeChange = ETestInstanceDataObjectFruitAlternate::Orange;
+	AltStructData.GrainTypeAndPropertyChange = ETestInstanceDataObjectGrainAlternateEnumClass::Corn;
+	AltStructData.FruitTypeAndPropertyChange = ETestInstanceDataObjectFruitAlternateNamespace::Orange;
+	AltStructData.Point.U = 1;
+	AltStructData.Point.V = 2;
+	AltStructData.Point.W = 3;
+
+	TArray<uint8> BinaryData;
+	{
+		FMemoryWriter Ar(BinaryData, /*bIsPersistent*/ true);
+		FBinaryArchiveFormatter Formatter(Ar);
+		FStructuredArchive StructuredAr(Formatter);
+		FTestInstanceDataObjectStructAlternate::StaticStruct()->SerializeTaggedProperties(StructuredAr.Open(), (uint8*)&AltStructData, nullptr, nullptr);
+	}
+
+#if WITH_TEXT_ARCHIVE_SUPPORT
+	TArray<uint8> JsonData;
+	{
+		FMemoryWriter Ar(JsonData, /*bIsPersistent*/ true);
+		FJsonArchiveOutputFormatter Formatter(Ar);
+		FStructuredArchive StructuredAr(Formatter);
+		FTestInstanceDataObjectStructAlternate::StaticStruct()->SerializeTaggedProperties(StructuredAr.Open(), (uint8*)&AltStructData, nullptr, nullptr);
+	}
+#endif // WITH_TEXT_ARCHIVE_SUPPORT
+
+	LOG_SCOPE_VERBOSITY_OVERRIDE(LogClass, ELogVerbosity::Error);
+	LOG_SCOPE_VERBOSITY_OVERRIDE(LogEnum, ELogVerbosity::Error);
+
+	FUnknownPropertyTree TreeAccessor(Owner);
+	CHECK(!TreeAccessor.Find());
+
+	FTestInstanceDataObjectStruct StructData;
+
+	{
+		FMemoryReader Ar(BinaryData, /*bIsPersistent*/ true);
+		FBinaryArchiveFormatter Formatter(Ar);
+		FStructuredArchive StructuredAr(Formatter);
+		FTestInstanceDataObjectStruct::StaticStruct()->SerializeTaggedProperties(StructuredAr.Open(), (uint8*)&StructData, nullptr, nullptr);
+	}
+
+	CHECK(StructData.A == -1);
+	CHECK(StructData.B == 2);
+	CHECK(StructData.C == 3);
+	CHECK(StructData.D == 4);
+	CHECK(StructData.Bird == TIDOB_Raven);
+	CHECK(StructData.Grain == ETestInstanceDataObjectGrain::Corn);
+	CHECK(StructData.Fruit == ETestInstanceDataObjectFruit::Orange);
+	CHECK(StructData.Direction == (ETestInstanceDataObjectDirection::North | ETestInstanceDataObjectDirection::West));
+	CHECK(StructData.GrainTypeChange == ETestInstanceDataObjectGrain::Corn);
+	CHECK(StructData.FruitTypeChange == ETestInstanceDataObjectFruit::Orange);
+	CHECK(StructData.GrainTypeAndPropertyChange == ETestInstanceDataObjectGrain::Corn);
+	CHECK(StructData.FruitTypeAndPropertyChange == ETestInstanceDataObjectFruit::Orange);
+	CHECK(StructData.Point.X == 0);
+	CHECK(StructData.Point.Y == 0);
+	CHECK(StructData.Point.Z == 0);
+#if WITH_METADATA
+	CHECK(StructData.Point.W == 3);
+#endif
+
+	TSharedPtr<FPropertyPathNameTree> Tree = TreeAccessor.FindOrCreate();
+	REQUIRE(Tree);
+	CHECK(TreeAccessor.Find() == Tree);
+	{
+		FSerializedPropertyPathScope Path(SerializeContext, {"B", MakePropertyTypeName(NAME_FloatProperty)});
+		CHECK(Tree->Find(SerializeContext->SerializedPropertyPath));
+	}
+	{
+		FSerializedPropertyPathScope Path(SerializeContext, {"C", MakePropertyTypeName(NAME_Int64Property)});
+		CHECK(Tree->Find(SerializeContext->SerializedPropertyPath));
+	}
+	{
+		FSerializedPropertyPathScope Path(SerializeContext, {"E", MakePropertyTypeName(NAME_IntProperty)});
+		CHECK(Tree->Find(SerializeContext->SerializedPropertyPath));
+	}
+	{
+		FProperty* PointProperty = FTestInstanceDataObjectStructAlternate::StaticStruct()->FindPropertyByName("Point");
+		CHECKED_IF(PointProperty)
+		{
+			FPropertyTypeNameBuilder Builder;
+			PointProperty->SaveTypeName(Builder);
+			FSerializedPropertyPathScope Path(SerializeContext, {"Point", Builder.Build()});
+			CHECK(Tree->Find(SerializeContext->SerializedPropertyPath));
+		#if WITH_METADATA
+			{
+				FSerializedPropertyPathScope SubPath(SerializeContext, {"U", MakePropertyTypeName(NAME_IntProperty)});
+				CHECK(Tree->Find(SerializeContext->SerializedPropertyPath));
+			}
+			{
+				FSerializedPropertyPathScope SubPath(SerializeContext, {"V", MakePropertyTypeName(NAME_IntProperty)});
+				CHECK(Tree->Find(SerializeContext->SerializedPropertyPath));
+			}
+			{
+				FSerializedPropertyPathScope SubPath(SerializeContext, {"W", MakePropertyTypeName(NAME_IntProperty)});
+				CHECK_FALSE(Tree->Find(SerializeContext->SerializedPropertyPath));
+			}
+		#endif
+		}
+	}
+	{
+		FSerializedPropertyPathScope Path(SerializeContext, {"GrainTypeChange", SavePropertyTypeName(FindFProperty<FProperty>(StaticStruct<FTestInstanceDataObjectStructAlternate>(), TEXT("GrainTypeChange")))});
+		CHECK(Tree->Find(SerializeContext->SerializedPropertyPath));
+	}
+	{
+		FSerializedPropertyPathScope Path(SerializeContext, {"FruitTypeChange", SavePropertyTypeName(FindFProperty<FProperty>(StaticStruct<FTestInstanceDataObjectStructAlternate>(), TEXT("FruitTypeChange")))});
+		CHECK(Tree->Find(SerializeContext->SerializedPropertyPath));
+	}
+	{
+		FSerializedPropertyPathScope Path(SerializeContext, {"GrainTypeAndPropertyChange", SavePropertyTypeName(FindFProperty<FProperty>(StaticStruct<FTestInstanceDataObjectStructAlternate>(), TEXT("GrainTypeAndPropertyChange")))});
+		CHECK(Tree->Find(SerializeContext->SerializedPropertyPath));
+	}
+	{
+		FSerializedPropertyPathScope Path(SerializeContext, {"FruitTypeAndPropertyChange", SavePropertyTypeName(FindFProperty<FProperty>(StaticStruct<FTestInstanceDataObjectStructAlternate>(), TEXT("FruitTypeAndPropertyChange")))});
+		CHECK(Tree->Find(SerializeContext->SerializedPropertyPath));
+	}
+	TreeAccessor.Destroy();
+	CHECK(!TreeAccessor.Find());
+
+#if WITH_TEXT_ARCHIVE_SUPPORT
+	StructData = {};
+
+	{
+		FMemoryReader Ar(JsonData, /*bIsPersistent*/ true);
+		FJsonArchiveInputFormatter Formatter(Ar);
+		FStructuredArchive StructuredAr(Formatter);
+		FTestInstanceDataObjectStruct::StaticStruct()->SerializeTaggedProperties(StructuredAr.Open(), (uint8*)&StructData, nullptr, nullptr);
+	}
+
+	CHECK(StructData.A == -1);
+	CHECK(StructData.B == 2);
+	CHECK(StructData.C == 3);
+	CHECK(StructData.D == 4);
+	CHECK(StructData.Bird == TIDOB_Raven);
+	CHECK(StructData.Grain == ETestInstanceDataObjectGrain::Corn);
+	CHECK(StructData.Fruit == ETestInstanceDataObjectFruit::Orange);
+	CHECK(StructData.Direction == (ETestInstanceDataObjectDirection::North | ETestInstanceDataObjectDirection::West));
+	CHECK(StructData.GrainTypeChange == ETestInstanceDataObjectGrain::Corn);
+	CHECK(StructData.FruitTypeChange == ETestInstanceDataObjectFruit::Orange);
+	CHECK(StructData.GrainTypeAndPropertyChange == ETestInstanceDataObjectGrain::Corn);
+	CHECK(StructData.FruitTypeAndPropertyChange == ETestInstanceDataObjectFruit::Orange);
+	CHECK(StructData.Point.X == 0);
+	CHECK(StructData.Point.Y == 0);
+	CHECK(StructData.Point.Z == 0);
+#if WITH_METADATA
+	CHECK(StructData.Point.W == 3);
+#endif
+
+	// Testing of the unknown property tree is skipped because it is not supported by the text format.
+
+	TreeAccessor.Destroy();
+	CHECK(!TreeAccessor.Find());
+#endif // WITH_TEXT_ARCHIVE_SUPPORT
+}
+
+TEST_CASE_NAMED(FTrackUnknownEnumNamesTest, "CoreUObject::Serialization::TrackUnknownEnumNames", "[CoreUObject][EngineFilter]")
+{
+	const auto MakePropertyTypeName = [](const UEnum* Enum)
+	{
+		FPropertyTypeNameBuilder Builder;
+		Builder.AddPath(Enum);
+		return Builder.Build();
+	};
+
+	const auto ParsePropertyTypeName = [](const TCHAR* Name) -> FPropertyTypeName
+	{
+		FPropertyTypeNameBuilder Builder;
+		CHECK(Builder.TryParse(Name));
+		return Builder.Build();
+	};
+
+	UObject* Owner = NewObject<UTestInstanceDataObjectClass>();
+
+	FUObjectSerializeContext* SerializeContext = FUObjectThreadContext::Get().GetSerializeContext();
+	TGuardValue<UObject*> SerializedObjectScope(SerializeContext->SerializedObject, Owner);
+	TGuardValue<bool> TrackSerializedPropertyPathScope(SerializeContext->bTrackSerializedPropertyPath, true);
+	TGuardValue<bool> TrackUnknownPropertiesScope(SerializeContext->bTrackUnknownProperties, true);
+	TGuardValue<bool> TrackUnknownEnumNamesScope(SerializeContext->bTrackUnknownEnumNames, true);
+	TGuardValue<bool> TrackImpersonatePropertiesScope(SerializeContext->bImpersonateProperties, true);
+	FSerializedPropertyPathScope SerializedObjectPath(SerializeContext, {"Struct"});
+
+	FTestInstanceDataObjectStructAlternate AltStructData;
+	AltStructData.Grain = ETestInstanceDataObjectGrainAlternate::Rye;
+	AltStructData.Fruit = ETestInstanceDataObjectFruitAlternate::Cherry;
+	AltStructData.Direction = ETestInstanceDataObjectDirectionAlternate::North | ETestInstanceDataObjectDirectionAlternate::West |
+		ETestInstanceDataObjectDirectionAlternate::Up | ETestInstanceDataObjectDirectionAlternate::Down;
+	AltStructData.GrainFromEnumClass = ETestInstanceDataObjectGrainAlternateEnumClass::Corn;
+	AltStructData.FruitFromNamespace = ETestInstanceDataObjectFruitAlternateNamespace::Orange;
+	AltStructData.GrainTypeChange = ETestInstanceDataObjectGrainAlternate::Corn;
+	AltStructData.FruitTypeChange = ETestInstanceDataObjectFruitAlternate::Orange;
+	AltStructData.DeletedGrain = ETestInstanceDataObjectGrainAlternate::Rice;
+	AltStructData.DeletedFruit = ETestInstanceDataObjectFruitAlternate::Apple;
+	AltStructData.DeletedDirection = ETestInstanceDataObjectDirectionAlternate::South | ETestInstanceDataObjectDirectionAlternate::Up;
+
+	TArray<uint8> BinaryData;
+	{
+		FMemoryWriter Ar(BinaryData, /*bIsPersistent*/ true);
+		FBinaryArchiveFormatter Formatter(Ar);
+		FStructuredArchive StructuredAr(Formatter);
+		FTestInstanceDataObjectStructAlternate::StaticStruct()->SerializeTaggedProperties(StructuredAr.Open(), (uint8*)&AltStructData, nullptr, nullptr);
+	}
+
+#if WITH_TEXT_ARCHIVE_SUPPORT
+	TArray<uint8> JsonData;
+	{
+		FMemoryWriter Ar(JsonData, /*bIsPersistent*/ true);
+		FJsonArchiveOutputFormatter Formatter(Ar);
+		FStructuredArchive StructuredAr(Formatter);
+		FTestInstanceDataObjectStructAlternate::StaticStruct()->SerializeTaggedProperties(StructuredAr.Open(), (uint8*)&AltStructData, nullptr, nullptr);
+	}
+#endif // WITH_TEXT_ARCHIVE_SUPPORT
+
+	LOG_SCOPE_VERBOSITY_OVERRIDE(LogClass, ELogVerbosity::Error);
+	LOG_SCOPE_VERBOSITY_OVERRIDE(LogEnum, ELogVerbosity::Error);
+
+	const FPropertyTypeName GrainTypeName = MakePropertyTypeName(StaticEnum<ETestInstanceDataObjectGrain::Type>());
+	const FPropertyTypeName FruitTypeName = MakePropertyTypeName(StaticEnum<ETestInstanceDataObjectFruit>());
+	const FPropertyTypeName DirectionTypeName = MakePropertyTypeName(StaticEnum<ETestInstanceDataObjectDirection>());
+
+	FUnknownEnumNames UnknownEnumNames(Owner);
+	TArray<FName> Names{NAME_None};
+	bool bHasFlags = false;
+
+	FTestInstanceDataObjectStruct StructData;
+
+	{
+		FMemoryReader Ar(BinaryData, /*bIsPersistent*/ true);
+		FBinaryArchiveFormatter Formatter(Ar);
+		FStructuredArchive StructuredAr(Formatter);
+		FTestInstanceDataObjectStruct::StaticStruct()->SerializeTaggedProperties(StructuredAr.Open(), (uint8*)&StructData, nullptr, nullptr);
+	}
+
+	CHECK(StructData.Grain == (ETestInstanceDataObjectGrain::Type)((uint8)ETestInstanceDataObjectGrain::Wheat + 1));
+	CHECK(StructData.Fruit == (ETestInstanceDataObjectFruit)((uint8)ETestInstanceDataObjectFruit::Orange + 1));
+	CHECK(StructData.Direction == (ETestInstanceDataObjectDirection)MAX_uint16);
+	CHECK(StructData.GrainFromEnumClass == ETestInstanceDataObjectGrain::Corn);
+	CHECK(StructData.FruitFromNamespace == ETestInstanceDataObjectFruit::Orange);
+	CHECK(StructData.GrainTypeChange == ETestInstanceDataObjectGrain::Corn);
+	CHECK(StructData.FruitTypeChange == ETestInstanceDataObjectFruit::Orange);
+
+#if WITH_METADATA
+	UnknownEnumNames.Find(GrainTypeName, Names, bHasFlags);
+	CHECKED_IF(Names.Num() == 1)
+	{
+		CHECK(Names[0] == "Rye");
+	}
+	CHECK_FALSE(bHasFlags);
+
+	UnknownEnumNames.Find(FruitTypeName, Names, bHasFlags);
+	CHECKED_IF(Names.Num() == 1)
+	{
+		CHECK(Names[0] == "Cherry");
+	}
+	CHECK_FALSE(bHasFlags);
+
+	UnknownEnumNames.Find(DirectionTypeName, Names, bHasFlags);
+	CHECKED_IF(Names.Num() == 2)
+	{
+		CHECK(Names[0] == "Up");
+		CHECK(Names[1] == "Down");
+	}
+	CHECK(bHasFlags);
+
+	UnknownEnumNames.Find(ParsePropertyTypeName(TEXT("ETestInstanceDataObjectDeletedGrain(/Script/CoreUObject)")), Names, bHasFlags);
+	CHECKED_IF(Names.Num() == 1)
+	{
+		CHECK(Names[0] == "Rice");
+	}
+	CHECK_FALSE(bHasFlags);
+
+	UnknownEnumNames.Find(ParsePropertyTypeName(TEXT("ETestInstanceDataObjectDeletedFruit(/Script/CoreUObject)")), Names, bHasFlags);
+	CHECKED_IF(Names.Num() == 1)
+	{
+		CHECK(Names[0] == "Apple");
+	}
+	CHECK_FALSE(bHasFlags);
+
+	UnknownEnumNames.Find(ParsePropertyTypeName(TEXT("ETestInstanceDataObjectDeletedDirection(/Script/CoreUObject)")), Names, bHasFlags);
+	CHECKED_IF(Names.Num() == 2)
+	{
+		CHECK(Names[0] == "Up");
+		CHECK(Names[1] == "South");
+	}
+	CHECK(bHasFlags);
+#endif // WITH_METADATA
+
+	UnknownEnumNames.Destroy();
+
+#if WITH_TEXT_ARCHIVE_SUPPORT
+	StructData = {};
+
+	{
+		FMemoryReader Ar(JsonData, /*bIsPersistent*/ true);
+		FJsonArchiveInputFormatter Formatter(Ar);
+		FStructuredArchive StructuredAr(Formatter);
+		FTestInstanceDataObjectStruct::StaticStruct()->SerializeTaggedProperties(StructuredAr.Open(), (uint8*)&StructData, nullptr, nullptr);
+	}
+
+	CHECK(StructData.Grain == (ETestInstanceDataObjectGrain::Type)((uint8)ETestInstanceDataObjectGrain::Wheat + 1));
+	CHECK(StructData.Fruit == (ETestInstanceDataObjectFruit)((uint8)ETestInstanceDataObjectFruit::Orange + 1));
+	CHECK(StructData.Direction == (ETestInstanceDataObjectDirection)MAX_uint16);
+	CHECK(StructData.GrainFromEnumClass == ETestInstanceDataObjectGrain::Corn);
+	CHECK(StructData.FruitFromNamespace == ETestInstanceDataObjectFruit::Orange);
+	CHECK(StructData.GrainTypeChange == ETestInstanceDataObjectGrain::Corn);
+	CHECK(StructData.FruitTypeChange == ETestInstanceDataObjectFruit::Orange);
+
+#if WITH_METADATA
+	UnknownEnumNames.Find(GrainTypeName, Names, bHasFlags);
+	CHECKED_IF(Names.Num() == 1)
+	{
+		CHECK(Names[0] == "Rye");
+	}
+	CHECK_FALSE(bHasFlags);
+
+	UnknownEnumNames.Find(FruitTypeName, Names, bHasFlags);
+	CHECKED_IF(Names.Num() == 1)
+	{
+		CHECK(Names[0] == "Cherry");
+	}
+	CHECK_FALSE(bHasFlags);
+
+	UnknownEnumNames.Find(DirectionTypeName, Names, bHasFlags);
+	CHECKED_IF(Names.Num() == 2)
+	{
+		CHECK(Names[0] == "Up");
+		CHECK(Names[1] == "Down");
+	}
+	CHECK(bHasFlags);
+
+	// Testing of the unknown property tree is skipped because it is not supported by the text format.
+#endif // WITH_METADATA
+
+	UnknownEnumNames.Destroy();
+#endif // WITH_TEXT_ARCHIVE_SUPPORT
+}
+
+TEST_CASE_NAMED(FUnknownEnumNamesTest, "CoreUObject::Serialization::UnknownEnumNames", "[CoreUObject][EngineFilter]")
+{
+	UObject* Owner = NewObject<UTestInstanceDataObjectClass>();
+
+	TArray<FName> Names{NAME_None};
+	bool bHasFlags = true;
+
+	// Test a non-flags enum...
+
+	FPropertyTypeName FruitTypeName = []
+	{
+		FPropertyTypeNameBuilder Builder;
+		Builder.AddPath(StaticEnum<ETestInstanceDataObjectFruit>());
+		return Builder.Build();
+	}();
+
+	FUnknownEnumNames UnknownEnumNames(Owner);
+
+	UnknownEnumNames.Find(FruitTypeName, Names, bHasFlags);
+	CHECK(Names.IsEmpty());
+	CHECK_FALSE(bHasFlags);
+
+	const FName NAME_Cherry = "Cherry";
+	const FName NAME_Pear = "Pear";
+
+	UnknownEnumNames.Add(nullptr, FruitTypeName, NAME_Pear);
+	UnknownEnumNames.Add(StaticEnum<ETestInstanceDataObjectFruit>(), {}, NAME_Cherry);
+	UnknownEnumNames.Add(StaticEnum<ETestInstanceDataObjectFruit>(), {}, NAME_Pear);
+	UnknownEnumNames.Add(nullptr, FruitTypeName, NAME_Cherry);
+
+	UnknownEnumNames.Find(FruitTypeName, Names, bHasFlags);
+	CHECKED_IF(Names.Num() == 2)
+	{
+		CHECK(Names[0] == NAME_Pear);
+		CHECK(Names[1] == NAME_Cherry);
+	}
+	CHECK_FALSE(bHasFlags);
+
+	// Test a flags enum by name only...
+
+	FPropertyTypeName FlagsTypeName = []
+	{
+		FPropertyTypeNameBuilder Builder;
+		Builder.AddPath(StaticEnum<ETestInstanceDataObjectDirection>());
+		return Builder.Build();
+	}();
+
+	const FName NAME_South = "South";
+	const FName NAME_Down = "Down";
+	const FName NAME_Up = "Up";
+
+	TStringBuilder<128> FlagsString;
+	FlagsString.Join(MakeArrayView({NAME_Up, NAME_Down, NAME_South}), TEXTVIEW(" | "));
+
+	UnknownEnumNames.Add(nullptr, FlagsTypeName, NAME_Down);
+	UnknownEnumNames.Add(nullptr, FlagsTypeName, *FlagsString);
+
+	UnknownEnumNames.Find(FlagsTypeName, Names, bHasFlags);
+	CHECKED_IF(Names.Num() == 3)
+	{
+		CHECK(Names[0] == NAME_Down);
+		CHECK(Names[1] == NAME_Up);
+		CHECK(Names[2] == NAME_South);
+	}
+	CHECK(bHasFlags);
+
+	// Test resetting unknown enum names for an owner...
+
+	UnknownEnumNames.Destroy();
+
+	UnknownEnumNames.Find(FlagsTypeName, Names, bHasFlags);
+	CHECK(Names.IsEmpty());
+	CHECK_FALSE(bHasFlags);
+
+	// Test a flags enum by enum...
+
+	UnknownEnumNames.Add(StaticEnum<ETestInstanceDataObjectDirection>(), {}, NAME_Up);
+
+	UnknownEnumNames.Find(FlagsTypeName, Names, bHasFlags);
+	CHECK(Names.Num() == 1);
+	CHECK(bHasFlags);
+
+	UnknownEnumNames.Add(StaticEnum<ETestInstanceDataObjectDirection>(), FlagsTypeName, *FlagsString);
+
+	UnknownEnumNames.Find(FlagsTypeName, Names, bHasFlags);
+	CHECKED_IF(Names.Num() == 2)
+	{
+		CHECK(Names[0] == NAME_Up);
+		CHECK(Names[1] == NAME_Down);
+	}
+	CHECK(bHasFlags);
+}
+
+} // UE
+
+#endif // WITH_TESTS && WITH_EDITORONLY_DATA
diff --git a/Engine/Source/Runtime/CoreUObject/Tests/UObject/PropertyStateTrackingTest.h b/Engine/Source/Runtime/CoreUObject/Tests/UObject/PropertyStateTrackingTest.h
new file mode 100644
--- /dev/null
+++ b/Engine/Source/Runtime/CoreUObject/Tests/UObject/PropertyStateTrackingTest.h
@@ -0,0 +1,314 @@
+// Copyright Epic Games, Inc. All Rights Reserved.
+
+#pragma once
+
+#include "CoreTypes.h"
+#include "Misc/EnumClassFlags.h"
+#include "UObject/Object.h"
+#include "PropertyStateTrackingTest.generated.h"
+
+UENUM()
+enum ETestInstanceDataObjectBird : uint8
+{
+	TIDOB_None = 0,
+	TIDOB_Cardinal,
+	TIDOB_Crow,
+	TIDOB_Eagle,
+	TIDOB_Hawk,
+	TIDOB_Owl,
+	TIDOB_Raven,
+};
+
+UENUM()
+namespace ETestInstanceDataObjectGrain
+{
+	enum Type : uint8
+	{
+		None = 0,
+		Barley,
+		Corn,
+		Quinoa,
+		Rice,
+		Wheat,
+	};
+}
+
+UENUM()
+namespace ETestInstanceDataObjectGrainAlternate
+{
+	enum Type : uint8
+	{
+		None = 0,
+		Corn,
+		Rice,
+		Rye,
+		Wheat,
+	};
+}
+
+UENUM()
+enum class ETestInstanceDataObjectGrainAlternateEnumClass : uint8
+{
+	None = 0,
+	Corn,
+	Rice,
+	Rye,
+	Wheat,
+};
+
+static_assert((uint8)ETestInstanceDataObjectGrain::Corn != (uint8)ETestInstanceDataObjectGrainAlternate::Corn);
+static_assert((uint8)ETestInstanceDataObjectGrain::Corn != (uint8)ETestInstanceDataObjectGrainAlternateEnumClass::Corn);
+
+UENUM()
+enum class ETestInstanceDataObjectFruit : uint8
+{
+	None = 0,
+	Apple,
+	Banana,
+	Lemon,
+	Orange,
+};
+
+UENUM()
+enum class ETestInstanceDataObjectFruitAlternate : uint8
+{
+	None = 0,
+	Apple,
+	Cherry,
+	Orange,
+	Pear,
+};
+
+UENUM()
+namespace ETestInstanceDataObjectFruitAlternateNamespace
+{
+	enum Type : uint8
+	{
+		None = 0,
+		Apple,
+		Cherry,
+		Orange,
+		Pear,
+	};
+}
+
+static_assert((uint8)ETestInstanceDataObjectFruit::Orange != (uint8)ETestInstanceDataObjectFruitAlternate::Orange);
+static_assert((uint8)ETestInstanceDataObjectFruit::Orange != (uint8)ETestInstanceDataObjectFruitAlternateNamespace::Orange);
+
+UENUM(Flags)
+enum class ETestInstanceDataObjectDirection : uint16
+{
+	None = 0,
+	North = 1 << 0,
+	East = 1 << 1,
+	South = 1 << 2,
+	West = 1 << 3,
+};
+
+ENUM_CLASS_FLAGS(ETestInstanceDataObjectDirection);
+
+UENUM(Flags)
+enum class ETestInstanceDataObjectDirectionAlternate : uint16
+{
+	None = 0,
+	Up = 1 << 0,
+	Down = 1 << 1,
+	North = 1 << 2,
+	East = 1 << 3,
+	South = 1 << 4,
+	West = 1 << 5,
+};
+
+ENUM_CLASS_FLAGS(ETestInstanceDataObjectDirectionAlternate);
+
+UENUM(Flags)
+enum class ETestInstanceDataObjectFullFlags : uint8
+{
+	None = 0,
+	Flag0 = 1 << 0,
+	Flag1 = 1 << 1,
+	Flag2 = 1 << 2,
+	// Flag3 skipped for testing insertion of an unknown flag in the middle.
+	Flag4 = 1 << 4,
+	Flag5 = 1 << 5,
+	Flag6 = 1 << 6,
+	Flag7 = 1 << 7,
+};
+
+ENUM_CLASS_FLAGS(ETestInstanceDataObjectFullFlags);
+
+USTRUCT()
+struct FTestInstanceDataObjectPoint
+{
+	GENERATED_BODY()
+
+public:
+	UPROPERTY()
+	int32 X = 0;
+
+	UPROPERTY()
+	int32 Y = 0;
+
+	UPROPERTY()
+	int32 Z = 0;
+
+	UPROPERTY()
+	int32 W = 0;
+};
+
+USTRUCT()
+struct FTestInstanceDataObjectPointAlternate
+{
+	GENERATED_BODY()
+
+public:
+	UPROPERTY()
+	int32 U = 0;
+
+	UPROPERTY()
+	int32 V = 0;
+
+	UPROPERTY()
+	int32 W = 0;
+};
+
+USTRUCT()
+struct FTestInstanceDataObjectStruct
+{
+	GENERATED_BODY()
+
+public:
+	UPROPERTY()
+	int32 A = -1;
+
+	UPROPERTY()
+	int32 B = -1;
+
+	UPROPERTY()
+	int32 C = -1;
+
+	UPROPERTY()
+	int32 D = -1;
+
+	UPROPERTY()
+	TEnumAsByte<ETestInstanceDataObjectBird> Bird = TIDOB_None;
+
+	UPROPERTY()
+	TEnumAsByte<ETestInstanceDataObjectGrain::Type> Grain = ETestInstanceDataObjectGrain::None;
+
+	UPROPERTY()
+	ETestInstanceDataObjectFruit Fruit = ETestInstanceDataObjectFruit::None;
+
+	UPROPERTY()
+	ETestInstanceDataObjectDirection Direction = ETestInstanceDataObjectDirection::None;
+
+	UPROPERTY()
+	ETestInstanceDataObjectFullFlags FullFlags = ETestInstanceDataObjectFullFlags::None;
+
+	UPROPERTY()
+	TEnumAsByte<ETestInstanceDataObjectGrain::Type> GrainFromEnumClass = ETestInstanceDataObjectGrain::None;
+
+	UPROPERTY()
+	ETestInstanceDataObjectFruit FruitFromNamespace = ETestInstanceDataObjectFruit::None;
+
+	UPROPERTY()
+	TEnumAsByte<ETestInstanceDataObjectGrain::Type> GrainTypeChange = ETestInstanceDataObjectGrain::None;
+
+	UPROPERTY()
+	ETestInstanceDataObjectFruit FruitTypeChange = ETestInstanceDataObjectFruit::None;
+
+	UPROPERTY()
+	TEnumAsByte<ETestInstanceDataObjectGrain::Type> GrainTypeAndPropertyChange = ETestInstanceDataObjectGrain::None;
+
+	UPROPERTY()
+	ETestInstanceDataObjectFruit FruitTypeAndPropertyChange = ETestInstanceDataObjectFruit::None;
+
+	UPROPERTY()
+	FTestInstanceDataObjectPoint Point;
+};
+
+USTRUCT()
+struct FTestInstanceDataObjectStructAlternate
+{
+	GENERATED_BODY()
+
+public:
+	UPROPERTY()
+	float B = -1;
+
+	UPROPERTY()
+	int64 C = -1;
+
+	UPROPERTY()
+	int32 D = -1;
+
+	UPROPERTY()
+	int32 E = -1;
+
+	UPROPERTY()
+	TEnumAsByte<ETestInstanceDataObjectBird> Bird = TIDOB_None;
+
+	UPROPERTY(meta=(OriginalType="ETestInstanceDataObjectGrain(/Script/CoreUObject)"))
+	TEnumAsByte<ETestInstanceDataObjectGrainAlternate::Type> Grain = ETestInstanceDataObjectGrainAlternate::None;
+
+	UPROPERTY(meta=(OriginalType="ETestInstanceDataObjectFruit(/Script/CoreUObject)"))
+	ETestInstanceDataObjectFruitAlternate Fruit = ETestInstanceDataObjectFruitAlternate::None;
+
+	UPROPERTY(meta=(OriginalType="ETestInstanceDataObjectDirection(/Script/CoreUObject)"))
+	ETestInstanceDataObjectDirectionAlternate Direction = ETestInstanceDataObjectDirectionAlternate::None;
+
+	UPROPERTY(meta=(OriginalType="ETestInstanceDataObjectGrain(/Script/CoreUObject)"))
+	ETestInstanceDataObjectGrainAlternateEnumClass GrainFromEnumClass = ETestInstanceDataObjectGrainAlternateEnumClass::None;
+
+	UPROPERTY(meta=(OriginalType="ETestInstanceDataObjectFruit(/Script/CoreUObject)"))
+	TEnumAsByte<ETestInstanceDataObjectFruitAlternateNamespace::Type> FruitFromNamespace = ETestInstanceDataObjectFruitAlternateNamespace::None;
+
+	UPROPERTY()
+	TEnumAsByte<ETestInstanceDataObjectGrainAlternate::Type> GrainTypeChange = ETestInstanceDataObjectGrainAlternate::None;
+
+	UPROPERTY()
+	ETestInstanceDataObjectFruitAlternate FruitTypeChange = ETestInstanceDataObjectFruitAlternate::None;
+
+	UPROPERTY()
+	ETestInstanceDataObjectGrainAlternateEnumClass GrainTypeAndPropertyChange = ETestInstanceDataObjectGrainAlternateEnumClass::None;
+
+	UPROPERTY()
+	TEnumAsByte<ETestInstanceDataObjectFruitAlternateNamespace::Type> FruitTypeAndPropertyChange = ETestInstanceDataObjectFruitAlternateNamespace::None;
+
+	UPROPERTY(meta=(OriginalType="ETestInstanceDataObjectDeletedGrain(/Script/CoreUObject)"))
+	TEnumAsByte<ETestInstanceDataObjectGrainAlternate::Type> DeletedGrain = ETestInstanceDataObjectGrainAlternate::None;
+
+	UPROPERTY(meta=(OriginalType="ETestInstanceDataObjectDeletedFruit(/Script/CoreUObject)"))
+	ETestInstanceDataObjectFruitAlternate DeletedFruit = ETestInstanceDataObjectFruitAlternate::None;
+
+	UPROPERTY(meta=(OriginalType="ETestInstanceDataObjectDeletedDirection(/Script/CoreUObject)"))
+	ETestInstanceDataObjectDirectionAlternate DeletedDirection = ETestInstanceDataObjectDirectionAlternate::None;
+
+	UPROPERTY(meta=(OriginalType="TestInstanceDataObjectPoint(/Script/CoreUObject)"))
+	FTestInstanceDataObjectPointAlternate Point;
+};
+
+UCLASS()
+class UTestInstanceDataObjectClass : public UObject
+{
+	GENERATED_BODY()
+
+public:
+	UPROPERTY()
+	int32 A[4]{-1, -1, -1, -1};
+
+	UPROPERTY()
+	float B = -1;
+
+	UPROPERTY()
+	int64 C = -1;
+
+	UPROPERTY()
+	int32 D = -1;
+
+	UPROPERTY()
+	int32 E = -1;
+
+	UPROPERTY()
+	FTestInstanceDataObjectStruct Struct;
+};
