SavePackage: Remove handling of deprecated UObject::PreSave, UObject::PreSaveRoot, and UObject::PostSaveRoot functions.
[RN] Minor, Foundation

diff --git a/Engine/Plugins/Experimental/CommonConversation/Source/CommonConversationRuntime/Private/ConversationDatabase.cpp b/Engine/Plugins/Experimental/CommonConversation/Source/CommonConversationRuntime/Private/ConversationDatabase.cpp
--- a/Engine/Plugins/Experimental/CommonConversation/Source/CommonConversationRuntime/Private/ConversationDatabase.cpp
+++ b/Engine/Plugins/Experimental/CommonConversation/Source/CommonConversationRuntime/Private/ConversationDatabase.cpp
@@ -71,13 +71,6 @@ FLinearColor UConversationDatabase::GetDebugParticipantColor(FGameplayTag Partic
 }
 
 #if WITH_EDITOR
-void UConversationDatabase::PreSave(const class ITargetPlatform* TargetPlatform)
-{
-	PRAGMA_DISABLE_DEPRECATION_WARNINGS;
-	Super::PreSave(TargetPlatform);
-	PRAGMA_ENABLE_DEPRECATION_WARNINGS;
-}
-
 void UConversationDatabase::PreSave(FObjectPreSaveContext ObjectSaveContext)
 {
 	Super::PreSave(ObjectSaveContext);
diff --git a/Engine/Plugins/Experimental/CommonConversation/Source/CommonConversationRuntime/Public/ConversationDatabase.h b/Engine/Plugins/Experimental/CommonConversation/Source/CommonConversationRuntime/Public/ConversationDatabase.h
--- a/Engine/Plugins/Experimental/CommonConversation/Source/CommonConversationRuntime/Public/ConversationDatabase.h
+++ b/Engine/Plugins/Experimental/CommonConversation/Source/CommonConversationRuntime/Public/ConversationDatabase.h
@@ -81,10 +81,6 @@ public:
 	virtual void PostDuplicate(bool bDuplicateForPIE) override;
 	virtual EDataValidationResult IsDataValid(class FDataValidationContext& Context) const override;
 	virtual EDataValidationResult ValidateOutBoundConnections(class FDataValidationContext& Context) const;
-	PRAGMA_DISABLE_DEPRECATION_WARNINGS // Suppress compiler warning on override of deprecated function
-	UE_DEPRECATED(5.0, "Use version that takes FObjectPreSaveContext instead.")
-	virtual void PreSave(const class ITargetPlatform* TargetPlatform) override;
-	PRAGMA_ENABLE_DEPRECATION_WARNINGS
 	virtual void PreSave(FObjectPreSaveContext ObjectSaveContext) override;
 
 	const TMap<FGuid, TObjectPtr<UConversationNode>>& GetFullNodeMap() const { return FullNodeMap; }
diff --git a/Engine/Plugins/Experimental/VirtualProductionUtilities/Source/VPUtilitiesEditor/Private/VPCameraBase.cpp b/Engine/Plugins/Experimental/VirtualProductionUtilities/Source/VPUtilitiesEditor/Private/VPCameraBase.cpp
--- a/Engine/Plugins/Experimental/VirtualProductionUtilities/Source/VPUtilitiesEditor/Private/VPCameraBase.cpp
+++ b/Engine/Plugins/Experimental/VirtualProductionUtilities/Source/VPUtilitiesEditor/Private/VPCameraBase.cpp
@@ -88,13 +88,6 @@ void AVPCameraBase::BeginDestroy()
 	Super::BeginDestroy();
 }
 
-void AVPCameraBase::PreSave(const class ITargetPlatform* TargetPlatform)
-{
-	PRAGMA_DISABLE_DEPRECATION_WARNINGS;
-	Super::PreSave(TargetPlatform);
-	PRAGMA_ENABLE_DEPRECATION_WARNINGS;
-}
-
 void AVPCameraBase::PreSave(FObjectPreSaveContext ObjectSaveContext)
 {
 	RemoveUser(true);
diff --git a/Engine/Plugins/Experimental/VirtualProductionUtilities/Source/VPUtilitiesEditor/Private/VPCameraBase.h b/Engine/Plugins/Experimental/VirtualProductionUtilities/Source/VPUtilitiesEditor/Private/VPCameraBase.h
--- a/Engine/Plugins/Experimental/VirtualProductionUtilities/Source/VPUtilitiesEditor/Private/VPCameraBase.h
+++ b/Engine/Plugins/Experimental/VirtualProductionUtilities/Source/VPUtilitiesEditor/Private/VPCameraBase.h
@@ -26,10 +26,6 @@ public:
 	virtual bool ShouldTickIfViewportsOnly() const override;
 	virtual void Tick(float DeltaSeconds) override;
 	virtual void BeginDestroy() override;
-	PRAGMA_DISABLE_DEPRECATION_WARNINGS // Suppress compiler warning on override of deprecated function
-	UE_DEPRECATED(5.0, "Use version that takes FObjectPreSaveContext instead.")
-	virtual void PreSave(const class ITargetPlatform* TargetPlatform) override;
-	PRAGMA_ENABLE_DEPRECATION_WARNINGS
 	virtual void PreSave(FObjectPreSaveContext ObjectSaveContext) override;
 
 private:
diff --git a/Engine/Plugins/Runtime/HairStrands/Source/HairStrandsCore/Private/GroomBindingAsset.cpp b/Engine/Plugins/Runtime/HairStrands/Source/HairStrandsCore/Private/GroomBindingAsset.cpp
--- a/Engine/Plugins/Runtime/HairStrands/Source/HairStrandsCore/Private/GroomBindingAsset.cpp
+++ b/Engine/Plugins/Runtime/HairStrands/Source/HairStrandsCore/Private/GroomBindingAsset.cpp
@@ -572,13 +572,6 @@ void UGroomBindingAsset::PostLoad()
 #endif
 }
 
-void UGroomBindingAsset::PreSave(const class ITargetPlatform* TargetPlatform)
-{
-	PRAGMA_DISABLE_DEPRECATION_WARNINGS;
-	Super::PreSave(TargetPlatform);
-	PRAGMA_ENABLE_DEPRECATION_WARNINGS;
-}
-
 void UGroomBindingAsset::PreSave(FObjectPreSaveContext ObjectSaveContext)
 {
 	Super::PreSave(ObjectSaveContext);
@@ -587,13 +580,6 @@ void UGroomBindingAsset::PreSave(FObjectPreSaveContext ObjectSaveContext)
 #endif
 }
 
-void UGroomBindingAsset::PostSaveRoot(bool bCleanupIsRequired)
-{
-	PRAGMA_DISABLE_DEPRECATION_WARNINGS;
-	Super::PostSaveRoot(bCleanupIsRequired);
-	PRAGMA_ENABLE_DEPRECATION_WARNINGS;
-}
-
 void UGroomBindingAsset::PostSaveRoot(FObjectPostSaveRootContext ObjectSaveContext)
 {
 	Super::PostSaveRoot(ObjectSaveContext);
diff --git a/Engine/Plugins/Runtime/HairStrands/Source/HairStrandsCore/Public/GroomBindingAsset.h b/Engine/Plugins/Runtime/HairStrands/Source/HairStrandsCore/Public/GroomBindingAsset.h
--- a/Engine/Plugins/Runtime/HairStrands/Source/HairStrandsCore/Public/GroomBindingAsset.h
+++ b/Engine/Plugins/Runtime/HairStrands/Source/HairStrandsCore/Public/GroomBindingAsset.h
@@ -374,12 +374,6 @@ public:
 	//~ Begin UObject Interface.
 	virtual void GetResourceSizeEx(FResourceSizeEx& CumulativeResourceSize) override;
 	virtual void PostLoad() override;
-	PRAGMA_DISABLE_DEPRECATION_WARNINGS // Suppress compiler warning on override of deprecated function
-	UE_DEPRECATED(5.0, "Use version that takes FObjectPreSaveContext instead.")
-	virtual void PreSave(const class ITargetPlatform* TargetPlatform) override;
-	UE_DEPRECATED(5.0, "Use version that takes FObjectPostSaveRootContext instead.")
-	virtual void PostSaveRoot(bool bCleanupIsRequired) override;
-	PRAGMA_ENABLE_DEPRECATION_WARNINGS
 	virtual void PreSave(FObjectPreSaveContext ObjectSaveContext) override;
 	virtual void PostSaveRoot(FObjectPostSaveRootContext ObjectSaveContext) override;
 	virtual void BeginDestroy() override;
diff --git a/Engine/Plugins/Runtime/RigVM/Source/RigVMDeveloper/Private/RigVMBlueprint.cpp b/Engine/Plugins/Runtime/RigVM/Source/RigVMDeveloper/Private/RigVMBlueprint.cpp
--- a/Engine/Plugins/Runtime/RigVM/Source/RigVMDeveloper/Private/RigVMBlueprint.cpp
+++ b/Engine/Plugins/Runtime/RigVM/Source/RigVMDeveloper/Private/RigVMBlueprint.cpp
@@ -844,13 +844,6 @@ void URigVMBlueprint::Serialize(FArchive& Ar)
 	}
 }
 
-void URigVMBlueprint::PreSave(const class ITargetPlatform* TargetPlatform)
-{
-	PRAGMA_DISABLE_DEPRECATION_WARNINGS;
-	Super::PreSave(TargetPlatform);
-	PRAGMA_ENABLE_DEPRECATION_WARNINGS;
-}
-
 void URigVMBlueprint::PreSave(FObjectPreSaveContext ObjectSaveContext)
 {
 	Super::PreSave(ObjectSaveContext);
diff --git a/Engine/Plugins/Runtime/RigVM/Source/RigVMDeveloper/Public/RigVMBlueprint.h b/Engine/Plugins/Runtime/RigVM/Source/RigVMDeveloper/Public/RigVMBlueprint.h
--- a/Engine/Plugins/Runtime/RigVM/Source/RigVMDeveloper/Public/RigVMBlueprint.h
+++ b/Engine/Plugins/Runtime/RigVM/Source/RigVMDeveloper/Public/RigVMBlueprint.h
@@ -280,10 +280,6 @@ public:
 	virtual void GetTypeActions(FBlueprintActionDatabaseRegistrar& ActionRegistrar) const override;
 	virtual void GetInstanceActions(FBlueprintActionDatabaseRegistrar& ActionRegistrar) const override;
 	virtual void SetObjectBeingDebugged(UObject* NewObject) override;
-	PRAGMA_DISABLE_DEPRECATION_WARNINGS // Suppress compiler warning on override of deprecated function
-	UE_DEPRECATED(5.0, "Use version that takes FObjectPreSaveContext instead.")
-	virtual void PreSave(const class ITargetPlatform* TargetPlatform) override;
-	PRAGMA_ENABLE_DEPRECATION_WARNINGS
 	virtual void PreSave(FObjectPreSaveContext ObjectSaveContext) override;
 	virtual void PostSaveRoot(FObjectPostSaveRootContext ObjectSaveContext) override;
 	virtual void PostLoad() override;
diff --git a/Engine/Source/Editor/BlueprintGraph/Classes/K2Node_FunctionEntry.h b/Engine/Source/Editor/BlueprintGraph/Classes/K2Node_FunctionEntry.h
--- a/Engine/Source/Editor/BlueprintGraph/Classes/K2Node_FunctionEntry.h
+++ b/Engine/Source/Editor/BlueprintGraph/Classes/K2Node_FunctionEntry.h
@@ -55,10 +55,6 @@ class UK2Node_FunctionEntry : public UK2Node_FunctionTerminator
 
 	//~ Begin UObject Interface
 	virtual void Serialize(FArchive& Ar) override;
-	PRAGMA_DISABLE_DEPRECATION_WARNINGS // Suppress compiler warning on override of deprecated function
-	UE_DEPRECATED(5.0, "Use version that takes FObjectPreSaveContext instead.")
-	virtual void PreSave(const class ITargetPlatform* TargetPlatform) override;
-	PRAGMA_ENABLE_DEPRECATION_WARNINGS
 	virtual void PreSave(FObjectPreSaveContext ObjectSaveContext) override;
 	virtual void PostLoad() override;
 	//~ End UObject Interface
diff --git a/Engine/Source/Editor/BlueprintGraph/Private/K2Node_FunctionEntry.cpp b/Engine/Source/Editor/BlueprintGraph/Private/K2Node_FunctionEntry.cpp
--- a/Engine/Source/Editor/BlueprintGraph/Private/K2Node_FunctionEntry.cpp
+++ b/Engine/Source/Editor/BlueprintGraph/Private/K2Node_FunctionEntry.cpp
@@ -230,13 +230,6 @@ UK2Node_FunctionEntry::UK2Node_FunctionEntry(const FObjectInitializer& ObjectIni
 	bCanRenameNode = bIsEditable;
 }
 
-void UK2Node_FunctionEntry::PreSave(const class ITargetPlatform* TargetPlatform)
-{
-	PRAGMA_DISABLE_DEPRECATION_WARNINGS;
-	Super::PreSave(TargetPlatform);
-	PRAGMA_ENABLE_DEPRECATION_WARNINGS;
-}
-
 void UK2Node_FunctionEntry::PreSave(FObjectPreSaveContext ObjectSaveContext)
 {
 	Super::PreSave(ObjectSaveContext);
diff --git a/Engine/Source/Editor/UMGEditor/Private/WidgetBlueprint.cpp b/Engine/Source/Editor/UMGEditor/Private/WidgetBlueprint.cpp
--- a/Engine/Source/Editor/UMGEditor/Private/WidgetBlueprint.cpp
+++ b/Engine/Source/Editor/UMGEditor/Private/WidgetBlueprint.cpp
@@ -694,13 +694,6 @@ void UWidgetBlueprint::ReplaceDeprecatedNodes()
 }
 
 #if WITH_EDITORONLY_DATA
-void UWidgetBlueprint::PreSave(const class ITargetPlatform* TargetPlatform)
-{
-	PRAGMA_DISABLE_DEPRECATION_WARNINGS;
-	Super::PreSave(TargetPlatform);
-	PRAGMA_ENABLE_DEPRECATION_WARNINGS;
-}
-
 void UWidgetBlueprint::PreSave(FObjectPreSaveContext ObjectSaveContext)
 {
 	Super::PreSave(ObjectSaveContext);
diff --git a/Engine/Source/Editor/UMGEditor/Public/WidgetBlueprint.h b/Engine/Source/Editor/UMGEditor/Public/WidgetBlueprint.h
--- a/Engine/Source/Editor/UMGEditor/Public/WidgetBlueprint.h
+++ b/Engine/Source/Editor/UMGEditor/Public/WidgetBlueprint.h
@@ -253,10 +253,6 @@ public:
 	virtual void PostDuplicate(bool bDuplicateForPIE) override;
 
 #if WITH_EDITORONLY_DATA
-	PRAGMA_DISABLE_DEPRECATION_WARNINGS // Suppress compiler warning on override of deprecated function
-	UE_DEPRECATED(5.0, "Use version that takes FObjectPreSaveContext instead.")
-	virtual void PreSave(const class ITargetPlatform* TargetPlatform) override;
-	PRAGMA_ENABLE_DEPRECATION_WARNINGS
 	virtual void PreSave(FObjectPreSaveContext ObjectSaveContext) override;
 #endif // WITH_EDITORONLY_DATA
 
diff --git a/Engine/Source/Runtime/AIModule/Classes/EnvironmentQuery/Tests/EnvQueryTest_GameplayTags.h b/Engine/Source/Runtime/AIModule/Classes/EnvironmentQuery/Tests/EnvQueryTest_GameplayTags.h
--- a/Engine/Source/Runtime/AIModule/Classes/EnvironmentQuery/Tests/EnvQueryTest_GameplayTags.h
+++ b/Engine/Source/Runtime/AIModule/Classes/EnvironmentQuery/Tests/EnvQueryTest_GameplayTags.h
@@ -35,10 +35,6 @@ protected:
 
 	bool SatisfiesTest(const IGameplayTagAssetInterface* ItemGameplayTagAssetInterface) const;
 
-	PRAGMA_DISABLE_DEPRECATION_WARNINGS // Suppress compiler warning on override of deprecated function
-	UE_DEPRECATED(5.0, "Use version that takes FObjectPreSaveContext instead.")
-	virtual void PreSave(const class ITargetPlatform* TargetPlatform) override;
-	PRAGMA_ENABLE_DEPRECATION_WARNINGS
 	/**
 	 * Presave function. Gets called once before an object gets serialized for saving. This function is necessary
 	 * for save time computation as Serialize gets called three times per object from within SavePackage.
diff --git a/Engine/Source/Runtime/AIModule/Private/EnvironmentQuery/Tests/EnvQueryTest_GameplayTags.cpp b/Engine/Source/Runtime/AIModule/Private/EnvironmentQuery/Tests/EnvQueryTest_GameplayTags.cpp
--- a/Engine/Source/Runtime/AIModule/Private/EnvironmentQuery/Tests/EnvQueryTest_GameplayTags.cpp
+++ b/Engine/Source/Runtime/AIModule/Private/EnvironmentQuery/Tests/EnvQueryTest_GameplayTags.cpp
@@ -24,13 +24,6 @@ UEnvQueryTest_GameplayTags::UEnvQueryTest_GameplayTags(const FObjectInitializer&
 	bRejectIncompatibleItems = false;
 }
 
-void UEnvQueryTest_GameplayTags::PreSave(const class ITargetPlatform* TargetPlatform)
-{
-	PRAGMA_DISABLE_DEPRECATION_WARNINGS;
-	Super::PreSave(TargetPlatform);
-	PRAGMA_ENABLE_DEPRECATION_WARNINGS;
-}
-
 void UEnvQueryTest_GameplayTags::PreSave(FObjectPreSaveContext ObjectSaveContext)
 {
 	// If this is a new object, it's already converted to new format.  If it was old, this should've already been set to
diff --git a/Engine/Source/Runtime/CoreUObject/Private/UObject/Obj.cpp b/Engine/Source/Runtime/CoreUObject/Private/UObject/Obj.cpp
--- a/Engine/Source/Runtime/CoreUObject/Private/UObject/Obj.cpp
+++ b/Engine/Source/Runtime/CoreUObject/Private/UObject/Obj.cpp
@@ -1456,27 +1456,14 @@ void UObject::ConditionalPostLoadSubobjects( FObjectInstancingGraph* OuterInstan
 	CheckDefaultSubobjects();
 }
 
-bool UObject::PreSaveRoot(const TCHAR* Filename)
-{
-	return false;
-}
-
 void UObject::PreSaveRoot(FObjectPreSaveRootContext ObjectSaveContext)
 {
 }
 
-void UObject::PostSaveRoot(bool bCleanupIsRequired)
-{
-}
-
 void UObject::PostSaveRoot(FObjectPostSaveRootContext ObjectSaveContext)
 {
 }
 
-void UObject::PreSave(const class ITargetPlatform* TargetPlatform)
-{
-}
-
 void UObject::PreSave(FObjectPreSaveContext SaveContext)
 {
 	SaveContext.Data.bBaseClassCalled = true;
diff --git a/Engine/Source/Runtime/CoreUObject/Private/UObject/ObjectRedirector.cpp b/Engine/Source/Runtime/CoreUObject/Private/UObject/ObjectRedirector.cpp
--- a/Engine/Source/Runtime/CoreUObject/Private/UObject/ObjectRedirector.cpp
+++ b/Engine/Source/Runtime/CoreUObject/Private/UObject/ObjectRedirector.cpp
@@ -19,13 +19,6 @@
 -----------------------------------------------------------------------------*/
 
 
-void UObjectRedirector::PreSave(const class ITargetPlatform* TargetPlatform)
-{
-	PRAGMA_DISABLE_DEPRECATION_WARNINGS;
-	Super::PreSave(TargetPlatform);
-	PRAGMA_ENABLE_DEPRECATION_WARNINGS;
-}
-
 // If this object redirector is pointing to an object that won't be serialized anyway, set the RF_Transient flag
 // so that this redirector is also removed from the package.
 void UObjectRedirector::PreSave(FObjectPreSaveContext ObjectSaveContext)
diff --git a/Engine/Source/Runtime/CoreUObject/Private/UObject/SavePackage/SavePackageUtilities.cpp b/Engine/Source/Runtime/CoreUObject/Private/UObject/SavePackage/SavePackageUtilities.cpp
--- a/Engine/Source/Runtime/CoreUObject/Private/UObject/SavePackage/SavePackageUtilities.cpp
+++ b/Engine/Source/Runtime/CoreUObject/Private/UObject/SavePackage/SavePackageUtilities.cpp
@@ -666,9 +666,6 @@ bool IsProceduralSave(bool bIsCooking, const FPackagePath& TargetPackagePath, ui
 void CallPreSave(UObject* Object, FObjectSaveContextData& ObjectSaveContext)
 {
 	SCOPED_SAVETIMER_TEXT(*WriteToString<256>(GetClassTraceScope(Object), TEXTVIEW("_PreSave")));
-	PRAGMA_DISABLE_DEPRECATION_WARNINGS;
-	Object->PreSave(ObjectSaveContext.TargetPlatform);
-	PRAGMA_ENABLE_DEPRECATION_WARNINGS;
 
 	FObjectPreSaveContext ObjectPreSaveContext(ObjectSaveContext);
 	ObjectSaveContext.bBaseClassCalled = false;
@@ -689,22 +686,15 @@ void CallPreSave(UObject* Object, FObjectSaveContextData& ObjectSaveContext)
 void CallPreSaveRoot(UObject* Object, FObjectSaveContextData& ObjectSaveContext)
 {
 	SCOPED_SAVETIMER_TEXT(*WriteToString<256>(GetClassTraceScope(Object), TEXTVIEW("_PreSave")));
-	PRAGMA_DISABLE_DEPRECATION_WARNINGS;
-	bool bLegacyNeedsCleanup = Object->PreSaveRoot(*ObjectSaveContext.TargetFilename);
-	PRAGMA_ENABLE_DEPRECATION_WARNINGS;
 
 	ObjectSaveContext.bCleanupRequired = false;
 	ObjectSaveContext.Object = Object;
 	Object->PreSaveRoot(FObjectPreSaveRootContext(ObjectSaveContext));
-	ObjectSaveContext.bCleanupRequired |= bLegacyNeedsCleanup;
 }
 
 void CallPostSaveRoot(UObject* Object, FObjectSaveContextData& ObjectSaveContext, bool bNeedsCleanup)
 {
 	SCOPED_SAVETIMER_TEXT(*WriteToString<256>(GetClassTraceScope(Object), TEXTVIEW("_PreSave")));
-	PRAGMA_DISABLE_DEPRECATION_WARNINGS;
-	Object->PostSaveRoot(bNeedsCleanup);
-	PRAGMA_ENABLE_DEPRECATION_WARNINGS;
 
 	ObjectSaveContext.Object = Object;
 	ObjectSaveContext.bCleanupRequired = bNeedsCleanup;
diff --git a/Engine/Source/Runtime/CoreUObject/Public/UObject/Object.h b/Engine/Source/Runtime/CoreUObject/Public/UObject/Object.h
--- a/Engine/Source/Runtime/CoreUObject/Public/UObject/Object.h
+++ b/Engine/Source/Runtime/CoreUObject/Public/UObject/Object.h
@@ -258,9 +258,6 @@ public:
 	}
 #endif
 
-	UE_DEPRECATED(5.0, "Use version that takes FObjectPreSaveContext instead.")
-	COREUOBJECT_API virtual bool PreSaveRoot(const TCHAR* Filename);
-
 	/**
 	 * Called from within SavePackage on the passed in base/root object. The return value of this function will be passed to PostSaveRoot. 
 	 * This is used to allow objects used as a base to perform required actions before saving and cleanup afterwards.
@@ -270,9 +267,6 @@ public:
 	 */
 	COREUOBJECT_API virtual void PreSaveRoot(FObjectPreSaveRootContext ObjectSaveContext);
 
-	UE_DEPRECATED(5.0, "Use version that takes FObjectPostSaveContext instead.")
-	COREUOBJECT_API virtual void PostSaveRoot(bool bCleanupIsRequired);
-
 	/**
 	 * Called from within SavePackage on the passed in base/root object. 
 	 * This function is called after the package has been saved and can perform cleanup.
@@ -281,9 +275,6 @@ public:
 	 */
 	COREUOBJECT_API virtual void PostSaveRoot(FObjectPostSaveRootContext ObjectSaveContext);
 
-	UE_DEPRECATED(5.0, "Use version that takes FObjectPreSaveContext instead.")
-	COREUOBJECT_API virtual void PreSave(const class ITargetPlatform* TargetPlatform);
-
 	/**
 	 * Presave function. Gets called once before an object gets serialized for saving. This function is necessary
 	 * for save time computation as Serialize gets called three times per object from within SavePackage.
diff --git a/Engine/Source/Runtime/CoreUObject/Public/UObject/ObjectRedirector.h b/Engine/Source/Runtime/CoreUObject/Public/UObject/ObjectRedirector.h
--- a/Engine/Source/Runtime/CoreUObject/Public/UObject/ObjectRedirector.h
+++ b/Engine/Source/Runtime/CoreUObject/Public/UObject/ObjectRedirector.h
@@ -30,10 +30,6 @@ class UObjectRedirector : public UObject
 	// Variables.
 	UObject*		DestinationObject;
 	// UObject interface.
-	PRAGMA_DISABLE_DEPRECATION_WARNINGS // Suppress compiler warning on override of deprecated function
-	UE_DEPRECATED(5.0, "Use version that takes FObjectPreSaveContext instead.")
-	virtual void PreSave(const class ITargetPlatform* TargetPlatform) override;
-	PRAGMA_ENABLE_DEPRECATION_WARNINGS
 	virtual void PreSave(FObjectPreSaveContext ObjectSaveContext) override;
 	void Serialize(FArchive& Ar) override;
 	void Serialize(FStructuredArchive::FRecord Record) override;
diff --git a/Engine/Source/Runtime/Engine/Classes/Animation/AnimMontage.h b/Engine/Source/Runtime/Engine/Classes/Animation/AnimMontage.h
--- a/Engine/Source/Runtime/Engine/Classes/Animation/AnimMontage.h
+++ b/Engine/Source/Runtime/Engine/Classes/Animation/AnimMontage.h
@@ -729,10 +729,6 @@ class UAnimMontage : public UAnimCompositeBase
 public:
 	//~ Begin UObject Interface
 	virtual void PostLoad() override;
-	PRAGMA_DISABLE_DEPRECATION_WARNINGS // Suppress compiler warning on override of deprecated function
-	UE_DEPRECATED(5.0, "Use version that takes FObjectPreSaveContext instead.")
-	virtual void PreSave(const class ITargetPlatform* TargetPlatform) override;
-	PRAGMA_ENABLE_DEPRECATION_WARNINGS
 	virtual void PreSave(FObjectPreSaveContext ObjectSaveContext) override;
 
 	virtual FFrameRate GetSamplingFrameRate() const override;
diff --git a/Engine/Source/Runtime/Engine/Classes/Animation/AnimNotifies/AnimNotify.h b/Engine/Source/Runtime/Engine/Classes/Animation/AnimNotifies/AnimNotify.h
--- a/Engine/Source/Runtime/Engine/Classes/Animation/AnimNotifies/AnimNotify.h
+++ b/Engine/Source/Runtime/Engine/Classes/Animation/AnimNotifies/AnimNotify.h
@@ -118,11 +118,7 @@ class UAnimNotify : public UObject
 
 	/** UObject Interface */
 	ENGINE_API virtual void PostLoad() override;
-	PRAGMA_DISABLE_DEPRECATION_WARNINGS // Suppress compiler warning on override of deprecated function
-	UE_DEPRECATED(5.0, "Use version that takes FObjectPreSaveContext instead.")
-	ENGINE_API virtual void PreSave(const class ITargetPlatform* TargetPlatform) override;
-	ENGINE_API PRAGMA_ENABLE_DEPRECATION_WARNINGS
-	virtual void PreSave(FObjectPreSaveContext ObjectSaveContext) override;
+	ENGINE_API virtual void PreSave(FObjectPreSaveContext ObjectSaveContext) override;
 	/** End UObject Interface */
 
 	/** This notify is always a branching point when used on Montages. */
diff --git a/Engine/Source/Runtime/Engine/Classes/Animation/AnimSequence.h b/Engine/Source/Runtime/Engine/Classes/Animation/AnimSequence.h
--- a/Engine/Source/Runtime/Engine/Classes/Animation/AnimSequence.h
+++ b/Engine/Source/Runtime/Engine/Classes/Animation/AnimSequence.h
@@ -396,11 +396,7 @@ public:
 	ENGINE_API virtual void Serialize(FArchive& Ar) override;
 	ENGINE_API virtual void PostInitProperties() override;
 	ENGINE_API virtual void PostLoad() override;
-	PRAGMA_DISABLE_DEPRECATION_WARNINGS // Suppress compiler warning on override of deprecated function
-	UE_DEPRECATED(5.0, "Use version that takes FObjectPreSaveContext instead.")
-	ENGINE_API virtual void PreSave(const class ITargetPlatform* TargetPlatform) override;
-	ENGINE_API PRAGMA_ENABLE_DEPRECATION_WARNINGS
-	virtual void PreSave(FObjectPreSaveContext ObjectSaveContext) override;
+	ENGINE_API virtual void PreSave(FObjectPreSaveContext ObjectSaveContext) override;
 	ENGINE_API virtual void GetPreloadDependencies(TArray<UObject*>& OutDeps) override;
 #if WITH_EDITOR
 	ENGINE_API virtual void PostEditChangeProperty(FPropertyChangedEvent& PropertyChangedEvent) override;
diff --git a/Engine/Source/Runtime/Engine/Classes/Animation/AnimStreamable.h b/Engine/Source/Runtime/Engine/Classes/Animation/AnimStreamable.h
--- a/Engine/Source/Runtime/Engine/Classes/Animation/AnimStreamable.h
+++ b/Engine/Source/Runtime/Engine/Classes/Animation/AnimStreamable.h
@@ -208,10 +208,6 @@ public:
 	bool bUseNormalizedRootMotionScale;
 
 	//~ Begin UObject Interface
-	PRAGMA_DISABLE_DEPRECATION_WARNINGS // Suppress compiler warning on override of deprecated function
-	UE_DEPRECATED(5.0, "Use version that takes FObjectPreSaveContext instead.")
-	virtual void PreSave(const class ITargetPlatform* TargetPlatform) override;
-	PRAGMA_ENABLE_DEPRECATION_WARNINGS
 	virtual void PreSave(FObjectPreSaveContext ObjectSaveContext) override;
 	virtual void Serialize(FArchive& Ar) override;
 	virtual void PostLoad() override;
diff --git a/Engine/Source/Runtime/Engine/Classes/Animation/PoseAsset.h b/Engine/Source/Runtime/Engine/Classes/Animation/PoseAsset.h
--- a/Engine/Source/Runtime/Engine/Classes/Animation/PoseAsset.h
+++ b/Engine/Source/Runtime/Engine/Classes/Animation/PoseAsset.h
@@ -251,10 +251,6 @@ public:
 	virtual void PostLoad() override;
 	virtual bool IsPostLoadThreadSafe() const override;
 	virtual void Serialize(FArchive& Ar) override;
-	PRAGMA_DISABLE_DEPRECATION_WARNINGS // Suppress compiler warning on override of deprecated function
-	UE_DEPRECATED(5.0, "Use version that takes FObjectPreSaveContext instead.")
-	virtual void PreSave(const class ITargetPlatform* TargetPlatform) override;
-	PRAGMA_ENABLE_DEPRECATION_WARNINGS
 	virtual void PreSave(FObjectPreSaveContext ObjectSaveContext) override;
 	virtual void GetAssetRegistryTags(FAssetRegistryTagsContext Context) const override;
 	UE_DEPRECATED(5.4, "Implement the version that takes FAssetRegistryTagsContext instead.")
diff --git a/Engine/Source/Runtime/Engine/Classes/Components/LightComponent.h b/Engine/Source/Runtime/Engine/Classes/Components/LightComponent.h
--- a/Engine/Source/Runtime/Engine/Classes/Components/LightComponent.h
+++ b/Engine/Source/Runtime/Engine/Classes/Components/LightComponent.h
@@ -448,11 +448,7 @@ public:
 	ENGINE_API virtual void Serialize(FArchive& Ar) override;
 	ENGINE_API virtual void PostLoad() override;
 #if WITH_EDITOR
-	PRAGMA_DISABLE_DEPRECATION_WARNINGS // Suppress compiler warning on override of deprecated function
-	UE_DEPRECATED(5.0, "Use version that takes FObjectPreSaveContext instead.")
-	ENGINE_API virtual void PreSave(const class ITargetPlatform* TargetPlatform) override;
-	ENGINE_API PRAGMA_ENABLE_DEPRECATION_WARNINGS
-	virtual void PreSave(FObjectPreSaveContext ObjectSaveContext) override;
+	ENGINE_API virtual void PreSave(FObjectPreSaveContext ObjectSaveContext) override;
 	ENGINE_API virtual bool CanEditChange(const FProperty* InProperty) const override;
 	ENGINE_API virtual void PostEditChangeProperty(FPropertyChangedEvent& PropertyChangedEvent) override;
 	ENGINE_API virtual void UpdateLightSpriteTexture() override;
diff --git a/Engine/Source/Runtime/Engine/Classes/Components/PrimitiveComponent.h b/Engine/Source/Runtime/Engine/Classes/Components/PrimitiveComponent.h
--- a/Engine/Source/Runtime/Engine/Classes/Components/PrimitiveComponent.h
+++ b/Engine/Source/Runtime/Engine/Classes/Components/PrimitiveComponent.h
@@ -2501,11 +2501,7 @@ public:
 	ENGINE_API virtual bool CanEditChange(const FProperty* InProperty) const override;
 	ENGINE_API virtual void UpdateCollisionProfile();
 	ENGINE_API virtual void PostEditImport() override;
-	PRAGMA_DISABLE_DEPRECATION_WARNINGS // Suppress compiler warning on override of deprecated function
-	UE_DEPRECATED(5.0, "Use version that takes FObjectPreSaveContext instead.")
-	ENGINE_API virtual void PreSave(const class ITargetPlatform* TargetPlatform) override;
-	ENGINE_API PRAGMA_ENABLE_DEPRECATION_WARNINGS
-	virtual void PreSave(FObjectPreSaveContext ObjectSaveContext) override;
+	ENGINE_API virtual void PreSave(FObjectPreSaveContext ObjectSaveContext) override;
 #endif // WITH_EDITOR
 	//~ End UObject Interface.
 
diff --git a/Engine/Source/Runtime/Engine/Classes/Components/StaticMeshComponent.h b/Engine/Source/Runtime/Engine/Classes/Components/StaticMeshComponent.h
--- a/Engine/Source/Runtime/Engine/Classes/Components/StaticMeshComponent.h
+++ b/Engine/Source/Runtime/Engine/Classes/Components/StaticMeshComponent.h
@@ -511,11 +511,7 @@ public:
 	ENGINE_API virtual void UpdateBounds() override;
 #endif // WITH_EDITOR
 #if WITH_EDITORONLY_DATA
-	PRAGMA_DISABLE_DEPRECATION_WARNINGS // Suppress compiler warning on override of deprecated function
-	UE_DEPRECATED(5.0, "Use version that takes FObjectPreSaveContext instead.")
-	ENGINE_API virtual void PreSave(const class ITargetPlatform* TargetPlatform) override;
-	ENGINE_API PRAGMA_ENABLE_DEPRECATION_WARNINGS
-	virtual void PreSave(FObjectPreSaveContext ObjectSaveContext) override;
+	ENGINE_API virtual void PreSave(FObjectPreSaveContext ObjectSaveContext) override;
 #endif // WITH_EDITORONLY_DATA
 	ENGINE_API virtual void PostLoad() override;
 	ENGINE_API virtual bool IsPostLoadThreadSafe() const override;
diff --git a/Engine/Source/Runtime/Engine/Classes/EdGraph/EdGraphNode.h b/Engine/Source/Runtime/Engine/Classes/EdGraph/EdGraphNode.h
--- a/Engine/Source/Runtime/Engine/Classes/EdGraph/EdGraphNode.h
+++ b/Engine/Source/Runtime/Engine/Classes/EdGraph/EdGraphNode.h
@@ -483,11 +483,7 @@ private:
 public:
 	// UObject interface
 	static ENGINE_API void AddReferencedObjects(UObject* InThis, FReferenceCollector& Collector);
-	PRAGMA_DISABLE_DEPRECATION_WARNINGS // Suppress compiler warning on override of deprecated function
-	UE_DEPRECATED(5.0, "Use version that takes FObjectPreSaveContext instead.")
-	ENGINE_API virtual void PreSave(const class ITargetPlatform* TargetPlatform) override;
-	ENGINE_API PRAGMA_ENABLE_DEPRECATION_WARNINGS
-	virtual void PreSave(FObjectPreSaveContext ObjectSaveContext) override;
+	ENGINE_API virtual void PreSave(FObjectPreSaveContext ObjectSaveContext) override;
 	ENGINE_API virtual void PostLoad() override;
 	ENGINE_API virtual void PostEditChangeProperty(FPropertyChangedEvent& PropertyChangedEvent) override;
 	ENGINE_API virtual void PostEditUndo() override;
diff --git a/Engine/Source/Runtime/Engine/Classes/Engine/Blueprint.h b/Engine/Source/Runtime/Engine/Classes/Engine/Blueprint.h
--- a/Engine/Source/Runtime/Engine/Classes/Engine/Blueprint.h
+++ b/Engine/Source/Runtime/Engine/Classes/Engine/Blueprint.h
@@ -996,11 +996,7 @@ public:
 
 	//~ Begin UObject Interface
 #if WITH_EDITORONLY_DATA
-	PRAGMA_DISABLE_DEPRECATION_WARNINGS // Suppress compiler warning on override of deprecated function
-	UE_DEPRECATED(5.0, "Use version that takes FObjectPreSaveContext instead.")
-	ENGINE_API virtual void PreSave(const class ITargetPlatform* TargetPlatform) override;
-	ENGINE_API PRAGMA_ENABLE_DEPRECATION_WARNINGS
-	virtual void PreSave(FObjectPreSaveContext ObjectSaveContext) override;
+	ENGINE_API virtual void PreSave(FObjectPreSaveContext ObjectSaveContext) override;
 	ENGINE_API virtual void GetPreloadDependencies(TArray<UObject*>& OutDeps) override;
 #endif // WITH_EDITORONLY_DATA
 	ENGINE_API virtual void Serialize(FArchive& Ar) override;
diff --git a/Engine/Source/Runtime/Engine/Classes/Engine/DataAsset.h b/Engine/Source/Runtime/Engine/Classes/Engine/DataAsset.h
--- a/Engine/Source/Runtime/Engine/Classes/Engine/DataAsset.h
+++ b/Engine/Source/Runtime/Engine/Classes/Engine/DataAsset.h
@@ -57,10 +57,6 @@ public:
 	ENGINE_API virtual void UpdateAssetBundleData();
 
 	/** Updates AssetBundleData */
-	PRAGMA_DISABLE_DEPRECATION_WARNINGS // Suppress compiler warning on override of deprecated function
-	UE_DEPRECATED(5.0, "Use version that takes FObjectPreSaveContext instead.")
-	ENGINE_API virtual void PreSave(const class ITargetPlatform* TargetPlatform) override;
-	PRAGMA_ENABLE_DEPRECATION_WARNINGS
 	ENGINE_API virtual void PreSave(FObjectPreSaveContext ObjectSaveContext) override;
 
 protected:
diff --git a/Engine/Source/Runtime/Engine/Classes/Engine/Level.h b/Engine/Source/Runtime/Engine/Classes/Engine/Level.h
--- a/Engine/Source/Runtime/Engine/Classes/Engine/Level.h
+++ b/Engine/Source/Runtime/Engine/Classes/Engine/Level.h
@@ -928,10 +928,6 @@ public:
 #endif // WITH_EDITOR
 	virtual bool ResolveSubobject(const TCHAR* SubObjectPath, UObject*& OutObject, bool bLoadIfExists) override;
 	virtual void PostLoad() override;
-	PRAGMA_DISABLE_DEPRECATION_WARNINGS // Suppress compiler warning on override of deprecated function
-	UE_DEPRECATED(5.0, "Use version that takes FObjectPreSaveContext instead.")
-	virtual void PreSave(const class ITargetPlatform* TargetPlatform) override;
-	PRAGMA_ENABLE_DEPRECATION_WARNINGS
 	virtual void PreSave(FObjectPreSaveContext ObjectSaveContext) override;
 	virtual void PreDuplicate(FObjectDuplicationParameters& DupParams) override;
 	virtual bool CanBeClusterRoot() const override;
diff --git a/Engine/Source/Runtime/Engine/Classes/Engine/SkeletalMesh.h b/Engine/Source/Runtime/Engine/Classes/Engine/SkeletalMesh.h
--- a/Engine/Source/Runtime/Engine/Classes/Engine/SkeletalMesh.h
+++ b/Engine/Source/Runtime/Engine/Classes/Engine/SkeletalMesh.h
@@ -2474,11 +2474,7 @@ public:
 #endif // WITH_EDITOR
 	ENGINE_API virtual void BeginDestroy() override;
 	ENGINE_API virtual bool IsReadyForFinishDestroy() override;
-	PRAGMA_DISABLE_DEPRECATION_WARNINGS // Suppress compiler warning on override of deprecated function
-	UE_DEPRECATED(5.0, "Use version that takes FObjectPreSaveContext instead.")
-	ENGINE_API virtual void PreSave(const class ITargetPlatform* TargetPlatform) override;
-	ENGINE_API PRAGMA_ENABLE_DEPRECATION_WARNINGS
-	virtual void PreSave(FObjectPreSaveContext ObjectSaveContext) override;
+	ENGINE_API virtual void PreSave(FObjectPreSaveContext ObjectSaveContext) override;
 	ENGINE_API virtual void Serialize(FArchive& Ar) override;
 #if WITH_EDITORONLY_DATA
 	static ENGINE_API void DeclareCustomVersions(FArchive& Ar, const UClass* SpecificSubclass);
diff --git a/Engine/Source/Runtime/Engine/Classes/Engine/Texture.h b/Engine/Source/Runtime/Engine/Classes/Engine/Texture.h
--- a/Engine/Source/Runtime/Engine/Classes/Engine/Texture.h
+++ b/Engine/Source/Runtime/Engine/Classes/Engine/Texture.h
@@ -1950,10 +1950,6 @@ public:
 	ENGINE_API virtual void Serialize(FArchive& Ar) override;
 	ENGINE_API virtual void PostInitProperties() override;
 	ENGINE_API virtual void PostLoad() override;
-	PRAGMA_DISABLE_DEPRECATION_WARNINGS // Suppress compiler warning on override of deprecated function
-	UE_DEPRECATED(5.0, "Use version that takes FObjectPreSaveContext instead.")
-	ENGINE_API virtual void PreSave(const class ITargetPlatform* TargetPlatform) override;
-	PRAGMA_ENABLE_DEPRECATION_WARNINGS
 	ENGINE_API virtual void PreSave(FObjectPreSaveContext ObjectSaveContext) override;
 	ENGINE_API virtual void BeginDestroy() override;
 	ENGINE_API virtual bool IsReadyForFinishDestroy() override;
diff --git a/Engine/Source/Runtime/Engine/Classes/Engine/Texture2D.h b/Engine/Source/Runtime/Engine/Classes/Engine/Texture2D.h
--- a/Engine/Source/Runtime/Engine/Classes/Engine/Texture2D.h
+++ b/Engine/Source/Runtime/Engine/Classes/Engine/Texture2D.h
@@ -112,10 +112,6 @@ public:
 	virtual void BeginDestroy() override;
 	virtual bool IsReadyForAsyncPostLoad() const override;
 	virtual void PostLoad() override;
-	PRAGMA_DISABLE_DEPRECATION_WARNINGS // Suppress compiler warning on override of deprecated function
-	UE_DEPRECATED(5.0, "Use version that takes FObjectPreSaveContext instead.")
-	virtual void PreSave(const class ITargetPlatform* TargetPlatform) override;
-	PRAGMA_ENABLE_DEPRECATION_WARNINGS
 	virtual void PreSave(FObjectPreSaveContext ObjectSaveContext) override;
 	virtual void GetAssetRegistryTags(FAssetRegistryTagsContext Context) const override;
 	UE_DEPRECATED(5.4, "Implement the version that takes FAssetRegistryTagsContext instead.")
diff --git a/Engine/Source/Runtime/Engine/Classes/Engine/World.h b/Engine/Source/Runtime/Engine/Classes/Engine/World.h
--- a/Engine/Source/Runtime/Engine/Classes/Engine/World.h
+++ b/Engine/Source/Runtime/Engine/Classes/Engine/World.h
@@ -2802,12 +2802,6 @@ public:
 	static void DeclareConstructClasses(TArray<FTopLevelAssetPath>& OutConstructClasses, const UClass* SpecificSubclass);
 #endif
 	virtual void PreDuplicate(FObjectDuplicationParameters& DupParams) override;
-	PRAGMA_DISABLE_DEPRECATION_WARNINGS // Suppress compiler warning on override of deprecated function
-	UE_DEPRECATED(5.0, "Use version that takes FObjectPreSaveRootContext instead.")
-	virtual bool PreSaveRoot(const TCHAR* InFilename) override;
-	UE_DEPRECATED(5.0, "Use version that takes FObjectPostSaveRootContext instead.")
-	virtual void PostSaveRoot(bool bCleanupIsRequired) override;
-	PRAGMA_ENABLE_DEPRECATION_WARNINGS
 	virtual void PreSaveRoot(FObjectPreSaveRootContext ObjectSaveContext) override;
 	virtual void PostSaveRoot(FObjectPostSaveRootContext ObjectSaveContext) override;
 	virtual UWorld* GetWorld() const override;
diff --git a/Engine/Source/Runtime/Engine/Classes/GameFramework/Actor.h b/Engine/Source/Runtime/Engine/Classes/GameFramework/Actor.h
--- a/Engine/Source/Runtime/Engine/Classes/GameFramework/Actor.h
+++ b/Engine/Source/Runtime/Engine/Classes/GameFramework/Actor.h
@@ -2307,15 +2307,7 @@ public:
 	ENGINE_API virtual bool IsRuntimeOnly() const { return false; }
 	ENGINE_API virtual bool IsAsset() const override;
 
-	PRAGMA_DISABLE_DEPRECATION_WARNINGS // Suppress compiler warning on override of deprecated function
-	UE_DEPRECATED(5.0, "Use version that takes FObjectPreSaveRootContext instead.")
-	ENGINE_API virtual bool PreSaveRoot(const TCHAR* InFilename) override;
-	UE_DEPRECATED(5.0, "Use version that takes FObjectPostSaveRootContext instead.")
-	ENGINE_API virtual void PostSaveRoot(bool bCleanupIsRequired) override;
-	UE_DEPRECATED(5.0, "Use version that takes FObjectPreSaveContext instead.")
-	ENGINE_API virtual void PreSave(const ITargetPlatform* TargetPlatform) override;
-	ENGINE_API PRAGMA_ENABLE_DEPRECATION_WARNINGS
-	virtual void PreSaveRoot(FObjectPreSaveRootContext ObjectSaveContext) override;
+	ENGINE_API virtual void PreSaveRoot(FObjectPreSaveRootContext ObjectSaveContext) override;
 	ENGINE_API virtual void PostSaveRoot(FObjectPostSaveRootContext ObjectSaveContext) override;
 	ENGINE_API virtual void PreSave(FObjectPreSaveContext ObjectSaveContext) override;
 
diff --git a/Engine/Source/Runtime/Engine/Classes/Particles/ParticleSystem.h b/Engine/Source/Runtime/Engine/Classes/Particles/ParticleSystem.h
--- a/Engine/Source/Runtime/Engine/Classes/Particles/ParticleSystem.h
+++ b/Engine/Source/Runtime/Engine/Classes/Particles/ParticleSystem.h
@@ -421,10 +421,6 @@ public:
 #if WITH_EDITOR
 	virtual void PostEditChangeProperty(FPropertyChangedEvent& PropertyChangedEvent) override;
 #endif // WITH_EDITOR
-	PRAGMA_DISABLE_DEPRECATION_WARNINGS // Suppress compiler warning on override of deprecated function
-	UE_DEPRECATED(5.0, "Use version that takes FObjectPreSaveContext instead.")
-	virtual void PreSave(const class ITargetPlatform* TargetPlatform) override;
-	PRAGMA_ENABLE_DEPRECATION_WARNINGS
 	virtual void PreSave(FObjectPreSaveContext ObjectSaveContext) override;
 	virtual void PostLoad() override;
 	virtual bool IsPostLoadThreadSafe() const override;
diff --git a/Engine/Source/Runtime/Engine/Private/Actor.cpp b/Engine/Source/Runtime/Engine/Private/Actor.cpp
--- a/Engine/Source/Runtime/Engine/Private/Actor.cpp
+++ b/Engine/Source/Runtime/Engine/Private/Actor.cpp
@@ -485,13 +485,6 @@ bool AActor::IsAsset() const
 	return IsPackageExternal() && !GetPackage()->HasAnyFlags(RF_Transient) && !HasAnyFlags(RF_Transient | RF_ClassDefaultObject) && IsMainPackageActor() && !GetPackage()->HasAnyPackageFlags(PKG_PlayInEditor);
 }
 
-bool AActor::PreSaveRoot(const TCHAR* InFilename)
-{
-	PRAGMA_DISABLE_DEPRECATION_WARNINGS;
-	return Super::PreSaveRoot(InFilename);
-	PRAGMA_ENABLE_DEPRECATION_WARNINGS;
-}
-
 void AActor::PreSaveRoot(FObjectPreSaveRootContext ObjectSaveContext)
 {
 	Super::PreSaveRoot(ObjectSaveContext);
@@ -512,13 +505,6 @@ void AActor::PreSaveRoot(FObjectPreSaveRootContext ObjectSaveContext)
 #endif
 }
 
-void AActor::PostSaveRoot(bool bCleanupIsRequired)
-{
-	PRAGMA_DISABLE_DEPRECATION_WARNINGS;
-	Super::PostSaveRoot(bCleanupIsRequired);
-	PRAGMA_ENABLE_DEPRECATION_WARNINGS;
-}
-
 void AActor::PostSaveRoot(FObjectPostSaveRootContext  ObjectSaveContext)
 {
 	Super::PostSaveRoot(ObjectSaveContext);
@@ -537,13 +523,6 @@ void AActor::PostSaveRoot(FObjectPostSaveRootContext  ObjectSaveContext)
 #endif
 }
 
-void AActor::PreSave(const ITargetPlatform* TargetPlatform)
-{
-	PRAGMA_DISABLE_DEPRECATION_WARNINGS;
-	Super::PreSave(TargetPlatform);
-	PRAGMA_ENABLE_DEPRECATION_WARNINGS;
-}
-
 void AActor::PreSave(FObjectPreSaveContext ObjectSaveContext)
 {
 	Super::PreSave(ObjectSaveContext);
diff --git a/Engine/Source/Runtime/Engine/Private/Animation/AnimMontage.cpp b/Engine/Source/Runtime/Engine/Private/Animation/AnimMontage.cpp
--- a/Engine/Source/Runtime/Engine/Private/Animation/AnimMontage.cpp
+++ b/Engine/Source/Runtime/Engine/Private/Animation/AnimMontage.cpp
@@ -373,13 +373,6 @@ void UAnimMontage::UnregisterOnMontageChanged(void* Unregister)
 }
 #endif	//WITH_EDITOR
 
-void UAnimMontage::PreSave(const class ITargetPlatform* TargetPlatform)
-{
-	PRAGMA_DISABLE_DEPRECATION_WARNINGS;
-	Super::PreSave(TargetPlatform);
-	PRAGMA_ENABLE_DEPRECATION_WARNINGS;
-}
-
 void UAnimMontage::PreSave(FObjectPreSaveContext ObjectSaveContext)
 {
 #if WITH_EDITOR
diff --git a/Engine/Source/Runtime/Engine/Private/Animation/AnimNotify.cpp b/Engine/Source/Runtime/Engine/Private/Animation/AnimNotify.cpp
--- a/Engine/Source/Runtime/Engine/Private/Animation/AnimNotify.cpp
+++ b/Engine/Source/Runtime/Engine/Private/Animation/AnimNotify.cpp
@@ -95,13 +95,6 @@ void UAnimNotify::PostLoad()
 #endif
 }
 
-void UAnimNotify::PreSave(const class ITargetPlatform* TargetPlatform)
-{
-	PRAGMA_DISABLE_DEPRECATION_WARNINGS;
-	Super::PreSave(TargetPlatform);
-	PRAGMA_ENABLE_DEPRECATION_WARNINGS;
-}
-
 void UAnimNotify::PreSave(FObjectPreSaveContext ObjectSaveContext)
 {
 #if WITH_EDITOR
diff --git a/Engine/Source/Runtime/Engine/Private/Animation/AnimSequence.cpp b/Engine/Source/Runtime/Engine/Private/Animation/AnimSequence.cpp
--- a/Engine/Source/Runtime/Engine/Private/Animation/AnimSequence.cpp
+++ b/Engine/Source/Runtime/Engine/Private/Animation/AnimSequence.cpp
@@ -960,13 +960,6 @@ private:
 };
 #endif
 
-void UAnimSequence::PreSave(const class ITargetPlatform* TargetPlatform)
-{
-	PRAGMA_DISABLE_DEPRECATION_WARNINGS;
-	Super::PreSave(TargetPlatform);
-	PRAGMA_ENABLE_DEPRECATION_WARNINGS;
-}
-
 void UAnimSequence::PreSave(FObjectPreSaveContext ObjectSaveContext)
 {
 #if WITH_EDITOR
diff --git a/Engine/Source/Runtime/Engine/Private/Animation/AnimStreamable.cpp b/Engine/Source/Runtime/Engine/Private/Animation/AnimStreamable.cpp
--- a/Engine/Source/Runtime/Engine/Private/Animation/AnimStreamable.cpp
+++ b/Engine/Source/Runtime/Engine/Private/Animation/AnimStreamable.cpp
@@ -145,13 +145,6 @@ UAnimStreamable::UAnimStreamable(const FObjectInitializer& ObjectInitializer)
 	bUseRawDataOnly = true;
 }
 
-void UAnimStreamable::PreSave(const class ITargetPlatform* TargetPlatform)
-{
-	PRAGMA_DISABLE_DEPRECATION_WARNINGS;
-	Super::PreSave(TargetPlatform);
-	PRAGMA_ENABLE_DEPRECATION_WARNINGS;
-}
-
 void UAnimStreamable::PreSave(FObjectPreSaveContext ObjectSaveContext)
 {
 #if WITH_EDITOR
diff --git a/Engine/Source/Runtime/Engine/Private/Animation/PoseAsset.cpp b/Engine/Source/Runtime/Engine/Private/Animation/PoseAsset.cpp
--- a/Engine/Source/Runtime/Engine/Private/Animation/PoseAsset.cpp
+++ b/Engine/Source/Runtime/Engine/Private/Animation/PoseAsset.cpp
@@ -1005,13 +1005,6 @@ void UPoseAsset::Serialize(FArchive& Ar)
 	Super::Serialize(Ar);
 }
 
-void UPoseAsset::PreSave(const ITargetPlatform* TargetPlatform)
-{
-	PRAGMA_DISABLE_DEPRECATION_WARNINGS;
-	Super::PreSave(TargetPlatform);
-	PRAGMA_ENABLE_DEPRECATION_WARNINGS;
-}
-
 void UPoseAsset::PreSave(FObjectPreSaveContext ObjectSaveContext)
 {
 #if WITH_EDITOR
diff --git a/Engine/Source/Runtime/Engine/Private/Blueprint.cpp b/Engine/Source/Runtime/Engine/Private/Blueprint.cpp
--- a/Engine/Source/Runtime/Engine/Private/Blueprint.cpp
+++ b/Engine/Source/Runtime/Engine/Private/Blueprint.cpp
@@ -367,13 +367,6 @@ static TAutoConsoleVariable<bool> CVarBPForceOldSearchDataFormatVersionOnSave(
 	TEXT("Force Blueprint search metadata to use an old format version on save (for QA/testing purposes only). On an editor relaunch, it should include the BP in the out-of-date count after the first search."),
 	ECVF_Cheat);
 
-void UBlueprint::PreSave(const class ITargetPlatform* TargetPlatform)
-{
-	PRAGMA_DISABLE_DEPRECATION_WARNINGS;
-	Super::PreSave(TargetPlatform);
-	PRAGMA_ENABLE_DEPRECATION_WARNINGS;
-}
-
 void UBlueprint::PreSave(FObjectPreSaveContext ObjectSaveContext)
 {
 	Super::PreSave(ObjectSaveContext);
diff --git a/Engine/Source/Runtime/Engine/Private/Components/LightComponent.cpp b/Engine/Source/Runtime/Engine/Private/Components/LightComponent.cpp
--- a/Engine/Source/Runtime/Engine/Private/Components/LightComponent.cpp
+++ b/Engine/Source/Runtime/Engine/Private/Components/LightComponent.cpp
@@ -613,13 +613,6 @@ bool ULightComponent::CanTraceDistanceFieldShadows() const
 }
 
 #if WITH_EDITOR
-void ULightComponent::PreSave(const class ITargetPlatform* TargetPlatform)
-{
-	PRAGMA_DISABLE_DEPRECATION_WARNINGS;
-	Super::PreSave(TargetPlatform);
-	PRAGMA_ENABLE_DEPRECATION_WARNINGS;
-}
-
 void ULightComponent::PreSave(FObjectPreSaveContext ObjectSaveContext)
 {
 	Super::PreSave(ObjectSaveContext);
diff --git a/Engine/Source/Runtime/Engine/Private/Components/PrimitiveComponent.cpp b/Engine/Source/Runtime/Engine/Private/Components/PrimitiveComponent.cpp
--- a/Engine/Source/Runtime/Engine/Private/Components/PrimitiveComponent.cpp
+++ b/Engine/Source/Runtime/Engine/Private/Components/PrimitiveComponent.cpp
@@ -1603,13 +1603,6 @@ void UPrimitiveComponent::PostEditImport()
 	ResetCustomPrimitiveData();
 }
 
-void UPrimitiveComponent::PreSave(const class ITargetPlatform* TargetPlatform)
-{
-	PRAGMA_DISABLE_DEPRECATION_WARNINGS;
-	Super::PreSave(TargetPlatform);
-	PRAGMA_ENABLE_DEPRECATION_WARNINGS;
-}
-
 void UPrimitiveComponent::PreSave(FObjectPreSaveContext ObjectSaveContext)
 {
 	Super::PreSave(ObjectSaveContext);
diff --git a/Engine/Source/Runtime/Engine/Private/Components/StaticMeshComponent.cpp b/Engine/Source/Runtime/Engine/Private/Components/StaticMeshComponent.cpp
--- a/Engine/Source/Runtime/Engine/Private/Components/StaticMeshComponent.cpp
+++ b/Engine/Source/Runtime/Engine/Private/Components/StaticMeshComponent.cpp
@@ -514,13 +514,6 @@ bool UStaticMeshComponent::AreNativePropertiesIdenticalTo( UObject* Other ) cons
 }
 
 #if WITH_EDITORONLY_DATA
-void UStaticMeshComponent::PreSave(const class ITargetPlatform* TargetPlatform)
-{
-	PRAGMA_DISABLE_DEPRECATION_WARNINGS;
-	Super::PreSave(TargetPlatform);
-	PRAGMA_ENABLE_DEPRECATION_WARNINGS;
-}
-
 void UStaticMeshComponent::PreSave(FObjectPreSaveContext ObjectSaveContext)
 {
 	Super::PreSave(ObjectSaveContext);
diff --git a/Engine/Source/Runtime/Engine/Private/DataAsset.cpp b/Engine/Source/Runtime/Engine/Private/DataAsset.cpp
--- a/Engine/Source/Runtime/Engine/Private/DataAsset.cpp
+++ b/Engine/Source/Runtime/Engine/Private/DataAsset.cpp
@@ -37,13 +37,6 @@ void UPrimaryDataAsset::UpdateAssetBundleData()
 	}
 }
 
-void UPrimaryDataAsset::PreSave(const class ITargetPlatform* TargetPlatform)
-{
-	PRAGMA_DISABLE_DEPRECATION_WARNINGS;
-	Super::PreSave(TargetPlatform);
-	PRAGMA_ENABLE_DEPRECATION_WARNINGS;
-}
-
 void UPrimaryDataAsset::PreSave(FObjectPreSaveContext ObjectSaveContext)
 {
 	Super::PreSave(ObjectSaveContext);
diff --git a/Engine/Source/Runtime/Engine/Private/EdGraph/EdGraphNode.cpp b/Engine/Source/Runtime/Engine/Private/EdGraph/EdGraphNode.cpp
--- a/Engine/Source/Runtime/Engine/Private/EdGraph/EdGraphNode.cpp
+++ b/Engine/Source/Runtime/Engine/Private/EdGraph/EdGraphNode.cpp
@@ -650,13 +650,6 @@ void UEdGraphNode::AddReferencedObjects(UObject* InThis, FReferenceCollector& Co
 	}
 }
 
-void UEdGraphNode::PreSave(const class ITargetPlatform* TargetPlatform)
-{
-	PRAGMA_DISABLE_DEPRECATION_WARNINGS;
-	Super::PreSave(TargetPlatform);
-	PRAGMA_ENABLE_DEPRECATION_WARNINGS;
-}
-
 void UEdGraphNode::PreSave(FObjectPreSaveContext ObjectSaveContext)
 {
 	Super::PreSave(ObjectSaveContext);
diff --git a/Engine/Source/Runtime/Engine/Private/Level.cpp b/Engine/Source/Runtime/Engine/Private/Level.cpp
--- a/Engine/Source/Runtime/Engine/Private/Level.cpp
+++ b/Engine/Source/Runtime/Engine/Private/Level.cpp
@@ -1215,13 +1215,6 @@ void ULevel::SortActorList()
 	Actors = ObjectPtrWrap(MoveTemp(NewActors));
 }
 
-void ULevel::PreSave(const class ITargetPlatform* TargetPlatform)
-{
-	PRAGMA_DISABLE_DEPRECATION_WARNINGS;
-	Super::PreSave(TargetPlatform);
-	PRAGMA_ENABLE_DEPRECATION_WARNINGS;
-}
-
 void ULevel::PreSave(FObjectPreSaveContext ObjectSaveContext)
 {
 	Super::PreSave(ObjectSaveContext);
diff --git a/Engine/Source/Runtime/Engine/Private/Materials/Material.cpp b/Engine/Source/Runtime/Engine/Private/Materials/Material.cpp
--- a/Engine/Source/Runtime/Engine/Private/Materials/Material.cpp
+++ b/Engine/Source/Runtime/Engine/Private/Materials/Material.cpp
@@ -1104,13 +1104,6 @@ UMaterial::UMaterial(const FObjectInitializer& ObjectInitializer)
 	FloatPrecisionMode = EMaterialFloatPrecisionMode::MFPM_Default;
 }
 
-void UMaterial::PreSave(const class ITargetPlatform* TargetPlatform)
-{
-	PRAGMA_DISABLE_DEPRECATION_WARNINGS;
-	Super::PreSave(TargetPlatform);
-	PRAGMA_ENABLE_DEPRECATION_WARNINGS;
-}
-
 void UMaterial::PreSave(FObjectPreSaveContext ObjectSaveContext)
 {
 	Super::PreSave(ObjectSaveContext);
diff --git a/Engine/Source/Runtime/Engine/Private/Materials/MaterialInstance.cpp b/Engine/Source/Runtime/Engine/Private/Materials/MaterialInstance.cpp
--- a/Engine/Source/Runtime/Engine/Private/Materials/MaterialInstance.cpp
+++ b/Engine/Source/Runtime/Engine/Private/Materials/MaterialInstance.cpp
@@ -5024,13 +5024,6 @@ void UMaterialInstance::GetLightingGuidChain(bool bIncludeTextures, TArray<FGuid
 #endif
 }
 
-void UMaterialInstance::PreSave(const class ITargetPlatform* TargetPlatform)
-{
-	PRAGMA_DISABLE_DEPRECATION_WARNINGS;
-	Super::PreSave(TargetPlatform);
-	PRAGMA_ENABLE_DEPRECATION_WARNINGS;
-}
-
 void UMaterialInstance::PreSave(FObjectPreSaveContext ObjectSaveContext)
 {
 	// @TODO : Remove any duplicate data from parent? Aims at improving change propagation (if controlled by parent)
diff --git a/Engine/Source/Runtime/Engine/Private/Materials/MaterialInterface.cpp b/Engine/Source/Runtime/Engine/Private/Materials/MaterialInterface.cpp
--- a/Engine/Source/Runtime/Engine/Private/Materials/MaterialInterface.cpp
+++ b/Engine/Source/Runtime/Engine/Private/Materials/MaterialInterface.cpp
@@ -1885,13 +1885,6 @@ bool UMaterialInterface::UseAnyStreamingTexture() const
 	return false;
 }
 
-void UMaterialInterface::PreSave(const class ITargetPlatform* TargetPlatform)
-{
-	PRAGMA_DISABLE_DEPRECATION_WARNINGS;
-	Super::PreSave(TargetPlatform);
-	PRAGMA_ENABLE_DEPRECATION_WARNINGS;
-}
-
 void UMaterialInterface::PreSave(FObjectPreSaveContext ObjectSaveContext)
 {
 	Super::PreSave(ObjectSaveContext);
diff --git a/Engine/Source/Runtime/Engine/Private/Particles/ParticleSystem.cpp b/Engine/Source/Runtime/Engine/Private/Particles/ParticleSystem.cpp
--- a/Engine/Source/Runtime/Engine/Private/Particles/ParticleSystem.cpp
+++ b/Engine/Source/Runtime/Engine/Private/Particles/ParticleSystem.cpp
@@ -301,13 +301,6 @@ void UParticleSystem::PostEditChangeProperty(FPropertyChangedEvent& PropertyChan
 }
 #endif // WITH_EDITOR
 
-void UParticleSystem::PreSave(const class ITargetPlatform* TargetPlatform)
-{
-	PRAGMA_DISABLE_DEPRECATION_WARNINGS;
-	Super::PreSave(TargetPlatform);
-	PRAGMA_ENABLE_DEPRECATION_WARNINGS;
-}
-
 void UParticleSystem::PreSave(FObjectPreSaveContext ObjectSaveContext)
 {
 	Super::PreSave(ObjectSaveContext);
diff --git a/Engine/Source/Runtime/Engine/Private/SkeletalMesh.cpp b/Engine/Source/Runtime/Engine/Private/SkeletalMesh.cpp
--- a/Engine/Source/Runtime/Engine/Private/SkeletalMesh.cpp
+++ b/Engine/Source/Runtime/Engine/Private/SkeletalMesh.cpp
@@ -2258,13 +2258,6 @@ void USkeletalMesh::FinishAsyncTaskInternal(FSkinnedAsyncTaskContext& Context)
 
 #endif // #if WITH_EDITOR
 
-void USkeletalMesh::PreSave(const class ITargetPlatform* TargetPlatform)
-{
-	PRAGMA_DISABLE_DEPRECATION_WARNINGS;
-	Super::PreSave(TargetPlatform);
-	PRAGMA_ENABLE_DEPRECATION_WARNINGS;
-}
-
 void USkeletalMesh::PreSave(FObjectPreSaveContext ObjectSaveContext)
 {
 	// check the parent index of the root bone is invalid
diff --git a/Engine/Source/Runtime/Engine/Private/Texture.cpp b/Engine/Source/Runtime/Engine/Private/Texture.cpp
--- a/Engine/Source/Runtime/Engine/Private/Texture.cpp
+++ b/Engine/Source/Runtime/Engine/Private/Texture.cpp
@@ -1380,13 +1380,6 @@ void UTexture::FinishDestroy()
 #endif
 }
 
-void UTexture::PreSave(const class ITargetPlatform* TargetPlatform)
-{
-	PRAGMA_DISABLE_DEPRECATION_WARNINGS;
-	Super::PreSave(TargetPlatform);
-	PRAGMA_ENABLE_DEPRECATION_WARNINGS;
-}
-
 #if WITH_EDITOR
 class FTextureDeterminismHelper : public UE::Cook::IDeterminismHelper
 {
diff --git a/Engine/Source/Runtime/Engine/Private/Texture2D.cpp b/Engine/Source/Runtime/Engine/Private/Texture2D.cpp
--- a/Engine/Source/Runtime/Engine/Private/Texture2D.cpp
+++ b/Engine/Source/Runtime/Engine/Private/Texture2D.cpp
@@ -669,13 +669,6 @@ void UTexture2D::PostLoad()
 	Super::PostLoad();
 }
 
-void UTexture2D::PreSave(const class ITargetPlatform* TargetPlatform)
-{
-	PRAGMA_DISABLE_DEPRECATION_WARNINGS;
-	Super::PreSave(TargetPlatform);
-	PRAGMA_ENABLE_DEPRECATION_WARNINGS;
-}
-
 void UTexture2D::PreSave(FObjectPreSaveContext ObjectSaveContext)
 {
 	Super::PreSave(ObjectSaveContext);
diff --git a/Engine/Source/Runtime/Engine/Private/World.cpp b/Engine/Source/Runtime/Engine/Private/World.cpp
--- a/Engine/Source/Runtime/Engine/Private/World.cpp
+++ b/Engine/Source/Runtime/Engine/Private/World.cpp
@@ -1560,13 +1560,6 @@ void UWorld::PreDuplicate(FObjectDuplicationParameters& DupParams)
 	}
 }
 
-bool UWorld::PreSaveRoot(const TCHAR* Filename)
-{
-	PRAGMA_DISABLE_DEPRECATION_WARNINGS;
-	return Super::PreSaveRoot(Filename);
-	PRAGMA_ENABLE_DEPRECATION_WARNINGS;
-}
-
 void UWorld::PreSaveRoot(FObjectPreSaveRootContext ObjectSaveContext)
 {
 	if (!ObjectSaveContext.IsFirstConcurrentSave())
@@ -1596,13 +1589,6 @@ void UWorld::PreSaveRoot(FObjectPreSaveRootContext ObjectSaveContext)
 	}
 }
 
-void UWorld::PostSaveRoot(bool bCleanupIsRequired)
-{
-	PRAGMA_DISABLE_DEPRECATION_WARNINGS;
-	Super::PostSaveRoot(bCleanupIsRequired);
-	PRAGMA_ENABLE_DEPRECATION_WARNINGS;
-}
-
 void UWorld::PostSaveRoot( FObjectPostSaveRootContext ObjectSaveContext )
 {
 	Super::PostSaveRoot(ObjectSaveContext);
diff --git a/Engine/Source/Runtime/Engine/Private/WorldPartition/WorldPartitionRuntimeSpatialHash.cpp b/Engine/Source/Runtime/Engine/Private/WorldPartition/WorldPartitionRuntimeSpatialHash.cpp
--- a/Engine/Source/Runtime/Engine/Private/WorldPartition/WorldPartitionRuntimeSpatialHash.cpp
+++ b/Engine/Source/Runtime/Engine/Private/WorldPartition/WorldPartitionRuntimeSpatialHash.cpp
@@ -1185,13 +1185,6 @@ UWorldPartitionRuntimeSpatialHash::UWorldPartitionRuntimeSpatialHash(const FObje
 #endif
 }
 
-void UWorldPartitionRuntimeSpatialHash::PreSave(const class ITargetPlatform* TargetPlatform)
-{
-	PRAGMA_DISABLE_DEPRECATION_WARNINGS;
-	Super::PreSave(TargetPlatform);
-	PRAGMA_ENABLE_DEPRECATION_WARNINGS;
-}
-
 void UWorldPartitionRuntimeSpatialHash::PreSave(FObjectPreSaveContext ObjectSaveContext)
 {
 	Super::PreSave(ObjectSaveContext);
diff --git a/Engine/Source/Runtime/Engine/Public/Materials/Material.h b/Engine/Source/Runtime/Engine/Public/Materials/Material.h
--- a/Engine/Source/Runtime/Engine/Public/Materials/Material.h
+++ b/Engine/Source/Runtime/Engine/Public/Materials/Material.h
@@ -1266,10 +1266,6 @@ public:
 	//~ End UMaterialInterface Interface.
 
 	//~ Begin UObject Interface
-	PRAGMA_DISABLE_DEPRECATION_WARNINGS // Suppress compiler warning on override of deprecated function
-	UE_DEPRECATED(5.0, "Use version that takes FObjectPreSaveContext instead.")
-	ENGINE_API virtual void PreSave(const class ITargetPlatform* TargetPlatform) override;
-	PRAGMA_ENABLE_DEPRECATION_WARNINGS
 	ENGINE_API virtual void PreSave(FObjectPreSaveContext ObjectSaveContext) override;
 	ENGINE_API virtual void PostInitProperties() override;
 	ENGINE_API virtual void Serialize(FArchive& Ar) override;
diff --git a/Engine/Source/Runtime/Engine/Public/Materials/MaterialInstance.h b/Engine/Source/Runtime/Engine/Public/Materials/MaterialInstance.h
--- a/Engine/Source/Runtime/Engine/Public/Materials/MaterialInstance.h
+++ b/Engine/Source/Runtime/Engine/Public/Materials/MaterialInstance.h
@@ -749,10 +749,6 @@ public:
 	 */
 	class FMaterialInstanceResource* Resource;
 
-	PRAGMA_DISABLE_DEPRECATION_WARNINGS // Suppress compiler warning on override of deprecated function
-	UE_DEPRECATED(5.0, "Use version that takes FObjectPreSaveContext instead.")
-	ENGINE_API virtual void PreSave(const class ITargetPlatform* TargetPlatform) override;
-	PRAGMA_ENABLE_DEPRECATION_WARNINGS
 	ENGINE_API virtual void PreSave(FObjectPreSaveContext ObjectSaveContext) override;
 	ENGINE_API virtual float GetTextureDensity(FName TextureName, const struct FMeshUVChannelInfo& UVChannelData) const override;
 
diff --git a/Engine/Source/Runtime/Engine/Public/Materials/MaterialInterface.h b/Engine/Source/Runtime/Engine/Public/Materials/MaterialInterface.h
--- a/Engine/Source/Runtime/Engine/Public/Materials/MaterialInterface.h
+++ b/Engine/Source/Runtime/Engine/Public/Materials/MaterialInterface.h
@@ -1175,10 +1175,6 @@ public:
 	*/
 	ENGINE_API virtual float GetTextureDensity(FName TextureName, const struct FMeshUVChannelInfo& UVChannelData) const;
 
-	PRAGMA_DISABLE_DEPRECATION_WARNINGS // Suppress compiler warning on override of deprecated function
-	UE_DEPRECATED(5.0, "Use version that takes FObjectPreSaveContext instead.")
-	ENGINE_API virtual void PreSave(const class ITargetPlatform* TargetPlatform) override;
-	PRAGMA_ENABLE_DEPRECATION_WARNINGS
 	ENGINE_API virtual void PreSave(FObjectPreSaveContext ObjectSaveContext) override;
 
 	/**
diff --git a/Engine/Source/Runtime/Engine/Public/WorldPartition/WorldPartitionRuntimeSpatialHash.h b/Engine/Source/Runtime/Engine/Public/WorldPartition/WorldPartitionRuntimeSpatialHash.h
--- a/Engine/Source/Runtime/Engine/Public/WorldPartition/WorldPartitionRuntimeSpatialHash.h
+++ b/Engine/Source/Runtime/Engine/Public/WorldPartition/WorldPartitionRuntimeSpatialHash.h
@@ -308,11 +308,7 @@ class UWorldPartitionRuntimeSpatialHash : public UWorldPartitionRuntimeHash
 	GENERATED_UCLASS_BODY()
 
 public:
-	PRAGMA_DISABLE_DEPRECATION_WARNINGS // Suppress compiler warning on override of deprecated function
-	UE_DEPRECATED(5.0, "Use version that takes FObjectPreSaveContext instead.")
-	ENGINE_API virtual void PreSave(const class ITargetPlatform* TargetPlatform) override;
-	ENGINE_API PRAGMA_ENABLE_DEPRECATION_WARNINGS
-	virtual void PreSave(FObjectPreSaveContext ObjectSaveContext) override;
+	ENGINE_API virtual void PreSave(FObjectPreSaveContext ObjectSaveContext) override;
 
 	ENGINE_API virtual void Serialize(FArchive& Ar) override;
 	ENGINE_API virtual void PostLoad() override;
diff --git a/Engine/Source/Runtime/MediaAssets/Private/Assets/BaseMediaSource.cpp b/Engine/Source/Runtime/MediaAssets/Private/Assets/BaseMediaSource.cpp
--- a/Engine/Source/Runtime/MediaAssets/Private/Assets/BaseMediaSource.cpp
+++ b/Engine/Source/Runtime/MediaAssets/Private/Assets/BaseMediaSource.cpp
@@ -41,13 +41,6 @@ void UBaseMediaSource::GetAssetRegistryTags(FAssetRegistryTagsContext Context) c
 	Super::GetAssetRegistryTags(Context);
 }
 
-void UBaseMediaSource::PreSave(const class ITargetPlatform* TargetPlatform)
-{
-	PRAGMA_DISABLE_DEPRECATION_WARNINGS;
-	Super::PreSave(TargetPlatform);
-	PRAGMA_ENABLE_DEPRECATION_WARNINGS;
-}
-
 void UBaseMediaSource::PreSave(FObjectPreSaveContext ObjectSaveContext)
 {
 #if WITH_EDITORONLY_DATA
diff --git a/Engine/Source/Runtime/MediaAssets/Private/Assets/PlatformMediaSource.cpp b/Engine/Source/Runtime/MediaAssets/Private/Assets/PlatformMediaSource.cpp
--- a/Engine/Source/Runtime/MediaAssets/Private/Assets/PlatformMediaSource.cpp
+++ b/Engine/Source/Runtime/MediaAssets/Private/Assets/PlatformMediaSource.cpp
@@ -16,13 +16,6 @@
 #endif
 
 
-void UPlatformMediaSource::PreSave(const class ITargetPlatform* TargetPlatform)
-{
-	PRAGMA_DISABLE_DEPRECATION_WARNINGS;
-	Super::PreSave(TargetPlatform);
-	PRAGMA_ENABLE_DEPRECATION_WARNINGS;
-}
-
 void UPlatformMediaSource::PreSave(FObjectPreSaveContext ObjectSaveContext)
 {
 #if WITH_EDITORONLY_DATA
diff --git a/Engine/Source/Runtime/MediaAssets/Public/BaseMediaSource.h b/Engine/Source/Runtime/MediaAssets/Public/BaseMediaSource.h
--- a/Engine/Source/Runtime/MediaAssets/Public/BaseMediaSource.h
+++ b/Engine/Source/Runtime/MediaAssets/Public/BaseMediaSource.h
@@ -47,11 +47,7 @@ private:
 public:
 
 	//~ UObject interface
-	PRAGMA_DISABLE_DEPRECATION_WARNINGS // Suppress compiler warning on override of deprecated function
-	UE_DEPRECATED(5.0, "Use version that takes FObjectPreSaveContext instead.")
-	MEDIAASSETS_API virtual void PreSave(const class ITargetPlatform* TargetPlatform);
-	MEDIAASSETS_API PRAGMA_ENABLE_DEPRECATION_WARNINGS
-	virtual void PreSave(FObjectPreSaveContext ObjectSaveContext);
+	MEDIAASSETS_API virtual void PreSave(FObjectPreSaveContext ObjectSaveContext);
 	MEDIAASSETS_API virtual void GetAssetRegistryTags(FAssetRegistryTagsContext Context) const override;
 	UE_DEPRECATED(5.4, "Implement the version that takes FAssetRegistryTagsContext instead.")
 	MEDIAASSETS_API virtual void GetAssetRegistryTags(TArray<FAssetRegistryTag>& OutTags) const override;
diff --git a/Engine/Source/Runtime/MediaAssets/Public/PlatformMediaSource.h b/Engine/Source/Runtime/MediaAssets/Public/PlatformMediaSource.h
--- a/Engine/Source/Runtime/MediaAssets/Public/PlatformMediaSource.h
+++ b/Engine/Source/Runtime/MediaAssets/Public/PlatformMediaSource.h
@@ -47,11 +47,7 @@ private:
 
 public:
 	//~ UObject interface
-	PRAGMA_DISABLE_DEPRECATION_WARNINGS // Suppress compiler warning on override of deprecated function
-	UE_DEPRECATED(5.0, "Use version that takes FObjectPreSaveContext instead.")
-	MEDIAASSETS_API virtual void PreSave(const class ITargetPlatform* TargetPlatform);
-	MEDIAASSETS_API PRAGMA_ENABLE_DEPRECATION_WARNINGS
-	virtual void PreSave(FObjectPreSaveContext ObjectSaveContext);
+	MEDIAASSETS_API virtual void PreSave(FObjectPreSaveContext ObjectSaveContext);
 	MEDIAASSETS_API virtual void Serialize(FArchive& Ar) override;
 
 	//~ UMediaSource interface
