Weak Object Pointers will no longer attempt to resolve remote objects in the middle of Garbage Collection

diff --git a/Engine/Source/Runtime/CoreUObject/Private/UObject/RemoteObject.cpp b/Engine/Source/Runtime/CoreUObject/Private/UObject/RemoteObject.cpp
--- a/Engine/Source/Runtime/CoreUObject/Private/UObject/RemoteObject.cpp
+++ b/Engine/Source/Runtime/CoreUObject/Private/UObject/RemoteObject.cpp
@@ -234,13 +234,22 @@ bool IsRemote(const UObject* Object)
 	return bIsRemote;
 }
 
-UObject* ResolveObject(const FRemoteObjectStub* Stub)
+UObject* ResolveObject(const FRemoteObjectStub* Stub, ERemoteReferenceType RefType /*= ERemoteReferenceType::Strong*/)
 {
 	using namespace UE::RemoteObject::Transfer;
 	using namespace UE::RemoteObject::Private;
 
 	// This is a slightly faster version of IsRemote(FRemoteObjectId) because we already know a Stub exists and we are going to re-use the Object pointer
 	UObject* Object = ObjectMaps->FindObject(Stub->Id);
+	if (GIsGarbageCollecting && (Object || RefType == ERemoteReferenceType::Weak)) // Not a thread-safe test
+	{
+		// Begin/FinishDestroy overrides may attempt to access subobjects of objects that have been migrated in which case we
+		// don't want to accidentally migrate them back mid-purge and if the Object memory is still valid (but has EInternalObjectFlags::Remote flag)
+		// we can just return it and let the owner finish its cleanup.
+		// In case of weak object ptrs it's relatively safe to just return null if the object doesn't exist on this server (see CanResolveObject)
+		return Object;
+	}
+
 	const bool bRemoteObject = !Object || IsRemote(Object);
 	if (bRemoteObject)
 	{
@@ -252,7 +261,7 @@ UObject* ResolveObject(const FRemoteObjectStub* Stub)
 	return Object;
 }
 
-UObject* ResolveObject(UObject* Object)
+UObject* ResolveObject(UObject* Object, ERemoteReferenceType RefType /*= ERemoteReferenceType::Strong*/)
 {
 	using namespace UE::RemoteObject::Private;
 
@@ -263,11 +272,34 @@ UObject* ResolveObject(UObject* Object)
 	{
 		FRemoteObjectStub* Stub = ObjectMaps->FindRemoteObjectStub(UE::CoreUObject::Private::FObjectHandleUtils::GetRemoteId(Object));
 		checkf(Stub, TEXT("Failed to find remote object stub for %s"), *GetPathNameSafe(Object));
-		return ResolveObject(Stub);
+		return ResolveObject(Stub, RefType);
 	}
 	return Object;
 }
 
+bool CanResolveObject(FRemoteObjectId ObjectId)
+{
+	using namespace UE::RemoteObject::Private;
+
+	// Note: this function needs to reflect the logic of ResolveObject(...) functions
+	
+	if (UObject* Object = ObjectMaps->FindObject(ObjectId))
+	{
+		// Object memory is local and even if it's already been migrated we can resolve it
+		return true;
+	}
+
+	if (FRemoteObjectStub* Stub = ObjectMaps->FindRemoteObjectStub(ObjectId))
+	{
+		// A stub exists so the object memory is not local but we can (attempt to) migrate it back if we're not garbage collecting
+		// Note: GIsGarbageCollecting checks are not thread safe
+		return !GIsGarbageCollecting;
+	}
+
+	// ObjectId is local or represents an object that has never been migrated
+	return false;
+}
+
 } // namespace UE::RemoteObject::Handle
 
 namespace UE::CoreUObject::Private
diff --git a/Engine/Source/Runtime/CoreUObject/Private/UObject/WeakObjectPtr.cpp b/Engine/Source/Runtime/CoreUObject/Private/UObject/WeakObjectPtr.cpp
--- a/Engine/Source/Runtime/CoreUObject/Private/UObject/WeakObjectPtr.cpp
+++ b/Engine/Source/Runtime/CoreUObject/Private/UObject/WeakObjectPtr.cpp
@@ -202,11 +202,7 @@ FUObjectItem* FWeakObjectPtr::TryResolveRemoteObject() const
 
 bool FWeakObjectPtr::CanBeResolved() const
 {
-	// Object can be resolved if there's an existing FRemoteObjectSub with its id (stubs are created when an object is migrated to remote server)
-	// Note that this doesn't mean the object actually exists on any of the remote server. It could've been already GC'd and leaked to disk.
-	// Future plans: If an object was leaked to disk/database we could broadcast that to all servers and have them remove FRemoteObjectStub that points to a leaked object
-	// or we could check if the leaked object exists in the database and return false if it does.
-	return !!UE::RemoteObject::Private::FindRemoteObjectStub(ObjectRemoteId);
+	return UE::RemoteObject::Handle::CanResolveObject(ObjectRemoteId);
 }
 
 #endif // UE_WITH_REMOTE_OBJECT_HANDLE
\ No newline at end of file
diff --git a/Engine/Source/Runtime/CoreUObject/Public/UObject/RemoteObject.h b/Engine/Source/Runtime/CoreUObject/Public/UObject/RemoteObject.h
--- a/Engine/Source/Runtime/CoreUObject/Public/UObject/RemoteObject.h
+++ b/Engine/Source/Runtime/CoreUObject/Public/UObject/RemoteObject.h
@@ -36,19 +36,35 @@ namespace UE::RemoteObject::Handle
 		COREUOBJECT_API explicit FRemoteObjectStub(UObject* Object);
 	};
 
+	enum class ERemoteReferenceType
+	{
+		Strong = 0,
+		Weak = 1
+	};
+
 	/**
 	* Resolves a remote object given its stub
 	* @param Stub Basic data required to migrate a remote object
+	* @param RefType Reference type that wants to resolve an object
 	* @return Resolved object
 	*/
-	COREUOBJECT_API UObject* ResolveObject(const FRemoteObjectStub* Stub);
+	COREUOBJECT_API UObject* ResolveObject(const FRemoteObjectStub* Stub, ERemoteReferenceType RefType = ERemoteReferenceType::Strong);
 
 	/**
 	* Resolves a remote object
 	* @param Object Object to resolve (remote object memory that has not yet been GC'd)
+	* @param RefType Reference type that wants to resolve an object
 	* @return Resolved object
 	*/
-	COREUOBJECT_API UObject* ResolveObject(UObject* Object);
+	COREUOBJECT_API UObject* ResolveObject(UObject* Object, ERemoteReferenceType RefType = ERemoteReferenceType::Strong);
+
+	/**
+	* Resolves a remote object
+	* @param Object Object id to resolve
+	* @param RefType Reference type that wants to resolve an object
+	* @return True if a remote object can be resolved
+	*/
+	COREUOBJECT_API bool CanResolveObject(FRemoteObjectId ObjectId);
 
 	/**
 	* Checks if an object associated with the specified unique id is remote
