Fix for ActivateTrackingPropertyValueInitialized not getting called on Verse classes.
Excluding property tracking on one additional Verse test until it works properly with transactions.

#rnx

diff --git a/Engine/Source/Runtime/CoreUObject/Internal/UObject/PropertyStateTracking.h b/Engine/Source/Runtime/CoreUObject/Internal/UObject/PropertyStateTracking.h
--- a/Engine/Source/Runtime/CoreUObject/Internal/UObject/PropertyStateTracking.h
+++ b/Engine/Source/Runtime/CoreUObject/Internal/UObject/PropertyStateTracking.h
@@ -20,6 +20,10 @@ namespace UE { class FPropertyTypeName; }
 
 #define UE_API COREUOBJECT_API
 
+#ifndef PROPERTY_INIT_STATE_TRACKING_ENABLE_DEFAULT
+	#define PROPERTY_INIT_STATE_TRACKING_ENABLE_DEFAULT 1
+#endif
+
 namespace UE
 {
 
diff --git a/Engine/Source/Runtime/CoreUObject/Private/UObject/InstanceDataObjectUtils.cpp b/Engine/Source/Runtime/CoreUObject/Private/UObject/InstanceDataObjectUtils.cpp
--- a/Engine/Source/Runtime/CoreUObject/Private/UObject/InstanceDataObjectUtils.cpp
+++ b/Engine/Source/Runtime/CoreUObject/Private/UObject/InstanceDataObjectUtils.cpp
@@ -253,10 +253,6 @@ namespace UE
 		ExcludedLoosePropertyTypesVar,
 		TEXT("Comma separated list of property types that will be excluded from loose properties in IDOs.")
 	);
-
-#ifndef PROPERTY_INIT_STATE_TRACKING_ENABLE_DEFAULT
-#define PROPERTY_INIT_STATE_TRACKING_ENABLE_DEFAULT 1
-#endif // PROPERTY_INIT_STATE_TRACKING_ENABLE_DEFAULT
 	
 	bool bUseAnnotationsForPropertyStateInitTracking = PROPERTY_INIT_STATE_TRACKING_ENABLE_DEFAULT;
 	FAutoConsoleVariableRef CVarUseAnnotationBasedPropertyStateInitTracking(
diff --git a/Engine/Source/Runtime/CoreUObject/Private/UObject/UObjectGlobals.cpp b/Engine/Source/Runtime/CoreUObject/Private/UObject/UObjectGlobals.cpp
--- a/Engine/Source/Runtime/CoreUObject/Private/UObject/UObjectGlobals.cpp
+++ b/Engine/Source/Runtime/CoreUObject/Private/UObject/UObjectGlobals.cpp
@@ -4046,10 +4046,12 @@ void FObjectInitializer::PostConstructInit()
 	}
 #endif // USE_CIRCULAR_DEPENDENCY_LOAD_DEFERRING
 
-#if WITH_EDITORONLY_DATA
+	// We disable the activation of property initialization tracking when PROPERTY_INIT_STATE_TRACKING_ENABLE_DEFAULT is 0.
+	// This is a workaround for a number of Verse tests, as property tracking does not work correctly with AutoRTFM yet.
+#if WITH_EDITORONLY_DATA && PROPERTY_INIT_STATE_TRACKING_ENABLE_DEFAULT
 	// Activate tracking of initialized property state for Verse objects.
 	// The Class of the Class may not be available during very early initialization of CoreUObject.
-	if (UClass* ClassClass = Class->GetClass(); ClassClass && ClassClass->GetClass()->IsChildOf<UVerseClass>())
+	if (UClass* ClassClass = Class->GetClass(); ClassClass && ClassClass->IsChildOf<UVerseClass>())
 	{
 		UE::ActivateTrackingPropertyValueInitialized(Class, Obj);
 	}
