- Changed "reloadcfg" to use SafeUnloadBranch so that it will reload the dynamic layers as well as the static ones

diff --git a/Engine/Source/Runtime/CoreUObject/Private/UObject/Obj.cpp b/Engine/Source/Runtime/CoreUObject/Private/UObject/Obj.cpp
--- a/Engine/Source/Runtime/CoreUObject/Private/UObject/Obj.cpp
+++ b/Engine/Source/Runtime/CoreUObject/Private/UObject/Obj.cpp
@@ -5409,9 +5409,8 @@ bool StaticExec( UWorld* InWorld, const TCHAR* Cmd, FOutputDevice& Ar )
 			{
 				if (ObjectToReload->GetClass()->HasAnyClassFlags(EClassFlags::CLASS_Config))
 				{
-					// reload the object's config
-					FConfigContext Context = FConfigContext::ForceReloadIntoGConfig();
-					Context.Load(*ObjectToReload->GetClass()->GetConfigName());
+					// unload the branch so next access will load the static and dynamic layers
+					GConfig->SafeUnloadBranch(*ObjectToReload->GetClass()->GetConfigName());
 					
 					// now updates all the class properties now that the config was reloaded from disk
 					ObjectToReload->ReloadConfig();
