Fall back to tagged property serialization when an object's class changes on load. This helps us preserve data that was saved in an asset when native code decides to create a uobject of a given type

[RN] When loading assets in the editor we now fall back to tagged serialization if an object has changed to an unrelated type, rather than losing its data completely

diff --git a/Engine/Source/Runtime/CoreUObject/Private/UObject/LinkerLoad.cpp b/Engine/Source/Runtime/CoreUObject/Private/UObject/LinkerLoad.cpp
--- a/Engine/Source/Runtime/CoreUObject/Private/UObject/LinkerLoad.cpp
+++ b/Engine/Source/Runtime/CoreUObject/Private/UObject/LinkerLoad.cpp
@@ -4601,13 +4601,17 @@ void FLinkerLoad::Preload( UObject* Object )
 #if WITH_EDITOR
 				// for placeholder objects that have no explicit type, we only want to serialize the TPS stream
 				bool bSerializeOnlyScriptProperties = false;
+				const bool bDoesSavedClassMatchActualClass = DoesSavedClassMatchActualClass(ExportIndex);
 				FGuardValue_Bitfield(bIsLoadingToPropertyBagObject, UE::FPropertyBagRepository::IsPropertyBagPlaceholderObject(Object));
-				if (bIsLoadingToPropertyBagObject && UEVer() >= EUnrealEngineObjectUE5Version::SCRIPT_SERIALIZATION_OFFSET)
+				if(UEVer() >= EUnrealEngineObjectUE5Version::SCRIPT_SERIALIZATION_OFFSET)
 				{
-					// note: script start/end offsets are relative to the export's offset in the file
-					StartPos += Export.ScriptSerializationStartOffset;
-					ExpectedSerialSize = Export.ScriptSerializationEndOffset;
-					bSerializeOnlyScriptProperties = true;	// signals that we can safely narrow the load to SerializeScriptProperties()
+					if (bIsLoadingToPropertyBagObject || !bDoesSavedClassMatchActualClass)
+					{
+						// note: script start/end offsets are relative to the export's offset in the file
+						StartPos += Export.ScriptSerializationStartOffset;
+						ExpectedSerialSize = Export.ScriptSerializationEndOffset;
+						bSerializeOnlyScriptProperties = true;	// signals that we can safely narrow the load to SerializeScriptProperties()
+					}
 				}
 #endif
 				// move to the position in the file where this object's data
@@ -5387,9 +5391,25 @@ UObject* FLinkerLoad::CreateExport( int32 Index )
 
 		if (ActualObjectWithTheName && !ActualObjectWithTheName->GetClass()->IsChildOf(LoadClass))
 		{
-			UE_ASSET_LOG(LogLinker, Error, PackagePath, TEXT("Failed import: class '%s' name '%s' outer '%s'. There is another object (of '%s' class) at the path."),
-				*LoadClass->GetName(), *Export.ObjectName.ToString(), *ThisParent->GetName(), *ActualObjectWithTheName->GetClass()->GetName());
-			return NULL;
+			if(UEVer() >= EUnrealEngineObjectUE5Version::SCRIPT_SERIALIZATION_OFFSET)
+			{
+				UE_ASSET_LOG(LogLinker, Log, PackagePath, TEXT("Object changed type on load: %s had class %s but is now %s"),
+					*ActualObjectWithTheName->GetPathName(), *LoadClass->GetPathName(), *ActualObjectWithTheName->GetClass()->GetPathName());
+				// native code created a new object with the same name but a different class, lets load on top of that object:
+				Export.Object = ActualObjectWithTheName;
+				Export.Object->SetLinker(this, Index);
+				Export.Object->SetFlags(RF_NeedLoad|RF_NeedPostLoad|RF_NeedPostLoadSubobjects|RF_WasLoaded);
+				check(CurrentLoadContext);
+				CurrentLoadContext->AddLoadedObject(Export.Object);
+				return ActualObjectWithTheName;
+			}
+			else
+			{
+				UE_ASSET_LOG(LogLinker, Error, PackagePath, TEXT("Failed import: class '%s' name '%s' outer '%s'. There is another object (of '%s' class) at the path."),
+					*LoadClass->GetName(), *Export.ObjectName.ToString(), *ThisParent->GetName(), *ActualObjectWithTheName->GetClass()->GetName());
+				return NULL;
+	
+			}
 		}
 
 		// Create the export object, marking it with the appropriate flags to
@@ -6776,6 +6796,38 @@ void FLinkerLoad::ReplaceExportIndexes(const FPackageIndex& OldIndex, const FPac
 	}
 }
 
+bool FLinkerLoad::DoesSavedClassMatchActualClass(int32 ExportIndex) const
+{
+	const FObjectExport& Export = ExportMap[ExportIndex];
+	check(Export.Object);
+	FPackageIndex ClassIndex = Export.ClassIndex;
+	const UClass* LoadClass = Cast<UClass>(GetCurrentObjectAtIndex(ClassIndex));
+
+	if(!LoadClass)
+	{
+		return false;
+	}
+	
+	return Export.Object->GetClass()->IsChildOf(LoadClass);
+}
+
+const UObject* FLinkerLoad::GetCurrentObjectAtIndex(FPackageIndex ObjectIndex) const
+{
+	if (ObjectIndex.IsNull())
+	{
+		return nullptr;
+	}
+
+	if (ObjectIndex.IsImport())
+	{
+		return ImportMap[ObjectIndex.ToImport()].XObject;
+	}
+	else
+	{
+		return ExportMap[ObjectIndex.ToExport()].Object;
+	}
+}
+
 void FLinkerLoad::FixupDuplicateExports()
 {
 	// We need to operate on copy to avoid incorrect indexes after sorting
diff --git a/Engine/Source/Runtime/CoreUObject/Public/UObject/LinkerLoad.h b/Engine/Source/Runtime/CoreUObject/Public/UObject/LinkerLoad.h
--- a/Engine/Source/Runtime/CoreUObject/Public/UObject/LinkerLoad.h
+++ b/Engine/Source/Runtime/CoreUObject/Public/UObject/LinkerLoad.h
@@ -835,6 +835,18 @@ public:
 
 private:
 #if WITH_EDITOR
+	/**
+	 * Returns a bool indicating that the class the object (specified by ExportIndex) was saved with matches the class
+	 * that it was created with. Most commonly this returns false when native user (non-coreuobject) code has required
+	 * a specific class of a specific type
+	 */
+	 bool DoesSavedClassMatchActualClass(int32 ExportIndex) const;
+	
+	/** 
+	 * Helper that returns the current object (import or export) specified by package index. Does not attempt to resolve the object.
+	 * See related function: IndexToObject. 
+	 */
+	const UObject* GetCurrentObjectAtIndex(FPackageIndex ObjectIndex) const;
 
 	/**
 	 * Nulls duplicated exports and fixes indexes in ExportMap to point to original objects instead of duplicates.
