[BugFix] Loose object properties weren't loading correctly because they needed CPF_TObjectPtr and their PropertyClass member set

diff --git a/Engine/Source/Runtime/CoreUObject/Private/UObject/InstanceDataObjectUtils.cpp b/Engine/Source/Runtime/CoreUObject/Private/UObject/InstanceDataObjectUtils.cpp
--- a/Engine/Source/Runtime/CoreUObject/Private/UObject/InstanceDataObjectUtils.cpp
+++ b/Engine/Source/Runtime/CoreUObject/Private/UObject/InstanceDataObjectUtils.cpp
@@ -615,7 +615,7 @@ namespace UE
 				MarkPropertyAsLoose(InnerProperty);
 			}
 		}
-		else if (const FObjectProperty* AsObjectProperty = CastField<FObjectProperty>(Property))
+		else if (FObjectProperty* AsObjectProperty = CastField<FObjectProperty>(Property))
 		{
 			// Hack for now - the assumption is that IDOs are generated only for class types that impose this flag on all object properties.
 			// There is currently an implicit assumption in the serialization logic that all inner properties have this flag set for containers.
@@ -628,7 +628,10 @@ namespace UE
 			// effects on that data as part of the construction/serialization path.
 			// 
 			// @todo - Remove if/when this flag is no longer required to signal whether this value is to be resolved via a subobject instancing graph.
-			Property->SetPropertyFlags(CPF_PersistentInstance| CPF_InstancedReference);
+			AsObjectProperty->SetPropertyFlags(CPF_PersistentInstance| CPF_InstancedReference | CPF_TObjectPtr);
+
+			// also assign the property class to UObject because loose properties can't infer their PropertyClass from TPS data so we'll assume it's as lenient as possible
+			AsObjectProperty->PropertyClass = UObject::StaticClass();
 		}
 	}
 
