[BugFix] Impersonation wasn't being enabled for prefabs which was causing loose properties to fail to save

diff --git a/Engine/Source/Runtime/CoreUObject/Private/UObject/InstanceDataObjectUtils.cpp b/Engine/Source/Runtime/CoreUObject/Private/UObject/InstanceDataObjectUtils.cpp
--- a/Engine/Source/Runtime/CoreUObject/Private/UObject/InstanceDataObjectUtils.cpp
+++ b/Engine/Source/Runtime/CoreUObject/Private/UObject/InstanceDataObjectUtils.cpp
@@ -298,6 +298,13 @@ namespace UE
 			return true;
 		}
 
+		// Assume that if this object has an IDO that it's enabled. This assumption is important for objects
+		// that were reparented into the transient package but still need their loose properties CPFUOed to new instances
+		if (UE::FPropertyBagRepository::Get().HasInstanceDataObject(InObject))
+		{
+			return true;
+		}
+
 		//@todo FH: change to check trait when available or use config object
 		const UClass* ObjClass = InObject->GetClass();
 		if (!ObjClass->CanCreateInstanceDataObject())
