Fixed FEnumProperty::ImportText to skip leading whitespace.
Fixed some questionable null handling.

diff --git a/Engine/Source/Runtime/CoreUObject/Private/UObject/EnumProperty.cpp b/Engine/Source/Runtime/CoreUObject/Private/UObject/EnumProperty.cpp
--- a/Engine/Source/Runtime/CoreUObject/Private/UObject/EnumProperty.cpp
+++ b/Engine/Source/Runtime/CoreUObject/Private/UObject/EnumProperty.cpp
@@ -49,43 +49,50 @@ namespace UE::CoreUObject::Private
 
 	const TCHAR* ImportEnumFromBuffer(UEnum* Enum, const FProperty* PropertyToSet, const FNumericProperty* UnderlyingProp, const TCHAR* PropertyClassName, const TCHAR* Buffer, void* ContainerOrPropertyPtr, EPropertyPointerType PropertyPointerType, FOutputDevice* ErrorText)
 	{
+		if (!Buffer)
+		{
+			return nullptr;
+		}
+
 		bool    bIsEnumOfFlags = Enum->HasAnyEnumFlags(EEnumFlags::Flags);
 		int64   EnumValue      = 0;
 		FString Temp;
 		for (;;)
 		{
-			Buffer = FPropertyHelpers::ReadToken(Buffer, Temp, true);
-			if (Buffer)
+			SkipWhitespace(Buffer);
+			Buffer = FPropertyHelpers::ReadToken(Buffer, Temp, /*bDottedNames=*/true);
+			if (!Buffer)
 			{
-				int32 EnumIndex = Enum->GetIndexByName(*Temp, EGetByNameFlags::CheckAuthoredName);
-				if (!bIsEnumOfFlags && EnumIndex == INDEX_NONE && (Temp.IsNumeric() && !Algo::Find(Temp, TEXT('.'))))
-				{
-					int64 LexedEnumValue = INDEX_NONE;
-					LexFromString(LexedEnumValue, *Temp);
-					EnumIndex = Enum->GetIndexByValue(LexedEnumValue);
-				}
-				if (EnumIndex == INDEX_NONE)
-				{
-					Buffer = nullptr;
-					break;
-				}
+				break;
+			}
 
-				EnumValue |= Enum->GetValueByIndex(EnumIndex);
+			int32 EnumIndex = Enum->GetIndexByName(*Temp, EGetByNameFlags::CheckAuthoredName);
+			if (!bIsEnumOfFlags && EnumIndex == INDEX_NONE && (Temp.IsNumeric() && !Algo::Find(Temp, TEXT('.'))))
+			{
+				int64 LexedEnumValue = INDEX_NONE;
+				LexFromString(LexedEnumValue, *Temp);
+				EnumIndex = Enum->GetIndexByValue(LexedEnumValue);
+			}
+			if (EnumIndex == INDEX_NONE)
+			{
+				Buffer = nullptr;
+				break;
+			}
 
-				if (!bIsEnumOfFlags)
-				{
-					break;
-				}
+			EnumValue |= Enum->GetValueByIndex(EnumIndex);
 
-				Temp.Reset();
-				SkipWhitespace(Buffer);
-				if (*Buffer != TEXT('|'))
-				{
-					break;
-				}
-				++Buffer;
-				SkipWhitespace(Buffer);
+			if (!bIsEnumOfFlags)
+			{
+				break;
+			}
+
+			Temp.Reset();
+			SkipWhitespace(Buffer);
+			if (*Buffer != TEXT('|'))
+			{
+				break;
 			}
+			++Buffer;
 		}
 
 		if (!Buffer)
