Struct Utils: Make it possible to construct a TInstancedStruct from a ScriptStruct

diff --git a/Engine/Source/Runtime/CoreUObject/Public/StructUtils/InstancedStruct.h b/Engine/Source/Runtime/CoreUObject/Public/StructUtils/InstancedStruct.h
--- a/Engine/Source/Runtime/CoreUObject/Public/StructUtils/InstancedStruct.h
+++ b/Engine/Source/Runtime/CoreUObject/Public/StructUtils/InstancedStruct.h
@@ -312,6 +312,12 @@ public:
 		InitializeAsScriptStruct(InOther.GetScriptStruct(), InOther.GetMemory());
 	}
 
+	template<typename T = BaseStructT, typename = std::enable_if_t<std::is_base_of_v<BaseStructT, std::decay_t<T>>>>
+	explicit TInstancedStruct(const UScriptStruct* InScriptStruct)
+	{
+		InitializeAsScriptStruct(InScriptStruct);
+	}
+
 	template<typename T = BaseStructT, typename = std::enable_if_t<std::is_base_of_v<BaseStructT, std::decay_t<T>>>>
 	TInstancedStruct(const TInstancedStruct<T>& InOther)
 		: InstancedStruct(InOther.InstancedStruct)
