Fixes for property initialization state tracking

Skip tracking for the state tracking flags themselves.

Added a missing initialization in InitProperties for objects of native and intrinsic types that are initialized from a default object other than the CDO.

#rnx

diff --git a/Engine/Source/Runtime/CoreUObject/Private/UObject/InstanceDataObjectUtils.cpp b/Engine/Source/Runtime/CoreUObject/Private/UObject/InstanceDataObjectUtils.cpp
--- a/Engine/Source/Runtime/CoreUObject/Private/UObject/InstanceDataObjectUtils.cpp
+++ b/Engine/Source/Runtime/CoreUObject/Private/UObject/InstanceDataObjectUtils.cpp
@@ -47,6 +47,11 @@ public:
 
 	bool HasPropertyValueFlag(EPropertyValueFlags Flags, const void* Data, const FProperty* Property, int32 ArrayIndex) const final
 	{
+		if (Property == InitializedValuesProperty || Property == SerializedValuesProperty)
+		{
+			return true;
+		}
+
 		if (const FProperty* FlagsProperty = GetPropertyValueFlagsProperty(Flags))
 		{
 			const int32 PropertyIndex = CalculatePropertyIndex(this, Property, ArrayIndex);
@@ -66,6 +71,11 @@ public:
 
 	void SetPropertyValueFlag(EPropertyValueFlags Flags, bool bValue, void* Data, const FProperty* Property, int32 ArrayIndex) const final
 	{
+		if (Property == InitializedValuesProperty || Property == SerializedValuesProperty)
+		{
+			return;
+		}
+
 		if (const FProperty* FlagsProperty = GetPropertyValueFlagsProperty(Flags))
 		{
 			const int32 PropertyIndex = CalculatePropertyIndex(this, Property, ArrayIndex);
@@ -797,6 +807,9 @@ namespace UE
 
 		Result->Bind();
 		Result->StaticLink(/*bRelinkExistingProperties*/true);
+		checkf(PropertyCount <= Result->TotalFieldCount,
+			TEXT("Type %s had %d properties after linking when at least %d are expected."),
+			*Result->GetPathName(), Result->TotalFieldCount, PropertyCount);
 		return Result;
 	}
 
diff --git a/Engine/Source/Runtime/CoreUObject/Private/UObject/UObjectGlobals.cpp b/Engine/Source/Runtime/CoreUObject/Private/UObject/UObjectGlobals.cpp
--- a/Engine/Source/Runtime/CoreUObject/Private/UObject/UObjectGlobals.cpp
+++ b/Engine/Source/Runtime/CoreUObject/Private/UObject/UObjectGlobals.cpp
@@ -4320,6 +4320,9 @@ void FObjectInitializer::InitProperties(UObject* Obj, UClass* DefaultsClass, UOb
 						if (P->IsInContainer(DefaultsClass))
 						{
 							P->CopyCompleteValue_InContainer(Obj, DefaultData);
+						#if WITH_EDITORONLY_DATA
+							SetUpPropertyInitializedFlag(DefaultData, P);
+						#endif
 						}
 					}
 				}
