Fixing TObjectPtr comparisons crashing when comparing against raw pointers (remote object handle specific)
Comparing TObjectPtr against raw object pointer will no longer resolve remote objects

diff --git a/Engine/Source/Runtime/CoreUObject/Private/UObject/ObjectHandlePrivate.h b/Engine/Source/Runtime/CoreUObject/Private/UObject/ObjectHandlePrivate.h
--- a/Engine/Source/Runtime/CoreUObject/Private/UObject/ObjectHandlePrivate.h
+++ b/Engine/Source/Runtime/CoreUObject/Private/UObject/ObjectHandlePrivate.h
@@ -29,17 +29,19 @@ namespace UE::CoreUObject::Private
 		FORCEINLINE static FRemoteObjectId GetRemoteId(const UObjectBase* Object)
 		{
 #if UE_WITH_REMOTE_OBJECT_HANDLE
-			if (UNLIKELY(Object->InternalIndex < 0))
+			if (Object)
 			{
-				return Object->GetPendingRegistrantRemoteId();
+				if (UNLIKELY(Object->InternalIndex < 0))
+				{
+					return Object->GetPendingRegistrantRemoteId();
+				}
+				else
+				{
+					return GUObjectArray.GetRemoteId(Object->InternalIndex);
+				}
 			}
-			else
-			{
-				return GUObjectArray.GetRemoteId(Object->InternalIndex);
-			}
-#else
-			return FRemoteObjectId();
 #endif // UE_WITH_REMOTE_OBJECT_HANDLE
+			return FRemoteObjectId();
 		}
 	};
 }
diff --git a/Engine/Source/Runtime/CoreUObject/Private/UObject/RemoteObject.cpp b/Engine/Source/Runtime/CoreUObject/Private/UObject/RemoteObject.cpp
--- a/Engine/Source/Runtime/CoreUObject/Private/UObject/RemoteObject.cpp
+++ b/Engine/Source/Runtime/CoreUObject/Private/UObject/RemoteObject.cpp
@@ -323,11 +323,7 @@ namespace UE::CoreUObject::Private
 		{
 			return ToStub()->Id;
 		}
-		else if (PointerOrHandle)
-		{
-			return FRemoteObjectId(reinterpret_cast<const UObjectBase*>(PointerOrHandle));
-		}
-		return FRemoteObjectId();
+		return FRemoteObjectId(reinterpret_cast<const UObjectBase*>(PointerOrHandle));
 	}
 
 	FRemoteObjectHandlePrivate FRemoteObjectHandlePrivate::ConvertToRemoteHandle(UObject* Object)
diff --git a/Engine/Source/Runtime/CoreUObject/Public/UObject/ObjectHandle.h b/Engine/Source/Runtime/CoreUObject/Public/UObject/ObjectHandle.h
--- a/Engine/Source/Runtime/CoreUObject/Public/UObject/ObjectHandle.h
+++ b/Engine/Source/Runtime/CoreUObject/Public/UObject/ObjectHandle.h
@@ -279,8 +279,8 @@ inline bool operator==(UE::CoreUObject::Private::FRemoteObjectHandlePrivate LHS,
 	using namespace UE::CoreUObject::Private;
 	using namespace UE::RemoteObject::Handle;
 
-	bool bLhsResolved = IsObjectHandleResolved(LHS);
-	bool bRhsResolved = IsObjectHandleResolved(RHS);
+	const bool bLhsResolved = IsObjectHandleResolved_ForGC(LHS);
+	const bool bRhsResolved = IsObjectHandleResolved_ForGC(RHS);
 
 	//if both resolved or both unresolved compare the uintptr
 	if (bLhsResolved == bRhsResolved)
@@ -331,7 +331,7 @@ inline FRemoteObjectId GetRemoteObjectId(UE::CoreUObject::Private::FRemoteObject
 		return FRemoteObjectId();
 	}
 
-	if (IsObjectHandleResolved(Handle))
+	if (IsObjectHandleResolved_ForGC(Handle))
 	{
 		const UObject* Obj = ReadObjectHandlePointerNoCheck(Handle);
 		return FRemoteObjectId((const UObjectBase*)Obj);
@@ -494,8 +494,16 @@ namespace UE::CoreUObject::Private
 		return ResolvedObject;
 #elif UE_WITH_REMOTE_OBJECT_HANDLE
 		FObjectHandle LocalHandle = Handle;
-		check(!!(LocalHandle.PointerOrHandle & 1) || !LocalHandle.PointerOrHandle);
-		UObject* ResolvedObject = UE::RemoteObject::Handle::ResolveObject(LocalHandle.ToStub());
+		// Unresolved handle may mean two things: we still have the remote object memory (it hasn't been GC'd yet) or we only have a stub
+		UObject* ResolvedObject = nullptr;
+		if (IsObjectHandleResolved_ForGC(LocalHandle))
+		{
+			ResolvedObject = UE::RemoteObject::Handle::ResolveObject(ReadObjectHandlePointerNoCheck(LocalHandle));
+		}
+		else
+		{
+			ResolvedObject = UE::RemoteObject::Handle::ResolveObject(LocalHandle.ToStub());
+		}
 		LocalHandle = MakeObjectHandle(ResolvedObject);
 		Handle = LocalHandle;
 		return ResolvedObject;	
diff --git a/Engine/Source/Runtime/CoreUObject/Public/UObject/ObjectPtr.h b/Engine/Source/Runtime/CoreUObject/Public/UObject/ObjectPtr.h
--- a/Engine/Source/Runtime/CoreUObject/Public/UObject/ObjectPtr.h
+++ b/Engine/Source/Runtime/CoreUObject/Public/UObject/ObjectPtr.h
@@ -387,7 +387,12 @@ namespace ObjectPtr_Private
 			return false;	// from above, we already know that !Ptr is false
 		}
 #endif
+#if UE_WITH_REMOTE_OBJECT_HANDLE
+		// avoids resolving since we can compare against globally unique id
+		return Ptr.GetRemoteId() == FRemoteObjectId(reinterpret_cast<const UObjectBase*>(Other));
+#else
 		return Ptr.GetNoReadNoCheck() == Other;
+#endif
 	}
 
 	/** Perform shallow equality check between a TObjectPtr and another (non TObjectPtr) type that we can coerce to a pointer. */
@@ -728,6 +733,13 @@ public:
 		Slot << ObjectPtr;
 	}
 
+#if UE_WITH_REMOTE_OBJECT_HANDLE
+	FORCEINLINE FRemoteObjectId GetRemoteId() const
+	{
+		return ObjectPtr.GetRemoteId();
+	}
+#endif
+
 	friend ObjectPtr_Private::Friend;
 	friend struct FObjectPtr;
 	template <typename U> friend struct TObjectPtr;
