Introduce MakeGloballyUniqueObjectName that creates a globally unique name by combining a base name, a GUID and a timestamp. The object name returned is guaranteed to be globally unique. This is done by combining a GUID for the current user session and a microseconds timestamp, then converting the resulting binary value into base 36. The timestamp is split into 2 parts:  the lowest 30 bits are stored in the name number to minimize the total number of generated names, and the remaining top bits contributing to the generated name string.

diff --git a/Engine/Source/Runtime/CoreUObject/Private/UObject/UObjectGlobals.cpp b/Engine/Source/Runtime/CoreUObject/Private/UObject/UObjectGlobals.cpp
--- a/Engine/Source/Runtime/CoreUObject/Private/UObject/UObjectGlobals.cpp
+++ b/Engine/Source/Runtime/CoreUObject/Private/UObject/UObjectGlobals.cpp
@@ -9,6 +9,7 @@
 #include "HAL/PlatformFileManager.h"
 #include "HAL/FileManager.h"
 #include "Logging/StructuredLog.h"
+#include "Math/Int128.h"
 #include "Misc/AsciiSet.h"
 #include "Misc/EnumRange.h"
 #include "Misc/PackageAccessTrackingOps.h"
@@ -2893,6 +2894,55 @@ FName MakeUniqueObjectName(UObject* Parent, const UClass* Class, FName InBaseNam
 	return TestName;
 }
 
+FName MakeGloballyUniqueObjectName(UObject* Outer, FName BaseName)
+{
+	// Generate a unique generation context for this session.
+	static struct FContext
+	{
+		FGuid Guid = FGuid::NewGuid();
+		FDateTime Origin = FDateTime(2024, 10, 1);
+		std::atomic<uint32> Counter = 0;
+	} Context;
+
+	FName OutObjectName;
+
+	do
+	{
+		TStringBuilderWithBuffer<TCHAR, NAME_SIZE> StringBuilder;
+
+		// Append the provided base name if present.
+		if (!BaseName.IsNone())
+		{
+			StringBuilder += BaseName.ToString();
+			StringBuilder += TEXT("_");
+		}
+
+		// Compute the amount of microseconds since a specific date to construct our GUID.
+		const FTimespan Elapsed = FDateTime::Now() - Context.Origin;
+		const uint64 ElapsedUs = (uint64)Elapsed.GetTotalMilliseconds() * 1000 + (Context.Counter++ % 1000);
+
+		// The lowest 30 bits of the timestamp are stored in the name number below. This is to minimize the total names 
+		// generated for globally unique names (the string part will change every ~17 minutes).
+		FUInt128 Value(0, uint32(ElapsedUs >> 30), Context.Guid.A ^ Context.Guid.B, Context.Guid.C ^ Context.Guid.D);
+
+		// Convert the constructed value in base 36
+		while (!Value.IsZero())
+		{
+			uint32 Remainder;
+			Value = Value.Divide(36, Remainder);
+			StringBuilder.AppendChar(NibbleToTChar(uint8(Remainder)));
+		}
+
+		// The name number bit 30 is reserved to know if this is a globally unique name. This is not 100% safe, but
+		// should cover most cases. The name number bit 31 is reserved to preserve the fast path name generation
+		// (see GFastPathUniqueNameGeneration).
+		OutObjectName = FName(*StringBuilder, (ElapsedUs & 0x3fffffff) | (1 << 30));
+	}
+	while (StaticFindObjectFast(nullptr, Outer, OutObjectName));
+
+	return OutObjectName;
+}
+
 FName MakeObjectNameFromDisplayLabel(const FString& DisplayLabel, const FName CurrentObjectName)
 {
 	FString GeneratedName = SlugStringForValidName(DisplayLabel);
diff --git a/Engine/Source/Runtime/CoreUObject/Public/UObject/UObjectGlobals.h b/Engine/Source/Runtime/CoreUObject/Public/UObject/UObjectGlobals.h
--- a/Engine/Source/Runtime/CoreUObject/Public/UObject/UObjectGlobals.h
+++ b/Engine/Source/Runtime/CoreUObject/Public/UObject/UObjectGlobals.h
@@ -967,7 +967,7 @@ ENUM_CLASS_FLAGS(EUniqueObjectNameOptions);
  * Create a unique name by combining a base name and an arbitrary number string.
  * The object name returned is guaranteed not to exist.
  *
- * @param	Parent		the outer for the object that needs to be named
+ * @param	Outer		the outer for the object that needs to be named
  * @param	Class		the class for the object
  * @param	BaseName	optional base name to use when generating the unique object name; if not specified, the class's name is used
  * @param	Options		Additional options. See EUniqueObjectNameOptions.
@@ -977,6 +977,17 @@ ENUM_CLASS_FLAGS(EUniqueObjectNameOptions);
  */
 COREUOBJECT_API FName MakeUniqueObjectName( UObject* Outer, const UClass* Class, FName BaseName = NAME_None, EUniqueObjectNameOptions Options = EUniqueObjectNameOptions::None);
 
+/**
+ * Create a globally unique name by combining a base name, a GUID and a timestamp.
+ * The object name returned is guaranteed not to locally and globally exist.
+ *
+ * @param	Outer		the outer for the object that needs to be named
+ * @param	BaseName	optional base name to use when generating the globally unique object name.
+ * 
+ * @return	name is the form BaseName_##, where ## a globally unique identifier making sure that name is globally unique.
+  */
+COREUOBJECT_API FName MakeGloballyUniqueObjectName(UObject* Outer, FName InBaseName = NAME_None);
+
 /**
  * Given a display label string, generates an FName slug that is a valid FName for that label.
  * If the object's current name is already satisfactory, then that name will be returned.
@@ -994,7 +1005,7 @@ COREUOBJECT_API FName MakeObjectNameFromDisplayLabel(const FString& DisplayLabel
 /**
  * Returns whether an object is referenced, not counting references from itself
  *
- * @param	Obj			Object to check
+ * @param	Res			Object to check
  * @param	KeepFlags	Objects with these flags will be considered as being referenced
  * @param	InternalKeepFlags	Objects with these internal flags will be considered as being referenced
  * @param	bCheckSubObjects	Treat subobjects as if they are the same as passed in object
