ClearOverride for use in scene graph will clear property initialization state
* This introduces new flag to give context to FPropertyValueImpl and PropertyChangedEvent that values are being set due to a ResetToDefault. When ResetToDefault, initialization state is cleared.
* This change also fixes NonNullPtrs being able to be reset back to their default null state if ClearOverride->ResetToDefault results in a null default
* OverriddenPropertySet and UObject::PostEditChangeProperty were using switch-case over a field of flags, equality would fail on concatenated flags. Converted to if-else chain.

diff --git a/Engine/Source/Editor/PropertyEditor/Private/PropertyHandleImpl.cpp b/Engine/Source/Editor/PropertyEditor/Private/PropertyHandleImpl.cpp
--- a/Engine/Source/Editor/PropertyEditor/Private/PropertyHandleImpl.cpp
+++ b/Engine/Source/Editor/PropertyEditor/Private/PropertyHandleImpl.cpp
@@ -39,6 +39,7 @@
 #include "Modules/ModuleManager.h"
 #include "PropertyEditorModule.h"
 #include "Internationalization/TextPackageNamespaceUtil.h"
+#include "UObject/NonNullPropertyUtils.h"
 
 #define LOCTEXT_NAMESPACE "PropertyHandleImplementation"
 
@@ -554,7 +555,12 @@ FPropertyAccess::Result FPropertyValueImpl::ImportText( const TArray<FObjectBase
 			FPropertyValueImpl::GenerateArrayIndexMapToObjectNode(ArrayIndicesPerObject[ObjectIndex], InPropertyNode);
 		}
 
-		FPropertyChangedEvent ChangeEvent(NodeProperty, bFinished ? EPropertyChangeType::ValueSet : EPropertyChangeType::Interactive, MakeArrayView(TopLevelObjects));
+		EPropertyChangeType::Type ChangeEventFlags = bFinished ? EPropertyChangeType::ValueSet : EPropertyChangeType::Interactive;
+		if (Flags & EPropertyValueSetFlags::ResetToDefault)
+		{
+			ChangeEventFlags = ChangeEventFlags | EPropertyChangeType::ResetToDefault;
+		}
+		FPropertyChangedEvent ChangeEvent(NodeProperty, ChangeEventFlags, MakeArrayView(TopLevelObjects));
 		ChangeEvent.SetArrayIndexPerObject(ArrayIndicesPerObject);
 
 		// If PreEditChange was called, so should PostEditChange.
@@ -1004,12 +1010,13 @@ void FPropertyValueImpl::ResetToDefault()
 		const bool bUseDisplayName = false;
 		FScopedTransaction Transaction( NSLOCTEXT("UnrealEd", "PropertyWindowResetToDefault", "Reset to Default") );
 		TSharedPtr<FPropertyNode>& KeyNode = PropertyNodePin->GetPropertyKeyNode();
+		UE::CoreUObject::FScopedSetNullOnNonNullable ScopedSetNullOnNonNullable(UE::CoreUObject::EAllowSetNullOnNonNullableBehavior::Enabled);
 		if(KeyNode.IsValid())
 		{
 			FPropertyValueImpl(KeyNode, NotifyHook, PropertyUtilities.Pin())
-				.ImportText(KeyNode->GetDefaultValueAsString(bUseDisplayName), EPropertyValueSetFlags::InstanceObjects);
+				.ImportText(KeyNode->GetDefaultValueAsString(bUseDisplayName), EPropertyValueSetFlags::InstanceObjects | EPropertyValueSetFlags::ResetToDefault);
 		}
-		ImportText(PropertyNodePin->GetDefaultValueAsString(bUseDisplayName), EPropertyValueSetFlags::InstanceObjects);
+		ImportText(PropertyNodePin->GetDefaultValueAsString(bUseDisplayName), EPropertyValueSetFlags::InstanceObjects | EPropertyValueSetFlags::ResetToDefault);
 
 		PropertyNodePin->BroadcastPropertyResetToDefault();
 	}
diff --git a/Engine/Source/Editor/PropertyEditor/Public/PropertyHandle.h b/Engine/Source/Editor/PropertyEditor/Public/PropertyHandle.h
--- a/Engine/Source/Editor/PropertyEditor/Public/PropertyHandle.h
+++ b/Engine/Source/Editor/PropertyEditor/Public/PropertyHandle.h
@@ -30,6 +30,8 @@ namespace EPropertyValueSetFlags
 	inline const Type InteractiveChange = 1 << 1;
 	/** If the property being imported to is an instanced object, create a new object rather than simply setting the value literally */
 	inline const Type InstanceObjects = 1 << 2;
+	/** If the property is being set due to being reset to default */
+	inline const Type ResetToDefault = 1 << 3;
 
 };
 
diff --git a/Engine/Source/Runtime/CoreUObject/Private/UObject/NonNullPropertyUtils.cpp b/Engine/Source/Runtime/CoreUObject/Private/UObject/NonNullPropertyUtils.cpp
--- a/Engine/Source/Runtime/CoreUObject/Private/UObject/NonNullPropertyUtils.cpp
+++ b/Engine/Source/Runtime/CoreUObject/Private/UObject/NonNullPropertyUtils.cpp
@@ -1,10 +1,10 @@
 ﻿// Copyright Epic Games, Inc. All Rights Reserved.
 
-#include "NonNullPropertyUtils.h"
+#include "UObject/NonNullPropertyUtils.h"
 
 #include "HAL/IConsoleManager.h"
 
-namespace UE::CoreUObject::Private
+namespace UE::CoreUObject
 {
 #if WITH_EDITORONLY_DATA
 	TAutoConsoleVariable<bool> CVarEnableNullDeserializeToNonNullableOverride(
@@ -18,14 +18,14 @@ namespace UE::CoreUObject::Private
 	FScopedSetNullOnNonNullable::FScopedSetNullOnNonNullable(TOptional<EAllowSetNullOnNonNullableBehavior> InBehavior)
 		: Previous(nullptr)
 	{
-		if (InBehavior.IsSet() && Private::CVarEnableNullDeserializeToNonNullableOverride.GetValueOnAnyThread())
+		if (InBehavior.IsSet() && CVarEnableNullDeserializeToNonNullableOverride.GetValueOnAnyThread())
 		{
 			bIsBehaviorSet = true;
 			Behavior =  *InBehavior;
 
 			// Push the stack object to the list
-			Previous = Private::TLSScopedSetNullOnNonNullableBehaviorStackHead;
-			Private::TLSScopedSetNullOnNonNullableBehaviorStackHead = this;
+			Previous = TLSScopedSetNullOnNonNullableBehaviorStackHead;
+			TLSScopedSetNullOnNonNullableBehaviorStackHead = this;
 		}
 		else
 		{
@@ -39,8 +39,8 @@ namespace UE::CoreUObject::Private
 		if (bIsBehaviorSet)
 		{	
 			// Pop us from the stack and set the previous to head
-			check(Private::TLSScopedSetNullOnNonNullableBehaviorStackHead == this);
-			Private::TLSScopedSetNullOnNonNullableBehaviorStackHead = Previous;
+			check(TLSScopedSetNullOnNonNullableBehaviorStackHead == this);
+			TLSScopedSetNullOnNonNullableBehaviorStackHead = Previous;
 		}
 	}
 
diff --git a/Engine/Source/Runtime/CoreUObject/Private/UObject/Obj.cpp b/Engine/Source/Runtime/CoreUObject/Private/UObject/Obj.cpp
--- a/Engine/Source/Runtime/CoreUObject/Private/UObject/Obj.cpp
+++ b/Engine/Source/Runtime/CoreUObject/Private/UObject/Obj.cpp
@@ -5,7 +5,7 @@
 =============================================================================*/
 
 #include "CoreMinimal.h"
-#include "NonNullPropertyUtils.h"
+#include "UObject/NonNullPropertyUtils.h"
 #include "AssetRegistry/AssetData.h"
 #include "Misc/CoreMisc.h"
 #include "Misc/CommandLine.h"
@@ -474,20 +474,15 @@ void UObject::PostEditChangeProperty(FPropertyChangedEvent& PropertyChangedEvent
 #if WITH_EDITORONLY_DATA
 	if (const FProperty* Property = PropertyChangedEvent.MemberProperty)
 	{
-		switch (PropertyChangedEvent.ChangeType)
+		const EPropertyChangeType::Type SetPropertyValueInitializedStates =
+			EPropertyChangeType::ArrayAdd | EPropertyChangeType::ArrayRemove | EPropertyChangeType::ArrayClear | EPropertyChangeType::ArrayMove | EPropertyChangeType::ValueSet | EPropertyChangeType::Interactive;
+		if (PropertyChangedEvent.ChangeType & EPropertyChangeType::ResetToDefault)
+		{
+			UE::ClearPropertyValueInitialized(GetClass(), this, Property);
+		}
+		else if (PropertyChangedEvent.ChangeType & SetPropertyValueInitializedStates)
 		{
-		case EPropertyChangeType::ArrayAdd:
-		case EPropertyChangeType::ArrayRemove:
-		case EPropertyChangeType::ArrayClear: // Clearing the entire array. No need for an index. Clear also counts as initialization of the property.
-		case EPropertyChangeType::ArrayMove:
-		case EPropertyChangeType::ValueSet:
-		case EPropertyChangeType::Interactive:
-			// No tracking for elements of container. Any change counts as initialization.
 			UE::SetPropertyValueInitialized(GetClass(), this, Property);
-			break;
-		default:
-			// No op for now
-			break;
 		}
 	}
 #endif
@@ -1632,14 +1627,14 @@ void UObject::Serialize(FStructuredArchive::FRecord Record)
 #if WITH_EDITORONLY_DATA
 		// Allow Null to be deserialized into a NonNull property
 		// Enabled for Undo/Redo purposes
-		UE::CoreUObject::Private::FScopedSetNullOnNonNullable ScopedAllowNullPtrDeserialization = [&UnderlyingArchive]()
+		UE::CoreUObject::FScopedSetNullOnNonNullable ScopedAllowNullPtrDeserialization = [&UnderlyingArchive]()
 		{
-			TOptional<UE::CoreUObject::Private::EAllowSetNullOnNonNullableBehavior> Behaviour;
+			TOptional<UE::CoreUObject::EAllowSetNullOnNonNullableBehavior> Behaviour;
 			if (UnderlyingArchive.IsTransacting() && UnderlyingArchive.IsLoading())
 			{
-				Behaviour.Emplace(UE::CoreUObject::Private::EAllowSetNullOnNonNullableBehavior::Enabled);
+				Behaviour.Emplace(UE::CoreUObject::EAllowSetNullOnNonNullableBehavior::Enabled);
 			}
-			return UE::CoreUObject::Private::FScopedSetNullOnNonNullable(Behaviour);
+			return UE::CoreUObject::FScopedSetNullOnNonNullable(Behaviour);
 		}();
 #endif
 		
diff --git a/Engine/Source/Runtime/CoreUObject/Private/UObject/OverriddenPropertySet.cpp b/Engine/Source/Runtime/CoreUObject/Private/UObject/OverriddenPropertySet.cpp
--- a/Engine/Source/Runtime/CoreUObject/Private/UObject/OverriddenPropertySet.cpp
+++ b/Engine/Source/Runtime/CoreUObject/Private/UObject/OverriddenPropertySet.cpp
@@ -865,57 +865,57 @@ void FOverriddenPropertySet::NotifyPropertyChange(FOverriddenPropertyNode* Paren
 						return;
 					}
 
-					switch(ChangeType)
+					if (ChangeType & EPropertyChangeType::ValueSet)
 					{
-					case EPropertyChangeType::ValueSet:
 						checkf(ArrayIndex != INDEX_NONE, TEXT("ValueSet change type should have associated indexes"));
-						// Intentional fall thru
-					case EPropertyChangeType::Unspecified:
-						{
-							if (ArrayIndex != INDEX_NONE)
-							{
-								// Overriding a single entry in the array
-								ArrayRemoveImpl();
-								ArrayAddImpl();
-							}
-							else
-							{
-								ArrayReplace();
-							}
-							return;
-						}
-					case EPropertyChangeType::ArrayAdd:
+					}
+
+					if (ChangeType & (EPropertyChangeType::ValueSet | EPropertyChangeType::Unspecified))
+					{
+						if (ArrayIndex != INDEX_NONE)
 						{
+							// Overriding a single entry in the array
+							ArrayRemoveImpl();
 							ArrayAddImpl();
-							return;
 						}
-					case EPropertyChangeType::ArrayRemove:
+						else
 						{
-							ArrayRemoveImpl();
-							return;
+							ArrayReplace();
 						}
-					case EPropertyChangeType::ArrayClear:
-						{
-							checkf(ArrayIndex == INDEX_NONE, TEXT("ArrayClear change type should not have associated indexes"));
+						return;
+					}
 
-							for (int i = 0; i < PreEditSubObjectsArrayHelper.Num(); ++i)
-							{
-								ArrayIndex = i;
-								ArrayRemoveImpl();
-							}
-							return;
-						}
-					case EPropertyChangeType::ArrayMove:
-						{
-							UE_LOG(LogOverridableObject, Warning, TEXT("ArrayMove change type is not going to change anything as ordering of object isn't supported yet"));
-							return;
-						}
-					default:
+					if (ChangeType & EPropertyChangeType::ArrayAdd)
+					{
+						ArrayAddImpl();
+						return;
+					}
+
+					if (ChangeType & EPropertyChangeType::ArrayRemove)
+					{
+						ArrayRemoveImpl();
+						return;
+					}
+			
+					if (ChangeType & EPropertyChangeType::ArrayClear)
+					{
+						checkf(ArrayIndex == INDEX_NONE, TEXT("ArrayClear change type should not have associated indexes"));
+
+						for (int i = 0; i < PreEditSubObjectsArrayHelper.Num(); ++i)
 						{
-							UE_LOG(LogOverridableObject, Warning, TEXT("Property change type is not supported will default to full array override"));
-							break;
+							ArrayIndex = i;
+							ArrayRemoveImpl();
 						}
+						return;
 					}
+			
+					if (ChangeType & EPropertyChangeType::ArrayMove)
+					{
+						UE_LOG(LogOverridableObject, Warning, TEXT("ArrayMove change type is not going to change anything as ordering of object isn't supported yet"));
+						return;
+					}
+
+					UE_LOG(LogOverridableObject, Warning, TEXT("Property change type is not supported will default to full array override"));
 				}
 				// Can only forward to subobject if we have a valid index
 				else if (ArrayHelper.IsValidIndex(ArrayIndex))
@@ -1080,57 +1080,63 @@ void FOverriddenPropertySet::NotifyPropertyChange(FOverriddenPropertyNode* Paren
 				return;
 			}
 
-			switch (ChangeType)
+			if (ChangeType & EPropertyChangeType::ValueSet)
 			{
-			case EPropertyChangeType::ValueSet:
 				checkf(LogicalMapIndex != INDEX_NONE, TEXT("ValueSet change type should have associated indexes"));
-				// Intentional fall thru
-			case EPropertyChangeType::Unspecified:
-				{
-					if(LogicalMapIndex != INDEX_NONE)
-					{
-						// Overriding a single entry in the map
-						MapRemoveImpl();
-						MapAddImpl();
-					}
-					else
-					{
-						MapReplace();
-					}
-					return;
-				}
-			case EPropertyChangeType::ArrayAdd:
+			}
+
+			if (ChangeType & (EPropertyChangeType::ValueSet | EPropertyChangeType::Unspecified))
+			{
+				if(LogicalMapIndex != INDEX_NONE)
 				{
+					// Overriding a single entry in the map
+					MapRemoveImpl();
 					MapAddImpl();
-					return;
 				}
-			case EPropertyChangeType::ArrayRemove:
+				else
 				{
-					MapRemoveImpl();
-					return;
+					MapReplace();
 				}
-			case EPropertyChangeType::ArrayClear:
-				{
-					checkf(InternalPreEditMapIndex == INDEX_NONE, TEXT("ArrayClear change type should not have associated indexes"));
+				return;
+			}
 
-					for (FScriptMapHelper::FIterator It(PreEditMapHelper); It; ++It)
-					{
-						InternalPreEditMapIndex = It.GetInternalIndex();
-						MapRemoveImpl();
-					}
-					return;
-				}
-			case EPropertyChangeType::ArrayMove:
-				{
-					UE_LOG(LogOverridableObject, Warning, TEXT("ArrayMove change type is not going to change anything as ordering of object isn't supported yet"));
-					return;
-				}
-			default:
+			if (ChangeType & EPropertyChangeType::ArrayAdd)
+			{
+				MapAddImpl();
+				return;
+			}
+
+			if (ChangeType & EPropertyChangeType::ArrayRemove)
+			{
+				MapRemoveImpl();
+				return;
+			}
+			
+			if (ChangeType & EPropertyChangeType::ArrayClear)
+			{
+				checkf(InternalPreEditMapIndex == INDEX_NONE, TEXT("ArrayClear change type should not have associated indexes"));
+
+				for (FScriptMapHelper::FIterator It(PreEditMapHelper); It; ++It)
 				{
-					UE_LOG(LogOverridableObject, Warning, TEXT("Property change type is not supported will default to full array override"));
-					break;
+					InternalPreEditMapIndex = It.GetInternalIndex();
+					MapRemoveImpl();
 				}
+				return;
 			}
+			
+			if (ChangeType & EPropertyChangeType::ArrayMove)
+			{
+				UE_LOG(LogOverridableObject, Warning, TEXT("ArrayMove change type is not going to change anything as ordering of object isn't supported yet"));
+				return;
+			}
+
+			if (ChangeType & EPropertyChangeType::ArrayAdd)
+			{
+				MapAddImpl();
+				return;
+			}
+
+			UE_LOG(LogOverridableObject, Warning, TEXT("Property change type is not supported will default to full array override"));;
 		}
 		// Can only forward to subobject if we have a valid index
 		else if (MapHelper.IsValidIndex(InternalMapIndex))
diff --git a/Engine/Source/Runtime/CoreUObject/Private/UObject/PropertyBaseObject.cpp b/Engine/Source/Runtime/CoreUObject/Private/UObject/PropertyBaseObject.cpp
--- a/Engine/Source/Runtime/CoreUObject/Private/UObject/PropertyBaseObject.cpp
+++ b/Engine/Source/Runtime/CoreUObject/Private/UObject/PropertyBaseObject.cpp
@@ -2,7 +2,7 @@
 
 #include "CoreMinimal.h"
 #include "HAL/IConsoleManager.h"
-#include "NonNullPropertyUtils.h"
+#include "UObject/NonNullPropertyUtils.h"
 #include "UObject/ObjectMacros.h"
 #include "Templates/Casts.h"
 #include "UObject/Package.h"
@@ -667,7 +667,7 @@ void FObjectPropertyBase::SetObjectPtrPropertyValueUnchecked_InContainer(void* C
 
 void FObjectPropertyBase::SetObjectPropertyValue(void* PropertyValueAddress, UObject* Value) const
 {
-	if (Value || !HasAnyPropertyFlags(CPF_NonNullable) || (UE::CoreUObject::Private::AllowSetNullOnNonNullableBehavior() == UE::CoreUObject::Private::EAllowSetNullOnNonNullableBehavior::Enabled))
+	if (Value || !HasAnyPropertyFlags(CPF_NonNullable) || (UE::CoreUObject::AllowSetNullOnNonNullableBehavior() == UE::CoreUObject::EAllowSetNullOnNonNullableBehavior::Enabled))
 	{
 		SetObjectPropertyValueUnchecked(PropertyValueAddress, Value);
 	}
@@ -679,7 +679,7 @@ void FObjectPropertyBase::SetObjectPropertyValue(void* PropertyValueAddress, UOb
 
 void FObjectPropertyBase::SetObjectPtrPropertyValue(void* PropertyValueAddress, TObjectPtr<UObject> Ptr) const
 {
-	if (Ptr || !HasAnyPropertyFlags(CPF_NonNullable) || (UE::CoreUObject::Private::AllowSetNullOnNonNullableBehavior() == UE::CoreUObject::Private::EAllowSetNullOnNonNullableBehavior::Enabled))
+	if (Ptr || !HasAnyPropertyFlags(CPF_NonNullable) || (UE::CoreUObject::AllowSetNullOnNonNullableBehavior() == UE::CoreUObject::EAllowSetNullOnNonNullableBehavior::Enabled))
 	{
 		SetObjectPtrPropertyValueUnchecked(PropertyValueAddress, Ptr);
 	}
@@ -691,7 +691,7 @@ void FObjectPropertyBase::SetObjectPtrPropertyValue(void* PropertyValueAddress,
 
 void FObjectPropertyBase::SetObjectPropertyValue_InContainer(void* ContainerAddress, UObject* Value, int32 ArrayIndex) const
 {
-	if (Value || !HasAnyPropertyFlags(CPF_NonNullable) || (UE::CoreUObject::Private::AllowSetNullOnNonNullableBehavior() == UE::CoreUObject::Private::EAllowSetNullOnNonNullableBehavior::Enabled))
+	if (Value || !HasAnyPropertyFlags(CPF_NonNullable) || (UE::CoreUObject::AllowSetNullOnNonNullableBehavior() == UE::CoreUObject::EAllowSetNullOnNonNullableBehavior::Enabled))
 	{
 		SetObjectPropertyValueUnchecked_InContainer(ContainerAddress, Value, ArrayIndex);
 	}
@@ -703,7 +703,7 @@ void FObjectPropertyBase::SetObjectPropertyValue_InContainer(void* ContainerAddr
 
 void FObjectPropertyBase::SetObjectPtrPropertyValue_InContainer(void* ContainerAddress, TObjectPtr<UObject> Ptr, int32 ArrayIndex) const
 {
-	if (Ptr || !HasAnyPropertyFlags(CPF_NonNullable) || (UE::CoreUObject::Private::AllowSetNullOnNonNullableBehavior() == UE::CoreUObject::Private::EAllowSetNullOnNonNullableBehavior::Enabled))
+	if (Ptr || !HasAnyPropertyFlags(CPF_NonNullable) || (UE::CoreUObject::AllowSetNullOnNonNullableBehavior() == UE::CoreUObject::EAllowSetNullOnNonNullableBehavior::Enabled))
 	{
 		SetObjectPtrPropertyValueUnchecked_InContainer(ContainerAddress, Ptr, ArrayIndex);
 	}
diff --git a/Engine/Source/Runtime/CoreUObject/Public/UObject/UnrealType.h b/Engine/Source/Runtime/CoreUObject/Public/UObject/UnrealType.h
--- a/Engine/Source/Runtime/CoreUObject/Public/UObject/UnrealType.h
+++ b/Engine/Source/Runtime/CoreUObject/Public/UObject/UnrealType.h
@@ -6491,6 +6491,8 @@ namespace EPropertyChangeType
 	inline const Type ArrayMove			= 1 << 8;
 	// Edit Condition State has changed
 	inline const Type ToggleEditable	= 1 << 9;
+	//
+	inline const Type ResetToDefault    = 1 << 10;
 };
 
 /**
