Fix Uninitialized Property alerts for POD data in Outliner + Enable uninitialized alerts by default

diff --git a/Engine/Source/Runtime/CoreUObject/Private/UObject/InstanceDataObjectUtils.cpp b/Engine/Source/Runtime/CoreUObject/Private/UObject/InstanceDataObjectUtils.cpp
--- a/Engine/Source/Runtime/CoreUObject/Private/UObject/InstanceDataObjectUtils.cpp
+++ b/Engine/Source/Runtime/CoreUObject/Private/UObject/InstanceDataObjectUtils.cpp
@@ -271,7 +271,7 @@ namespace UE
 		TEXT("Enables uobject annotation method to track property initialization state if IDOs disabled.")
 	);
 
-	bool bEnableUninitializedUI = false;
+	bool bEnableUninitializedUI = true;
 	FAutoConsoleVariableRef EnableUninitializedUI(
 		TEXT("IDO.EnableUninitializedAlertUI"),
 		bEnableUninitializedUI,
