[StructUtils] use actual UScriptStruct member pointer for serialization, instead of local copy.

diff --git a/Engine/Source/Runtime/CoreUObject/Private/StructUtils/InstancedStruct.cpp b/Engine/Source/Runtime/CoreUObject/Private/StructUtils/InstancedStruct.cpp
--- a/Engine/Source/Runtime/CoreUObject/Private/StructUtils/InstancedStruct.cpp
+++ b/Engine/Source/Runtime/CoreUObject/Private/StructUtils/InstancedStruct.cpp
@@ -143,24 +143,28 @@ bool FInstancedStruct::Serialize(FArchive& Ar)
 		}
 	}
 
-	UScriptStruct* NonConstStruct = const_cast<UScriptStruct*>(GetScriptStruct());
-
 	if (Ar.IsLoading())
 	{
 		// UScriptStruct type
-		Ar << NonConstStruct;
-		if (NonConstStruct)
+		UScriptStruct* SerializedScriptStruct = nullptr;
+		Ar << SerializedScriptStruct;
+		if (SerializedScriptStruct)
 		{
-			Ar.Preload(NonConstStruct);
+			Ar.Preload(SerializedScriptStruct);
+		}
+
+		// Initialize only if the type changes.
+		if (ScriptStruct != SerializedScriptStruct)
+		{
+			InitializeAs(SerializedScriptStruct);
 		}
-		InitializeAs(NonConstStruct);
 
 		// Size of the serialized memory
 		int32 SerialSize = 0; 
 		Ar << SerialSize;
 
 		// Serialized memory
-		if (NonConstStruct == nullptr && SerialSize > 0)
+		if (ScriptStruct == nullptr && SerialSize > 0)
 		{
 			// A null struct indicates an old struct or an unsupported one for the current target.
 			// In this case we manually seek in the archive to skip its serialized content. 
@@ -172,16 +176,16 @@ bool FInstancedStruct::Serialize(FArchive& Ar)
 				Ar.GetLinker() ? *GetPathNameSafe(Ar.GetLinker()->LinkerRoot) : TEXT("NoLinker"));
 			Ar.Seek(Ar.Tell() + SerialSize);
 		}
-		else if (NonConstStruct != nullptr && ensureMsgf(GetMutableMemory() != nullptr, TEXT("A valid script struct should always have allocated memory")))
+		else if (ScriptStruct != nullptr && ensureMsgf(GetMutableMemory() != nullptr, TEXT("A valid script struct should always have allocated memory")))
 		{
-			NonConstStruct->SerializeItem(Ar, GetMutableMemory(), /* Defaults */ nullptr);
+			ConstCast(ScriptStruct)->SerializeItem(Ar, GetMutableMemory(), /* Defaults */ nullptr);
 		}
 	}
 	else if (Ar.IsSaving())
 	{
 		// UScriptStruct type
 #if WITH_EDITOR
-		const UUserDefinedStruct* UserDefinedStruct = Cast<UUserDefinedStruct>(NonConstStruct);
+		const UUserDefinedStruct* UserDefinedStruct = Cast<UUserDefinedStruct>(ScriptStruct);
 		if (UserDefinedStruct
 			&& UserDefinedStruct->Status == EUserDefinedStructureStatus::UDSS_Duplicate
 			&& UserDefinedStruct->PrimaryStruct.IsValid())
@@ -194,7 +198,7 @@ bool FInstancedStruct::Serialize(FArchive& Ar)
 		else
 #endif // WITH_EDITOR
 		{
-			Ar << NonConstStruct;
+			Ar << ScriptStruct;
 		}
 	
 		// Size of the serialized memory (reserve location)
@@ -204,9 +208,9 @@ bool FInstancedStruct::Serialize(FArchive& Ar)
 		
 		// Serialized memory
 		const int64 InitialOffset = Ar.Tell(); // Position before struct serialization to compute its serial size
-		if (NonConstStruct != nullptr && ensureMsgf(GetMutableMemory() != nullptr, TEXT("A valid script struct should always have allocated memory")))
+		if (ScriptStruct != nullptr && ensureMsgf(GetMutableMemory() != nullptr, TEXT("A valid script struct should always have allocated memory")))
 		{
-			NonConstStruct->SerializeItem(Ar, GetMutableMemory(), /* Defaults */ nullptr);
+			ConstCast(ScriptStruct)->SerializeItem(Ar, GetMutableMemory(), /* Defaults */ nullptr);
 		}
 		const int64 FinalOffset = Ar.Tell(); // Keep current offset to reset the archive pos after write the serial size
 
@@ -219,12 +223,12 @@ bool FInstancedStruct::Serialize(FArchive& Ar)
 	else if (Ar.IsCountingMemory() || Ar.IsModifyingWeakAndStrongReferences() || Ar.IsObjectReferenceCollector())
 	{
 		// Report type
-		Ar << NonConstStruct;
+		Ar << ScriptStruct;
 	
 		// Report value
-		if (NonConstStruct != nullptr && ensureMsgf(GetMutableMemory() != nullptr, TEXT("A valid script struct should always have allocated memory")))
+		if (ScriptStruct != nullptr && ensureMsgf(GetMutableMemory() != nullptr, TEXT("A valid script struct should always have allocated memory")))
 		{
-			NonConstStruct->SerializeItem(Ar, GetMutableMemory(), /* Defaults */ nullptr);
+			ConstCast(ScriptStruct)->SerializeItem(Ar, GetMutableMemory(), /* Defaults */ nullptr);
 		}
 	}
 
@@ -367,9 +371,9 @@ bool FInstancedStruct::SerializeFromMismatchedTag(const FPropertyTag& Tag, FStru
 
 void FInstancedStruct::GetPreloadDependencies(TArray<UObject*>& OutDeps)
 {
-	if (UScriptStruct* NonConstStruct = const_cast<UScriptStruct*>(GetScriptStruct()))
+	if (ScriptStruct)
 	{
-		OutDeps.Add(NonConstStruct);
+		OutDeps.Add(ConstCast(ScriptStruct));
 
 		// Report direct dependencies of the instanced struct
 		if (UScriptStruct::ICppStructOps* CppStructOps = GetScriptStruct()->GetCppStructOps())
@@ -479,14 +483,14 @@ bool FInstancedStruct::NetSerialize(FArchive& Ar, UPackageMap* Map, bool& bOutSu
 	{
 		if (Ar.IsLoading())
 		{
-			UScriptStruct* NonConstStruct = nullptr;
+			UScriptStruct* SerializedScriptStruct = nullptr;
 
-			Ar << NonConstStruct;
+			Ar << SerializedScriptStruct;
 
 			// Initialize only if the type changes.
-			if (ScriptStruct != NonConstStruct)
+			if (ScriptStruct != SerializedScriptStruct)
 			{
-				InitializeAs(NonConstStruct);
+				InitializeAs(SerializedScriptStruct);
 			}
 
 			if (!IsValid())
diff --git a/Engine/Source/Runtime/CoreUObject/Private/StructUtils/InstancedStructContainer.cpp b/Engine/Source/Runtime/CoreUObject/Private/StructUtils/InstancedStructContainer.cpp
--- a/Engine/Source/Runtime/CoreUObject/Private/StructUtils/InstancedStructContainer.cpp
+++ b/Engine/Source/Runtime/CoreUObject/Private/StructUtils/InstancedStructContainer.cpp
@@ -535,13 +535,13 @@ bool FInstancedStructContainer::Serialize(FArchive& Ar)
 			Views.Reserve(NumItemsSerialized);
 			for (int32 Index = 0; Index < NumItemsSerialized; Index++)
 			{
-				UScriptStruct* NonConstStruct = nullptr;
-				Ar << NonConstStruct;
-				if (NonConstStruct)
+				UScriptStruct* SerializedScriptStruct = nullptr;
+				Ar << SerializedScriptStruct;
+				if (SerializedScriptStruct)
 				{
-					Ar.Preload(NonConstStruct);
+					Ar.Preload(SerializedScriptStruct);
 				}
-				Views.Add(FConstStructView(NonConstStruct, nullptr));
+				Views.Add(FConstStructView(SerializedScriptStruct, nullptr));
 			}
 
 			// Allocate and init items.
@@ -554,16 +554,13 @@ bool FInstancedStructContainer::Serialize(FArchive& Ar)
 			for (int32 Index = 0; Index < NumItemsSerialized; Index++)
 			{
 				const FItem& Item = GetItem(Index);
-				auto& NonConstStruct = ConstCast(Item.ScriptStruct);
 
-				check(NonConstStruct == Item.ScriptStruct);
-				
 				// Size of the serialized memory
 				int32 SerialSize = 0; 
 				Ar << SerialSize;
 
 				// Serialized memory
-				if (NonConstStruct == nullptr && SerialSize > 0)
+				if (Item.ScriptStruct == nullptr && SerialSize > 0)
 				{
 					// A null struct indicates an old struct or an unsupported one for the current target.
 					// In this case we manually seek in the archive to skip its serialized content. 
@@ -572,9 +569,9 @@ bool FInstancedStructContainer::Serialize(FArchive& Ar)
 					UE_LOG(LogCore, Warning, TEXT("Unable to find serialized UScriptStruct -> Advance %u bytes in the archive and reset to empty FInstancedStructContainer. LinkerRoot:%s"), SerialSize, Ar.GetLinker() ? *GetPathNameSafe(Ar.GetLinker()->LinkerRoot) : TEXT("NoLinker"));
 					Ar.Seek(Ar.Tell() + SerialSize);
 				}
-				else if (NonConstStruct != nullptr)
+				else if (Item.ScriptStruct != nullptr)
 				{
-					NonConstStruct->SerializeItem(Ar, Memory + Item.Offset, /* Defaults */ nullptr);
+					ConstCast(Item.ScriptStruct)->SerializeItem(Ar, Memory + Item.Offset, /* Defaults */ nullptr);
 				}
 			}
 		}
@@ -607,7 +604,6 @@ bool FInstancedStructContainer::Serialize(FArchive& Ar)
 			for (int32 Index = 0; Index < NumItemsSerialized; Index++)
 			{
 				const FItem& Item = GetItem(Index);
-				auto& NonConstStruct = ConstCast(Item.ScriptStruct);
 
 				// Size of the serialized memory (reserve location)
 				const int64 SizeOffset = Ar.Tell(); // Position to write the actual size after struct serialization
@@ -616,9 +612,9 @@ bool FInstancedStructContainer::Serialize(FArchive& Ar)
 				
 				// Serialized memory
 				const int64 InitialOffset = Ar.Tell(); // Position before struct serialization to compute its serial size
-				if (NonConstStruct != nullptr)
+				if (Item.ScriptStruct != nullptr)
 				{
-					NonConstStruct->SerializeItem(Ar, Memory + Item.Offset, /* Defaults */ nullptr);
+					ConstCast(Item.ScriptStruct)->SerializeItem(Ar, Memory + Item.Offset, /* Defaults */ nullptr);
 				}
 				const int64 FinalOffset = Ar.Tell(); // Keep current offset to reset the archive pos after write the serial size
 
@@ -642,10 +638,9 @@ bool FInstancedStructContainer::Serialize(FArchive& Ar)
 			for (int32 Index = 0; Index < NumItemsSerialized; Index++)
 			{
 				const FItem& Item = GetItem(Index);
-				auto& NonConstStruct = ConstCast(Item.ScriptStruct);
-				if (NonConstStruct != nullptr)
+				if (Item.ScriptStruct != nullptr)
 				{
-					NonConstStruct->SerializeItem(Ar, Memory + Item.Offset, /* Defaults */ nullptr);
+					ConstCast(Item.ScriptStruct)->SerializeItem(Ar, Memory + Item.Offset, /* Defaults */ nullptr);
 				}
 			}
 		}
@@ -659,9 +654,9 @@ void FInstancedStructContainer::GetPreloadDependencies(TArray<UObject*>& OutDeps
 	for (int32 Index = 0; Index < NumItems; Index++)
 	{
 		FItem& Item = GetItem(Index);
-		if (auto& NonConstStruct = ConstCast(Item.ScriptStruct))
+		if (Item.ScriptStruct)
 		{
-			OutDeps.Add(NonConstStruct);
+			OutDeps.Add(ConstCast(Item.ScriptStruct));
 
 			// Report direct dependencies of the instanced struct
 			if (UScriptStruct::ICppStructOps* CppStructOps = Item.ScriptStruct->GetCppStructOps())
