Improve error message for when we attempt to get the annotation of an invalid UObject. In some cases users are passing in old valid objects with no entry in the GUObjectArray, so it's helpful to provide the object path when possible.

The check is done in two steps to avoid potentially accessing a corrupt ptr when logging which would make us crash before we know why from the log file along what went wrong. Since failing the check in the original code would exit the application, this best-effort logging seems like an improvement for the cases where users have valid objects at the slight quality of life reduction by making the log contain a crash accessing invalid mem in rare cases where we pass in bad ptrs to GetAnnotation.

#rnx

diff --git a/Engine/Source/Runtime/CoreUObject/Public/UObject/UObjectAnnotation.h b/Engine/Source/Runtime/CoreUObject/Public/UObject/UObjectAnnotation.h
--- a/Engine/Source/Runtime/CoreUObject/Public/UObject/UObjectAnnotation.h
+++ b/Engine/Source/Runtime/CoreUObject/Public/UObject/UObjectAnnotation.h
@@ -8,6 +8,7 @@
 
 #include "AutoRTFM/AutoRTFM.h"
 #include "UObject/UObjectArray.h"
+#include "UObject/UObjectBaseUtility.h"
 #include "Misc/ScopeLock.h"
 #include "Misc/TransactionallySafeCriticalSection.h"
 #include "Misc/TransactionallySafeRWLock.h"
@@ -1206,7 +1207,13 @@ public:
 	FORCEINLINE TAnnotation GetAnnotation(const UObjectBase *Object)
 	{
 		check(Object);
-		return GetAnnotation(GUObjectArray.ObjectToIndex(Object));
+		int32 Index = GUObjectArray.ObjectToIndex(Object);
+		// Since we are potentially dealing with a corrupt ptr validation is broken into two checks.
+		// First to tell if something is wrong and the second to try to provide
+		// more information to the user (however in a potentially unsafe way by reading the ptr).
+		ensureAlwaysMsgf(Index >= 0, TEXT("Corrupt object with negative Index passed to GetAnnotation: Index == %d"), Index);
+		checkf(Index >= 0, TEXT("Failed GetAnnotation for Object '%s'"), *static_cast<const UObjectBaseUtility*>(Object)->GetPathName());
+		return GetAnnotation(Index);
 	}
 
 	/**
