Networking
* Replicated UProperties no longer need to register themselves via DOREPLIFETIME in GetLifetimeReplicatedProps.
* Now unregistered properties will be replicated using default settings.
* This simulates the default behavior when Iris is enabled
* Revert to the previous behavior of not replicating missing DOREPLIFETIME by setting Net.AutoRegisterReplicatedProperties=false
* You can now also ensure when unregistered properties are detected with Net.EnsureOnMissingReplicatedPropertiesRegister=true

diff --git a/Engine/Source/Runtime/CoreUObject/Private/UObject/Class.cpp b/Engine/Source/Runtime/CoreUObject/Private/UObject/Class.cpp
--- a/Engine/Source/Runtime/CoreUObject/Private/UObject/Class.cpp
+++ b/Engine/Source/Runtime/CoreUObject/Private/UObject/Class.cpp
@@ -5323,184 +5323,127 @@ void UClass::Link(FArchive& Ar, bool bRelinkExistingProperties)
 	Super::Link(Ar, bRelinkExistingProperties);
 }
 
+namespace UE::Net::Private
+{
 #if (UE_BUILD_SHIPPING)
-static int32 GValidateReplicatedProperties = 0;
+	static bool GValidateReplicatedProperties = false;
 #else 
-static int32 GValidateReplicatedProperties = 1;
+	static bool GValidateReplicatedProperties = true;
 #endif
 
-static FAutoConsoleVariable CVarValidateReplicatedPropertyRegistration(TEXT("net.ValidateReplicatedPropertyRegistration"), GValidateReplicatedProperties, TEXT("Warns if replicated properties were not registered in GetLifetimeReplicatedProps."));
+static FAutoConsoleVariable CVarValidateReplicatedPropertyRegistration(TEXT("net.ValidateReplicatedPropertyRegistration"), GValidateReplicatedProperties, TEXT("Warns if NetDeltaSerialize properties are not top-level"));
+
+void ValidateNetDeltaSerializeProperties(UClass* ClassToValidate);
+
+} // end namespace UE::Net::Private
 
 void UClass::SetUpRuntimeReplicationData()
 {
-	if (!HasAnyClassFlags(CLASS_ReplicationDataIsSetUp))
+	// Nothing to do if we already built the replication data
+	if (HasAnyClassFlags(CLASS_ReplicationDataIsSetUp))
 	{
-		NetFields.Empty();
+		return;
+	}
+	
+	NetFields.Empty();
 
-		if (UClass* SuperClass = GetSuperClass())
-		{
-			SuperClass->SetUpRuntimeReplicationData();
-			ClassReps = SuperClass->ClassReps;
-			FirstOwnedClassRep = ClassReps.Num();
-		}
-		else
-		{
-			ClassReps.Empty();
-			FirstOwnedClassRep = 0;
-		}
+	if (UClass* SuperClass = GetSuperClass())
+	{
+		SuperClass->SetUpRuntimeReplicationData();
+		ClassReps = SuperClass->ClassReps;
+		FirstOwnedClassRep = ClassReps.Num();
+	}
+	else
+	{
+		ClassReps.Empty();
+		FirstOwnedClassRep = 0;
+	}
 
-		// Track properties so me can ensure they are sorted by offsets at the end
-		TArray<FProperty*> NetProperties;
-		for (TFieldIterator<FField> It(this, EFieldIteratorFlags::ExcludeSuper); It; ++It)
+	// Track properties so me can ensure they are sorted by offsets at the end
+	TArray<FProperty*> NetProperties;
+	for (TFieldIterator<FField> It(this, EFieldIteratorFlags::ExcludeSuper); It; ++It)
+	{
+		if (FProperty* Prop = CastField<FProperty>(*It))
 		{
-			if (FProperty* Prop = CastField<FProperty>(*It))
+			if ((Prop->PropertyFlags & CPF_Net) && Prop->GetOwner<UObject>() == this)
 			{
-				if ((Prop->PropertyFlags & CPF_Net) && Prop->GetOwner<UObject>() == this)
-				{
-					NetProperties.Add(Prop);
-				}
+				NetProperties.Add(Prop);
 			}
 		}
+	}
 
-		for(TFieldIterator<UField> It(this,EFieldIteratorFlags::ExcludeSuper); It; ++It)
+	for(TFieldIterator<UField> It(this,EFieldIteratorFlags::ExcludeSuper); It; ++It)
+	{
+		if (UFunction * Func = Cast<UFunction>(*It))
 		{
-			if (UFunction * Func = Cast<UFunction>(*It))
+			// When loading reflection data (e.g. from blueprints), we may have references to placeholder functions, or reflection data 
+			// in children may be out of date. In that case we cannot enforce this check, but that is ok because reflection data will
+			// be regenerated by compile on load anyway:
+			const bool bCanCheck = (!GIsEditor && !IsRunningCommandlet()) || !Func->HasAnyFlags(RF_WasLoaded);
+			check(!bCanCheck || (!Func->GetSuperFunction() || (Func->GetSuperFunction()->FunctionFlags&FUNC_NetFuncFlags) == (Func->FunctionFlags&FUNC_NetFuncFlags)));
+			if ((Func->FunctionFlags&FUNC_Net) && !Func->GetSuperFunction())
 			{
-				// When loading reflection data (e.g. from blueprints), we may have references to placeholder functions, or reflection data 
-				// in children may be out of date. In that case we cannot enforce this check, but that is ok because reflection data will
-				// be regenerated by compile on load anyway:
-				const bool bCanCheck = (!GIsEditor && !IsRunningCommandlet()) || !Func->HasAnyFlags(RF_WasLoaded);
-				check(!bCanCheck || (!Func->GetSuperFunction() || (Func->GetSuperFunction()->FunctionFlags&FUNC_NetFuncFlags) == (Func->FunctionFlags&FUNC_NetFuncFlags)));
-				if ((Func->FunctionFlags&FUNC_Net) && !Func->GetSuperFunction())
-				{
-					NetFields.Add(Func);
-				}
+				NetFields.Add(Func);
 			}
 		}
+	}
 
-		const bool bIsNativeClass = HasAnyClassFlags(CLASS_Native);
-		if (!bIsNativeClass)
+	const bool bIsNativeClass = HasAnyClassFlags(CLASS_Native);
+	if (!bIsNativeClass)
+	{
+		// Sort NetProperties so that their ClassReps are sorted by memory offset
+		struct FComparePropertyOffsets
 		{
-			// Sort NetProperties so that their ClassReps are sorted by memory offset
-			struct FComparePropertyOffsets
+			FORCEINLINE bool operator()(FProperty* A, FProperty* B) const
 			{
-				FORCEINLINE bool operator()(FProperty* A, FProperty* B) const
+				// Ensure stable sort
+				if (A->GetOffset_ForGC() == B->GetOffset_ForGC())
 				{
-					// Ensure stable sort
-					if (A->GetOffset_ForGC() == B->GetOffset_ForGC())
-					{
-						return A->GetName() < B->GetName();
-					}
-
-					return A->GetOffset_ForGC() < B->GetOffset_ForGC();
+					return A->GetName() < B->GetName();
 				}
-			};
-
-			Algo::Sort(NetProperties, FComparePropertyOffsets());
-		}
 
-		ClassReps.Reserve(ClassReps.Num() + NetProperties.Num());
-		for (int32 i = 0; i < NetProperties.Num(); i++)
-		{
-			NetProperties[i]->RepIndex = (uint16)ClassReps.Num();
-			for (int32 j = 0; j < NetProperties[i]->ArrayDim; j++)
-			{
-				ClassReps.Emplace(NetProperties[i], j);
+				return A->GetOffset_ForGC() < B->GetOffset_ForGC();
 			}
-		}
-		check(ClassReps.Num() <= 65535);
-
-		if (bIsNativeClass && GValidateReplicatedProperties)
-		{
-			GetDefaultObject()->ValidateGeneratedRepEnums(ClassReps);
-		}
-
-		NetFields.Shrink();
-		
-		Algo::SortBy(NetFields, &UField::GetFName, FNameLexicalLess());
+		};
 
-		ClassFlags |= CLASS_ReplicationDataIsSetUp;
+		Algo::Sort(NetProperties, FComparePropertyOffsets());
+	}
 
-		if (GValidateReplicatedProperties != 0)
+	ClassReps.Reserve(ClassReps.Num() + NetProperties.Num());
+	for (int32 i = 0; i < NetProperties.Num(); i++)
+	{
+		NetProperties[i]->RepIndex = (uint16)ClassReps.Num();
+		for (int32 j = 0; j < NetProperties[i]->ArrayDim; j++)
 		{
-			ValidateRuntimeReplicationData();
+			ClassReps.Emplace(NetProperties[i], j);
 		}
 	}
-}
-
-void UClass::ValidateRuntimeReplicationData()
-{
-	DECLARE_SCOPE_CYCLE_COUNTER(TEXT("Class ValidateRuntimeReplicationData"), STAT_Class_ValidateRuntimeReplicationData, STATGROUP_Game);
-
-	bool bCheckLifetimeProps = true;
+	check(ClassReps.Num() <= 65535);
 
-	if (HasAnyClassFlags(CLASS_CompiledFromBlueprint) || bLayoutChanging)
+	if (bIsNativeClass && UE::Net::Private::GValidateReplicatedProperties)
 	{
-		// Blueprint classes don't always generate a GetLifetimeReplicatedProps function. 
-		// Assume the Blueprint compiler was ok to do this.
-		bCheckLifetimeProps = false;
+		GetDefaultObject()->ValidateGeneratedRepEnums(ClassReps);
 	}
 
-	if (bCheckLifetimeProps && HasAnyClassFlags(CLASS_ReplicationDataIsSetUp) == false)
-	{
-		UE_LOG(LogClass, Warning, TEXT("ValidateRuntimeReplicationData for class %s called before ReplicationData was setup."), *GetName());
-		bCheckLifetimeProps = false;
-	}
+	NetFields.Shrink();
+		
+	Algo::SortBy(NetFields, &UField::GetFName, FNameLexicalLess());
 
-	// Our replication data was set up, but there are no class reps, so there's nothing to do.
-	if (ClassReps.Num() == 0)
-	{
-		bCheckLifetimeProps = false;
-	}
+	ClassFlags |= CLASS_ReplicationDataIsSetUp;
 
-	// Let's compare the CDO's registered lifetime properties with the Class's net properties
-	if (bCheckLifetimeProps)
+	if (UE::Net::Private::GValidateReplicatedProperties)
 	{
-		TArray<FLifetimeProperty> LifetimeProps;
-		LifetimeProps.Reserve(ClassReps.Num());
-
-		const UObject* Object = GetDefaultObject();
-		Object->GetLifetimeReplicatedProps(LifetimeProps);
-
-		// Find which properties where not registered by the user code
-		if (LifetimeProps.Num() != ClassReps.Num())
-		{
-			for (int32 RepIndex=0; RepIndex<ClassReps.Num(); ++RepIndex)
-			{
-				const FProperty* RepProp = ClassReps[RepIndex].Property;
-
-				const FLifetimeProperty* LifetimeProp = LifetimeProps.FindByPredicate(
-					[&RepIndex](const FLifetimeProperty& Var) { return Var.RepIndex == RepIndex; });
-
-				if (LifetimeProp == nullptr)
-				{
-					// Check if this unregistered property type uses a custom delta serializer
-					if (const FStructProperty* StructProperty = CastField<FStructProperty>(RepProp))
-					{
-						const UScriptStruct* Struct = StructProperty->Struct;
-
-						if (EnumHasAnyFlags(Struct->StructFlags, STRUCT_NetDeltaSerializeNative))
-						{
-							UE_LOG(LogClass, Warning,
-									TEXT("Property %s::%s (SourceClass: %s) with custom net delta serializer was not registered in "
-									     "GetLifetimeReplicatedProps. This property will replicate but you should still register it."),
-									ToCStr(GetName()), ToCStr(RepProp->GetName()), ToCStr(RepProp->GetOwnerClass()->GetName()));
-
-							continue;
-						}
-					}
-
-					UE_LOG(LogClass, Warning,
-							TEXT("Property %s::%s (SourceClass: %s) was not registered in GetLifetimeReplicatedProps. "
-							     "This property will not be replicated. Use DISABLE_REPLICATED_PROPERTY if not replicating was intentional."),
-							ToCStr(GetName()), ToCStr(RepProp->GetName()), ToCStr(RepProp->GetOwnerClass()->GetName()));
-				}
-			}
-		}
+		UE::Net::Private::ValidateNetDeltaSerializeProperties(this);
 	}
+}
 
+namespace UE::Net::Private
+{
+
+void ValidateNetDeltaSerializeProperties(UClass* ClassToValidate)
+{
 	// Find NetDeltaSerialize properties which are not top-level/parent properties (e.g. nested in another struct)
-	for (const FRepRecord& RepRecord : ClassReps)
+	for (const FRepRecord& RepRecord : ClassToValidate->ClassReps)
 	{
 		const FProperty* BaseProp = RepRecord.Property;
 		TArray<const UStruct*, TInlineAllocator<16>> RecursiveStructList;
@@ -5512,91 +5455,92 @@ void UClass::ValidateRuntimeReplicationData()
 		};
 
 		auto CheckStructRecursive = [&BaseProp, &RecursiveStructList]
-									(auto CheckStructRecursive, const FProperty* InProp, ECheckStructType CheckType)
-									-> void
+		(auto CheckStructRecursive, const FProperty* InProp, ECheckStructType CheckType)
+		-> void
+		{
+			auto BadStruct = [&BaseProp](const FStructProperty* StructProp)
 			{
-				auto BadStruct = [&BaseProp](const FStructProperty* StructProp)
-					{
-						UE_LOG(LogClass, Warning, TEXT("Property %s contained nested NetDeltaSerialize struct '%s', ")
-								TEXT("when this is not supported. Only use NetDeltaSerialize structs at class level. ")
-								TEXT("Struct will replicate using non-delta serialization."),
-								ToCStr(BaseProp->GetPathName()), ToCStr(StructProp->GetPathName()));
-					};
+				UE_LOG(LogClass, Warning, TEXT("Property %s contained nested NetDeltaSerialize struct '%s', ")
+					TEXT("when this is not supported. Only use NetDeltaSerialize structs at class level. ")
+					TEXT("Struct will replicate using non-delta serialization."),
+					ToCStr(BaseProp->GetPathName()), ToCStr(StructProp->GetPathName()));
+			};
 
-				const FStructProperty* ParamStructProp = nullptr;
-				bool bNonDeltaArrayStruct = false;
+			const FStructProperty* ParamStructProp = nullptr;
+			bool bNonDeltaArrayStruct = false;
 
-				if (!EnumHasAnyFlags(InProp->PropertyFlags, CPF_RepSkip))
+			if (!EnumHasAnyFlags(InProp->PropertyFlags, CPF_RepSkip))
+			{
+				if (const FStructProperty* StructProp = CastField<FStructProperty>(InProp))
 				{
-					if (const FStructProperty* StructProp = CastField<FStructProperty>(InProp))
-					{
-						ParamStructProp = StructProp;
-					}
-					else if (const FArrayProperty* ArrayProp = CastField<FArrayProperty>(InProp))
+					ParamStructProp = StructProp;
+				}
+				else if (const FArrayProperty* ArrayProp = CastField<FArrayProperty>(InProp))
+				{
+					if (const FStructProperty* ArrayStructProp = CastField<FStructProperty>(ArrayProp->Inner))
 					{
-						if (const FStructProperty* ArrayStructProp = CastField<FStructProperty>(ArrayProp->Inner))
+						if (EnumHasAnyFlags(ArrayStructProp->Struct->StructFlags, STRUCT_NetDeltaSerializeNative))
 						{
-							if (EnumHasAnyFlags(ArrayStructProp->Struct->StructFlags, STRUCT_NetDeltaSerializeNative))
-							{
-								BadStruct(ArrayStructProp);
-							}
-							else
-							{
-								ParamStructProp = ArrayStructProp;
-								bNonDeltaArrayStruct = true;
-							}
+							BadStruct(ArrayStructProp);
+						}
+						else
+						{
+							ParamStructProp = ArrayStructProp;
+							bNonDeltaArrayStruct = true;
 						}
 					}
 				}
+			}
 
 
-				const bool bTopLevelNonDeltaStruct = ParamStructProp != nullptr && CheckType == ECheckStructType::TopLevel &&
-					(bNonDeltaArrayStruct || !EnumHasAnyFlags(ParamStructProp->Struct->StructFlags, STRUCT_NetDeltaSerializeNative));
+			const bool bTopLevelNonDeltaStruct = ParamStructProp != nullptr && CheckType == ECheckStructType::TopLevel &&
+			(bNonDeltaArrayStruct || !EnumHasAnyFlags(ParamStructProp->Struct->StructFlags, STRUCT_NetDeltaSerializeNative));
 
-				const bool bFreshRecursiveStruct = ParamStructProp != nullptr && CheckType == ECheckStructType::Recursive &&
-													!RecursiveStructList.Contains(ParamStructProp->Struct);
+			const bool bFreshRecursiveStruct = ParamStructProp != nullptr && CheckType == ECheckStructType::Recursive &&
+			!RecursiveStructList.Contains(ParamStructProp->Struct);
 
-				if (bTopLevelNonDeltaStruct || bFreshRecursiveStruct || bNonDeltaArrayStruct)
-				{
-					RecursiveStructList.Add(ParamStructProp->Struct);
+			if (bTopLevelNonDeltaStruct || bFreshRecursiveStruct || bNonDeltaArrayStruct)
+			{
+				RecursiveStructList.Add(ParamStructProp->Struct);
 
-					for (TFieldIterator<FProperty> It(ParamStructProp->Struct); It; ++It)
+				for (TFieldIterator<FProperty> It(ParamStructProp->Struct); It; ++It)
+				{
+					if (!EnumHasAnyFlags(It->PropertyFlags, CPF_RepSkip))
 					{
-						if (!EnumHasAnyFlags(It->PropertyFlags, CPF_RepSkip))
-						{
-							const FStructProperty* CurLevelStructProp = nullptr;
+						const FStructProperty* CurLevelStructProp = nullptr;
 
-							if (const FStructProperty* StructProp = CastField<FStructProperty>(*It))
+						if (const FStructProperty* StructProp = CastField<FStructProperty>(*It))
+						{
+							CurLevelStructProp = StructProp;
+						}
+						else if (const FArrayProperty* ArrayProp = CastField<FArrayProperty>(*It))
+						{
+							if (const FStructProperty* ArrayStructProp = CastField<FStructProperty>(ArrayProp->Inner))
 							{
-								CurLevelStructProp = StructProp;
+								CurLevelStructProp = ArrayStructProp;
 							}
-							else if (const FArrayProperty* ArrayProp = CastField<FArrayProperty>(*It))
+						}
+
+						if (CurLevelStructProp != nullptr)
+						{
+							if (EnumHasAnyFlags(CurLevelStructProp->Struct->StructFlags, STRUCT_NetDeltaSerializeNative))
 							{
-								if (const FStructProperty* ArrayStructProp = CastField<FStructProperty>(ArrayProp->Inner))
-								{
-									CurLevelStructProp = ArrayStructProp;
-								}
+								BadStruct(CurLevelStructProp);
 							}
-
-							if (CurLevelStructProp != nullptr)
+							else
 							{
-								if (EnumHasAnyFlags(CurLevelStructProp->Struct->StructFlags, STRUCT_NetDeltaSerializeNative))
-								{
-									BadStruct(CurLevelStructProp);
-								}
-								else
-								{
-									CheckStructRecursive(CheckStructRecursive, CurLevelStructProp, ECheckStructType::Recursive);
-								}
+								CheckStructRecursive(CheckStructRecursive, CurLevelStructProp, ECheckStructType::Recursive);
 							}
 						}
 					}
 				}
-			};
+			}
+		};
 
 		CheckStructRecursive(CheckStructRecursive, BaseProp, ECheckStructType::TopLevel);
 	}
 }
+} // end namespace UE::Net::Private
 
 void UClass::InternalCreateDefaultObjectWrapper() const
 {
diff --git a/Engine/Source/Runtime/CoreUObject/Public/UObject/Class.h b/Engine/Source/Runtime/CoreUObject/Public/UObject/Class.h
--- a/Engine/Source/Runtime/CoreUObject/Public/UObject/Class.h
+++ b/Engine/Source/Runtime/CoreUObject/Public/UObject/Class.h
@@ -3780,7 +3780,6 @@ private:
 	/**
 	 * Tests if all properties tagged with Replicate were registered in GetLifetimeReplicatedProps
 	 */
-	COREUOBJECT_API void ValidateRuntimeReplicationData();
 	COREUOBJECT_API void AssembleReferenceTokenStreamInternal(bool bForce = false);
 	COREUOBJECT_API void InternalCreateDefaultObjectWrapper() const;
 protected:
diff --git a/Engine/Source/Runtime/Engine/Private/RepLayout.cpp b/Engine/Source/Runtime/Engine/Private/RepLayout.cpp
--- a/Engine/Source/Runtime/Engine/Private/RepLayout.cpp
+++ b/Engine/Source/Runtime/Engine/Private/RepLayout.cpp
@@ -7,23 +7,30 @@
 #include "Net/RepLayout.h"
 #include "Containers/StaticBitArray.h"
 #include "Misc/MemStack.h"
-#include "Net/Core/PropertyConditions/RepChangedPropertyTracker.h"
+
+#include "Net/Core/NetBitArray.h"
+#include "Net/Core/NetCoreModule.h"
+#include "Net/Core/Misc/NetCVars.h"
 #include "Net/Core/NetToken/NetTokenExportContext.h"
-#include "Net/Serialization/FastArraySerializer.h"
+#include "Net/Core/PropertyConditions/RepChangedPropertyTracker.h"
+#include "Net/Core/PushModel/PushModel.h"
+#include "Net/Core/PushModel/PushModelMacros.h"
+#include "Net/Core/PushModel/Types/PushModelPerNetDriverState.h"
+#include "Net/Core/Trace/NetTrace.h"
+#include "Net/Core/Trace/Private/NetTraceInternal.h"
+
 #include "EngineStats.h"
 #include "Engine/PackageMapClient.h"
 #include "Engine/NetConnection.h"
-#include "Net/Core/PushModel/PushModel.h"
-#include "Net/Core/PushModel/PushModelMacros.h"
+
+#include "Net/Serialization/FastArraySerializer.h"
 #include "Net/NetworkProfiler.h"
 #include "Engine/ActorChannel.h"
 #include "Misc/App.h"
-#include "Net/Core/Trace/Private/NetTraceInternal.h"
+
 #include "Net/NetworkGranularMemoryLogging.h"
 #include "Serialization/ArchiveCountMem.h"
-#include "Net/Core/PushModel/Types/PushModelPerNetDriverState.h"
-#include "Net/Core/Trace/NetTrace.h"
-#include "Net/Core/NetCoreModule.h"
+
 #include "Stats/StatsTrace.h"
 #include "UObject/EnumProperty.h"
 #if UE_WITH_IRIS
@@ -102,18 +109,11 @@ namespace UE::Net::Private
 	static bool bDeltaInitialFastArrayElements = false;
 	static FAutoConsoleVariableRef CVarDeltaInitialFastArrayElements(TEXT("net.DeltaInitialFastArrayElements"), bDeltaInitialFastArrayElements, TEXT("If true, send delta struct changelists for initial fast array elements."));
 
-	/* FastArrays and other custom delta properties may have order dependencies due to callbacks being fired during serialization at which time other custom delta properties have not yet received their state.
-	 * This cvar toggles the behavior between using the RepIndex of the property or the order of appearance in the lifetime property array filled during a GetLifetimeReplicatedProps() call.
-	 * Default is false to keep the legacy behavior of using the GetLifetimeReplicatedProps() order for the custom delta properties.
-	 * The cvar is used in ReplicationStateDescriptorBuilder as well. Search for the cvar name in the code base before removing it.
-	 */
-	static bool bReplicateCustomDeltaPropertiesInRepIndexOrder = false;
-	static FAutoConsoleVariableRef CVarReplicateCustomDeltaPropertiesInRepIndexOrder(TEXT("net.ReplicateCustomDeltaPropertiesInRepIndexOrder"), bReplicateCustomDeltaPropertiesInRepIndexOrder, TEXT("If false (default) custom delta properties will replicate in the same order as they're added to the lifetime property array during the call to GetLifetimeReplicatedProps. If true custom delta properties will be replicated in the property RepIndex order, which is typically in increasing property offset order. Note that custom delta properties are always serialized after regular properties."));
-
 	static bool bAlwaysUpdateGuidReferenceMapForNetSerializeObjectStruct = false;
 	static FAutoConsoleVariableRef CVarAlwaysUpdateGuidReferenceMapForNetSerializeStruct(TEXT("net.AlwaysUpdateGuidReferenceMapForNetSerializeObjectStruct"), bAlwaysUpdateGuidReferenceMapForNetSerializeObjectStruct,
 		TEXT("Requires net.TrackNetSerializeObjectReferences. If true, entries in the GuidReferenceMap for NetSerialize struct properties with object properties will always be updated, not just when the Guid changes or goes NULL. This should prevent issues with old property data being applied when an unmapped actor ref in the struct is mapped."));
-}
+
+} // end namespace UE::Net::Private
 
 extern int32 GNumSharedSerializationHit;
 extern int32 GNumSharedSerializationMiss;
@@ -6176,16 +6176,44 @@ void FRepLayout::InitFromClass(
 		NetSerializeLayouts.Shrink();
 	}
 
-	// Initialize lifetime props
-	// Properties that replicate for the lifetime of the channel
+	// Initialize replicated properties
+
+	// The list of all replicated properties
 	TArray<FLifetimeProperty> LifetimeProps;
 	LifetimeProps.Reserve(Parents.Num());
-
+	
+	// Ask the user code for explicit settings
 	UObject* Object = InObjectClass->GetDefaultObject();
-
 	Object->GetLifetimeReplicatedProps(LifetimeProps);
+
+	// Set default settings for properties not explicitely registered
+	{
+		UE::Net::FNetBitArray RegisteredProperties(InObjectClass->ClassReps.Num());
+		for (const FLifetimeProperty& UserProp : LifetimeProps)
+		{
+			RegisteredProperties.SetBit(UserProp.RepIndex);
+		}
+
+		RegisteredProperties.ForAllUnsetBits([&LifetimeProps, InObjectClass](uint32 Index)
+		{
+			const uint32 PropIndex = LifetimeProps.Emplace(FLifetimeProperty(Index));
+
+			// Check if the registration is missing if we don't automatically replicate unregistered variables
+			if (!UE::Net::bAutoRegisterReplicatedProperties)
+			{
+				LifetimeProps[PropIndex].Condition = COND_Never;
+
+				if (UE::Net::bEnsureForMissingProperties)
+				{
+					ensureMsgf(false, TEXT("Property %s::%s (SourceClass: %s) was not registered in GetLifetimeReplicatedProps. This property will not be replicated. Use DISABLE_REPLICATED_PROPERTY if not replicating was intentional."), 
+							   *GetNameSafe(InObjectClass), *GetNameSafe(InObjectClass->ClassReps[Index].Property), *GetNameSafe(InObjectClass->ClassReps[Index].Property->GetOwnerClass()));
+				}
+			}
+		});
+	}
+
 	// If there are custom delta properties we may have to change the order we traverse the replicated props.
-	if (UE::Net::Private::bReplicateCustomDeltaPropertiesInRepIndexOrder && (HighestCustomDeltaRepIndex != INDEX_NONE))
+	if (UE::Net::bReplicateCustomDeltaPropertiesInRepIndexOrder && (HighestCustomDeltaRepIndex != INDEX_NONE))
 	{
 		Algo::SortBy(LifetimeProps, [](const FLifetimeProperty& Element) { return Element.RepIndex; }, TLess<decltype(FLifetimeProperty::RepIndex)>());
 	}
@@ -6201,57 +6229,76 @@ void FRepLayout::InitFromClass(
 	int32 NumberOfFastArrayProperties = 0;
 	int32 NumberOfFastArrayPushModelProperties = 0;
 
+	// Use default settings when auto-registering properties
+	FLifetimeProperty DefaultLifetimeProp;
+
 	// Setup lifetime replicated properties
-	for (int32 i = 0; i < LifetimeProps.Num(); i++)
+	for (int32 ParentIndex=0; ParentIndex < InObjectClass->ClassReps.Num(); ++ParentIndex)
 	{
-		const int32 ParentIndex = LifetimeProps[i].RepIndex;
-
-		if (!ensureMsgf(Parents.IsValidIndex(ParentIndex), TEXT("Parents array index %d out of bounds! i = %d, LifetimeProps.Num() = %d, Parents.Num() = %d, InObjectClass = %s"),
-				ParentIndex, i, LifetimeProps.Num(), Parents.Num(), *GetFullNameSafe(InObjectClass)))
+		const FRepRecord& ReplicatedProp = InObjectClass->ClassReps[ParentIndex];
+	
+		const bool bIsValid = ensureMsgf(Parents.IsValidIndex(ParentIndex), TEXT("ClassReps RepRecord of %s point to invalid RepIndex=%i for replicated property: %s. Parents.Num() = %d"), *GetFullNameSafe(InObjectClass), ParentIndex, *GetNameSafe(ReplicatedProp.Property), Parents.Num());
+		if (!bIsValid)
 		{
 			continue;
 		}
 
+		FRepParentCmd& RepParentCmd = Parents[ParentIndex];
+
+		// Find a lifetime condition defined for this replicated property
+		const FLifetimeProperty* UserLifetimeProp = LifetimeProps.FindByPredicate([ParentIndex](const FLifetimeProperty& Item)
+		{
+			return Item.RepIndex == ParentIndex;
+		});
+
+		
+
+		// Set the correct RepIndex just in case its refereced.
+		DefaultLifetimeProp.RepIndex = ParentIndex;
+
+		// Use Default settings if the user didn't 
+		const FLifetimeProperty& LifetimeProp = UserLifetimeProp ? *UserLifetimeProp : DefaultLifetimeProp;
+
 		// Don't bother doing any setup work for COND_Never properties.
 		// These are never expected to replicate.
-		if (COND_Never == LifetimeProps[i].Condition)
+		if (COND_Never == LifetimeProp.Condition)
 		{
 			continue;
 		}
 
 		// Store the condition on the parent in case we need it
-		Parents[ParentIndex].Condition = LifetimeProps[i].Condition;
-		Parents[ParentIndex].RepNotifyCondition = LifetimeProps[i].RepNotifyCondition;
+		RepParentCmd.Condition = LifetimeProp.Condition;
+		RepParentCmd.RepNotifyCondition = LifetimeProp.RepNotifyCondition;
 
-		if (UFunction* RepNotifyFunc = InObjectClass->FindFunctionByName(Parents[ParentIndex].Property->RepNotifyFunc))
+		if (UFunction* RepNotifyFunc = InObjectClass->FindFunctionByName(RepParentCmd.Property->RepNotifyFunc))
 		{
-			Parents[ParentIndex].RepNotifyNumParams = RepNotifyFunc->NumParms;
+			RepParentCmd.RepNotifyNumParams = RepNotifyFunc->NumParms;
 		}
 
-		if (!EnumHasAnyFlags(Parents[ParentIndex].Flags, ERepParentFlags::IsCustomDelta))
+		if (!EnumHasAnyFlags(RepParentCmd.Flags, ERepParentFlags::IsCustomDelta))
 		{
-			Parents[ParentIndex].Flags |= ERepParentFlags::IsLifetime;
-			if (LifetimeProps[i].Condition == COND_None)
+			RepParentCmd.Flags |= ERepParentFlags::IsLifetime;
+			if (RepParentCmd.Condition == COND_None)
 			{
-				Parents[ParentIndex].Flags &= ~ERepParentFlags::IsConditional;
+				RepParentCmd.Flags &= ~ERepParentFlags::IsConditional;
 			}
-			else if (LifetimeProps[i].Condition == COND_InitialOnly)
+			else if (RepParentCmd.Condition == COND_InitialOnly)
 			{
 				Flags |= ERepLayoutFlags::HasInitialOnlyProperties;
 			}
-			else if (LifetimeProps[i].Condition == COND_Dynamic)
+			else if (RepParentCmd.Condition == COND_Dynamic)
 			{
 				Flags |= ERepLayoutFlags::HasDynamicConditionProperties;
 			}
 
-			if (EnumHasAnyFlags(Parents[ParentIndex].Flags, ERepParentFlags::HasNetSerializeProperties | ERepParentFlags::HasObjectProperties))
+			if (EnumHasAnyFlags(RepParentCmd.Flags, ERepParentFlags::HasNetSerializeProperties | ERepParentFlags::HasObjectProperties))
 			{
 				Flags |= ERepLayoutFlags::HasObjectOrNetSerializeProperties;
 			}
 
 			++NumberOfLifetimeProperties;
 #if WITH_PUSH_MODEL
-			if (bIsPushModelEnabled && LifetimeProps[i].bIsPushBased)
+			if (bIsPushModelEnabled && LifetimeProp.bIsPushBased)
 			{
 				++NumberOfLifetimePushModelProperties;
 				PushModelProperties[ParentIndex] = true;
@@ -6291,11 +6338,11 @@ void FRepLayout::InitFromClass(
 
 			bool bAddedFastArray = false;
 
-			if (EnumHasAnyFlags(Parents[ParentIndex].Flags, ERepParentFlags::IsFastArray))
+			if (EnumHasAnyFlags(RepParentCmd.Flags, ERepParentFlags::IsFastArray))
 			{
 				int32 FastArrayItemArrayCmd = INDEX_NONE;
 
-				for (int32 CmdIndex = Parents[ParentIndex].CmdStart; CmdIndex < Parents[ParentIndex].CmdEnd; ++CmdIndex)
+				for (int32 CmdIndex = RepParentCmd.CmdStart; CmdIndex < RepParentCmd.CmdEnd; ++CmdIndex)
 				{
 					const FRepLayoutCmd& Cmd = Cmds[CmdIndex];
 					if (ERepLayoutCmdType::DynamicArray == Cmd.Type)
@@ -6324,7 +6371,7 @@ void FRepLayout::InitFromClass(
 
 								++NumberOfFastArrayProperties;
 #if WITH_PUSH_MODEL
-								if (bIsPushModelEnabled && LifetimeProps[i].bIsPushBased)
+								if (bIsPushModelEnabled && LifetimeProp.bIsPushBased)
 								{
 									++NumberOfFastArrayPushModelProperties;
 									PushModelProperties[ParentIndex] = true;
@@ -6341,7 +6388,7 @@ void FRepLayout::InitFromClass(
 
 				if (!bAddedFastArray)
 				{
-					UE_LOG(LogRep, Warning, TEXT("FRepLayout::InitFromClass: Unable to find Fast Array Item array in Fast Array Serializer: %s"), *Parents[ParentIndex].CachedPropertyName.ToString());
+					UE_LOG(LogRep, Warning, TEXT("FRepLayout::InitFromClass: Unable to find Fast Array Item array in Fast Array Serializer: %s"), *RepParentCmd.CachedPropertyName.ToString());
 				}
 			}
 
diff --git a/Engine/Source/Runtime/Engine/Public/Tests/NetTestHelpers.h b/Engine/Source/Runtime/Engine/Public/Tests/NetTestHelpers.h
--- a/Engine/Source/Runtime/Engine/Public/Tests/NetTestHelpers.h
+++ b/Engine/Source/Runtime/Engine/Public/Tests/NetTestHelpers.h
@@ -5,7 +5,8 @@
 #include "HAL/Platform.h"
 #include "Engine/NetDriver.h"
 
-class APlayerController;
+#include "GameFramework/PlayerController.h"
+
 class IConsoleVariable;
 class UGameInstance;
 class UWorld;
diff --git a/Engine/Source/Runtime/Experimental/Iris/Core/Private/Iris/ReplicationState/ReplicationStateDescriptorBuilder.cpp b/Engine/Source/Runtime/Experimental/Iris/Core/Private/Iris/ReplicationState/ReplicationStateDescriptorBuilder.cpp
--- a/Engine/Source/Runtime/Experimental/Iris/Core/Private/Iris/ReplicationState/ReplicationStateDescriptorBuilder.cpp
+++ b/Engine/Source/Runtime/Experimental/Iris/Core/Private/Iris/ReplicationState/ReplicationStateDescriptorBuilder.cpp
@@ -7,6 +7,7 @@
 #include "Iris/Core/IrisProfiler.h"
 #include "Iris/Core/MemoryLayoutUtil.h"
 #include "Net/Core/NetBitArray.h"
+#include "Net/Core/Misc/NetCVars.h"
 #include "Net/Core/Trace/NetDebugName.h"
 #include "Iris/ReplicationState/PropertyNetSerializerInfoRegistry.h"
 #include "Iris/ReplicationState/InternalPropertyReplicationState.h"
@@ -1995,7 +1996,7 @@ bool FPropertyReplicationStateDescriptorBuilder::IsSupportedProperty(FMemberProp
 	OutMemberProperty.ExternalSizeAndAlignment = { 0, 0 };
 
 	// Init traits
-	GetIrisPropertyTraits(OutMemberProperty, Property, Parameters. LifeTimeProperties, Parameters.InObjectClass);
+	GetIrisPropertyTraits(OutMemberProperty, Property, Parameters.LifeTimeProperties, Parameters.InObjectClass);
 	GetSerializerTraits(OutMemberProperty, Property, NetSerializerInfo, Parameters.bAllowFastArrayWithExtraReplicatedProperties);
 
 	// $IRIS: $TODO: Add property attribute to control how many changemask bits we want to use per property
@@ -2324,8 +2325,6 @@ EMemberPropertyTraits FPropertyReplicationStateDescriptorBuilder::GetFastArrayPr
 namespace UE::Net
 {
 
-const IConsoleVariable* FReplicationStateDescriptorBuilder::CVarReplicateCustomDeltaPropertiesInRepIndexOrder = nullptr;
-
 FReplicationStateDescriptorBuilder::FParameters::FParameters()
 : DescriptorRegistry(nullptr),
   ReplicationSystem(nullptr),
@@ -2636,18 +2635,38 @@ SIZE_T FReplicationStateDescriptorBuilder::CreateDescriptorsForClass(FResult& Cr
 	{
 		IRIS_PROFILER_SCOPE(FReplicationStateDescriptorBuilder_GetLifeTimeProperties);
 
+		// Keep track of non user registered properties
+		FNetBitArray RegisteredProperty(InObjectClass->ClassReps.Num());
+
 		LifeTimePropertiesOriginalOrder.Reserve(InObjectClass->ClassReps.Num());
 		UObject* Object = InObjectClass->GetDefaultObject();
 		Object->GetLifetimeReplicatedProps(LifeTimePropertiesOriginalOrder);
 
 		// Lifetime properties aren't sorted and searching for a property with a specific RepIndex in the array is linear.
 		// Let's put them where we want them.
-		static_assert(TIsZeroConstructType<FLifetimeProperty>::Value, "Need to use SetNum() instead of SetNumZeroed()");
-		LifeTimeProperties.SetNumZeroed(InObjectClass->ClassReps.Num());
+		LifeTimeProperties.SetNum(InObjectClass->ClassReps.Num());
 		for (const FLifetimeProperty& LifetimeProperty : LifeTimePropertiesOriginalOrder)
 		{
 			LifeTimeProperties[LifetimeProperty.RepIndex] = LifetimeProperty;
+			RegisteredProperty.SetBit(LifetimeProperty.RepIndex);
 		}
+
+		RegisteredProperty.ForAllUnsetBits([&LifeTimeProperties, InObjectClass](uint32 Index)
+		{
+			LifeTimeProperties[Index].RepIndex = static_cast<uint16>(Index);
+
+			if (UE::Net::bAutoRegisterReplicatedProperties == false)
+			{
+				// Never replicate this property when auto-registration is disabled
+				LifeTimeProperties[Index].Condition = COND_Never;
+
+				if (UE::Net::bEnsureForMissingProperties)
+				{
+					ensureMsgf(false, TEXT("Property %s::%s (SourceClass: %s) was not registered in GetLifetimeReplicatedProps. This property will not be replicated. Use DISABLE_REPLICATED_PROPERTY if not replicating was intentional."), 
+							   *GetNameSafe(InObjectClass), *GetNameSafe(InObjectClass->ClassReps[Index].Property), *GetNameSafe(InObjectClass->ClassReps[Index].Property->GetOwnerClass()));
+				}
+			}
+		});
 	}
 
 	// For a UClass we split up our state into separate replication states based on common traits (Conditionals)
@@ -2827,10 +2846,8 @@ SIZE_T FReplicationStateDescriptorBuilder::CreateDescriptorsForClass(FResult& Cr
 	// We may have to re-order custom properties if there are more than one.
 	if (CustomProperties.Num() > 1)
 	{
-		InitCVarReplicateCustomDeltaPropertiesInRepIndexOrder();
-
-		// If CVarReplicateCustomDeltaPropertiesInRepIndexOrder is false then we use the GetLifetimeReplicatedPropsOrder instead, which is the RepLayout legacy behavior.
-		if (CVarReplicateCustomDeltaPropertiesInRepIndexOrder != nullptr && !CVarReplicateCustomDeltaPropertiesInRepIndexOrder->GetBool())
+		// If bReplicateCustomDeltaPropertiesInRepIndexOrder is false then we use the GetLifetimeReplicatedPropsOrder instead, which is the RepLayout legacy behavior.
+		if (!UE::Net::bReplicateCustomDeltaPropertiesInRepIndexOrder)
 		{
 			TArray<FPropertyReplicationStateDescriptorBuilder::FMemberProperty, TInlineAllocator<8>> NewCustomProperties;
 			NewCustomProperties.Reserve(CustomProperties.Num());
@@ -2884,13 +2901,4 @@ SIZE_T FReplicationStateDescriptorBuilder::CreateDescriptorsForClass(FResult& Cr
 	return CreatedDescriptors.Num();
 }
 
-void FReplicationStateDescriptorBuilder::InitCVarReplicateCustomDeltaPropertiesInRepIndexOrder()
-{
-	if (CVarReplicateCustomDeltaPropertiesInRepIndexOrder == nullptr)
-	{
-		CVarReplicateCustomDeltaPropertiesInRepIndexOrder = IConsoleManager::Get().FindConsoleVariable(TEXT("net.ReplicateCustomDeltaPropertiesInRepIndexOrder"));
-		ensureMsgf(CVarReplicateCustomDeltaPropertiesInRepIndexOrder != nullptr, TEXT("Unable to find cvar net.ReplicateCustomDeltaPropertiesInRepIndexOrder"));
-	}
-}
-
 }
diff --git a/Engine/Source/Runtime/Experimental/Iris/Core/Public/Iris/ReplicationState/ReplicationStateDescriptorBuilder.h b/Engine/Source/Runtime/Experimental/Iris/Core/Public/Iris/ReplicationState/ReplicationStateDescriptorBuilder.h
--- a/Engine/Source/Runtime/Experimental/Iris/Core/Public/Iris/ReplicationState/ReplicationStateDescriptorBuilder.h
+++ b/Engine/Source/Runtime/Experimental/Iris/Core/Public/Iris/ReplicationState/ReplicationStateDescriptorBuilder.h
@@ -72,11 +72,6 @@ public:
 	 * Any new created ReplicationStateDescriptor will be registered in the provided registry.
 	 */
 	IRISCORE_API static TRefCountPtr<const FReplicationStateDescriptor> CreateDescriptorForFunction(const UFunction* Function, const FParameters& Parameters = FParameters());
-
-private:
-	static void InitCVarReplicateCustomDeltaPropertiesInRepIndexOrder();
-
-	static const IConsoleVariable* CVarReplicateCustomDeltaPropertiesInRepIndexOrder;
 };
 
 }
diff --git a/Engine/Source/Runtime/Net/Core/Private/Net/Core/Misc/NetCVars.cpp b/Engine/Source/Runtime/Net/Core/Private/Net/Core/Misc/NetCVars.cpp
--- a/Engine/Source/Runtime/Net/Core/Private/Net/Core/Misc/NetCVars.cpp
+++ b/Engine/Source/Runtime/Net/Core/Private/Net/Core/Misc/NetCVars.cpp
@@ -8,8 +8,44 @@ namespace UE::Net
 
 int32 CVar_ForceConnectionViewerPriority = 1;
 // Deprecated RepGraph specific name.
-static FAutoConsoleVariableRef CVarRepGraphForceConnectionViewerPriority(TEXT("Net.RepGraph.ForceConnectionViewerPriority"), CVar_ForceConnectionViewerPriority, TEXT("Force the connection's player controller and viewing pawn as topmost priority. Same as Net.ForceConnectionViewerPriority."));
+static FAutoConsoleVariableRef CVarRepGraphForceConnectionViewerPriority(
+	TEXT("Net.RepGraph.ForceConnectionViewerPriority"), 
+	CVar_ForceConnectionViewerPriority, 
+	TEXT("Force the connection's player controller and viewing pawn as topmost priority. Same as Net.ForceConnectionViewerPriority.")
+);
+
 // New name without RepGraph as this is used by Iris as well.
-static FAutoConsoleVariableRef CVarForceConnectionViewerPriority(TEXT("Net.ForceConnectionViewerPriority"), CVar_ForceConnectionViewerPriority, TEXT("Force the connection's player controller and viewing pawn as topmost priority."));
+static FAutoConsoleVariableRef CVarForceConnectionViewerPriority(
+	TEXT("Net.ForceConnectionViewerPriority"), 
+	CVar_ForceConnectionViewerPriority, 
+	TEXT("Force the connection's player controller and viewing pawn as topmost priority.")
+);
+
+bool bAutoRegisterReplicatedProperties = true;
+static FAutoConsoleVariableRef CVarAutoRegisterReplicatedProperties(
+	TEXT("Net.AutoRegisterReplicatedProperties"),
+	bAutoRegisterReplicatedProperties,
+	TEXT("Automatically register replicated variables if they are not registered by the class in GetLifetimeReplicatedProps."), 
+	ECVF_Default);
+
+bool bEnsureForMissingProperties = false;
+static FAutoConsoleVariableRef CVarEnsureForMissingProperties(
+	TEXT("Net.EnsureOnMissingReplicatedPropertiesRegister"),
+	bEnsureForMissingProperties,
+	TEXT("Ensure when we detect a missing replicated property in GetLifetimeReplicatedProps of the class."),
+	ECVF_Default);
+
+/*
+ * FastArrays and other custom delta properties may have order dependencies due to callbacks being fired during serialization at which time other custom delta properties have not yet received their state.
+ * This cvar toggles the behavior between using the RepIndex of the property or the order of appearance in the lifetime property array filled during a GetLifetimeReplicatedProps() call.
+ * Default is false to keep the legacy behavior of using the GetLifetimeReplicatedProps() order for the custom delta properties.
+ * The cvar is used in ReplicationStateDescriptorBuilder as well. Search for the cvar name in the code base before removing it.
+*/
+bool bReplicateCustomDeltaPropertiesInRepIndexOrder = false;
+static FAutoConsoleVariableRef CVarReplicateCustomDeltaPropertiesInRepIndexOrder(
+	TEXT("net.ReplicateCustomDeltaPropertiesInRepIndexOrder"),
+	bReplicateCustomDeltaPropertiesInRepIndexOrder,
+	TEXT("If false (default) custom delta properties will replicate in the same order as they're added to the lifetime property array during the call to GetLifetimeReplicatedProps. If true custom delta properties will be replicated in the property RepIndex order, which is typically in increasing property offset order. Note that custom delta properties are always serialized after regular properties.")
+);
 
 }
diff --git a/Engine/Source/Runtime/Net/Core/Public/Net/Core/Misc/NetCVars.h b/Engine/Source/Runtime/Net/Core/Public/Net/Core/Misc/NetCVars.h
--- a/Engine/Source/Runtime/Net/Core/Public/Net/Core/Misc/NetCVars.h
+++ b/Engine/Source/Runtime/Net/Core/Public/Net/Core/Misc/NetCVars.h
@@ -8,5 +8,8 @@ namespace UE::Net
 {
 
 NETCORE_API extern int32 CVar_ForceConnectionViewerPriority;
+NETCORE_API extern bool bAutoRegisterReplicatedProperties;
+NETCORE_API extern bool bEnsureForMissingProperties;
+NETCORE_API extern bool bReplicateCustomDeltaPropertiesInRepIndexOrder;
 
 }
diff --git a/Engine/Source/Runtime/Net/Core/Public/Net/Core/NetBitArray.h b/Engine/Source/Runtime/Net/Core/Public/Net/Core/NetBitArray.h
--- a/Engine/Source/Runtime/Net/Core/Public/Net/Core/NetBitArray.h
+++ b/Engine/Source/Runtime/Net/Core/Public/Net/Core/NetBitArray.h
@@ -279,6 +279,10 @@ public:
 	template<typename T>
 	void ForAllSetBits(T&& Functor) const;
 
+	/** Iterate over all the zero bits and invoke functor with signature void(uint32 BitIndex). */
+	template<typename T>
+	void ForAllUnsetBits(T&& Functor) const;
+
 	/** Iterate over all words and execute Functor, with signature void(uint32 BitIndex), for all set bits after executing WordOpFunctor */
 	template<typename T, typename V>
 	static void ForAllSetBits(const FNetBitArray& A, const FNetBitArray& B, T&& WordOpFunctor, V&& Functor);
@@ -497,6 +501,10 @@ public:
 	template <typename T>
 	void ForAllSetBits(T&& Functor) const;
 
+	/** Iterate over all the zero bits and invoke functor with signature void(uint32 BitIndex). */
+	template<typename T>
+	void ForAllUnsetBits(T&& Functor) const;
+
 	/** Iterate over all words and execute Functor, with signature void(uint32 BitIndex), for all set bits after executing WordOpFunctor */
 	template <typename T, typename V>
 	static void ForAllSetBits(const FNetBitArrayView& A, const FNetBitArrayView& B, T&& WordOpFunctor, V&& Functor);
@@ -847,6 +855,32 @@ private:
 		}
 	}
 
+	template<typename T> static void ForAllUnsetBits(const StorageWordType* Storage, const uint32 WordCount, const uint32 BitCount, T&& Functor)
+	{
+		const StorageWordType LastWordMask = ~StorageWordType(0) >> (uint32(-int32(BitCount)) & (WordBitCount - 1));
+		const uint32 LastWordIt = WordCount - 1;
+
+		for (uint32 WordIt = 0, CurrentBitIndex = 0; WordIt < WordCount; ++WordIt, CurrentBitIndex += WordBitCount)
+		{
+			const StorageWordType WordMask = ((WordIt == LastWordIt) ? LastWordMask : ~0U);
+
+			// Flip bits to test the zero values
+			StorageWordType CurrentWord = ~Storage[WordIt] & WordMask;
+			
+			// Test All bits in the CurrentWord and invoke functor if they are set
+			uint32 LocalBitIndex = CurrentBitIndex;
+			while (CurrentWord)
+			{
+				if (CurrentWord & 0x1)
+				{
+					Functor(LocalBitIndex);
+				}
+				CurrentWord >>= 1;
+				++LocalBitIndex;
+			}
+		}
+	}
+
 	template<typename T, typename V> static void ForAllSetBits(const StorageWordType* StorageA, const StorageWordType* StorageB, const uint32 WordCount, const uint32 BitCount, T&& WordOpFunctor, V&& Functor)
 	{
 		const StorageWordType LastWordMask = ~StorageWordType(0) >> (uint32(-int32(BitCount)) & (WordBitCount - 1));
@@ -1197,6 +1231,11 @@ template<typename T> inline void FNetBitArray::ForAllSetBits(T&& Functor) const
 	FNetBitArrayHelper::ForAllSetBits(Storage.GetData(), Storage.Num(), BitCount, Functor);
 }
 
+template<typename T> inline void FNetBitArray::ForAllUnsetBits(T&& Functor) const
+{
+	FNetBitArrayHelper::ForAllUnsetBits(Storage.GetData(), Storage.Num(), BitCount, Functor);
+}
+
 template <typename T, typename V> inline void FNetBitArray::ForAllSetBits(const FNetBitArray& A, const FNetBitArray& B, T&& WordOpFunctor, V&& Functor)
 {
 	UE_NETBITARRAY_CHECK(A.GetNumBits() == B.GetNumBits());
@@ -1444,6 +1483,12 @@ template<typename T> inline void FNetBitArrayView::ForAllSetBits(T&& Functor) co
 	FNetBitArrayHelper::ForAllSetBits(Storage, WordCount, BitCount, Functor);
 }
 
+
+template<typename T> inline void FNetBitArrayView::ForAllUnsetBits(T&& Functor) const
+{
+	FNetBitArrayHelper::ForAllUnsetBits(Storage, WordCount, BitCount, Functor);
+}
+
 template <typename T, typename V>
 void FNetBitArrayView::ForAllSetBits(const FNetBitArrayView& A, const FNetBitArrayView& B, T&& WordOpFunctor, V&& Functor)
 {
