Fix problem with OS and verse CDOs as it was not saving correctly the overridden values.

diff --git a/Engine/Source/Runtime/CoreUObject/Private/UObject/OverridableManager.cpp b/Engine/Source/Runtime/CoreUObject/Private/UObject/OverridableManager.cpp
--- a/Engine/Source/Runtime/CoreUObject/Private/UObject/OverridableManager.cpp
+++ b/Engine/Source/Runtime/CoreUObject/Private/UObject/OverridableManager.cpp
@@ -32,13 +32,36 @@ void FOverridableManager::Enable(UObject& Object)
 #endif // WITH_EDITORONLY_DATA
 }
 
-void FOverridableManager::Disable(UObject& Object)
+void FOverridableManager::Disable(UObject& Object, bool bPropagateToSubObject/*= false*/)
 {
 #if WITH_EDITORONLY_DATA
 	OverriddenObjectAnnotations.RemoveAnnotation(&Object);
+
+	if (bPropagateToSubObject)
+	{
+		TArray<UObject*> InstancedSubObjects;
+		constexpr bool bIncludeNestedObjects = false;
+		GetObjectsWithOuter(&Object, InstancedSubObjects, bIncludeNestedObjects);
+		for (UObject* InstancedSubObject : InstancedSubObjects)
+		{
+			checkf(InstancedSubObject, TEXT("Expecting non null SubObjects"));
+			DisableInstancedSubObject(Object, *InstancedSubObject);
+		}
+	}
 #endif // WITH_EDITORONLY_DATA
 }
 
+void FOverridableManager::DisableInstancedSubObject(const UObject& Object, UObject& InstancedSubObject)
+{
+#if WITH_EDITORONLY_DATA
+	if (InstancedSubObject.IsIn(&Object))
+	{
+		Disable(InstancedSubObject, /*bPropagateToSubObject*/ true);
+	}
+#endif // WITH_EDITORONLY_DATA
+}
+
+
 void FOverridableManager::InheritEnabledFrom(UObject& Object, const UObject* DefaultData)
 {
 #if WITH_EDITORONLY_DATA
@@ -154,7 +177,7 @@ void FOverridableManager::OverrideObject(UObject& Object)
 #endif // WITH_EDITORONLY_DATA
 }
 
-void FOverridableManager::OverrideInstancedSubObject(UObject& Object, UObject& InstancedSubObject)
+void FOverridableManager::OverrideInstancedSubObject(const UObject& Object, UObject& InstancedSubObject)
 {
 #if WITH_EDITORONLY_DATA
 	if (InstancedSubObject.IsIn(&Object))
@@ -250,7 +273,7 @@ void FOverridableManager::ClearOverrides(UObject& Object)
 #endif // WITH_EDITORONLY_DATA
 }
 
-void FOverridableManager::ClearInstancedSubObjectOverrides(UObject& Object, UObject& InstancedSubObject)
+void FOverridableManager::ClearInstancedSubObjectOverrides(const UObject& Object, UObject& InstancedSubObject)
 {
 #if WITH_EDITORONLY_DATA
 	if (InstancedSubObject.IsIn(&Object))
diff --git a/Engine/Source/Runtime/CoreUObject/Public/UObject/OverridableManager.h b/Engine/Source/Runtime/CoreUObject/Public/UObject/OverridableManager.h
--- a/Engine/Source/Runtime/CoreUObject/Public/UObject/OverridableManager.h
+++ b/Engine/Source/Runtime/CoreUObject/Public/UObject/OverridableManager.h
@@ -118,9 +118,17 @@ public:
 	COREUOBJECT_API void Enable(UObject& Object);
 
 	/**
-	 * Sets on the specified object to use overridable serialization
-	 * @param Object to disable */
-	COREUOBJECT_API void Disable(UObject& Object);
+	 * Sets on the specified object to not use overridable serialization
+	 * @param Object to disable
+	 * @param bPropagateToSubObjects true if wants to propagate the disable to all subobjects */
+	COREUOBJECT_API void Disable(UObject& Object, bool bPropagateToSubObjects = false);
+
+	/**
+	 * Sets on the specified object to not use overridable serialization
+	 * Sets the specified instanced subobject to not use overridable serialization if it is really owned by the referencer
+	 * @param Object referencing this subobject
+	 * @param InstancedSubObject the sub object to override all its properties */
+	COREUOBJECT_API void DisableInstancedSubObject(const UObject& Object, UObject& InstancedSubObject);
 
 	/**
 	 * Inherit if the specified object should enable overridable serialization. It inherits it from either its default object or its outer.
@@ -168,7 +176,7 @@ public:
 	 * Override all the properties of the specified instanced subobject if it is really owned by the referencer
 	 * @param Object referencing this subobject
 	 * @param InstancedSubObject the sub object to override all its properties */
-	COREUOBJECT_API void OverrideInstancedSubObject(UObject& Object, UObject& InstancedSubObject);
+	COREUOBJECT_API void OverrideInstancedSubObject(const UObject& Object, UObject& InstancedSubObject);
 
 	/**
 	 * Propagate override to all instanced sub object of the specified object
@@ -184,7 +192,7 @@ public:
 	 * Clear all overrides of the specified instanced subobject if it is really owned by the referencer
 	 * @param Object referencing the subobject
 	 * @param InstancedSubObject to clear overrides on */
-	COREUOBJECT_API void ClearInstancedSubObjectOverrides(UObject& Object, UObject& InstancedSubObject);
+	COREUOBJECT_API void ClearInstancedSubObjectOverrides(const UObject& Object, UObject& InstancedSubObject);
 
 	/**
 	 * Propagate the clear overrides to all instanced suboject of the specified object
