implement transactional STM actor ticks

The core functionality is in RemoteObjectTransfer ExecuteTransactional where we try to execute some code within an AutoRTFM transaction. If it fails, we request the missing object and retry the transaction until success.

diff --git a/Engine/Source/Runtime/CoreUObject/Private/UObject/RemoteObject.cpp b/Engine/Source/Runtime/CoreUObject/Private/UObject/RemoteObject.cpp
--- a/Engine/Source/Runtime/CoreUObject/Private/UObject/RemoteObject.cpp
+++ b/Engine/Source/Runtime/CoreUObject/Private/UObject/RemoteObject.cpp
@@ -18,6 +18,12 @@
 
 DEFINE_LOG_CATEGORY(LogRemoteObject);
 
+namespace UE::RemoteObject::Transfer
+{
+	extern FRemoteObjectId GNextRequestedObjectId;
+	extern UObject* GNextRequestedDestinationOuter;
+}
+
 namespace UE::RemoteObject::Private
 {
 
@@ -26,7 +32,7 @@ std::atomic<uint64> RemoteObjectSerialNumber(1);
 
 class FRemoteObjectMaps
 {
-	mutable FCriticalSection ObjectMapCritical;
+	mutable FTransactionallySafeCriticalSection ObjectMapCritical;
 	TMap<FRemoteObjectId, UE::RemoteObject::Handle::FRemoteObjectStub*> RemoteObjects;
 	TMap<FRemoteObjectId, UObject*> IdToObjectMap;
 
@@ -34,7 +40,7 @@ public:
 
 	bool IsMapped(FRemoteObjectId ObjectId) const
 	{
-		FScopeLock MapLock(&ObjectMapCritical);
+		UE::TScopeLock MapLock(ObjectMapCritical);
 		return IdToObjectMap.Contains(ObjectId);
 	}
 
@@ -42,7 +48,7 @@ public:
 	{
 		FRemoteObjectId Id = UE::CoreUObject::Private::FObjectHandleUtils::GetRemoteId(Object);
 		{
-			FScopeLock MapLock(&ObjectMapCritical);
+			UE::TScopeLock MapLock(ObjectMapCritical);
 			UObject*& MappedObject = IdToObjectMap.FindOrAdd(Id);
 			MappedObject = Object;
 		}
@@ -52,14 +58,14 @@ public:
 	{
 		FRemoteObjectId Id = UE::CoreUObject::Private::FObjectHandleUtils::GetRemoteId(Object);
 		{
-			FScopeLock MapLock(&ObjectMapCritical);
+			UE::TScopeLock MapLock(ObjectMapCritical);
 			IdToObjectMap.Remove(Id);
 		}
 	}
 
 	UObject* FindObject(FRemoteObjectId Id)
 	{
-		FScopeLock MapLock(&ObjectMapCritical);
+		UE::TScopeLock MapLock(ObjectMapCritical);
 		UObject** MappedObject = IdToObjectMap.Find(Id);
 		return MappedObject ? *MappedObject : nullptr;
 	}
@@ -68,7 +74,7 @@ public:
 	{
 		using namespace UE::RemoteObject::Handle;
 
-		FScopeLock MapLock(&ObjectMapCritical);
+		UE::TScopeLock MapLock(ObjectMapCritical);
 		if (FRemoteObjectStub** ExistingStub = RemoteObjects.Find(Id))
 		{
 			return *ExistingStub;
@@ -80,7 +86,7 @@ public:
 	{
 		using namespace UE::RemoteObject::Handle;
 
-		FScopeLock MapLock(&ObjectMapCritical);
+		UE::TScopeLock MapLock(ObjectMapCritical);
 		FRemoteObjectStub*& Stub = RemoteObjects.FindOrAdd(Id);
 		if (!Stub)
 		{
@@ -95,7 +101,7 @@ public:
 	{
 		using namespace UE::RemoteObject::Handle;
 
-		FScopeLock MapLock(&ObjectMapCritical);
+		UE::TScopeLock MapLock(ObjectMapCritical);
 		FRemoteObjectStub*& Stub = RemoteObjects.FindOrAdd(UE::CoreUObject::Private::FObjectHandleUtils::GetRemoteId(Object));
 		if (!Stub)
 		{
@@ -145,7 +151,7 @@ void InitRemoteObjects()
 		UE::RemoteObject::Transfer::RestoreRemoteObjectDataDelegate.BindStatic(&UE::RemoteObject::Serialization::Disk::LoadObjectFromDisk);
 	}
 
-	FCoreDelegates::OnEndFrame.AddStatic(&UE::RemoteObject::Transfer::TickRemoteObjectRequests);
+	FCoreDelegates::OnEndFrame.AddStatic(&UE::RemoteObject::Transfer::WaitForServersToSync);
 }
 
 void MarkAsRemote(UObject* Object, FRemoteServerId DestinationServerId)
@@ -155,6 +161,7 @@ void MarkAsRemote(UObject* Object, FRemoteServerId DestinationServerId)
 	ObjectItem->SetFlags(EInternalObjectFlags::Remote);
 	ObjectItem->ClearFlags(EInternalObjectFlags_RootFlags | EInternalObjectFlags::RemoteReference);
 	ObjectMaps->FindOrAddRemoteObjectStub(Object, DestinationServerId);
+	ObjectMaps->UnmapObject(Object);
 }
 
 void MarkAsLocal(UObject* Object)
@@ -256,11 +263,47 @@ UObject* ResolveObject(const FRemoteObjectStub* Stub, ERemoteReferenceType RefTy
 		checkf(!GIsGarbageCollecting, TEXT("Resolving remote objects while collecting garbage is not allowed (trying to resolve object %s (%s)"), *Stub->Id.ToString(), *Stub->Name.ToString());
 
 		UObject* Outer = Object ? Object->GetOuter() : ObjectMaps->FindObject(Stub->OuterId);
-		Object = MigrateObjectFromRemoteServer(Stub->Id, Stub->RemoteServerId, Outer);
+
+		// ensure the global writes are put into the Open so after we abort we can read them
+		UE_AUTORTFM_OPEN
+		{
+			GNextRequestedObjectId = Stub->Id;
+			GNextRequestedDestinationOuter = Outer;
+
+			UE_LOG(LogObj, Warning, TEXT("Aborting transaction, need obj id %s on server %s"),
+				*Stub->Id.ToString(),
+				*Stub->RemoteServerId.ToString());
+		};
+
+		AutoRTFM::AbortTransaction();
 	}
 	return Object;
 }
 
+bool RequestObject(const FRemoteObjectStub* Stub)
+{
+	using namespace UE::RemoteObject::Transfer;
+	using namespace UE::RemoteObject::Private;
+
+	UObject* Object = ObjectMaps->FindObject(Stub->Id);
+	bool bRemoteObject = false;
+	if (Object)
+	{
+		bRemoteObject = IsRemote(Object);
+	}
+	else
+	{
+		bRemoteObject = IsRemote(Stub->Id);
+	}
+
+	if (bRemoteObject)
+	{
+		RequestObjectData(Stub->Id, Stub->RemoteServerId, GetGlobalServerId());
+	}
+
+	return !bRemoteObject;
+}
+
 UObject* ResolveObject(UObject* Object, ERemoteReferenceType RefType /*= ERemoteReferenceType::Strong*/)
 {
 	using namespace UE::RemoteObject::Private;
@@ -277,6 +320,22 @@ UObject* ResolveObject(UObject* Object, ERemoteReferenceType RefType /*= ERemote
 	return Object;
 }
 
+bool RequestObject(FRemoteObjectId ObjectId)
+{
+	using namespace UE::RemoteObject::Private;
+
+	if (IsRemote(ObjectId))
+	{
+		FRemoteObjectStub* Stub = ObjectMaps->FindRemoteObjectStub(ObjectId);
+		checkf(Stub, TEXT("Failed to find remote object stub for %s"), *ObjectId.ToString());
+		return RequestObject(Stub);
+	}
+	else
+	{
+		return true;
+	}
+}
+
 bool CanResolveObject(FRemoteObjectId ObjectId)
 {
 	using namespace UE::RemoteObject::Private;
diff --git a/Engine/Source/Runtime/CoreUObject/Private/UObject/RemoteObjectTransfer.cpp b/Engine/Source/Runtime/CoreUObject/Private/UObject/RemoteObjectTransfer.cpp
--- a/Engine/Source/Runtime/CoreUObject/Private/UObject/RemoteObjectTransfer.cpp
+++ b/Engine/Source/Runtime/CoreUObject/Private/UObject/RemoteObjectTransfer.cpp
@@ -12,230 +12,48 @@
 #include "UObject/RemoteObjectPrivate.h"
 #include "UObject/ObjectHandlePrivate.h"
 #include "HAL/CriticalSection.h"
-
+#include "AutoRTFM/AutoRTFM.h"
 
 namespace UE::RemoteObject::Transfer
 {
+FRemoteObjectId GNextRequestedObjectId;
+UObject* GNextRequestedDestinationOuter;
 
 const FRemoteServerId DatabaseId(0x3ff);
 
 TDelegate<void(FRemoteObjectId ObjectId, FRemoteServerId CurrentOwnerServerId, FRemoteServerId DestinationServerId)> RequestRemoteObjectDelegate;
 TDelegate<void(FRemoteObjectId ObjectId, FRemoteServerId DestinationServerId, FRemoteObjectData& ObjectData)> RemoteObjectTransferDelegate;
+TDelegate<void(FRemoteObjectId ObjectId, FRemoteServerId DestinationServerId)> RemoteObjectDeniedTransferDelegate;
 TDelegate<void()> FlushRemoteObjectRequestsDelegate;
+TDelegate<void()> WaitForServersToSyncDelegate;
 
 TMulticastDelegate<void(const FRemoteObjectData& ObjectData, FRemoteServerId RemoteServerId)> OnObjectDataSentDelegate;
 TMulticastDelegate<void(const TSet<UObject*>& MigratedObjects, FRemoteServerId RemoteServerId)> OnObjectsSentDelegate;
 TMulticastDelegate<void(const FRemoteObjectData& ObjectData, FRemoteServerId RemoteServerId)> OnObjectDataReceivedDelegate;
+TMulticastDelegate<void(FRemoteObjectId ObjectId, FRemoteServerId RemoteServerId)> OnObjectDataDeniedDelegate;
 TMulticastDelegate<void(const TArray<UObject*>& MigratedObjects, FRemoteServerId RemoteServerId)> OnObjectsReceivedDelegate;
 
 TDelegate<void(FRemoteObjectId ObjectId, FRemoteServerId DatabaseId, FRemoteObjectData& ObjectData)> StoreRemoteObjectDataDelegate;
 TDelegate<void(FRemoteObjectId ObjectId, FRemoteServerId DatabaseId, FRemoteServerId DestinationServerId)> RestoreRemoteObjectDataDelegate;
 
+UObject* MigrateObjectFromObjectData(FRemoteObjectData& ObjectData, FRemoteServerId RemoteServerId, UObject* DestinationOuter);
+
 /**
  * Remote object transfer queue. 
  * Queued requests (both send and receive) are processed on the game thread since some systems (like RPCs) that are used for transfering object data are GT-only
  */
 class FRemoteObjectTransferQueue
 {
-	enum class ERequestState : uint8
-	{
-		None,
-		InProgress,
-		Complete
-	};
-
-	enum class ERequestType : uint8
-	{
-		Unknown,
-		Send,
-		Receive,
-		Forward
-	};
-
-	/** Base class for remote object requests */
-	struct FRemoteObjectRequestBase
-	{
-		/** Unique object id this request is related to */
-		FRemoteObjectId ObjectId;
-		/** Server Id this request will be sent to */
-		FRemoteServerId ServerId;
-		/** Type of this request */
-		ERequestType Type = ERequestType::Unknown;
-		/** Current state of this request */
-		ERequestState State = ERequestState::None;
-
-		/** Next request on the list */
-		FRemoteObjectRequestBase* Next = nullptr;
-
-		virtual void Process() = 0;
-		virtual ~FRemoteObjectRequestBase() = default;
-	};
-
-	/**
-	* Receive request - requests an object from a remote server
-	* Receive request is not complete until a remote server fullfils the request with a 'send' request
-	*/
-	struct FReceiveRemoteObjectRequest : public FRemoteObjectRequestBase
-	{		
-		TPromise<TUniquePtr<FRemoteObjectData>> ReceivePromise;
-		FRemoteObjectDataFuture ReceiveFuture;
-
-		FReceiveRemoteObjectRequest()
-		{
-			Type = ERequestType::Receive;
-		}
-		virtual void Process() override
-		{
-			if (ServerId != DatabaseId)
-			{
-				// Request an object from remote server
-				RequestRemoteObjectDelegate.Execute(ObjectId, /*CurrentOwner*/ ServerId, GetGlobalServerId());
-			}
-			else
-			{
-				RestoreRemoteObjectDataDelegate.Execute(ObjectId, DatabaseId, GetGlobalServerId());
-			}
-		}
-	};
-
-	/**
-	* Forward request - forwards receive request for an object to another remote server
-	*/
-	struct FForwardRemoteObjectRequest : public FRemoteObjectRequestBase
-	{
-		/** Server Id that wants to receive object data */
-		FRemoteServerId DestinationServerId;
-
-		FForwardRemoteObjectRequest()
-		{
-			Type = ERequestType::Forward;
-		}
-		virtual void Process() override
-		{
-			// Request an object from remote server to be sent to another remote server
-			RequestRemoteObjectDelegate.Execute(ObjectId, /*CurrentOwner*/ ServerId, DestinationServerId);
-			State = ERequestState::Complete;
-		}
-	};
-
-	/**
-	* Send request - sends object data to a remote server
-	* Send requests can be insigated either by the current server or be a response to a receive request sent by a remote server
-	*/
-	struct FSendRemoteObjectRequest : public FRemoteObjectRequestBase
-	{
-		TUniquePtr<FRemoteObjectData> ObjectData;
-
-		FSendRemoteObjectRequest()
-		{
-			Type = ERequestType::Send;
-		}
-		virtual void Process() override
-		{
-			// Send the associated object data to a remote server
-			RemoteObjectTransferDelegate.Execute(ObjectId, ServerId, *ObjectData);
-			State = ERequestState::Complete;
-		}
-	};
-	
-	FCriticalSection RequestCritical;
-	// Request queue. Order of requests needs to be preserved
-	FRemoteObjectRequestBase* Head = nullptr;
-	FRemoteObjectRequestBase* Tail = nullptr;
-
-	// For quick lookups, receive requests are stored in ID -> Request map
-	TMap<FRemoteObjectId, FReceiveRemoteObjectRequest*> IdToReceiveRequestMap;
-
-	/**
-	* Find an existing receive request by object id
-	* @param ObjectId Remote object id
-	* @return Existing receive request or null if not found
-	*/
-	FORCEINLINE FReceiveRemoteObjectRequest* FindReceiveRequestUnsafe(FRemoteObjectId ObjectId)
-	{
-		FReceiveRemoteObjectRequest** FoundRequest = IdToReceiveRequestMap.Find(ObjectId);
-		if (FoundRequest)
-		{
-			return *FoundRequest;
-		}
-		return nullptr;
-	}
-
-	/**
-	* Appends (links to the list tail) new request
-	* @param Request request to link
-	*/
-	void AddRequest(FRemoteObjectRequestBase* Request)
-	{
-		if (Tail)
-		{
-			Tail->Next = Request;
-			Tail = Request;
-		}
-		else
-		{
-			Tail = Request;
-			Head = Request;
-		}
-		if (Request->Type == ERequestType::Receive)
-		{
-			IdToReceiveRequestMap.Add(Request->ObjectId, static_cast<FReceiveRemoteObjectRequest*>(Request));
-		}
-	}
-
-	/**
-	* Removes (unlinks) and destroys an existing request
-	* @param Request request to unlink and destroy
-	* @param Prev Request that was previous to the one being destroyed
-	* @return Request following the request being destroyed
-	*/
-	FRemoteObjectRequestBase* RemoveRequest(FRemoteObjectRequestBase* Request, FRemoteObjectRequestBase* Prev)
-	{
-		FRemoteObjectRequestBase* Next = Request->Next;
-		if (Prev)
-		{
-			Prev->Next = Next;
-		}
-		if (Request == Head)
-		{
-			Head = Next;
-		}
-		if (Request == Tail)
-		{
-			Tail = Prev;
-		}
-		if (Request->Type == ERequestType::Receive)
-		{
-			IdToReceiveRequestMap.Remove(Request->ObjectId);
-		}
-		delete Request;
-		return Next;
-	}
-	
 public:
+	TArray<FRemoteObjectId> LockedObjects;
+	TArray<FRemoteObjectId> RequestedObjects;
+	TArray<TObjectPtr<UObject>> RequestedObjectsDestinationOuter;
+	int32 PendingRequests;
 
-	/**
-	* Requests an object from a remote server (adds a receive request)
-	* @param ObjectId Id of the object being requested
-	* @param SourceServerId Remote server id
-	* @return Future that's going to receive the object data
-	*/
-	FRemoteObjectDataFuture RequestRemoteObject(FRemoteObjectId ObjectId, FRemoteServerId CurrentOwnerServerId)
+	void RequestRemoteObject(FRemoteObjectId ObjectId, FRemoteServerId CurrentOwnerServerId)
 	{
-		FReceiveRemoteObjectRequest* Request = nullptr;
-		{
-			FScopeLock RequestLock(&RequestCritical);
-			Request = FindReceiveRequestUnsafe(ObjectId);
-			if (!Request)
-			{
-				Request = new FReceiveRemoteObjectRequest();
-				Request->ObjectId = ObjectId;
-				Request->ServerId = CurrentOwnerServerId;
-				Request->ReceiveFuture = Request->ReceivePromise.GetFuture().Share();
-				AddRequest(Request);
-			}
-		}
-		return Request->ReceiveFuture;
+		// Request an object from remote server
+		RequestRemoteObjectDelegate.Execute(ObjectId, CurrentOwnerServerId, GetGlobalServerId());
 	}
 
 	/**
@@ -246,13 +64,7 @@ public:
 	*/
 	void ForwardRemoteObjectRequest(FRemoteObjectId ObjectId, FRemoteServerId CurrentOwnerServerId, FRemoteServerId DestinationServerId)
 	{
-		FForwardRemoteObjectRequest* Request = new FForwardRemoteObjectRequest();
-		Request->ObjectId = ObjectId;
-		Request->ServerId = CurrentOwnerServerId;
-		Request->DestinationServerId = DestinationServerId;
-
-		FScopeLock RequestLock(&RequestCritical);
-		AddRequest(Request);
+		RequestRemoteObjectDelegate.Execute(ObjectId, CurrentOwnerServerId, DestinationServerId);
 	}
 
 	/**
@@ -264,13 +76,12 @@ public:
 	*/
 	void SendRemoteObject(FRemoteObjectId ObjectId, FRemoteServerId DestinationServerId, FRemoteObjectData* ObjectData)
 	{		
-		FSendRemoteObjectRequest* Request = new FSendRemoteObjectRequest();
-		Request->ObjectId = ObjectId;
-		Request->ServerId = DestinationServerId;
-		Request->ObjectData = MakeUnique<FRemoteObjectData>(MoveTemp(*ObjectData));
+		RemoteObjectTransferDelegate.Execute(ObjectId, DestinationServerId, *ObjectData);
+	}
 
-		FScopeLock RequestLock(&RequestCritical);
-		AddRequest(Request);
+	void SendDeniedRemoteObject(FRemoteObjectId ObjectId, FRemoteServerId DestinationServerId)
+	{
+		RemoteObjectDeniedTransferDelegate.Execute(ObjectId, DestinationServerId);
 	}
 
 	/**
@@ -279,77 +90,46 @@ public:
 	* @param Data Object data. Data ownership is transferred to the receive request (if it exists)
 	* @return True if a receive request for the specified object exists (and object data ownership was transferred), false otherwise
 	*/
-	bool FulfillReceiveRequest(FRemoteObjectId ObjectId, FRemoteObjectData* Data)
+	void FulfillReceiveRequest(FRemoteObjectId ObjectId, FRemoteServerId RemoteServerId, FRemoteObjectData* ObjectData)
 	{
-		FScopeLock RequestLock(&RequestCritical);
-		FReceiveRemoteObjectRequest* Request = FindReceiveRequestUnsafe(ObjectId);
-		if (Request)
+		UObject* DestinationOuter = nullptr;
+
+		bool bFoundObject = false;
+		for(int32 RequestedObjectIndex = 0; RequestedObjectIndex < RequestedObjects.Num(); RequestedObjectIndex++)
 		{
-			check(Request->State == ERequestState::InProgress);
-			Request->State = ERequestState::Complete;
-			TUniquePtr<FRemoteObjectData> UniqueData(Data);
-			Request->ReceivePromise.SetValue(MoveTemp(UniqueData));
-			return true;
+			if(RequestedObjects[RequestedObjectIndex] == ObjectId)
+			{
+				DestinationOuter = RequestedObjectsDestinationOuter[RequestedObjectIndex];
+				bFoundObject = true;
+			}			
 		}
-		return false;
-	}
 
-	/**
-	* Finds an existing receive request
-	* @param ObjectId Id of the object the receive request is to be found
-	* @return Object data future associated with the existing request (invalid if a receive request for the specified object id was not found)
-	*/
-	FRemoteObjectDataFuture FindReceiveRequest(FRemoteObjectId ObjectId)
-	{
-		FScopeLock RequestLock(&RequestCritical);
-		FReceiveRemoteObjectRequest* Request = FindReceiveRequestUnsafe(ObjectId);
-		if (Request)
+		if (bFoundObject)
 		{
-			return Request->ReceiveFuture;
+			UE_LOG(LogRemoteObject, Verbose, TEXT("     : FulfillReceiveRequest %s (was requested)"), *ObjectId.ToString());
 		}
-		return FRemoteObjectDataFuture();
-	}
-
-	/**
-	* Processes exiting requests
-	* @return true if all requests have been processed and the queue is empty
-	*/
-	FORCENOINLINE bool Tick()
-	{
-		FScopeLock RequestLock(&RequestCritical);
-		FRemoteObjectRequestBase* Prev = nullptr;
-		for (FRemoteObjectRequestBase* Request = Head; Request;)
+		else
 		{
-			if (Request->State == ERequestState::None)
-			{
-				Request->State = ERequestState::InProgress;
-				Request->Process();
-			}
+			UE_LOG(LogRemoteObject, Verbose, TEXT("     : FulfillReceiveRequest %s (was not requested)"), *ObjectId.ToString());
+		}
 
-			if (Request->State == ERequestState::Complete)
-			{
-				Request = RemoveRequest(Request, Prev);
-			}
-			else
-			{
-				Prev = Request;
-				Request = Request->Next;
-			}
+		UObject* Result = MigrateObjectFromObjectData(*ObjectData, RemoteServerId, DestinationOuter);
+		checkf(UE::CoreUObject::Private::FObjectHandleUtils::GetRemoteId(Result) == ObjectId, TEXT("PeerId:%s requested object %s migration from PeerId:%s but received %s"),
+			*GetGlobalServerId().ToString(), *ObjectId.ToString(), *RemoteServerId.ToString(), *UE::CoreUObject::Private::FObjectHandleUtils::GetRemoteId(Result).ToString());
+
+		if (bFoundObject)
+		{
+			PendingRequests--;
 		}
-		return !Head;
 	}
 
-	/**
-	* Processes exiting requests until the receive request associated with the specified object id is fullfilled
-	*/
-	FORCENOINLINE void FlushReceiveRequest(FRemoteObjectId ObjectId)
+	void DenyReceiveRequest(FRemoteObjectId ObjectId)
 	{
-		FRemoteObjectDataFuture RequestFuture = FindReceiveRequest(ObjectId);
-		while (RequestFuture.IsValid() && !RequestFuture.IsReady())
+		UE_LOG(LogRemoteObject, Verbose, TEXT("     : DenyReceiveRequest %s"), *ObjectId.ToString());
+
+		if (RequestedObjects.Find(ObjectId) != INDEX_NONE)
 		{
-			FlushRemoteObjectRequestsDelegate.ExecuteIfBound();
-			Tick();
-			FPlatformProcess::Sleep(0.0f);
+			PendingRequests--;
 		}
 	}
 };
@@ -357,17 +137,15 @@ public:
 FRemoteObjectTransferQueue GTransferQueue;
 
 
-FRemoteObjectDataFuture RequestObjectData(FRemoteObjectId ObjectId, FRemoteServerId CurrentOwnerServerId, FRemoteServerId DestinationServerId)
+void RequestObjectData(FRemoteObjectId ObjectId, FRemoteServerId CurrentOwnerServerId, FRemoteServerId DestinationServerId)
 {
 	if (DestinationServerId == GetGlobalServerId())
 	{
-		return GTransferQueue.RequestRemoteObject(ObjectId, CurrentOwnerServerId);
+		GTransferQueue.RequestRemoteObject(ObjectId, CurrentOwnerServerId);
 	}
 	else
 	{
 		GTransferQueue.ForwardRemoteObjectRequest(ObjectId, CurrentOwnerServerId, DestinationServerId);
-		// We're not expecting to receive anything on this server
-		return FRemoteObjectDataFuture();
 	}
 }
 
@@ -388,6 +166,42 @@ UObject* MigrateObjectFromObjectData(FRemoteObjectData& ObjectData, FRemoteServe
 
 	checkf(ReceivedObjects.Num() > 0, TEXT("PeerId:%s unable to deserialize object data (%d bytes)"), *GetGlobalServerId().ToString(), ObjectData.Bytes.Num());
 
+	if (ReceivedObjects.Num() > 0)
+	{
+		FRemoteObjectId RootObjectId = UE::CoreUObject::Private::FObjectHandleUtils::GetRemoteId(ReceivedObjects[0]);
+
+		if (GTransferQueue.RequestedObjects.Find(RootObjectId) != INDEX_NONE)
+		{
+			if (GTransferQueue.LockedObjects.Find(RootObjectId) == INDEX_NONE)
+			{
+				UE_LOG(LogRemoteObject, Verbose, TEXT("     : MigrateObjectFromObjectData received object %s"), *RootObjectId.ToString());
+				GTransferQueue.LockedObjects.Add(RootObjectId);
+			}
+			else
+			{
+				UE_LOG(LogRemoteObject, Verbose, TEXT("     : MigrateObjectFromObjectData ???? received an object that was already locked %s"), *RootObjectId.ToString());
+			}
+
+			for(int32 SubObjectIndex = 1; SubObjectIndex < ReceivedObjects.Num(); SubObjectIndex++)
+			{
+				FRemoteObjectId SubObjectId = UE::CoreUObject::Private::FObjectHandleUtils::GetRemoteId(ReceivedObjects[SubObjectIndex]);
+				UE_LOG(LogRemoteObject, Verbose, TEXT("         (subobject %d) : %s"), SubObjectIndex, *SubObjectId.ToString());
+			}
+			UE_LOG(LogRemoteObject, Verbose, TEXT("         <done listing subobjects>"));
+		}
+		else
+		{
+			// got an object that we didn't request - ownership migration?
+			UE_LOG(LogRemoteObject, Verbose, TEXT("     : MigrateObjectFromObjectData received object %s (was not requested)"), *RootObjectId.ToString());
+			for (int32 SubObjectIndex = 1; SubObjectIndex < ReceivedObjects.Num(); SubObjectIndex++)
+			{
+				FRemoteObjectId SubObjectId = UE::CoreUObject::Private::FObjectHandleUtils::GetRemoteId(ReceivedObjects[SubObjectIndex]);
+				UE_LOG(LogRemoteObject, Verbose, TEXT("         (subobject %d) : %s"), SubObjectIndex, *SubObjectId.ToString());
+			}
+			UE_LOG(LogRemoteObject, Verbose, TEXT("         <done listing subobjects>"));
+		}
+	}
+
 #if UE_WITH_REMOTE_OBJECT_HANDLE
 	for (UObject* ReceivedObject : ReceivedObjects)
 	{
@@ -401,48 +215,78 @@ UObject* MigrateObjectFromObjectData(FRemoteObjectData& ObjectData, FRemoteServe
 	return ReceivedObjects[0];
 } 
 
-UObject* MigrateObjectFromRemoteServer(FRemoteObjectId ObjectId, FRemoteServerId CurrentOwnerServerId, UObject* DestinationOuter)
+void MigrateObjectToRemoteServer(UObject* Object, FRemoteServerId DestinationServerId)
 {
-	FRemoteObjectDataFuture ObjectFuture = RequestObjectData(ObjectId, CurrentOwnerServerId, /*DestinationServerId*/ GetGlobalServerId());
-	if (!ObjectFuture.IsReady())
+	using namespace UE::RemoteObject::Serialization;
+	using namespace UE::RemoteObject::Private;
+
+	FRemoteObjectId ObjectId = UE::CoreUObject::Private::FObjectHandleUtils::GetRemoteId(Object);
+
+	UE_LOG(LogRemoteObject, Verbose, TEXT("     : MigrateObjectToRemoteServer sending obj %s (isremote %d) (%d locked objs)"), *ObjectId.ToString(), Handle::IsRemote(Object), GTransferQueue.LockedObjects.Num());
+
+	// is the object locked?
+	bool bIsObjectLocked = false;
+
+	for (int32 LockedObjectIndex = 0; LockedObjectIndex < GTransferQueue.LockedObjects.Num(); LockedObjectIndex++)
 	{
-		// RobM todo: this function should not expect to receive data immediately. Instead we should abort the current transaction after RequestObjectData is called.
-		GTransferQueue.FlushReceiveRequest(ObjectId);
+		FRemoteObjectId LockedObjectId = GTransferQueue.LockedObjects[LockedObjectIndex];
+
+		if (LockedObjectId == ObjectId)
+		{
+			// do we give it up, or deny their request?
+			if(GetGlobalServerId() < DestinationServerId)
+			{
+				// we are higher priority, send a deny response
+				bIsObjectLocked = true;
+
+				UE_LOG(LogRemoteObject, Verbose, TEXT("     : MigrateObjectToRemoteServer DENYING request for %s because it's locked"), *LockedObjectId.ToString());
+				RemoteObjectDeniedTransferDelegate.Execute(LockedObjectId, DestinationServerId);
+			}
+			else
+			{
+				UE_LOG(LogRemoteObject, Verbose, TEXT("     : MigrateObjectToRemoteServer GIVING UP our locked %s"), *LockedObjectId.ToString());
+
+				// unlock the object from ourselves
+				GTransferQueue.LockedObjects.RemoveAt(LockedObjectIndex);
+			}
+
+			break;
+		}
 	}
 
-	// The below code should probably be moved to FulfillReceiveRequest
-	FRemoteObjectData& ObjectData = *ObjectFuture.Get();
-	UObject* Result = MigrateObjectFromObjectData(ObjectData, CurrentOwnerServerId, DestinationOuter);
-	checkf(UE::CoreUObject::Private::FObjectHandleUtils::GetRemoteId(Result) == ObjectId, TEXT("PeerId:%s requested object %s migration from PeerId:%s but received %s"),
-		*GetGlobalServerId().ToString(), *ObjectId.ToString(), *CurrentOwnerServerId.ToString(), *UE::CoreUObject::Private::FObjectHandleUtils::GetRemoteId(Result).ToString());
+	if (!bIsObjectLocked)
+	{
+		UE_LOG(LogRemoteObject, Verbose, TEXT("     : ... MigrateObjectToRemoteServer serializing obj %s (isremote %d)"), *ObjectId.ToString(), Handle::IsRemote(Object));
 
-	return Result;
-}
+		// send the object
+		TSet<UObject*> SentObjects;
+		FRemoteObjectData* ObjectData = SerializeObjectData(Object, SentObjects);
 
-void MigrateObjectToRemoteServer(UObject* Object, FRemoteServerId DestinationServerId)
-{
-	using namespace UE::RemoteObject::Serialization;
-	using namespace UE::RemoteObject::Private;
+		UE_LOG(LogRemoteObject, Verbose, TEXT("     : ... MigrateObjectToRemoteServer serialization complete obj %s (isremote %d)"), *ObjectId.ToString(), Handle::IsRemote(Object));
 
-	TSet<UObject*> SentObjects;	
-	FRemoteObjectData* ObjectData = SerializeObjectData(Object, SentObjects);
+		for(UObject* SO : SentObjects)
+		{
+			FRemoteObjectId SubObjectId = UE::CoreUObject::Private::FObjectHandleUtils::GetRemoteId(SO);
+			UE_LOG(LogRemoteObject, Verbose, TEXT("         sent subobject : %s"), *SubObjectId.ToString());
+		}
 
-	// Notify that object data has been migrated (actually will be but SendObjectData does a MoveTemp on the data)
-	OnObjectDataSentDelegate.Broadcast(*ObjectData, DestinationServerId);
+		// Notify that object data has been migrated (actually will be but SendObjectData does a MoveTemp on the data)
+		OnObjectDataSentDelegate.Broadcast(*ObjectData, DestinationServerId);
 
-	SendObjectData(Object, DestinationServerId, ObjectData);
+		SendObjectData(Object, DestinationServerId, ObjectData);
 
-	checkf(SentObjects.Num() > 0, TEXT("PeerId:%s unable to migrate object %s to PeerId:%s"), *GetGlobalServerId().ToString(), *UE::CoreUObject::Private::FObjectHandleUtils::GetRemoteId(Object).ToString(), *DestinationServerId.ToString());	
+		checkf(SentObjects.Num() > 0, TEXT("PeerId:%s unable to migrate object %s to PeerId:%s"), *GetGlobalServerId().ToString(), *UE::CoreUObject::Private::FObjectHandleUtils::GetRemoteId(Object).ToString(), *DestinationServerId.ToString());
 
-	// Notify that objects have been migrated to a remote server
-	OnObjectsSentDelegate.Broadcast(SentObjects, DestinationServerId);
+		// Notify that objects have been migrated to a remote server
+		OnObjectsSentDelegate.Broadcast(SentObjects, DestinationServerId);
 
-	for (UObject* SentObject : SentObjects)
-	{
+		for (UObject* SentObject : SentObjects)
+		{
 #if UE_WITH_REMOTE_OBJECT_HANDLE
-		SentObject->PostMigrate(DestinationServerId);
+			SentObject->PostMigrate(DestinationServerId);
 #endif
-		MarkAsRemote(SentObject, DestinationServerId);
+			MarkAsRemote(SentObject, DestinationServerId);
+		}
 	}
 }
 
@@ -452,46 +296,153 @@ void MigrateObjectToRemoteServer(FRemoteObjectId ObjectId, FRemoteServerId Desti
 	using namespace UE::RemoteObject::Serialization;
 	using namespace UE::RemoteObject::Private;
 
+	UE_LOG(LogRemoteObject, Verbose, TEXT(" >>> MigrateObjectToRemoteServer obj %s to %s"), *ObjectId.ToString(), *DestinationServerId.ToString());
+
 	UObject* Object = FindObject(ObjectId);
 	if (Object)
 	{
+		UE_LOG(LogRemoteObject, Verbose, TEXT("   > MigrateObjectToRemoteServer found obj %s"), *ObjectId.ToString());
 		MigrateObjectToRemoteServer(Object, DestinationServerId);
 	}
 	else
-	{		
+	{
+		UE_LOG(LogRemoteObject, Verbose, TEXT("   > MigrateObjectToRemoteServer did not find obj %s"), *ObjectId.ToString());
+
 		FRemoteObjectStub* RemoteObjectStub = FindRemoteObjectStub(ObjectId);
 		FRemoteServerId CurrentOwnerServerId = RemoteObjectStub ? RemoteObjectStub->RemoteServerId : ObjectId.GetServerId();
 		checkf(CurrentOwnerServerId != GetGlobalServerId(), TEXT("PeerId:%s attempting to forward request for object %s to PeerId:%s but object location is unknown"), 
 			*GetGlobalServerId().ToString(), *ObjectId.ToString(), *DestinationServerId.ToString());
 
-		// Forward the request to CurrentOwnerServerId
-		RequestObjectData(ObjectId, CurrentOwnerServerId, DestinationServerId);
+		if (CurrentOwnerServerId != DestinationServerId)
+		{
+			UE_LOG(LogRemoteObject, Verbose, TEXT("     : MigrateObjectToRemoteServer --FORWARDING REQUEST--"));
+
+			// Forward the request to CurrentOwnerServerId
+			RequestObjectData(ObjectId, CurrentOwnerServerId, DestinationServerId);
+		}
+		else
+		{
+			UE_LOG(LogRemoteObject, Verbose, TEXT("     : MigrateObjectToRemoteServer --SENDING DENY-- %s"), *ObjectId.ToString());
+			RemoteObjectDeniedTransferDelegate.Execute(ObjectId, DestinationServerId);
+		}
 	}
 }
 
 void OnObjectDataReceived(FRemoteObjectId ObjectId, FRemoteServerId RemoteServerId, FRemoteObjectData* Data)
 {
-	if (!GTransferQueue.FulfillReceiveRequest(ObjectId, Data))
-	{
-		MigrateObjectFromObjectData(*Data, RemoteServerId, nullptr);
-	}
-} 
+	UE_LOG(LogRemoteObject, Verbose, TEXT("     : OnObjectDataReceived %s from %s"), *ObjectId.ToString(), *RemoteServerId.ToString());
 
-void TickRemoteObjectRequests()
+	GTransferQueue.FulfillReceiveRequest(ObjectId, RemoteServerId, Data);
+}
+
+void OnObjectDataDenied(FRemoteObjectId ObjectId, FRemoteServerId RemoteServerId)
 {
-	UE::RemoteObject::Transfer::GTransferQueue.Tick();
+	UE_LOG(LogRemoteObject, Verbose, TEXT("     : OnObjectDataDenied %s from %s"), *ObjectId.ToString(), *RemoteServerId.ToString());
+
+	GTransferQueue.DenyReceiveRequest(ObjectId);
 }
 
-void FlushRemoteObjectRequests()
+void ExecuteTransactional(const TFunction<void(void)>& Work)
 {
-	do
+	UE_LOG(LogRemoteObject, Verbose, TEXT("-- ExecuteTransactional -"));
+
+	GTransferQueue.LockedObjects.Empty();
+	GTransferQueue.RequestedObjects.Empty();
+	GTransferQueue.RequestedObjectsDestinationOuter.Empty();
+	GTransferQueue.PendingRequests = 0;
+
+	for (;;)
 	{
-		FlushRemoteObjectRequestsDelegate.ExecuteIfBound();
+		for (FRemoteObjectId ObjectId : GTransferQueue.LockedObjects)
+		{
+			UE_LOG(LogRemoteObject, Verbose, TEXT("   : locked obj : %s"), *ObjectId.ToString());
+		}
+
+		for (FRemoteObjectId ObjectId : GTransferQueue.RequestedObjects)
+		{
+			UE_LOG(LogRemoteObject, Verbose, TEXT("   : requested obj : %s"), *ObjectId.ToString());
+		}
+
+		TArray<FRemoteObjectId> ObjectsToRequest;
+
+		// try to migrate all RequestedObjects -> LockedObjects
+		for (FRemoteObjectId RequestedObjectId : GTransferQueue.RequestedObjects)
+		{
+			if (GTransferQueue.LockedObjects.Find(RequestedObjectId) == INDEX_NONE)
+			{
+				ObjectsToRequest.Add(RequestedObjectId);
+			}
+		}
+
+		UE_LOG(LogRemoteObject, Verbose, TEXT("   : ObjectsToRequest: %d entries"), ObjectsToRequest.Num());
+
+		// send the requests...
+		int32 RequestsSent = 0;
+		for(FRemoteObjectId RequestedObjectId : ObjectsToRequest)
+		{
+			if (UE::RemoteObject::Handle::RequestObject(RequestedObjectId))
+			{
+				// true means we already have the object locally and didn't actually send a request for it
+				GTransferQueue.LockedObjects.Add(RequestedObjectId);
+			}
+			else
+			{
+				RequestsSent++;
+			}
+		}
+
+		GTransferQueue.PendingRequests = RequestsSent;
+
+		// now tick the network until we get a response for each request made (either OK or DENIED)
+		while(GTransferQueue.PendingRequests > 0)
+		{
+			UE_LOG(LogRemoteObject, Verbose, TEXT("   : waiting for %d responses..."), GTransferQueue.PendingRequests);
+			FlushRemoteObjectRequestsDelegate.ExecuteIfBound();
+		}
+
+		// now we also might have migrated objects away, so check to see if LockedObjects == RequestedObjects
+		if(GTransferQueue.LockedObjects.Num() != GTransferQueue.RequestedObjects.Num())
+		{
+			UE_LOG(LogRemoteObject, Verbose, TEXT("   : did we give up locked objects? trying again..."));
+
+			// try again!
+			continue;
+		}
+
+		UE_LOG(LogRemoteObject, Verbose, TEXT("   : READY"));
+
+		// all of the requested objects are locked, try to perform the work
+		
+		AutoRTFM::ETransactionResult WorkTransactResult = AutoRTFM::Transact([&Work](){ Work(); });
+
+		if (WorkTransactResult == AutoRTFM::ETransactionResult::AbortedByRequest)
+		{
+			UE_LOG(LogRemoteObject, Verbose, TEXT("   : work aborted, need object %s"), *GNextRequestedObjectId.ToString());
+
+			// the work aborted, insert a new requested object and retry
+			GTransferQueue.RequestedObjects.Add(GNextRequestedObjectId);
+			GTransferQueue.RequestedObjectsDestinationOuter.Add(GNextRequestedDestinationOuter);
+		}
+		else
+		{
+			UE_LOG(LogRemoteObject, Verbose, TEXT("   : work OK!"));
+
+			// all done!
+			GTransferQueue.RequestedObjects.Empty();
+			GTransferQueue.RequestedObjectsDestinationOuter.Empty();
+			GTransferQueue.LockedObjects.Empty();
+			break;
+		}
 	}
-	while (!UE::RemoteObject::Transfer::GTransferQueue.Tick());
 }
 
-} // // namespace UE::RemoteObject::Transfer
+void WaitForServersToSync()
+{
+	UE_LOG(LogRemoteObject, Verbose, TEXT("-- WaitForServersToSync -"));
+	WaitForServersToSyncDelegate.ExecuteIfBound();
+}
+
+} // namespace UE::RemoteObject::Transfer
 
 namespace UE::RemoteObject::Transfer::Private
 {
diff --git a/Engine/Source/Runtime/CoreUObject/Public/UObject/RemoteObject.h b/Engine/Source/Runtime/CoreUObject/Public/UObject/RemoteObject.h
--- a/Engine/Source/Runtime/CoreUObject/Public/UObject/RemoteObject.h
+++ b/Engine/Source/Runtime/CoreUObject/Public/UObject/RemoteObject.h
@@ -43,7 +43,7 @@ namespace UE::RemoteObject::Handle
 	};
 
 	/**
-	* Resolves a remote object given its stub
+	* Resolves a remote object given its stub, aborting the active transaction if the object is unavailable
 	* @param Stub Basic data required to migrate a remote object
 	* @param RefType Reference type that wants to resolve an object
 	* @return Resolved object
@@ -51,13 +51,20 @@ namespace UE::RemoteObject::Handle
 	COREUOBJECT_API UObject* ResolveObject(const FRemoteObjectStub* Stub, ERemoteReferenceType RefType = ERemoteReferenceType::Strong);
 
 	/**
-	* Resolves a remote object
+	* Resolves a remote object, aborting the active transaction if the object is unavailable
 	* @param Object Object to resolve (remote object memory that has not yet been GC'd)
 	* @param RefType Reference type that wants to resolve an object
 	* @return Resolved object
 	*/
 	COREUOBJECT_API UObject* ResolveObject(UObject* Object, ERemoteReferenceType RefType = ERemoteReferenceType::Strong);
 
+	/**
+	* Sends a request to a remote server for the object
+	* @param Object Object to resolve
+	* @return true if the object is already local, false if a request was sent
+	*/
+	COREUOBJECT_API bool RequestObject(FRemoteObjectId ObjectId);
+
 	/**
 	* Resolves a remote object
 	* @param Object Object id to resolve
diff --git a/Engine/Source/Runtime/CoreUObject/Public/UObject/RemoteObjectTransfer.h b/Engine/Source/Runtime/CoreUObject/Public/UObject/RemoteObjectTransfer.h
--- a/Engine/Source/Runtime/CoreUObject/Public/UObject/RemoteObjectTransfer.h
+++ b/Engine/Source/Runtime/CoreUObject/Public/UObject/RemoteObjectTransfer.h
@@ -10,7 +10,11 @@
 
 namespace UE::RemoteObject::Transfer
 {
-	using FRemoteObjectDataFuture = TSharedFuture<TUniquePtr<FRemoteObjectData>>;
+	/**
+	* Execute the Work under an transaction that can catch remote object requests
+	* as well as negotiating remote object transfers, retrying the transaction until it succeeds.
+	*/
+	COREUOBJECT_API void ExecuteTransactional(const TFunction<void(void)>& Work);
 
 	/**
 	* Called when remote object data has been received from a remote server
@@ -20,14 +24,20 @@ namespace UE::RemoteObject::Transfer
 	*/
 	COREUOBJECT_API void OnObjectDataReceived(FRemoteObjectId ObjectId, FRemoteServerId RemoteServerId, FRemoteObjectData* Data);
 
+	/**
+	* Called when a remote object request was denied by a remote server
+	* @param ObjectId Remote object id
+	* @param RemoteServerId Server that owns the object data
+	*/
+	COREUOBJECT_API void OnObjectDataDenied(FRemoteObjectId ObjectId, FRemoteServerId RemoteServerId);
+
 	/**
 	* Requests remote object data to be transferred from Current Owner Server to Destination Server
 	* @param ObjectId Remote object id
 	* @param CurrentOwnerServerId Id of a remote server that is known to currently own the object
 	* @param DestinationServerId Id of a remote server that wants to own the object (may be different than the the server that sent the request)
-	* @return Future that will receive the object data
 	*/
-	COREUOBJECT_API FRemoteObjectDataFuture RequestObjectData(FRemoteObjectId ObjectId, FRemoteServerId CurrentOwnerServerId, FRemoteServerId DestinationServerId);
+	COREUOBJECT_API void RequestObjectData(FRemoteObjectId ObjectId, FRemoteServerId CurrentOwnerServerId, FRemoteServerId DestinationServerId);
 
 	/**
 	* Migrates an object to a remote server
@@ -43,15 +53,6 @@ namespace UE::RemoteObject::Transfer
 	*/
 	COREUOBJECT_API void MigrateObjectToRemoteServer(FRemoteObjectId Id, FRemoteServerId DestinationServerId);
 
-	/**
-	* Migrates an object from a remote server (temp function)
-	* @param ObjectId Remote object id
-	* @param CurrentOwnerServerId Id of a remote server that currently owns the object
-	* @param DestinationOuter Optional outer the migrated object should be reparented to
-	* @return Migrated object
-	*/
-	COREUOBJECT_API UObject* MigrateObjectFromRemoteServer(FRemoteObjectId ObjectId, FRemoteServerId CurrentOwnerServerId, UObject* DestinationOuter = nullptr);
-
 	/**
 	* Utility that converts FRemoteObjectId to TObjectPtr of a given type. Will not resolve the object if it's not local.
 	*/
@@ -65,25 +66,30 @@ namespace UE::RemoteObject::Transfer
 	}
 
 	/**
-	* Ticks remote object requests (send and receive requests)
+	* Ticks remote object requests until all have been processed (send and receive requests)
 	*/
-	COREUOBJECT_API void TickRemoteObjectRequests();
+	COREUOBJECT_API void FlushRemoteObjectRequests();
 
 	/**
-	* Ticks remote object requests until all have been processed (send and receive requests)
+	* Establishes a synchronization barrier across all servers
 	*/
-	COREUOBJECT_API void FlushRemoteObjectRequests();
+	COREUOBJECT_API void WaitForServersToSync();
 
 	extern COREUOBJECT_API const FRemoteServerId DatabaseId;
 
 	/** Delegate that transfers object data to another server */
 	extern COREUOBJECT_API TDelegate<void(FRemoteObjectId /*ObjectId*/, FRemoteServerId /*DestinationServerId*/, FRemoteObjectData& /*ObjectData*/)> RemoteObjectTransferDelegate;
 
+	/** Delegate that handles an object request being denied */
+	extern COREUOBJECT_API TDelegate<void(FRemoteObjectId /*ObjectId*/, FRemoteServerId /*DestinationServerId*/)> RemoteObjectDeniedTransferDelegate;
+
 	/** Delegate that requests remote object data from CurrentOwnerServerId to be transferred to DestinationServerId */
 	extern COREUOBJECT_API TDelegate<void(FRemoteObjectId /*ObjectId*/, FRemoteServerId /*CurrentOwnerServerId*/, FRemoteServerId /*DestinationServerId*/)> RequestRemoteObjectDelegate;
 
 	/** Delegate that flushes remote object requests */
 	extern COREUOBJECT_API TDelegate<void()> FlushRemoteObjectRequestsDelegate;
+	/** Delegate that syncs all servers */
+	extern COREUOBJECT_API TDelegate<void()> WaitForServersToSyncDelegate;
 
 	/** 
 	 * Delegate executed when objects have been migrated from another server. 
diff --git a/Engine/Source/Runtime/Engine/Classes/Engine/EngineBaseTypes.h b/Engine/Source/Runtime/Engine/Classes/Engine/EngineBaseTypes.h
--- a/Engine/Source/Runtime/Engine/Classes/Engine/EngineBaseTypes.h
+++ b/Engine/Source/Runtime/Engine/Classes/Engine/EngineBaseTypes.h
@@ -503,7 +503,11 @@ struct FActorTickFunction : public FTickFunction
 	GENERATED_USTRUCT_BODY()
 
 	/**  AActor  that is the target of this tick **/
+#if UE_WITH_REMOTE_OBJECT_HANDLE
+	TObjectPtr<AActor> Target;
+#else
 	class AActor*	Target;
+#endif
 
 	/** 
 		* Abstract function actually execute the tick. 
@@ -513,9 +517,16 @@ struct FActorTickFunction : public FTickFunction
 		* @param MyCompletionGraphEvent - completion event for this task. Useful for holding the completetion of this task until certain child tasks are complete.
 	**/
 	ENGINE_API virtual void ExecuteTick(float DeltaTime, ELevelTick TickType, ENamedThreads::Type CurrentThread, const FGraphEventRef& MyCompletionGraphEvent) override;
+	ENGINE_API void ExecuteActorTickImpl(float DeltaTime, ELevelTick TickType, ENamedThreads::Type CurrentThread, const FGraphEventRef& MyCompletionGraphEvent);
+
 	/** Abstract function to describe this tick. Used to print messages about illegal cycles in the dependency graph **/
 	ENGINE_API virtual FString DiagnosticMessage() override;
 	ENGINE_API virtual FName DiagnosticContext(bool bDetailed) override;
+
+#if UE_WITH_REMOTE_OBJECT_HANDLE
+	/** Exposes references to GC system */
+	ENGINE_API void AddStructReferencedObjects(FReferenceCollector& Collector);
+#endif
 };
 
 template<>
@@ -523,6 +534,9 @@ struct TStructOpsTypeTraits<FActorTickFunction> : public TStructOpsTypeTraitsBas
 {
 	enum
 	{
+#if UE_WITH_REMOTE_OBJECT_HANDLE
+		WithAddStructReferencedObjects = true,
+#endif
 		WithCopy = false
 	};
 };
diff --git a/Engine/Source/Runtime/Engine/Private/Actor.cpp b/Engine/Source/Runtime/Engine/Private/Actor.cpp
--- a/Engine/Source/Runtime/Engine/Private/Actor.cpp
+++ b/Engine/Source/Runtime/Engine/Private/Actor.cpp
@@ -77,6 +77,7 @@
 #include "WorldPartition/DataLayer/DataLayerManager.h"
 #include "WorldPartition/ContentBundle/ContentBundlePaths.h"
 #include "WorldPartition/ActorInstanceGuids.h"
+#include "UObject/RemoteObjectTransfer.h"
 
 DEFINE_LOG_CATEGORY(LogActor);
 
@@ -277,6 +278,20 @@ void AActor::InitializeDefaults()
 }
 
 void FActorTickFunction::ExecuteTick(float DeltaTime, enum ELevelTick TickType, ENamedThreads::Type CurrentThread, const FGraphEventRef& MyCompletionGraphEvent)
+{
+#if UE_WITH_REMOTE_OBJECT_HANDLE
+	auto ExecuteTickWork = [this, DeltaTime, TickType, CurrentThread, &MyCompletionGraphEvent]()
+	{
+		ExecuteActorTickImpl(DeltaTime, TickType, CurrentThread, MyCompletionGraphEvent);
+	};
+
+	UE::RemoteObject::Transfer::ExecuteTransactional(ExecuteTickWork);
+#else
+	ExecuteActorTickImpl(DeltaTime, TickType, CurrentThread, MyCompletionGraphEvent);
+#endif
+}
+
+void FActorTickFunction::ExecuteActorTickImpl(float DeltaTime, enum ELevelTick TickType, ENamedThreads::Type CurrentThread, const FGraphEventRef& MyCompletionGraphEvent)
 {
 	if (IsValid(Target))
 	{
@@ -307,6 +322,13 @@ FName FActorTickFunction::DiagnosticContext(bool bDetailed)
 	}
 }
 
+#if UE_WITH_REMOTE_OBJECT_HANDLE
+void FActorTickFunction::AddStructReferencedObjects(FReferenceCollector& Collector)
+{
+	Collector.AddReferencedObject(Target);
+}
+#endif
+
 bool AActor::CheckDefaultSubobjectsInternal() const
 {
 	bool Result = Super::CheckDefaultSubobjectsInternal();
