build fix

diff --git a/Engine/Source/Runtime/CoreUObject/Public/UObject/RemoteObjectTransfer.h b/Engine/Source/Runtime/CoreUObject/Public/UObject/RemoteObjectTransfer.h
--- a/Engine/Source/Runtime/CoreUObject/Public/UObject/RemoteObjectTransfer.h
+++ b/Engine/Source/Runtime/CoreUObject/Public/UObject/RemoteObjectTransfer.h
@@ -65,11 +65,6 @@ namespace UE::RemoteObject::Transfer
 #endif
 	}
 
-	/**
-	* Ticks remote object requests until all have been processed (send and receive requests)
-	*/
-	COREUOBJECT_API void FlushRemoteObjectRequests();
-
 	/**
 	* Establishes a synchronization barrier across all servers
 	*/
