Converted the property value state tracking API to be object-based

Group the functions for property value state tracking together and makes it easier to work with objects since there has been confusion around that.

Fixed object initialization and most other callers in CoreUObject to pass the index for static arrays.

#rnx

diff --git a/Engine/Plugins/Experimental/InstanceDataObjectFixupTool/Source/Private/InstanceDataObjectFixupPanel.cpp b/Engine/Plugins/Experimental/InstanceDataObjectFixupTool/Source/Private/InstanceDataObjectFixupPanel.cpp
--- a/Engine/Plugins/Experimental/InstanceDataObjectFixupTool/Source/Private/InstanceDataObjectFixupPanel.cpp
+++ b/Engine/Plugins/Experimental/InstanceDataObjectFixupTool/Source/Private/InstanceDataObjectFixupPanel.cpp
@@ -11,6 +11,7 @@
 #include "Editor.h"
 #include "Editor/PropertyEditor/Private/PropertyNode.h"
 #include "UObject/PropertyBagRepository.h"
+#include "UObject/PropertyStateTracking.h"
 
 #include "Elements/Columns/TypedElementAlertColumns.h"
 #include "Elements/Common/EditorDataStorageFeatures.h"
@@ -932,29 +933,18 @@ void FInstanceDataObjectFixupPanel::OnRedirectProperty(FPropertyPath From, FProp
 static void InitRedirectedPropertyTreeRec(const TSharedPtr<FRedirectedPropertyNode>& Node, FProperty* Property, void* Value, TSet<UObject*>& EnteredObjects);
 static void InitRedirectedPropertyTreeRec(const TSharedPtr<FRedirectedPropertyNode>& Node, UStruct* Struct, void* StructValue, TSet<UObject*>& EnteredObjects)
 {
+	const UE::FSerializedPropertyValueState SerializedState(Struct, StructValue);
 	for (FProperty* Property : TFieldRange<FProperty>(Struct))
 	{
-		if (Property->ArrayDim == 1)
+		for (int32 StaticArrayIndex = 0; StaticArrayIndex < Property->ArrayDim; ++StaticArrayIndex)
 		{
-			if (UE::FPropertyBagRepository::WasPropertyValueSerialized(Struct, StructValue, Property))
+			if (SerializedState.IsSet(Property, StaticArrayIndex))
 			{
-				const TSharedPtr<FRedirectedPropertyNode>& ChildNode = Node->FindOrAdd(FPropertyInfo(Property));
-				void* Value = Property->ContainerPtrToValuePtr<void>(StructValue);
+				const TSharedPtr<FRedirectedPropertyNode>& ChildNode = Node->FindOrAdd(FPropertyInfo(Property, StaticArrayIndex));
+				void* Value = Property->ContainerPtrToValuePtr<void>(StructValue, StaticArrayIndex);
 				InitRedirectedPropertyTreeRec(ChildNode, Property, Value, EnteredObjects);
 			}
 		}
-		else
-		{
-			for (int32 StaticArrayIndex = 0; StaticArrayIndex < Property->ArrayDim; ++StaticArrayIndex)
-            {
-            	if (UE::FPropertyBagRepository::WasPropertyValueSerialized(Struct, StructValue, Property, StaticArrayIndex))
-            	{
-            		const TSharedPtr<FRedirectedPropertyNode>& ChildNode = Node->FindOrAdd(FPropertyInfo(Property, StaticArrayIndex));
-            		void* Value = Property->ContainerPtrToValuePtr<void>(StructValue, StaticArrayIndex);
-            		InitRedirectedPropertyTreeRec(ChildNode, Property, Value, EnteredObjects);
-            	}
-            }
-		}
 	}
 }
 
diff --git a/Engine/Source/Runtime/CoreUObject/Internal/UObject/InstanceDataObjectUtils.h b/Engine/Source/Runtime/CoreUObject/Internal/UObject/InstanceDataObjectUtils.h
--- a/Engine/Source/Runtime/CoreUObject/Internal/UObject/InstanceDataObjectUtils.h
+++ b/Engine/Source/Runtime/CoreUObject/Internal/UObject/InstanceDataObjectUtils.h
@@ -46,10 +46,6 @@ UClass* CreateInstanceDataObjectClass(const FPropertyPathNameTree* PropertyTree,
 /** Copy whether each property was set by serialization from one IDO to another. */
 void CopyPropertyValueSerializedData(const FFieldVariant& OldField, void* OldDataPtr, const FFieldVariant& NewField, void* NewDataPtr);
 
-UE_API bool SupportsPropertyInitStateTracking();
-
-UE_API bool SupportsPropertyInitStateTracking(const UStruct* Struct, void* StructData);
-
 } // UE
 
 #undef UE_API
diff --git a/Engine/Source/Runtime/CoreUObject/Internal/UObject/PropertyBagRepository.h b/Engine/Source/Runtime/CoreUObject/Internal/UObject/PropertyBagRepository.h
--- a/Engine/Source/Runtime/CoreUObject/Internal/UObject/PropertyBagRepository.h
+++ b/Engine/Source/Runtime/CoreUObject/Internal/UObject/PropertyBagRepository.h
@@ -100,9 +100,6 @@ public:
 
 	COREUOBJECT_API const UObject* FindInstanceForDataObject(const UObject* InstanceDataObject) const;
 
-	// query whether a property in Struct was set when the struct was deserialized
-	COREUOBJECT_API static bool WasPropertyValueSerialized(const UStruct* Struct, const void* StructData, const FProperty* Property, int32 ArrayIndex = 0);
-
 	// FGCObject interface
 	virtual void AddReferencedObjects(FReferenceCollector& Collector) override;
 	virtual FString GetReferencerName() const override;
diff --git a/Engine/Source/Runtime/CoreUObject/Internal/UObject/PropertyStateTracking.h b/Engine/Source/Runtime/CoreUObject/Internal/UObject/PropertyStateTracking.h
--- a/Engine/Source/Runtime/CoreUObject/Internal/UObject/PropertyStateTracking.h
+++ b/Engine/Source/Runtime/CoreUObject/Internal/UObject/PropertyStateTracking.h
@@ -7,7 +7,9 @@
 #if WITH_EDITORONLY_DATA
 
 #include "Containers/ContainersFwd.h"
+#include "Serialization/StructuredArchive.h"
 #include "Templates/SharedPointerFwd.h"
+#include "UObject/Class.h"
 
 class FName;
 class FProperty;
@@ -20,45 +22,109 @@ namespace UE { class FPropertyTypeName; }
 
 #define UE_API COREUOBJECT_API
 
-#ifndef PROPERTY_INIT_STATE_TRACKING_ENABLE_DEFAULT
-	#define PROPERTY_INIT_STATE_TRACKING_ENABLE_DEFAULT 1
-#endif
-
 namespace UE
 {
 
+/**
+ * An accessor for a specific property value flag.
+ *
+ * All functions that take FProperty require the property to be owned directly by the struct/class
+ * that this is constructed with, or a super struct/class of that type.
+ */
+template <EPropertyValueFlags Flag>
+class TPropertyValueState
+{
+public:
+	/** Construct an accessor for the property value flags of instance of InStruct pointed to by InData. */
+	inline explicit TPropertyValueState(const UStruct* InStruct, void* InData)
+		: Struct(InStruct)
+		, Data(InData)
+	{
+	}
+
+	/** Construct an accessor for the property value flags of the object. */
+	inline explicit TPropertyValueState(UObject* Object)
+		: Struct(Object->GetClass())
+		, Data(Object)
+	{
+	}
+
+	/** Try to activate tracking of the property value flag. Returns true if tracking is active, otherwise false. */
+	inline bool ActivateTracking()
+	{
+		return Struct->ActivateTrackingPropertyValueFlag(Flag, Data);
+	}
+
+	/** Query whether property value flag is being tracked. */
+	inline bool IsTracking() const
+	{
+		return Struct->IsTrackingPropertyValueFlag(Flag, Data);
+	}
+
+	/** Query whether the property has the property value flag set. */
+	inline bool IsSet(const FProperty* Property, int32 ArrayIndex = 0) const
+	{
+		return Struct->HasPropertyValueFlag(Flag, Data, Property, ArrayIndex);
+	}
+
+	/** Set the property value flag for the property to bValue. */
+	inline void SetTo(bool bValue, const FProperty* Property, int32 ArrayIndex = 0)
+	{
+		Struct->SetPropertyValueFlag(Flag, bValue, Data, Property, ArrayIndex);
+	}
+
+	/** Set the property value flag for the property. */
+	inline void Set(const FProperty* Property, int32 ArrayIndex = 0)
+	{
+		SetTo(/*bValue*/ true, Property, ArrayIndex);
+	}
+
+	/** Clear the property value flag for the property. */
+	inline void Clear(const FProperty* Property, int32 ArrayIndex = 0)
+	{
+		SetTo(/*bValue*/ false, Property, ArrayIndex);
+	}
+
+	/** Reset the property value flag for every property in the type. */
+	inline void Reset()
+	{
+		Struct->ResetPropertyValueFlags(Flag, Data);
+	}
+
+	/** Serialize the property value flags for every property in the type. */
+	inline void Serialize(FStructuredArchiveRecord Record)
+	{
+		Struct->SerializePropertyValueFlags(Flag, Data, Record, GetArchiveFieldName());
+	}
+
+private:
+	inline static FArchiveFieldName GetArchiveFieldName()
+	{
+		static_assert(Flag == EPropertyValueFlags::Initialized || Flag == EPropertyValueFlags::Serialized);
+		if constexpr (Flag == EPropertyValueFlags::Initialized)
+		{
+			return TEXT("PropertyValueInitializedFlags");
+		}
+		else
+		{
+			return TEXT("PropertyValueSerializedFlags");
+		}
+	}
+
+	const UStruct* Struct = nullptr;
+	void* Data = nullptr;
+};
+
 // A property is initialized if it has a value set on it.
 // A property may be uninitialized if it has no default and has never been set.
 // A property may be uninitialized if it is an unknown property in an InstanceDataObject (IDO)
 // and has no value on a particular instance or has had its value unset while being fixed up.
-
-/** Try to activate tracking of property value initialization for the struct. */
-UE_API bool ActivateTrackingPropertyValueInitialized(const UStruct* Struct, void* Data);
-/** Query whether property value initialization is being tracked for the struct. */
-UE_API bool IsTrackingPropertyValueInitialized(const UStruct* Struct, const void* Data);
-/** Query whether the property value is initialized for a property in the struct. */
-UE_API bool IsPropertyValueInitialized(const UStruct* Struct, const void* Data, const FProperty* Property, int32 ArrayIndex = 0);
-/** Set the flag that the property value is initialized for a property in the struct. */
-UE_API void SetPropertyValueInitialized(const UStruct* Struct, void* Data, const FProperty* Property, int32 ArrayIndex = 0);
-/** Clear the flag that the property value is initialized for a property in the struct. */
-UE_API void ClearPropertyValueInitialized(const UStruct* Struct, void* Data, const FProperty* Property, int32 ArrayIndex = 0);
-/** Reset the property value initialized flags for every property in the struct. */
-UE_API void ResetPropertyValueInitialized(const UStruct* Struct, void* Data);
-/** Serialize the property value initialized flags for every property in the struct. */
-UE_API void SerializePropertyValueInitialized(const UStruct* Struct, void* Data, FStructuredArchiveRecord Record);
+using FInitializedPropertyValueState = TPropertyValueState<EPropertyValueFlags::Initialized>;
 
 // A property is serialized if a value was loaded into it during serialization.
 // A property may be initialized without having been serialized because a template object provided a default.
 // A property may be serialized without being initialized if its value was removed following serialization.
-
-/** Try to activate tracking of property value serialization for the struct. */
-bool ActivateTrackingPropertyValueSerialized(const UStruct* Struct, void* Data);
-/** Query whether property value serialization is being tracked for the struct. */
-bool IsTrackingPropertyValueSerialized(const UStruct* Struct, const void* Data);
-/** Set the flag that the property value was deserialized for a property in the struct. */
-void MarkPropertyValueSerialized(const UStruct* Struct, void* Data, const FProperty* Property, int32 ArrayIndex = 0);
-/** Query whether the property value was deserialized for a property in a struct. */
-bool WasPropertyValueSerialized(const UStruct* Struct, const void* Data, const FProperty* Property, int32 ArrayIndex = 0);
+using FSerializedPropertyValueState = TPropertyValueState<EPropertyValueFlags::Serialized>;
 
 /**
  * A tree of unknown properties found during serialization. Used with bTrackUnknownProperties.
diff --git a/Engine/Source/Runtime/CoreUObject/Private/UObject/Class.cpp b/Engine/Source/Runtime/CoreUObject/Private/UObject/Class.cpp
--- a/Engine/Source/Runtime/CoreUObject/Private/UObject/Class.cpp
+++ b/Engine/Source/Runtime/CoreUObject/Private/UObject/Class.cpp
@@ -1396,7 +1396,7 @@ void UStruct::LoadTaggedPropertiesFromText(FStructuredArchive::FSlot Slot, uint8
 				#if WITH_EDITORONLY_DATA
 					if (SerializeContext->bTrackInitializedProperties)
 					{
-						UE::SetPropertyValueInitialized(Struct, Data, Property, ArrayIndex);
+						UE::FInitializedPropertyValueState(Struct, Data).Set(Property, ArrayIndex);
 					}
 				#endif
 				};
@@ -1718,7 +1718,7 @@ void UStruct::SerializeVersionedTaggedProperties(FStructuredArchive::FSlot Slot,
 				#if WITH_EDITORONLY_DATA
 					if (SerializeContext->bTrackInitializedProperties && Property)
 					{
-						UE::SetPropertyValueInitialized(this, Data, Property, Tag.ArrayIndex);
+						UE::FInitializedPropertyValueState(this, Data).Set(Property, Tag.ArrayIndex);
 					}
 				#endif
 					continue;
@@ -1761,7 +1761,7 @@ void UStruct::SerializeVersionedTaggedProperties(FStructuredArchive::FSlot Slot,
 #if WITH_EDITORONLY_DATA
 				if (SerializeContext->bTrackSerializedProperties && Property)
 				{
-					UE::MarkPropertyValueSerialized(this, Data, Property, Tag.ArrayIndex);
+					UE::FSerializedPropertyValueState(this, Data).Set(Property, Tag.ArrayIndex);
 				}
 #endif
 
@@ -1926,7 +1926,7 @@ void UStruct::SerializeVersionedTaggedProperties(FStructuredArchive::FSlot Slot,
 				#if WITH_EDITORONLY_DATA
 					if (SerializeContext->bTrackInitializedProperties && Property)
 					{
-						UE::SetPropertyValueInitialized(this, Data, Property, Tag.ArrayIndex);
+						UE::FInitializedPropertyValueState(this, Data).Set(Property, Tag.ArrayIndex);
 					}
 				#endif
 				}
@@ -1986,7 +1986,7 @@ void UStruct::SerializeVersionedTaggedProperties(FStructuredArchive::FSlot Slot,
 						(FOverridableSerializationLogic::IsEnabled() && FOverridableSerializationLogic::GetOverriddenPropertyOperation(UnderlyingArchive, Property, DataPtr, DefaultValue) != EOverriddenPropertyOperation::None) ||
 						(!FOverridableSerializationLogic::IsEnabled() && !Property->Identical(DataPtr, DefaultValue, UnderlyingArchive.GetPortFlags())));
 				#if WITH_EDITORONLY_DATA
-					const bool bInitializedValue = !SerializeContext->bTrackInitializedProperties || UE::IsPropertyValueInitialized(this, Data, Property, Idx);
+					const bool bInitializedValue = !SerializeContext->bTrackInitializedProperties || UE::FInitializedPropertyValueState(this, Data).IsSet(Property, Idx);
 					if (bInitializedValue && (bSerializeValue || SerializeContext->bTrackInitializedProperties))
 				#else
 					if (bSerializeValue)
diff --git a/Engine/Source/Runtime/CoreUObject/Private/UObject/InstanceDataObjectUtils.cpp b/Engine/Source/Runtime/CoreUObject/Private/UObject/InstanceDataObjectUtils.cpp
--- a/Engine/Source/Runtime/CoreUObject/Private/UObject/InstanceDataObjectUtils.cpp
+++ b/Engine/Source/Runtime/CoreUObject/Private/UObject/InstanceDataObjectUtils.cpp
@@ -2,14 +2,12 @@
 
 #include "UObject/InstanceDataObjectUtils.h"
 
-#include "Serialization/ObjectReader.h"
-#include "Serialization/ObjectWriter.h"
-#include "VerseVM/VVMVerseClass.h"
-
 #if WITH_EDITORONLY_DATA
 
 #include "HAL/IConsoleManager.h"
 #include "Misc/ReverseIterate.h"
+#include "Serialization/ObjectReader.h"
+#include "Serialization/ObjectWriter.h"
 #include "Serialization/StructuredArchive.h"
 #include "UObject/Class.h"
 #include "UObject/EnumProperty.h"
@@ -21,6 +19,7 @@
 #include "UObject/PropertyPathNameTree.h"
 #include "UObject/UnrealType.h"
 #include "UObject/UObjectAnnotation.h"
+#include "VerseVM/VVMVerseClass.h"
 
 // Implemented in Class.cpp
 int32 CalculatePropertyIndex(const UStruct* Struct, const FProperty* Property, int32 ArrayIndex);
@@ -263,13 +262,6 @@ namespace UE
 		ExcludedLoosePropertyTypesVar,
 		TEXT("Comma separated list of property types that will be excluded from loose properties in IDOs.")
 	);
-	
-	bool bUseAnnotationsForPropertyStateInitTracking = PROPERTY_INIT_STATE_TRACKING_ENABLE_DEFAULT;
-	FAutoConsoleVariableRef CVarUseAnnotationBasedPropertyStateInitTracking(
-		TEXT("AnnotationBasedPropertyInitStateTracking"),
-		bUseAnnotationsForPropertyStateInitTracking,
-		TEXT("Enables uobject annotation method to track property initialization state if IDOs disabled.")
-	);
 
 	bool bEnableUninitializedUI = true;
 	FAutoConsoleVariableRef EnableUninitializedUI(
@@ -988,7 +980,9 @@ namespace UE
 			{
 				SerializedValuesProperty->InitializeValue_InContainer(NewDataPtr);
 			}
-			
+
+			const FSerializedPropertyValueState OldSerializedState(OldAsStruct, OldDataPtr);
+			FSerializedPropertyValueState NewSerializedState(NewAsStruct, NewDataPtr);
 			for (const FProperty* OldSubProperty : TFieldRange<FProperty>(OldAsStruct))
 			{
 				if (const FProperty* NewSubProperty = FindMatchingProperty(NewAsStruct, OldSubProperty))
@@ -996,14 +990,14 @@ namespace UE
 					for (int32 ArrayIndex = 0; ArrayIndex < FMath::Min(OldSubProperty->ArrayDim, NewSubProperty->ArrayDim); ++ArrayIndex)
 					{
 						// copy set flags to new struct instance
-						if (WasPropertyValueSerialized(OldAsStruct, OldDataPtr, OldSubProperty, ArrayIndex))
+						if (OldSerializedState.IsSet(OldSubProperty, ArrayIndex))
 						{
-							MarkPropertyValueSerialized(NewAsStruct, NewDataPtr, NewSubProperty, ArrayIndex);
+							NewSerializedState.Set(NewSubProperty, ArrayIndex);
 						}
 						else if (NewSubProperty->GetBoolMetaData(NAME_IsLooseMetadata))
 						{
 							// loose properties should be marked as serialized regardless of whether the old struct marked them as such
-							MarkPropertyValueSerialized(NewAsStruct, NewDataPtr, NewSubProperty, ArrayIndex);
+							NewSerializedState.Set(NewSubProperty, ArrayIndex);
 						}
 					
 						// recurse
@@ -1016,20 +1010,6 @@ namespace UE
 		}
 	}
 
-	bool SupportsPropertyInitStateTracking()
-	{
-		return bUseAnnotationsForPropertyStateInitTracking || bEnableIDOSupport;
-	}
-
-	bool SupportsPropertyInitStateTracking(const UStruct* Struct, void* StructData)
-	{
-		if (bUseAnnotationsForPropertyStateInitTracking || bEnableIDOSupport)
-		{
-			return IsTrackingPropertyValueInitialized(Struct, StructData);
-		}
-		return false;
-	}
-
 } // UE
 
 #endif // WITH_EDITORONLY_DATA
diff --git a/Engine/Source/Runtime/CoreUObject/Private/UObject/Obj.cpp b/Engine/Source/Runtime/CoreUObject/Private/UObject/Obj.cpp
--- a/Engine/Source/Runtime/CoreUObject/Private/UObject/Obj.cpp
+++ b/Engine/Source/Runtime/CoreUObject/Private/UObject/Obj.cpp
@@ -478,11 +478,11 @@ void UObject::PostEditChangeProperty(FPropertyChangedEvent& PropertyChangedEvent
 			EPropertyChangeType::ArrayAdd | EPropertyChangeType::ArrayRemove | EPropertyChangeType::ArrayClear | EPropertyChangeType::ArrayMove | EPropertyChangeType::ValueSet | EPropertyChangeType::Interactive;
 		if (PropertyChangedEvent.ChangeType & EPropertyChangeType::ResetToDefault)
 		{
-			UE::ClearPropertyValueInitialized(GetClass(), this, Property);
+			UE::FInitializedPropertyValueState(this).Clear(Property);
 		}
 		else if (PropertyChangedEvent.ChangeType & SetPropertyValueInitializedStates)
 		{
-			UE::SetPropertyValueInitialized(GetClass(), this, Property);
+			UE::FInitializedPropertyValueState(this).Set(Property);
 		}
 	}
 #endif
@@ -1715,7 +1715,7 @@ void UObject::Serialize(FStructuredArchive::FRecord Record)
 			{
 				FOverridableManager::Get().SerializeOverriddenProperties(*this, Record);
 #if WITH_EDITORONLY_DATA
-				UE::SerializePropertyValueInitialized(GetClass(), this, Record);
+				UE::FInitializedPropertyValueState(this).Serialize(Record);
 #endif
 			}
 
diff --git a/Engine/Source/Runtime/CoreUObject/Private/UObject/Property.cpp b/Engine/Source/Runtime/CoreUObject/Private/UObject/Property.cpp
--- a/Engine/Source/Runtime/CoreUObject/Private/UObject/Property.cpp
+++ b/Engine/Source/Runtime/CoreUObject/Private/UObject/Property.cpp
@@ -2406,10 +2406,7 @@ const TCHAR* FProperty::ImportSingleProperty( const TCHAR* Str, void* DestData,
 					const TCHAR* Result = ImportText(Str, Property, Property->ContainerPtrToValuePtr<void>(DestData, Index));
 
 #if WITH_EDITORONLY_DATA
-					if (UE::SupportsPropertyInitStateTracking(ObjectStruct, DestData))
-					{
-						UE::SetPropertyValueInitialized(ObjectStruct, DestData, Property);
-					}
+					UE::FInitializedPropertyValueState(ObjectStruct, DestData).Set(Property, Index);
 #endif
 					
 					// in the failure case, don't return NULL so the caller can potentially skip less and get values further in the string
diff --git a/Engine/Source/Runtime/CoreUObject/Private/UObject/PropertyBagRepository.cpp b/Engine/Source/Runtime/CoreUObject/Private/UObject/PropertyBagRepository.cpp
--- a/Engine/Source/Runtime/CoreUObject/Private/UObject/PropertyBagRepository.cpp
+++ b/Engine/Source/Runtime/CoreUObject/Private/UObject/PropertyBagRepository.cpp
@@ -797,11 +797,6 @@ const UObject* FPropertyBagRepository::FindInstanceForDataObject(const UObject*
 	return Owner ? *Owner : nullptr;
 }
 
-bool FPropertyBagRepository::WasPropertyValueSerialized(const UStruct* Struct, const void* StructData, const FProperty* Property, int32 ArrayIndex)
-{
-	return UE::WasPropertyValueSerialized(Struct, StructData, Property, ArrayIndex);
-}
-
 void FPropertyBagRepository::AddReferencedObjects(FReferenceCollector& Collector)
 {
 	FPropertyBagPlaceholderTypeRegistry::Get().AddReferencedObjects(Collector);
diff --git a/Engine/Source/Runtime/CoreUObject/Private/UObject/PropertyStateTracking.cpp b/Engine/Source/Runtime/CoreUObject/Private/UObject/PropertyStateTracking.cpp
--- a/Engine/Source/Runtime/CoreUObject/Private/UObject/PropertyStateTracking.cpp
+++ b/Engine/Source/Runtime/CoreUObject/Private/UObject/PropertyStateTracking.cpp
@@ -7,7 +7,6 @@
 #include "Containers/Array.h"
 #include "Containers/Map.h"
 #include "Containers/Set.h"
-#include "Serialization/StructuredArchive.h"
 #include "String/ParseTokens.h"
 #include "Templates/SharedPointer.h"
 #include "UObject/Class.h"
@@ -19,65 +18,6 @@
 namespace UE
 {
 
-bool ActivateTrackingPropertyValueInitialized(const UStruct* Struct, void* Data)
-{
-	return Struct->ActivateTrackingPropertyValueFlag(EPropertyValueFlags::Initialized, Data);
-}
-
-bool IsTrackingPropertyValueInitialized(const UStruct* Struct, const void* Data)
-{
-	return Struct->IsTrackingPropertyValueFlag(EPropertyValueFlags::Initialized, Data);
-}
-
-bool IsPropertyValueInitialized(const UStruct* Struct, const void* Data, const FProperty* Property, int32 ArrayIndex)
-{
-	return Struct->HasPropertyValueFlag(EPropertyValueFlags::Initialized, Data, Property, ArrayIndex);
-}
-
-void SetPropertyValueInitialized(const UStruct* Struct, void* Data, const FProperty* Property, int32 ArrayIndex)
-{
-	Struct->SetPropertyValueFlag(EPropertyValueFlags::Initialized, /*bValue*/ true, Data, Property, ArrayIndex);
-}
-
-void ClearPropertyValueInitialized(const UStruct* Struct, void* Data, const FProperty* Property, int32 ArrayIndex)
-{
-	Struct->SetPropertyValueFlag(EPropertyValueFlags::Initialized, /*bValue*/ false, Data, Property, ArrayIndex);
-}
-
-void ResetPropertyValueInitialized(const UStruct* Struct, void* Data)
-{
-	Struct->ResetPropertyValueFlags(EPropertyValueFlags::Initialized, Data);
-}
-
-void SerializePropertyValueInitialized(const UStruct* Struct, void* Data, FStructuredArchiveRecord Record)
-{
-	Struct->SerializePropertyValueFlags(EPropertyValueFlags::Initialized, Data, Record, TEXT("PropertyValueInitializedFlags"));
-}
-
-///////////////////////////////////////////////////////////////////////////////////////////////////////////////////////
-
-bool ActivateTrackingPropertyValueSerialized(const UStruct* Struct, void* Data)
-{
-	return Struct->ActivateTrackingPropertyValueFlag(EPropertyValueFlags::Serialized, Data);
-}
-
-bool IsTrackingPropertyValueSerialized(const UStruct* Struct, const void* Data)
-{
-	return Struct->IsTrackingPropertyValueFlag(EPropertyValueFlags::Serialized, Data);
-}
-
-void MarkPropertyValueSerialized(const UStruct* Struct, void* Data, const FProperty* Property, int32 ArrayIndex)
-{
-	Struct->SetPropertyValueFlag(EPropertyValueFlags::Serialized, /*bValue*/ true, Data, Property, ArrayIndex);
-}
-
-bool WasPropertyValueSerialized(const UStruct* Struct, const void* Data, const FProperty* Property, int32 ArrayIndex)
-{
-	return Struct->HasPropertyValueFlag(EPropertyValueFlags::Serialized, Data, Property, ArrayIndex);
-}
-
-///////////////////////////////////////////////////////////////////////////////////////////////////////////////////////
-
 struct FUnknownPropertyTreeAnnotation
 {
 	bool IsDefault() const
diff --git a/Engine/Source/Runtime/CoreUObject/Private/UObject/UObjectGlobals.cpp b/Engine/Source/Runtime/CoreUObject/Private/UObject/UObjectGlobals.cpp
--- a/Engine/Source/Runtime/CoreUObject/Private/UObject/UObjectGlobals.cpp
+++ b/Engine/Source/Runtime/CoreUObject/Private/UObject/UObjectGlobals.cpp
@@ -112,6 +112,20 @@ static FAutoConsoleVariableRef CVarAllowParseObjectLoading(
 	ECVF_Default
 );
 
+#ifndef PROPERTY_INIT_STATE_TRACKING_ENABLE_DEFAULT
+#define PROPERTY_INIT_STATE_TRACKING_ENABLE_DEFAULT 1
+#endif
+#if WITH_EDITORONLY_DATA
+// We disable the activation of property initialization tracking when PROPERTY_INIT_STATE_TRACKING_ENABLE_DEFAULT is 0.
+// This is a workaround for an AutoRTFM issue with the lifetime of static local variables.
+bool bUseAnnotationsForPropertyInitializedStateTracking = PROPERTY_INIT_STATE_TRACKING_ENABLE_DEFAULT;
+FAutoConsoleVariableRef CVarUseAnnotationBasedPropertyInitializedStateTracking(
+	TEXT("AnnotationBasedPropertyInitStateTracking"),
+	bUseAnnotationsForPropertyInitializedStateTracking,
+	TEXT("If true, enables annotation-based tracking of UObject property initialization state.")
+);
+#endif
+
 void EndLoad(FUObjectSerializeContext* LoadContext, TArray<UPackage*>* OutLoadedPackages);
 
 COREUOBJECT_API bool GetAllowNativeComponentClassOverrides()
@@ -4097,14 +4111,15 @@ void FObjectInitializer::PostConstructInit()
 	}
 #endif // USE_CIRCULAR_DEPENDENCY_LOAD_DEFERRING
 
-	// We disable the activation of property initialization tracking when PROPERTY_INIT_STATE_TRACKING_ENABLE_DEFAULT is 0.
-	// This is a workaround for a number of Verse tests, as property tracking does not work correctly with AutoRTFM yet.
-#if WITH_EDITORONLY_DATA && PROPERTY_INIT_STATE_TRACKING_ENABLE_DEFAULT
-	// Activate tracking of initialized property state for Verse objects.
-	// The Class of the Class may not be available during very early initialization of CoreUObject.
-	if (UClass* ClassClass = Class->GetClass(); ClassClass && ClassClass->IsChildOf<UVerseClass>())
+#if WITH_EDITORONLY_DATA
+	if (bUseAnnotationsForPropertyInitializedStateTracking)
 	{
-		UE::ActivateTrackingPropertyValueInitialized(Class, Obj);
+		// Activate tracking of initialized property state for Verse objects.
+		// The Class of the Class may not be available during very early initialization of CoreUObject.
+		if (UClass* ClassClass = Class->GetClass(); ClassClass && ClassClass->IsChildOf<UVerseClass>())
+		{
+			UE::FInitializedPropertyValueState(Obj).ActivateTracking();
+		}
 	}
 #endif
 
@@ -4341,15 +4356,17 @@ void FObjectInitializer::InitProperties(UObject* Obj, UClass* DefaultsClass, UOb
 	}
 
 #if WITH_EDITORONLY_DATA
-	auto SetUpPropertyInitializedFlag = [Obj, Class, bTrackInitializedProperties = UE::IsTrackingPropertyValueInitialized(Class, Obj)](UObject* Defaults, FProperty* Property)
+	UE::FInitializedPropertyValueState InitializedState(Obj);
+	auto SetUpPropertyInitializedFlag = [InitializedState, bTrackInitializedProperties = InitializedState.IsTracking()](UObject* Defaults, FProperty* Property) mutable
 	{
 		if (bTrackInitializedProperties)
 		{
-			// A required property is only initialized if the copied default is initialized.
-			const bool bInitialized =
-				!Property->HasAllPropertyFlags(CPF_RequiredParm) ||
-				Defaults->GetClass()->HasPropertyValueFlag(EPropertyValueFlags::Initialized, Defaults, Property);
-			Class->SetPropertyValueFlag(EPropertyValueFlags::Initialized, bInitialized, Obj, Property);
+			for (int32 ArrayIndex = 0; ArrayIndex < Property->ArrayDim; ++ArrayIndex)
+			{
+				// A required property is only initialized if the copied default is initialized.
+				const bool bInitialized = !Property->HasAllPropertyFlags(CPF_RequiredParm) || UE::FInitializedPropertyValueState(Defaults).IsSet(Property, ArrayIndex);
+				InitializedState.SetTo(bInitialized, Property, ArrayIndex);
+			}
 		}
 	};
 #endif // WITH_EDITORONLY_DATA
diff --git a/Engine/Source/Runtime/CoreUObject/Tests/UObject/PropertyStateTrackingTest.cpp b/Engine/Source/Runtime/CoreUObject/Tests/UObject/PropertyStateTrackingTest.cpp
--- a/Engine/Source/Runtime/CoreUObject/Tests/UObject/PropertyStateTrackingTest.cpp
+++ b/Engine/Source/Runtime/CoreUObject/Tests/UObject/PropertyStateTrackingTest.cpp
@@ -24,112 +24,114 @@ namespace UE
 
 TEST_CASE_NAMED(FPropertyValueFlagsTest, "CoreUObject::Serialization::PropertyValueFlags", "[CoreUObject][EngineFilter]")
 {
-	const auto SavePropertyValueInitializedFlags = [](const UStruct* Struct, void* Data, TArray<uint8>& OutFlags)
+	const auto SavePropertyValueInitializedFlags = [](FInitializedPropertyValueState& InitializedState, TArray<uint8>& OutFlags)
 	{
 		FMemoryWriter Ar(OutFlags, /*bIsPersistent*/ true);
 		FBinaryArchiveFormatter Formatter(Ar);
 		FStructuredArchive StructuredAr(Formatter);
-		SerializePropertyValueInitialized(Struct, Data, StructuredAr.Open().EnterRecord());
+		InitializedState.Serialize(StructuredAr.Open().EnterRecord());
 	};
 
-	const auto LoadPropertyValueInitializedFlags = [](const UStruct* Struct, void* Data, TArray<uint8>& OutFlags)
+	const auto LoadPropertyValueInitializedFlags = [](FInitializedPropertyValueState& InitializedState, TArray<uint8>& OutFlags)
 	{
 		FMemoryReader Ar(OutFlags, /*bIsPersistent*/ true);
 		FBinaryArchiveFormatter Formatter(Ar);
 		FStructuredArchive StructuredAr(Formatter);
-		SerializePropertyValueInitialized(Struct, Data, StructuredAr.Open().EnterRecord());
+		InitializedState.Serialize(StructuredAr.Open().EnterRecord());
 	};
 
 	UTestInstanceDataObjectClass* BaseObject = NewObject<UTestInstanceDataObjectClass>();
 
 	// Test BaseObject w/o IDOs
 	{
-		UObject* TestObject = BaseObject;
-		UClass* TestClass = UTestInstanceDataObjectClass::StaticClass();
+		FInitializedPropertyValueState InitializedState(BaseObject);
 
-		CHECK_FALSE(IsTrackingPropertyValueInitialized(TestClass, TestObject));
+		CHECK_FALSE(InitializedState.IsTracking());
 
 		// Test serialization of empty flags.
 		{
 			TArray<uint8> FlagsData;
-			SavePropertyValueInitializedFlags(TestClass, TestObject, FlagsData);
-			LoadPropertyValueInitializedFlags(TestClass, TestObject, FlagsData);
-			CHECK_FALSE(IsTrackingPropertyValueInitialized(TestClass, TestObject));
+			SavePropertyValueInitializedFlags(InitializedState, FlagsData);
+			LoadPropertyValueInitializedFlags(InitializedState, FlagsData);
+			CHECK_FALSE(InitializedState.IsTracking());
 		}
 
-		CHECK(ActivateTrackingPropertyValueInitialized(TestClass, TestObject));
-		CHECK(IsTrackingPropertyValueInitialized(TestClass, TestObject));
+		CHECK(InitializedState.ActivateTracking());
+		CHECK(InitializedState.IsTracking());
 
-		FIntProperty* AProperty = FindFProperty<FIntProperty>(TestClass, TEXT("A"));
-		FFloatProperty* BProperty = FindFProperty<FFloatProperty>(TestClass, TEXT("B"));
-		FInt64Property* CProperty = FindFProperty<FInt64Property>(TestClass, TEXT("C"));
-		FIntProperty* DProperty = FindFProperty<FIntProperty>(TestClass, TEXT("D"));
-		FIntProperty* EProperty = FindFProperty<FIntProperty>(TestClass, TEXT("E"));
+		const UClass* SearchClass = BaseObject->GetClass();
+		FIntProperty* AProperty = FindFProperty<FIntProperty>(SearchClass, TEXT("A"));
+		FFloatProperty* BProperty = FindFProperty<FFloatProperty>(SearchClass, TEXT("B"));
+		FInt64Property* CProperty = FindFProperty<FInt64Property>(SearchClass, TEXT("C"));
+		FIntProperty* DProperty = FindFProperty<FIntProperty>(SearchClass, TEXT("D"));
+		FIntProperty* EProperty = FindFProperty<FIntProperty>(SearchClass, TEXT("E"));
 		REQUIRE(AProperty);
 		REQUIRE(BProperty);
 		REQUIRE(CProperty);
 		REQUIRE(DProperty);
 		REQUIRE(EProperty);
 
-		CHECK_FALSE(IsPropertyValueInitialized(TestClass, TestObject, AProperty, 0));
-		CHECK_FALSE(IsPropertyValueInitialized(TestClass, TestObject, AProperty, 1));
-		CHECK_FALSE(IsPropertyValueInitialized(TestClass, TestObject, AProperty, 2));
-		CHECK_FALSE(IsPropertyValueInitialized(TestClass, TestObject, AProperty, 3));
-		CHECK_FALSE(IsPropertyValueInitialized(TestClass, TestObject, BProperty));
-		CHECK_FALSE(IsPropertyValueInitialized(TestClass, TestObject, CProperty));
-		CHECK_FALSE(IsPropertyValueInitialized(TestClass, TestObject, DProperty));
-		CHECK_FALSE(IsPropertyValueInitialized(TestClass, TestObject, EProperty));
-
-		SetPropertyValueInitialized(TestClass, TestObject, AProperty, 1);
-		SetPropertyValueInitialized(TestClass, TestObject, AProperty, 3);
-		SetPropertyValueInitialized(TestClass, TestObject, CProperty);
-		SetPropertyValueInitialized(TestClass, TestObject, EProperty);
-
-		CHECK_FALSE(IsPropertyValueInitialized(TestClass, TestObject, AProperty, 0));
-		CHECK(IsPropertyValueInitialized(TestClass, TestObject, AProperty, 1));
-		CHECK_FALSE(IsPropertyValueInitialized(TestClass, TestObject, AProperty, 2));
-		CHECK(IsPropertyValueInitialized(TestClass, TestObject, AProperty, 3));
-		CHECK_FALSE(IsPropertyValueInitialized(TestClass, TestObject, BProperty));
-		CHECK(IsPropertyValueInitialized(TestClass, TestObject, CProperty));
-		CHECK_FALSE(IsPropertyValueInitialized(TestClass, TestObject, DProperty));
-		CHECK(IsPropertyValueInitialized(TestClass, TestObject, EProperty));
+		CHECK_FALSE(InitializedState.IsSet(AProperty, 0));
+		CHECK_FALSE(InitializedState.IsSet(AProperty, 1));
+		CHECK_FALSE(InitializedState.IsSet(AProperty, 2));
+		CHECK_FALSE(InitializedState.IsSet(AProperty, 3));
+		CHECK_FALSE(InitializedState.IsSet(BProperty));
+		CHECK_FALSE(InitializedState.IsSet(CProperty));
+		CHECK_FALSE(InitializedState.IsSet(DProperty));
+		CHECK_FALSE(InitializedState.IsSet(EProperty));
+
+		InitializedState.Set(AProperty, 1);
+		InitializedState.Set(AProperty, 3);
+		InitializedState.Set(CProperty);
+		InitializedState.Set(EProperty);
+
+		CHECK_FALSE(InitializedState.IsSet(AProperty, 0));
+		CHECK(InitializedState.IsSet(AProperty, 1));
+		CHECK_FALSE(InitializedState.IsSet(AProperty, 2));
+		CHECK(InitializedState.IsSet(AProperty, 3));
+		CHECK_FALSE(InitializedState.IsSet(BProperty));
+		CHECK(InitializedState.IsSet(CProperty));
+		CHECK_FALSE(InitializedState.IsSet(DProperty));
+		CHECK(InitializedState.IsSet(EProperty));
 
 		TArray<uint8> FlagsData;
-		SavePropertyValueInitializedFlags(TestClass, TestObject, FlagsData);
-
-		ResetPropertyValueInitialized(TestClass, TestObject);
-
-		CHECK_FALSE(IsPropertyValueInitialized(TestClass, TestObject, AProperty, 0));
-		CHECK_FALSE(IsPropertyValueInitialized(TestClass, TestObject, AProperty, 1));
-		CHECK_FALSE(IsPropertyValueInitialized(TestClass, TestObject, AProperty, 2));
-		CHECK_FALSE(IsPropertyValueInitialized(TestClass, TestObject, AProperty, 3));
-		CHECK_FALSE(IsPropertyValueInitialized(TestClass, TestObject, BProperty));
-		CHECK_FALSE(IsPropertyValueInitialized(TestClass, TestObject, CProperty));
-		CHECK_FALSE(IsPropertyValueInitialized(TestClass, TestObject, DProperty));
-		CHECK_FALSE(IsPropertyValueInitialized(TestClass, TestObject, EProperty));
-
-		LoadPropertyValueInitializedFlags(TestClass, TestObject, FlagsData);
-
-		CHECK_FALSE(IsPropertyValueInitialized(TestClass, TestObject, AProperty, 0));
-		CHECK(IsPropertyValueInitialized(TestClass, TestObject, AProperty, 1));
-		CHECK_FALSE(IsPropertyValueInitialized(TestClass, TestObject, AProperty, 2));
-		CHECK(IsPropertyValueInitialized(TestClass, TestObject, AProperty, 3));
-		CHECK_FALSE(IsPropertyValueInitialized(TestClass, TestObject, BProperty));
-		CHECK(IsPropertyValueInitialized(TestClass, TestObject, CProperty));
-		CHECK_FALSE(IsPropertyValueInitialized(TestClass, TestObject, DProperty));
-		CHECK(IsPropertyValueInitialized(TestClass, TestObject, EProperty));
+		SavePropertyValueInitializedFlags(InitializedState, FlagsData);
+
+		InitializedState.Reset();
+
+		CHECK_FALSE(InitializedState.IsSet(AProperty, 0));
+		CHECK_FALSE(InitializedState.IsSet(AProperty, 1));
+		CHECK_FALSE(InitializedState.IsSet(AProperty, 2));
+		CHECK_FALSE(InitializedState.IsSet(AProperty, 3));
+		CHECK_FALSE(InitializedState.IsSet(BProperty));
+		CHECK_FALSE(InitializedState.IsSet(CProperty));
+		CHECK_FALSE(InitializedState.IsSet(DProperty));
+		CHECK_FALSE(InitializedState.IsSet(EProperty));
+
+		LoadPropertyValueInitializedFlags(InitializedState, FlagsData);
+
+		CHECK_FALSE(InitializedState.IsSet(AProperty, 0));
+		CHECK(InitializedState.IsSet(AProperty, 1));
+		CHECK_FALSE(InitializedState.IsSet(AProperty, 2));
+		CHECK(InitializedState.IsSet(AProperty, 3));
+		CHECK_FALSE(InitializedState.IsSet(BProperty));
+		CHECK(InitializedState.IsSet(CProperty));
+		CHECK_FALSE(InitializedState.IsSet(DProperty));
+		CHECK(InitializedState.IsSet(EProperty));
 
 		UTestInstanceDataObjectClass* EmptyObject = NewObject<UTestInstanceDataObjectClass>();
-		LoadPropertyValueInitializedFlags(TestClass, EmptyObject, FlagsData);
+		FInitializedPropertyValueState EmptyInitializedState(EmptyObject);
 
-		CHECK_FALSE(IsPropertyValueInitialized(TestClass, EmptyObject, AProperty, 0));
-		CHECK(IsPropertyValueInitialized(TestClass, EmptyObject, AProperty, 1));
-		CHECK_FALSE(IsPropertyValueInitialized(TestClass, EmptyObject, AProperty, 2));
-		CHECK(IsPropertyValueInitialized(TestClass, EmptyObject, AProperty, 3));
-		CHECK_FALSE(IsPropertyValueInitialized(TestClass, EmptyObject, BProperty));
-		CHECK(IsPropertyValueInitialized(TestClass, EmptyObject, CProperty));
-		CHECK_FALSE(IsPropertyValueInitialized(TestClass, EmptyObject, DProperty));
-		CHECK(IsPropertyValueInitialized(TestClass, EmptyObject, EProperty));
+		LoadPropertyValueInitializedFlags(EmptyInitializedState, FlagsData);
+
+		CHECK_FALSE(EmptyInitializedState.IsSet(AProperty, 0));
+		CHECK(EmptyInitializedState.IsSet(AProperty, 1));
+		CHECK_FALSE(EmptyInitializedState.IsSet(AProperty, 2));
+		CHECK(EmptyInitializedState.IsSet(AProperty, 3));
+		CHECK_FALSE(EmptyInitializedState.IsSet(BProperty));
+		CHECK(EmptyInitializedState.IsSet(CProperty));
+		CHECK_FALSE(EmptyInitializedState.IsSet(DProperty));
+		CHECK(EmptyInitializedState.IsSet(EProperty));
 	}
 
 	// Create IDO for BaseObject
@@ -141,69 +143,72 @@ TEST_CASE_NAMED(FPropertyValueFlagsTest, "CoreUObject::Serialization::PropertyVa
 
 	// Test BaseObject w/ IDOs
 	{
-		ResetPropertyValueInitialized(TestClass, TestObject);
+		FInitializedPropertyValueState InitializedState(TestObject);
+
+		InitializedState.Reset();
 
-		CHECK(ActivateTrackingPropertyValueInitialized(TestClass, TestObject));
-		CHECK(IsTrackingPropertyValueInitialized(TestClass, TestObject));
+		CHECK(InitializedState.ActivateTracking());
+		CHECK(InitializedState.IsTracking());
 
-		FIntProperty* AProperty = FindFProperty<FIntProperty>(TestClass, TEXT("A"));
-		FFloatProperty* BProperty = FindFProperty<FFloatProperty>(TestClass, TEXT("B"));
-		FInt64Property* CProperty = FindFProperty<FInt64Property>(TestClass, TEXT("C"));
-		FIntProperty* DProperty = FindFProperty<FIntProperty>(TestClass, TEXT("D"));
-		FIntProperty* EProperty = FindFProperty<FIntProperty>(TestClass, TEXT("E"));
+		const UClass* SearchClass = TestObject->GetClass();
+		FIntProperty* AProperty = FindFProperty<FIntProperty>(SearchClass, TEXT("A"));
+		FFloatProperty* BProperty = FindFProperty<FFloatProperty>(SearchClass, TEXT("B"));
+		FInt64Property* CProperty = FindFProperty<FInt64Property>(SearchClass, TEXT("C"));
+		FIntProperty* DProperty = FindFProperty<FIntProperty>(SearchClass, TEXT("D"));
+		FIntProperty* EProperty = FindFProperty<FIntProperty>(SearchClass, TEXT("E"));
 		REQUIRE(AProperty);
 		REQUIRE(BProperty);
 		REQUIRE(CProperty);
 		REQUIRE(DProperty);
 		REQUIRE(EProperty);
 
-		CHECK_FALSE(IsPropertyValueInitialized(TestClass, TestObject, AProperty, 0));
-		CHECK_FALSE(IsPropertyValueInitialized(TestClass, TestObject, AProperty, 1));
-		CHECK_FALSE(IsPropertyValueInitialized(TestClass, TestObject, AProperty, 2));
-		CHECK_FALSE(IsPropertyValueInitialized(TestClass, TestObject, AProperty, 3));
-		CHECK_FALSE(IsPropertyValueInitialized(TestClass, TestObject, BProperty));
-		CHECK_FALSE(IsPropertyValueInitialized(TestClass, TestObject, CProperty));
-		CHECK_FALSE(IsPropertyValueInitialized(TestClass, TestObject, DProperty));
-		CHECK_FALSE(IsPropertyValueInitialized(TestClass, TestObject, EProperty));
-
-		SetPropertyValueInitialized(TestClass, TestObject, AProperty, 1);
-		SetPropertyValueInitialized(TestClass, TestObject, AProperty, 3);
-		SetPropertyValueInitialized(TestClass, TestObject, CProperty);
-		SetPropertyValueInitialized(TestClass, TestObject, EProperty);
-
-		CHECK_FALSE(IsPropertyValueInitialized(TestClass, TestObject, AProperty, 0));
-		CHECK(IsPropertyValueInitialized(TestClass, TestObject, AProperty, 1));
-		CHECK_FALSE(IsPropertyValueInitialized(TestClass, TestObject, AProperty, 2));
-		CHECK(IsPropertyValueInitialized(TestClass, TestObject, AProperty, 3));
-		CHECK_FALSE(IsPropertyValueInitialized(TestClass, TestObject, BProperty));
-		CHECK(IsPropertyValueInitialized(TestClass, TestObject, CProperty));
-		CHECK_FALSE(IsPropertyValueInitialized(TestClass, TestObject, DProperty));
-		CHECK(IsPropertyValueInitialized(TestClass, TestObject, EProperty));
+		CHECK_FALSE(InitializedState.IsSet(AProperty, 0));
+		CHECK_FALSE(InitializedState.IsSet(AProperty, 1));
+		CHECK_FALSE(InitializedState.IsSet(AProperty, 2));
+		CHECK_FALSE(InitializedState.IsSet(AProperty, 3));
+		CHECK_FALSE(InitializedState.IsSet(BProperty));
+		CHECK_FALSE(InitializedState.IsSet(CProperty));
+		CHECK_FALSE(InitializedState.IsSet(DProperty));
+		CHECK_FALSE(InitializedState.IsSet(EProperty));
+
+		InitializedState.Set(AProperty, 1);
+		InitializedState.Set(AProperty, 3);
+		InitializedState.Set(CProperty);
+		InitializedState.Set(EProperty);
+
+		CHECK_FALSE(InitializedState.IsSet(AProperty, 0));
+		CHECK(InitializedState.IsSet(AProperty, 1));
+		CHECK_FALSE(InitializedState.IsSet(AProperty, 2));
+		CHECK(InitializedState.IsSet(AProperty, 3));
+		CHECK_FALSE(InitializedState.IsSet(BProperty));
+		CHECK(InitializedState.IsSet(CProperty));
+		CHECK_FALSE(InitializedState.IsSet(DProperty));
+		CHECK(InitializedState.IsSet(EProperty));
 
 		TArray<uint8> FlagsData;
-		SavePropertyValueInitializedFlags(TestClass, TestObject, FlagsData);
+		SavePropertyValueInitializedFlags(InitializedState, FlagsData);
 
-		ResetPropertyValueInitialized(TestClass, TestObject);
+		InitializedState.Reset();
 
-		CHECK_FALSE(IsPropertyValueInitialized(TestClass, TestObject, AProperty, 0));
-		CHECK_FALSE(IsPropertyValueInitialized(TestClass, TestObject, AProperty, 1));
-		CHECK_FALSE(IsPropertyValueInitialized(TestClass, TestObject, AProperty, 2));
-		CHECK_FALSE(IsPropertyValueInitialized(TestClass, TestObject, AProperty, 3));
-		CHECK_FALSE(IsPropertyValueInitialized(TestClass, TestObject, BProperty));
-		CHECK_FALSE(IsPropertyValueInitialized(TestClass, TestObject, CProperty));
-		CHECK_FALSE(IsPropertyValueInitialized(TestClass, TestObject, DProperty));
-		CHECK_FALSE(IsPropertyValueInitialized(TestClass, TestObject, EProperty));
+		CHECK_FALSE(InitializedState.IsSet(AProperty, 0));
+		CHECK_FALSE(InitializedState.IsSet(AProperty, 1));
+		CHECK_FALSE(InitializedState.IsSet(AProperty, 2));
+		CHECK_FALSE(InitializedState.IsSet(AProperty, 3));
+		CHECK_FALSE(InitializedState.IsSet(BProperty));
+		CHECK_FALSE(InitializedState.IsSet(CProperty));
+		CHECK_FALSE(InitializedState.IsSet(DProperty));
+		CHECK_FALSE(InitializedState.IsSet(EProperty));
 
-		LoadPropertyValueInitializedFlags(TestClass, TestObject, FlagsData);
+		LoadPropertyValueInitializedFlags(InitializedState, FlagsData);
 
-		CHECK_FALSE(IsPropertyValueInitialized(TestClass, TestObject, AProperty, 0));
-		CHECK(IsPropertyValueInitialized(TestClass, TestObject, AProperty, 1));
-		CHECK_FALSE(IsPropertyValueInitialized(TestClass, TestObject, AProperty, 2));
-		CHECK(IsPropertyValueInitialized(TestClass, TestObject, AProperty, 3));
-		CHECK_FALSE(IsPropertyValueInitialized(TestClass, TestObject, BProperty));
-		CHECK(IsPropertyValueInitialized(TestClass, TestObject, CProperty));
-		CHECK_FALSE(IsPropertyValueInitialized(TestClass, TestObject, DProperty));
-		CHECK(IsPropertyValueInitialized(TestClass, TestObject, EProperty));
+		CHECK_FALSE(InitializedState.IsSet(AProperty, 0));
+		CHECK(InitializedState.IsSet(AProperty, 1));
+		CHECK_FALSE(InitializedState.IsSet(AProperty, 2));
+		CHECK(InitializedState.IsSet(AProperty, 3));
+		CHECK_FALSE(InitializedState.IsSet(BProperty));
+		CHECK(InitializedState.IsSet(CProperty));
+		CHECK_FALSE(InitializedState.IsSet(DProperty));
+		CHECK(InitializedState.IsSet(EProperty));
 	}
 
 	// Test BaseObject.Struct with IDOs
@@ -212,7 +217,10 @@ TEST_CASE_NAMED(FPropertyValueFlagsTest, "CoreUObject::Serialization::PropertyVa
 		REQUIRE(StructProperty);
 		REQUIRE(StructProperty->Struct);
 		void* StructData = StructProperty->ContainerPtrToValuePtr<void>(TestObject);
-		ResetPropertyValueInitialized(StructProperty->Struct, StructData);
+
+		FInitializedPropertyValueState InitializedState(StructProperty->Struct, StructData);
+
+		InitializedState.Reset();
 
 		FIntProperty* AProperty = FindFProperty<FIntProperty>(StructProperty->Struct, TEXT("A"));
 		FIntProperty* BProperty = FindFProperty<FIntProperty>(StructProperty->Struct, TEXT("B"));
@@ -223,41 +231,41 @@ TEST_CASE_NAMED(FPropertyValueFlagsTest, "CoreUObject::Serialization::PropertyVa
 		REQUIRE(CProperty);
 		REQUIRE(DProperty);
 
-		CHECK_FALSE(IsPropertyValueInitialized(StructProperty->Struct, StructData, AProperty));
-		CHECK_FALSE(IsPropertyValueInitialized(StructProperty->Struct, StructData, BProperty));
-		CHECK_FALSE(IsPropertyValueInitialized(StructProperty->Struct, StructData, CProperty));
-		CHECK_FALSE(IsPropertyValueInitialized(StructProperty->Struct, StructData, DProperty));
+		CHECK_FALSE(InitializedState.IsSet(AProperty));
+		CHECK_FALSE(InitializedState.IsSet(BProperty));
+		CHECK_FALSE(InitializedState.IsSet(CProperty));
+		CHECK_FALSE(InitializedState.IsSet(DProperty));
 
-		SetPropertyValueInitialized(StructProperty->Struct, StructData, AProperty);
-		SetPropertyValueInitialized(StructProperty->Struct, StructData, CProperty);
-		CHECK(IsPropertyValueInitialized(StructProperty->Struct, StructData, AProperty));
-		CHECK_FALSE(IsPropertyValueInitialized(StructProperty->Struct, StructData, BProperty));
-		CHECK(IsPropertyValueInitialized(StructProperty->Struct, StructData, CProperty));
-		CHECK_FALSE(IsPropertyValueInitialized(StructProperty->Struct, StructData, DProperty));
+		InitializedState.Set(AProperty);
+		InitializedState.Set(CProperty);
+		CHECK(InitializedState.IsSet(AProperty));
+		CHECK_FALSE(InitializedState.IsSet(BProperty));
+		CHECK(InitializedState.IsSet(CProperty));
+		CHECK_FALSE(InitializedState.IsSet(DProperty));
 
-		ClearPropertyValueInitialized(StructProperty->Struct, StructData, AProperty);
-		ClearPropertyValueInitialized(StructProperty->Struct, StructData, CProperty);
-		CHECK_FALSE(IsPropertyValueInitialized(StructProperty->Struct, StructData, AProperty));
-		CHECK_FALSE(IsPropertyValueInitialized(StructProperty->Struct, StructData, CProperty));
+		InitializedState.Clear(AProperty);
+		InitializedState.Clear(CProperty);
+		CHECK_FALSE(InitializedState.IsSet(AProperty));
+		CHECK_FALSE(InitializedState.IsSet(CProperty));
 
-		SetPropertyValueInitialized(StructProperty->Struct, StructData, BProperty);
-		SetPropertyValueInitialized(StructProperty->Struct, StructData, DProperty);
+		InitializedState.Set(BProperty);
+		InitializedState.Set(DProperty);
 
 		TArray<uint8> FlagsData;
-		SavePropertyValueInitializedFlags(StructProperty->Struct, StructData, FlagsData);
+		SavePropertyValueInitializedFlags(InitializedState, FlagsData);
 
-		ResetPropertyValueInitialized(StructProperty->Struct, StructData);
-		CHECK_FALSE(IsPropertyValueInitialized(StructProperty->Struct, StructData, AProperty));
-		CHECK_FALSE(IsPropertyValueInitialized(StructProperty->Struct, StructData, BProperty));
-		CHECK_FALSE(IsPropertyValueInitialized(StructProperty->Struct, StructData, CProperty));
-		CHECK_FALSE(IsPropertyValueInitialized(StructProperty->Struct, StructData, DProperty));
+		InitializedState.Reset();
+		CHECK_FALSE(InitializedState.IsSet(AProperty));
+		CHECK_FALSE(InitializedState.IsSet(BProperty));
+		CHECK_FALSE(InitializedState.IsSet(CProperty));
+		CHECK_FALSE(InitializedState.IsSet(DProperty));
 
-		LoadPropertyValueInitializedFlags(StructProperty->Struct, StructData, FlagsData);
+		LoadPropertyValueInitializedFlags(InitializedState, FlagsData);
 
-		CHECK_FALSE(IsPropertyValueInitialized(StructProperty->Struct, StructData, AProperty));
-		CHECK(IsPropertyValueInitialized(StructProperty->Struct, StructData, BProperty));
-		CHECK_FALSE(IsPropertyValueInitialized(StructProperty->Struct, StructData, CProperty));
-		CHECK(IsPropertyValueInitialized(StructProperty->Struct, StructData, DProperty));
+		CHECK_FALSE(InitializedState.IsSet(AProperty));
+		CHECK(InitializedState.IsSet(BProperty));
+		CHECK_FALSE(InitializedState.IsSet(CProperty));
+		CHECK(InitializedState.IsSet(DProperty));
 	}
 }
 
@@ -309,7 +317,9 @@ TEST_CASE_NAMED(FTrackInitializedPropertiesTest, "CoreUObject::Serialization::Tr
 	UObject* TestObject = NewObject<UTestInstanceDataObjectClass>();
 	UClass* TestClass = TestObject->GetClass();
 	UObject* DefaultObject = TestClass->GetDefaultObject();
-	ActivateTrackingPropertyValueInitialized(TestClass, TestObject);
+	FInitializedPropertyValueState InitializedState(TestObject);
+	
+	InitializedState.ActivateTracking();
 
 	FStructProperty* StructProperty = FindFProperty<FStructProperty>(TestClass, TEXT("Struct"));
 	REQUIRE(StructProperty);
@@ -333,10 +343,10 @@ TEST_CASE_NAMED(FTrackInitializedPropertiesTest, "CoreUObject::Serialization::Tr
 	CProperty->SetPropertyValue_InContainer(TestObject, 6);
 	DProperty->SetPropertyValue_InContainer(TestObject, -1);
 
-	SetPropertyValueInitialized(TestClass, TestObject, AProperty, 0);
-	SetPropertyValueInitialized(TestClass, TestObject, AProperty, 2);
-	SetPropertyValueInitialized(TestClass, TestObject, BProperty);
-	SetPropertyValueInitialized(TestClass, TestObject, DProperty);
+	InitializedState.Set(AProperty, 0);
+	InitializedState.Set(AProperty, 2);
+	InitializedState.Set(BProperty);
+	InitializedState.Set(DProperty);
 
 	FUObjectSerializeContext* SerializeContext = FUObjectThreadContext::Get().GetSerializeContext();
 	TGuardValue<bool> TrackInitializedPropertiesScope(SerializeContext->bTrackInitializedProperties, true);
@@ -367,7 +377,7 @@ TEST_CASE_NAMED(FTrackInitializedPropertiesTest, "CoreUObject::Serialization::Tr
 	CProperty->SetPropertyValue_InContainer(TestObject, 7);
 	DProperty->SetPropertyValue_InContainer(TestObject, 7);
 
-	ResetPropertyValueInitialized(TestClass, TestObject);
+	InitializedState.Reset();
 
 	{
 		FMemoryReader Ar(BinaryData, /*bIsPersistent*/ true);
@@ -384,13 +394,13 @@ TEST_CASE_NAMED(FTrackInitializedPropertiesTest, "CoreUObject::Serialization::Tr
 	CHECK(CProperty->GetPropertyValue_InContainer(TestObject) == 7); // C unchanged because it was not initialized
 	CHECK(DProperty->GetPropertyValue_InContainer(TestObject) == 7); // D unchanged because it was serialized without its value
 
-	CHECK(IsPropertyValueInitialized(TestClass, TestObject, AProperty, 0));
-	CHECK_FALSE(IsPropertyValueInitialized(TestClass, TestObject, AProperty, 1));
-	CHECK(IsPropertyValueInitialized(TestClass, TestObject, AProperty, 2));
-	CHECK_FALSE(IsPropertyValueInitialized(TestClass, TestObject, AProperty, 3));
-	CHECK(IsPropertyValueInitialized(TestClass, TestObject, BProperty));
-	CHECK_FALSE(IsPropertyValueInitialized(TestClass, TestObject, CProperty));
-	CHECK(IsPropertyValueInitialized(TestClass, TestObject, DProperty));
+	CHECK(InitializedState.IsSet(AProperty, 0));
+	CHECK_FALSE(InitializedState.IsSet(AProperty, 1));
+	CHECK(InitializedState.IsSet(AProperty, 2));
+	CHECK_FALSE(InitializedState.IsSet(AProperty, 3));
+	CHECK(InitializedState.IsSet(BProperty));
+	CHECK_FALSE(InitializedState.IsSet(CProperty));
+	CHECK(InitializedState.IsSet(DProperty));
 
 #if WITH_TEXT_ARCHIVE_SUPPORT
 	AProperty->SetPropertyValue_InContainer(TestObject, 7, 0);
@@ -401,7 +411,7 @@ TEST_CASE_NAMED(FTrackInitializedPropertiesTest, "CoreUObject::Serialization::Tr
 	CProperty->SetPropertyValue_InContainer(TestObject, 7);
 	DProperty->SetPropertyValue_InContainer(TestObject, 7);
 
-	ResetPropertyValueInitialized(TestClass, TestObject);
+	InitializedState.Reset();
 
 	{
 		FMemoryReader Ar(JsonData, /*bIsPersistent*/ true);
@@ -419,13 +429,13 @@ TEST_CASE_NAMED(FTrackInitializedPropertiesTest, "CoreUObject::Serialization::Tr
 	CHECK(CProperty->GetPropertyValue_InContainer(TestObject) == 7); // C unchanged because it was not initialized
 	CHECK(DProperty->GetPropertyValue_InContainer(TestObject) == 7); // D unchanged because it was serialized without its value
 
-	CHECK(IsPropertyValueInitialized(TestClass, TestObject, AProperty, 0));
-	//CHECK_FALSE(IsPropertyValueInitialized(TestClass, TestObject, AProperty, 1));
-	//CHECK(IsPropertyValueInitialized(TestClass, TestObject, AProperty, 2));
-	//CHECK_FALSE(IsPropertyValueInitialized(TestClass, TestObject, AProperty, 3));
-	CHECK(IsPropertyValueInitialized(TestClass, TestObject, BProperty));
-	CHECK_FALSE(IsPropertyValueInitialized(TestClass, TestObject, CProperty));
-	CHECK(IsPropertyValueInitialized(TestClass, TestObject, DProperty));
+	CHECK(InitializedState.IsSet(AProperty, 0));
+	//CHECK_FALSE(InitializedState.IsSet(AProperty, 1));
+	//CHECK(InitializedState.IsSet(AProperty, 2));
+	//CHECK_FALSE(InitializedState.IsSet(AProperty, 3));
+	CHECK(InitializedState.IsSet(BProperty));
+	CHECK_FALSE(InitializedState.IsSet(CProperty));
+	CHECK(InitializedState.IsSet(DProperty));
 #endif // WITH_TEXT_ARCHIVE_SUPPORT
 }
 
