Fix that the movie scene system can cause package loading within a transaction.

diff --git a/Engine/Source/Programs/AutoRTFMEngineTests/Private/AutoRTFMEngineTests.cpp b/Engine/Source/Programs/AutoRTFMEngineTests/Private/AutoRTFMEngineTests.cpp
--- a/Engine/Source/Programs/AutoRTFMEngineTests/Private/AutoRTFMEngineTests.cpp
+++ b/Engine/Source/Programs/AutoRTFMEngineTests/Private/AutoRTFMEngineTests.cpp
@@ -68,6 +68,7 @@ INT32_MAIN_INT32_ARGC_TCHAR_ARGV()
 					"LogCsvProfiler off, "
 					"LogDeviceProfileManager off, "
 					"LogConfig off, "
+					"LogRHI off, "
 					"AutoRTFMEngineTests on\""));
 		FCommandLine::Set(ToCStr(CmdLineOverride));
 	}
@@ -80,6 +81,17 @@ INT32_MAIN_INT32_ARGC_TCHAR_ARGV()
 
 	GEngine = NewObject<UAutoRTFMTestEngine>(GetTransientPackage());
 	GEngine->DefaultPhysMaterial = NewObject<UPhysicalMaterial>();
+	GEngine->WorldSettingsClass = AWorldSettings::StaticClass();
+
+	FGenericDataDrivenShaderPlatformInfo::Initialize();
+
+	// This is a bodge because we can't setup RHI without setting up the engine
+	// for real. So instead we just lie that for all feature levels we are using
+	// SM5 and yolo it'll let our tests progress.
+	for (unsigned Index = 0; Index < ERHIFeatureLevel::Num; Index++)
+	{
+		GShaderPlatformForFeatureLevel[Index] = EShaderPlatform::SP_PCD3D_SM5;
+	}
 
 	UE_LOG(LogAutoRTFMEngineTests, Display, TEXT("AutoRTFMEngineTests"));
 
diff --git a/Engine/Source/Runtime/CoreUObject/Private/Serialization/AsyncLoading.cpp b/Engine/Source/Runtime/CoreUObject/Private/Serialization/AsyncLoading.cpp
--- a/Engine/Source/Runtime/CoreUObject/Private/Serialization/AsyncLoading.cpp
+++ b/Engine/Source/Runtime/CoreUObject/Private/Serialization/AsyncLoading.cpp
@@ -5429,7 +5429,7 @@ void FAsyncLoadingThread::CancelAsyncLoadingInternal()
 
 	{
 #if THREADSAFE_UOBJECTS
-		FScopeLock Lock(&PendingRequestsCritical);
+		UE::TScopeLock Lock(PendingRequestsCritical);
 #endif
 		PendingRequests.Empty();
 	}
@@ -7483,7 +7483,7 @@ PRAGMA_ENABLE_DEPRECATION_WARNINGS
 
 	// Generate new request ID and add it immediately to the global request list (it needs to be there before we exit
 	// this function, otherwise it would be added when the packages are being processed on the async thread).
-	int32 RequestID = IAsyncPackageLoader::GetNextRequestId();
+	const int32 RequestID = IAsyncPackageLoader::GetNextRequestId();
 	TRACE_LOADTIME_BEGIN_REQUEST(RequestID);
 	AddPendingRequest(RequestID);
 
diff --git a/Engine/Source/Runtime/CoreUObject/Private/Serialization/AsyncLoadingThread.h b/Engine/Source/Runtime/CoreUObject/Private/Serialization/AsyncLoadingThread.h
--- a/Engine/Source/Runtime/CoreUObject/Private/Serialization/AsyncLoadingThread.h
+++ b/Engine/Source/Runtime/CoreUObject/Private/Serialization/AsyncLoadingThread.h
@@ -12,6 +12,7 @@
 #include "Serialization/AsyncLoading.h"
 #include "Serialization/LoadTimeTracePrivate.h"
 #include "Misc/ScopeLock.h"
+#include "Misc/TransactionallySafeCriticalSection.h"
 #include "Misc/CommandLine.h"
 #include "Misc/App.h"
 #include "HAL/Runnable.h"
@@ -76,7 +77,7 @@ struct FAsyncLoadEvent
 /** [EDL] Event queue for the event driven loader */
 struct FAsyncLoadEventQueue
 {
-	//FCriticalSection AsyncLoadEventQueueCritical; //@todoio maybe this doesn't need to be protected by a critical section. 
+	//FTransactionallySafeCriticalSection AsyncLoadEventQueueCritical; //@todoio maybe this doesn't need to be protected by a critical section. 
 	int32 RunningSerialNumber;
 	TArray<FAsyncLoadEvent> EventQueue;
 
@@ -87,7 +88,7 @@ struct FAsyncLoadEventQueue
 
 	FORCEINLINE void AddAsyncEvent(int32 UserPriority, int32 PackageSerialNumber, int32 EventSystemPriority, TFunction<void(FAsyncLoadEventArgs& Args)>&& Payload)
 	{
-		//FScopeLock Lock(&AsyncLoadEventQueueCritical);
+		//UE::TScopeLock Lock(AsyncLoadEventQueueCritical);
 		//@todoio check(FAsyncLoadingThread::IsInAsyncLoadThread());
 		EventQueue.HeapPush(FAsyncLoadEvent(UserPriority, PackageSerialNumber, EventSystemPriority, ++RunningSerialNumber, Forward<TFunction<void(FAsyncLoadEventArgs& Args)>>(Payload)));
 	}
@@ -96,7 +97,7 @@ struct FAsyncLoadEventQueue
 		FAsyncLoadEvent Event;
 		bool bResult = false;
 		{
-			//FScopeLock Lock(&AsyncLoadEventQueueCritical);
+			//UE::TScopeLock Lock(AsyncLoadEventQueueCritical);
 			//@todoio check(FAsyncLoadingThread::IsInAsyncLoadThread());
 			if (EventQueue.Num())
 			{
@@ -312,7 +313,7 @@ private:
 	TSet<int32> PendingRequests;
 #if THREADSAFE_UOBJECTS
 	/** Synchronization object for PendingRequests list */
-	FCriticalSection PendingRequestsCritical;
+	FTransactionallySafeCriticalSection PendingRequestsCritical;
 #endif
 
 	/** [ASYNC/GAME THREAD] Number of package load requests in the async loading queue */
@@ -658,7 +659,7 @@ private:
 	void AddPendingRequest(int32 RequestID)
 	{
 #if THREADSAFE_UOBJECTS
-		FScopeLock Lock(&PendingRequestsCritical);
+		UE::TScopeLock Lock(PendingRequestsCritical);
 #endif
 		if (!PendingRequests.Contains(RequestID))
 		{
@@ -672,7 +673,7 @@ private:
 	void RemovePendingRequests(TArray<int32>& RequestIDs)
 	{
 #if THREADSAFE_UOBJECTS
-		FScopeLock Lock(&PendingRequestsCritical);
+		UE::TScopeLock Lock(PendingRequestsCritical);
 #endif
 		for (int32 ID : RequestIDs)
 		{
@@ -712,7 +713,7 @@ private:
 	bool ContainsRequestInternal(int32 RequestID)
 	{
 #if THREADSAFE_UOBJECTS
-		FScopeLock Lock(&PendingRequestsCritical);
+		UE::TScopeLock Lock(PendingRequestsCritical);
 #endif
 		return PendingRequests.Contains(RequestID);
 	}
@@ -723,7 +724,7 @@ private:
 	bool ContainsAnyRequestInternal(TConstArrayView<int32> RequestIDs)
 	{
 #if THREADSAFE_UOBJECTS
-		FScopeLock Lock(&PendingRequestsCritical);
+		UE::TScopeLock Lock(PendingRequestsCritical);
 #endif
 		return Algo::AnyOf(RequestIDs, [this](int32 RequestID){ return PendingRequests.Contains(RequestID); });
 	}
@@ -734,7 +735,7 @@ private:
 	bool ContainsAnyRequestInternal(TConstArrayView<std::atomic<int32>> RequestIDs)
 	{
 #if THREADSAFE_UOBJECTS
-		FScopeLock Lock(&PendingRequestsCritical);
+		UE::TScopeLock Lock(PendingRequestsCritical);
 #endif
 		return Algo::AnyOf(RequestIDs, [this](const std::atomic<int32>& RequestID){ return PendingRequests.Contains(RequestID.load(std::memory_order_relaxed)); });
 	}
diff --git a/Engine/Source/Runtime/CoreUObject/Private/Serialization/AsyncPackageLoader.cpp b/Engine/Source/Runtime/CoreUObject/Private/Serialization/AsyncPackageLoader.cpp
--- a/Engine/Source/Runtime/CoreUObject/Private/Serialization/AsyncPackageLoader.cpp
+++ b/Engine/Source/Runtime/CoreUObject/Private/Serialization/AsyncPackageLoader.cpp
@@ -14,6 +14,7 @@
 #include "IO/IoDispatcher.h"
 #include "IO/IoDispatcherInternal.h"
 #include "HAL/IConsoleManager.h"
+#include "AutoRTFM/AutoRTFM.h"
 
 static bool GUseOldLoaderAsFallback = false;
 static FAutoConsoleVariableRef CVarGUseOldLoaderAsFallback(
@@ -40,7 +41,7 @@ FThreadSafeCounter IAsyncPackageLoader::NextPackageRequestId;
 
 int32 IAsyncPackageLoader::GetNextRequestId()
 {
-	 return NextPackageRequestId.Increment();
+	return AutoRTFM::Open([] { return NextPackageRequestId.Increment(); });
 }
 
 class FEarlyRegistrationEventsRecorder
diff --git a/Engine/Source/Runtime/CoreUObject/Private/UObject/LinkerInstancingContext.cpp b/Engine/Source/Runtime/CoreUObject/Private/UObject/LinkerInstancingContext.cpp
--- a/Engine/Source/Runtime/CoreUObject/Private/UObject/LinkerInstancingContext.cpp
+++ b/Engine/Source/Runtime/CoreUObject/Private/UObject/LinkerInstancingContext.cpp
@@ -8,13 +8,15 @@
 #include "Containers/UnrealString.h"
 #include "Templates/Tuple.h"
 #include "Misc/AutomationTest.h"
+#include "Misc/TransactionallySafeRWLock.h"
+#include "Misc/TransactionallySafeRWScopeLock.h"
 
 LLM_DEFINE_TAG(Loading_LinkerInstancingContext);
 
 class FLinkerInstancingContext::FSharedLinkerInstancingContextData
 {
 	/** The shared content needs to be protected as it could be modified by the caller as well as the loading thread */
-	mutable FRWLock Lock;
+	mutable FTransactionallySafeRWLock Lock;
 	/** Map of original package name to their instance counterpart. */
 	FLinkerInstancedPackageMap InstancedPackageMap;
 	/** Optional function to map original package name to their instance counterpart. The result of this function should be immutable, as it will be cached. */
@@ -49,14 +51,14 @@ public:
 		
 	void AddPackageMapping(FName Original, FName Instanced)
 	{
-		FWriteScopeLock ScopeLock(Lock);
+		FTransactionallySafeWriteScopeLock ScopeLock(Lock);
 		LLM_SCOPE_BYTAG(Loading_LinkerInstancingContext);
 		InstancedPackageMap.AddPackageMapping(Original, Instanced);
 	}
 
 	bool FindPackageMapping(FName Original, FName& Instanced) const
 	{
-		FReadScopeLock ScopeLock(Lock);
+		FTransactionallySafeReadScopeLock ScopeLock(Lock);
 		LLM_SCOPE_BYTAG(Loading_LinkerInstancingContext);
 		if (const FName* InstancedPtr = InstancedPackageMap.InstancedPackageMapping.Find(Original))
 		{
@@ -68,14 +70,14 @@ public:
 
 	void FixupSoftObjectPath(FSoftObjectPath& InOutSoftObjectPath) const
 	{
-		FReadScopeLock ScopeLock(Lock);
+		FTransactionallySafeReadScopeLock ScopeLock(Lock);
 		LLM_SCOPE_BYTAG(Loading_LinkerInstancingContext);
 		InstancedPackageMap.FixupSoftObjectPath(InOutSoftObjectPath);
 	}
 
 	FSoftObjectPath RemapPath(const FSoftObjectPath& Path) const
 	{
-		FReadScopeLock ScopeLock(Lock);
+		FTransactionallySafeReadScopeLock ScopeLock(Lock);
 		if (const FTopLevelAssetPath* Remapped = PathMapping.Find(Path.GetAssetPath()))
 		{
 			return FSoftObjectPath(*Remapped, Path.GetSubPathString());
@@ -85,40 +87,40 @@ public:
 
 	bool IsInstanced() const
 	{
-		FReadScopeLock ScopeLock(Lock);
+		FTransactionallySafeReadScopeLock ScopeLock(Lock);
 		return InstancedPackageMap.IsInstanced() || InstancedPackageMapFunc || PathMapping.Num() > 0;
 	}
 
 	void EnableAutomationTest()
 	{
-		FWriteScopeLock ScopeLock(Lock);
+		FTransactionallySafeWriteScopeLock ScopeLock(Lock);
 		InstancedPackageMap.EnableAutomationTest();
 	}
 
 	void BuildPackageMapping(FName Original, FName Instanced, bool bInSoftObjectPathRemappingEnabled)
 	{
-		FWriteScopeLock ScopeLock(Lock);
+		FTransactionallySafeWriteScopeLock ScopeLock(Lock);
 		LLM_SCOPE_BYTAG(Loading_LinkerInstancingContext);
 		InstancedPackageMap.BuildPackageMapping(Original, Instanced, bInSoftObjectPathRemappingEnabled);
 	}
 
 	void AddTag(FName NewTag)
 	{
-		FWriteScopeLock ScopeLock(Lock);
+		FTransactionallySafeWriteScopeLock ScopeLock(Lock);
 		LLM_SCOPE_BYTAG(Loading_LinkerInstancingContext);
 		Tags.Add(NewTag);
 	}
 
 	void AppendTags(const TSet<FName>& NewTags)
 	{
-		FWriteScopeLock ScopeLock(Lock);
+		FTransactionallySafeWriteScopeLock ScopeLock(Lock);
 		LLM_SCOPE_BYTAG(Loading_LinkerInstancingContext);
 		Tags.Append(NewTags);
 	}
 
 	bool HasTag(FName Tag) const
 	{
-		FReadScopeLock ScopeLock(Lock);
+		FTransactionallySafeReadScopeLock ScopeLock(Lock);
 		return Tags.Contains(Tag);
 	}
 
@@ -134,14 +136,14 @@ public:
 
 	void AddPathMapping(FSoftObjectPath Original, FSoftObjectPath Instanced)
 	{
-		FWriteScopeLock ScopeLock(Lock);
+		FTransactionallySafeWriteScopeLock ScopeLock(Lock);
 		LLM_SCOPE_BYTAG(Loading_LinkerInstancingContext);
 		PathMapping.Emplace(Original.GetAssetPath(), Instanced.GetAssetPath());
 	}
 
 	void SetPackageMappingFunc(TFunction<FName(FName)> InInstancedPackageMapFunc)
 	{
-		FWriteScopeLock ScopeLock(Lock);
+		FTransactionallySafeWriteScopeLock ScopeLock(Lock);
 		InstancedPackageMapFunc = MoveTemp(InInstancedPackageMapFunc);
 	}
 
@@ -151,7 +153,7 @@ public:
 
 		FName RemappedPackageName;
 		{
-			FReadScopeLock ScopeLock(Lock);
+			FTransactionallySafeReadScopeLock ScopeLock(Lock);
 			RemappedPackageName = InstancedPackageMap.RemapPackage(PackageName);
 
 			if ((RemappedPackageName == PackageName) && InstancedPackageMapFunc)
@@ -167,7 +169,7 @@ public:
 
 		if (bAddPackageMapping)
 		{
-			FWriteScopeLock ScopeLock(Lock);
+			FTransactionallySafeWriteScopeLock ScopeLock(Lock);
 			LLM_SCOPE_BYTAG(Loading_LinkerInstancingContext);
 			InstancedPackageMap.AddPackageMapping(PackageName, RemappedPackageName);
 		}
diff --git a/Engine/Source/Runtime/Engine/Private/Tests/AutoRTFM/AutoRTFMActorComponentTests.cpp b/Engine/Source/Runtime/Engine/Private/Tests/AutoRTFM/AutoRTFMActorComponentTests.cpp
--- a/Engine/Source/Runtime/Engine/Private/Tests/AutoRTFM/AutoRTFMActorComponentTests.cpp
+++ b/Engine/Source/Runtime/Engine/Private/Tests/AutoRTFM/AutoRTFMActorComponentTests.cpp
@@ -13,6 +13,7 @@
 #include "AutoRTFMTestPrimitiveComponent.h"
 #include "Chaos/Core.h"
 #include "Components/SkeletalMeshComponent.h"
+#include "Engine/LevelStreamingPersistent.h"
 #include "Engine/SkeletalMesh.h"
 #include "Engine/World.h"
 #include "PBDRigidsSolver.h"
@@ -102,6 +103,7 @@ protected:
 		{
 			UWorld* World = NewObject<UWorld>();
 			World->CreatePhysicsScene(nullptr);
+			World->InitializeNewWorld();
 
 			UAutoRTFMTestLevel* Level = NewObject<UAutoRTFMTestLevel>();
 			Level->OwningWorld = World;
@@ -110,6 +112,8 @@ protected:
 
 			Run(World, Level, Actor, Component);
 
+			World->CleanupWorld();
+
 			if (Component->IsRegistered())
 			{
 				Component->UnregisterComponent();
@@ -467,6 +471,49 @@ AUTORTFM_ACTOR_COMPONENT_TEST(CloseDisregardForGC)
 	TEST_CHECK_TRUE(AutoRTFM::ETransactionResult::AbortedByRequest == Result);
 }
 
+AUTORTFM_ACTOR_COMPONENT_TEST(WorldGetWorldSettings)
+{
+	AWorldSettings* Settings = nullptr;
+
+	AutoRTFM::ETransactionResult Result = AutoRTFM::Transact([&]
+		{
+			Settings = World->GetWorldSettings();
+		});
+
+	TEST_CHECK_TRUE(AutoRTFM::ETransactionResult::Committed == Result);
+	TEST_CHECK_TRUE(Settings);
+}
+
+AUTORTFM_ACTOR_COMPONENT_TEST(WorldProcessLevelStreamingVolumes)
+{
+	AutoRTFM::ETransactionResult Result = AutoRTFM::Transact([&]
+		{
+			World->ProcessLevelStreamingVolumes();
+		});
+
+	TEST_CHECK_TRUE(AutoRTFM::ETransactionResult::Committed == Result);
+}
+
+AUTORTFM_ACTOR_COMPONENT_TEST(WorldBlockTillLevelStreamingCompleted)
+{
+	ULevelStreamingPersistent* LevelStreamingPersistent = NewObject<ULevelStreamingPersistent>(
+		World,
+		TEXT("WOWWEE"));
+	LevelStreamingPersistent->SetWorldAsset(World);
+	World->AddStreamingLevel(LevelStreamingPersistent);
+
+	TEST_CHECK_TRUE(World->HasStreamingLevelsToConsider());
+
+	AutoRTFM::ETransactionResult Result = AutoRTFM::Transact([&]
+		{
+			World->BlockTillLevelStreamingCompleted();
+		});
+
+	TEST_CHECK_TRUE(!World->HasStreamingLevelsToConsider());
+
+	TEST_CHECK_TRUE(AutoRTFM::ETransactionResult::Committed == Result);
+}
+
 }  // anonymous namespace
 
 #undef TEST_CHECK_FALSE
