IsRemote(Id) will no longer always assume that the fact that a remote object stub doesn't exist means an object is local and will inspect the server id associated with remote object id too.

diff --git a/Engine/Source/Runtime/CoreUObject/Private/UObject/RemoteObject.cpp b/Engine/Source/Runtime/CoreUObject/Private/UObject/RemoteObject.cpp
--- a/Engine/Source/Runtime/CoreUObject/Private/UObject/RemoteObject.cpp
+++ b/Engine/Source/Runtime/CoreUObject/Private/UObject/RemoteObject.cpp
@@ -228,10 +228,13 @@ bool IsRemote(FRemoteObjectId ObjectId)
 	{
 		return IsRemote(Object);
 	}
-	else
+	if (FindRemoteObjectStub(ObjectId))
 	{
-		return !!FindRemoteObjectStub(ObjectId);
+		return true;
 	}
+	const FRemoteServerId ServerId = ObjectId.GetServerId();
+	// Invalid server Id means local native classes which are created before a server has a chance to have an id assigned
+	return ServerId.IsValid() && ServerId != UE::RemoteObject::GetGlobalServerId();
 }
 
 bool IsRemote(const UObject* Object)
