Weak object pointers will no longer attempt to resolve remote objects when serializing remote object data

diff --git a/Engine/Source/Runtime/CoreUObject/Private/UObject/RemoteObjectSerialization.cpp b/Engine/Source/Runtime/CoreUObject/Private/UObject/RemoteObjectSerialization.cpp
--- a/Engine/Source/Runtime/CoreUObject/Private/UObject/RemoteObjectSerialization.cpp
+++ b/Engine/Source/Runtime/CoreUObject/Private/UObject/RemoteObjectSerialization.cpp
@@ -369,7 +369,18 @@ public:
 	virtual FArchive& operator<<(FSoftObjectPath& Value) override { return FArchiveUObject::SerializeSoftObjectPath(*this, Value); }
 	virtual FArchive& operator<<(FWeakObjectPtr& Value) override
 	{
-		FObjectPtr Ptr(Value.Get(/** bEvenIfGarbage */ true));
+		FObjectPtr Ptr;
+#if UE_WITH_REMOTE_OBJECT_HANDLE
+		const FRemoteObjectId WeakPtrId = Value.GetRemoteId();
+		if (UE::RemoteObject::Handle::IsRemote(WeakPtrId))
+		{
+			Ptr = UE::RemoteObject::Transfer::RemoteObjectIdToObjectPtr(WeakPtrId);
+		}
+		else
+#endif
+		{
+			Ptr = Value.Get(/** bEvenIfGarbage */ true);
+		}
 		WriteObjectReference(Ptr);
 
 		return *this;
