[BulkData]
 - Add function to perform deep copy

#rnx

diff --git a/Engine/Source/Runtime/CoreUObject/Public/Serialization/BulkData.h b/Engine/Source/Runtime/CoreUObject/Public/Serialization/BulkData.h
--- a/Engine/Source/Runtime/CoreUObject/Public/Serialization/BulkData.h
+++ b/Engine/Source/Runtime/CoreUObject/Public/Serialization/BulkData.h
@@ -1171,6 +1171,19 @@ public:
 		}
 		Formats.Empty();
 	}
+	
+	void DeepCopy(const FFormatContainer& Other)
+	{
+		// Safely clean up any entries that the current FFormatContainer has before creating/adding copies from the provided container.
+		FlushData();
+
+		for (const TPair<FName, FByteBulkData*>& Format : Other.Formats)
+		{
+			FByteBulkData* Copy = new FByteBulkData(*Format.Value);
+			Formats.Add(Format.Key, Copy);
+		}
+	}
+
 	COREUOBJECT_API void Serialize(FArchive& Ar, UObject* Owner, const TArray<FName>* FormatsToSave = nullptr, bool bSingleUse = true, uint16 InAlignment = DEFAULT_ALIGNMENT, bool bInline = true, bool bMapped = false);
 	COREUOBJECT_API void SerializeAttemptMappedLoad(FArchive& Ar, UObject* Owner);
 };
