Added a new FindPropertyByOffset function that finds a property with a given offset.

[RN] Added a new FindPropertyByOffset function that finds a property with a given offset.

diff --git a/Engine/Source/Runtime/CoreUObject/Private/UObject/Property.cpp b/Engine/Source/Runtime/CoreUObject/Private/UObject/Property.cpp
--- a/Engine/Source/Runtime/CoreUObject/Private/UObject/Property.cpp
+++ b/Engine/Source/Runtime/CoreUObject/Private/UObject/Property.cpp
@@ -2555,3 +2555,20 @@ FProperty* UStruct::FindPropertyByName(FName InName) const
 
 	return nullptr;
 }
+
+FProperty* UStruct::FindPropertyByOffset(int32 Offset, FProperty* FirstPropertyForSearch) const
+{
+	check(!FirstPropertyForSearch || (FirstPropertyForSearch->GetOwnerStruct() == PropertyLink->GetOwnerStruct()));
+
+	FProperty* FirstProperty = FirstPropertyForSearch ? FirstPropertyForSearch : PropertyLink;
+
+	for (FProperty* Property = FirstProperty; Property != nullptr; Property = Property->PropertyLinkNext)
+	{
+		if (Property->GetOffset_ForInternal() == Offset)
+		{
+			return Property;
+		}
+	}
+
+	return nullptr;
+}
diff --git a/Engine/Source/Runtime/CoreUObject/Public/UObject/Class.h b/Engine/Source/Runtime/CoreUObject/Public/UObject/Class.h
--- a/Engine/Source/Runtime/CoreUObject/Public/UObject/Class.h
+++ b/Engine/Source/Runtime/CoreUObject/Public/UObject/Class.h
@@ -496,6 +496,14 @@ public:
 	/** Searches property link chain for a property with the specified name */
 	COREUOBJECT_API FProperty* FindPropertyByName(FName InName) const;
 
+	/**
+	 * Searches the property link chain for a property with the specified offset
+	 *
+	 * @param	Offset						offset of the property to find
+	 * @param	FirstPropertyForSearch		starts iterating over properties starting with this one
+	 */
+	COREUOBJECT_API FProperty* FindPropertyByOffset(int32 Offset, FProperty* FirstPropertyForSearch = nullptr) const;
+
 	/**
 	 * Creates new copies of components
 	 * 
