Skip checking for initialized defaults if the default object does not track initialized state

#rnx

diff --git a/Engine/Source/Runtime/CoreUObject/Private/UObject/UObjectGlobals.cpp b/Engine/Source/Runtime/CoreUObject/Private/UObject/UObjectGlobals.cpp
--- a/Engine/Source/Runtime/CoreUObject/Private/UObject/UObjectGlobals.cpp
+++ b/Engine/Source/Runtime/CoreUObject/Private/UObject/UObjectGlobals.cpp
@@ -4361,10 +4361,12 @@ void FObjectInitializer::InitProperties(UObject* Obj, UClass* DefaultsClass, UOb
 	{
 		if (bTrackInitializedProperties)
 		{
+			UE::FInitializedPropertyValueState DefaultsInitializedState(Defaults);
+			const bool bTrackingDefaults = DefaultsInitializedState.IsTracking();
 			for (int32 ArrayIndex = 0; ArrayIndex < Property->ArrayDim; ++ArrayIndex)
 			{
 				// A required property is only initialized if the copied default is initialized.
-				const bool bInitialized = !Property->HasAllPropertyFlags(CPF_RequiredParm) || UE::FInitializedPropertyValueState(Defaults).IsSet(Property, ArrayIndex);
+				const bool bInitialized = !Property->HasAllPropertyFlags(CPF_RequiredParm) || (bTrackingDefaults && DefaultsInitializedState.IsSet(Property, ArrayIndex));
 				InitializedState.SetTo(bInitialized, Property, ArrayIndex);
 			}
 		}
