Fix UB in FJsonArchiveInputFormatter::Serialize from calling memcpy with a null argument.

diff --git a/Engine/Source/Runtime/CoreUObject/Private/Serialization/Formatters/JsonArchiveInputFormatter.cpp b/Engine/Source/Runtime/CoreUObject/Private/Serialization/Formatters/JsonArchiveInputFormatter.cpp
--- a/Engine/Source/Runtime/CoreUObject/Private/Serialization/Formatters/JsonArchiveInputFormatter.cpp
+++ b/Engine/Source/Runtime/CoreUObject/Private/Serialization/Formatters/JsonArchiveInputFormatter.cpp
@@ -569,7 +569,10 @@ void FJsonArchiveInputFormatter::Serialize(void* Data, uint64 DataSize)
 	TArray<uint8> Buffer;
 	Serialize(Buffer);
 	check(Buffer.Num() == DataSize);
-	memcpy(Data, Buffer.GetData(), DataSize);
+	if (DataSize)
+	{
+		memcpy(Data, Buffer.GetData(), DataSize);
+	}
 }
 
 FString FJsonArchiveInputFormatter::EscapeFieldName(const TCHAR* Name)
