Avoid resolving remote objects when serializing weak object pointer properties

diff --git a/Engine/Source/Runtime/CoreUObject/Private/UObject/PropertyWeakObjectPtr.cpp b/Engine/Source/Runtime/CoreUObject/Private/UObject/PropertyWeakObjectPtr.cpp
--- a/Engine/Source/Runtime/CoreUObject/Private/UObject/PropertyWeakObjectPtr.cpp
+++ b/Engine/Source/Runtime/CoreUObject/Private/UObject/PropertyWeakObjectPtr.cpp
@@ -70,10 +70,21 @@ void FWeakObjectProperty::SerializeItem( FStructuredArchive::FSlot Slot, void* V
 {
 	FArchive& UnderlyingArchive = Slot.GetUnderlyingArchive();
 
-	UObject* OldObjectValue = GetObjectPropertyValue(Value);
+	UObject* OldObjectValue = nullptr;
+#if UE_WITH_REMOTE_OBJECT_HANDLE
+	if (!UnderlyingArchive.IsMigratingRemoteObjects())
+#endif
+	{
+		OldObjectValue = GetObjectPropertyValue(Value);
+	}
+
 	Slot << *(FWeakObjectPtr*)Value;
 
-	if (UnderlyingArchive.IsLoading() || UnderlyingArchive.IsModifyingWeakAndStrongReferences())
+	if (
+#if UE_WITH_REMOTE_OBJECT_HANDLE
+	!UnderlyingArchive.IsMigratingRemoteObjects() &&
+#endif
+	   (UnderlyingArchive.IsLoading() || UnderlyingArchive.IsModifyingWeakAndStrongReferences()))
 	{
 		UObject* NewObjectValue = GetObjectPropertyValue(Value);
 
