Enable Overridable Serialization for T3D Copy/Paste

diff --git a/Engine/Source/Runtime/CoreUObject/Private/UObject/OverriddenPropertySet.cpp b/Engine/Source/Runtime/CoreUObject/Private/UObject/OverriddenPropertySet.cpp
--- a/Engine/Source/Runtime/CoreUObject/Private/UObject/OverriddenPropertySet.cpp
+++ b/Engine/Source/Runtime/CoreUObject/Private/UObject/OverriddenPropertySet.cpp
@@ -39,7 +39,7 @@ struct FOverridableSerializationLogicInternalAdapter
 
 namespace Private
 {
-	bool bEnableT3D = false;
+	bool bEnableT3D = true;
 	FAutoConsoleVariableRef CVar_bT3DOverrideSerializationEnabled(
 		TEXT("OverridableSerializationLogic.Capabilities.T3D"),
 		bEnableT3D,
