Add UE::RemoteObject::Transfer::RegisterRemoteObjectId and RegisterSharedObject APIs used for sharing objects between servers

diff --git a/Engine/Source/Runtime/CoreUObject/Private/UObject/RemoteObject.cpp b/Engine/Source/Runtime/CoreUObject/Private/UObject/RemoteObject.cpp
--- a/Engine/Source/Runtime/CoreUObject/Private/UObject/RemoteObject.cpp
+++ b/Engine/Source/Runtime/CoreUObject/Private/UObject/RemoteObject.cpp
@@ -82,7 +82,7 @@ public:
 		return nullptr;
 	}
 
-	UE::RemoteObject::Handle::FRemoteObjectStub* FindOrAddRemoteObjectStub(FRemoteObjectId Id)
+	UE::RemoteObject::Handle::FRemoteObjectStub* FindOrAddRemoteObjectStub(FRemoteObjectId Id, FRemoteServerId OwnerServerId = FRemoteServerId())
 	{
 		using namespace UE::RemoteObject::Handle;
 
@@ -92,7 +92,7 @@ public:
 		{
 			Stub = new FRemoteObjectStub();
 			Stub->Id = Id;
-			Stub->RemoteServerId = Id.GetServerId();
+			Stub->RemoteServerId = OwnerServerId.IsValid() ? OwnerServerId : Id.GetServerId();
 		}		
 		return Stub;
 	}
@@ -154,6 +154,17 @@ void InitRemoteObjects()
 	FCoreDelegates::OnEndFrame.AddStatic(&UE::RemoteObject::Transfer::WaitForServersToSync);
 }
 
+void RegisterRemoteObjectId(FRemoteObjectId ObjectId, FRemoteServerId OwnerServerId)
+{
+	ObjectMaps->FindOrAddRemoteObjectStub(ObjectId, OwnerServerId);
+}
+
+void RegisterSharedObject(UObject* Object)
+{
+	MapObject(Object);
+	MarkAsLocal(Object);
+}
+
 void MarkAsRemote(UObject* Object, FRemoteServerId DestinationServerId)
 {
 	static_assert(sizeof(FObjectHandle) == sizeof(UObject*));
diff --git a/Engine/Source/Runtime/CoreUObject/Private/UObject/RemoteObjectPrivate.h b/Engine/Source/Runtime/CoreUObject/Private/UObject/RemoteObjectPrivate.h
--- a/Engine/Source/Runtime/CoreUObject/Private/UObject/RemoteObjectPrivate.h
+++ b/Engine/Source/Runtime/CoreUObject/Private/UObject/RemoteObjectPrivate.h
@@ -43,6 +43,16 @@ namespace UE::RemoteObject::Private
 	*/
 	UObject* FindObject(FRemoteObjectId ObjectId);
 
+	/**
+	* Registers a stub for a remote object that is known to be owned by a specific server
+	*/
+	void RegisterRemoteObjectId(FRemoteObjectId ObjectId, FRemoteServerId OwnerServerId);
+
+	/**
+	* Registers object for sharing, marking it as owned by the current server
+	*/
+	void RegisterSharedObject(UObject* Object);
+
 	/**
 	* Finds remote object stub
 	*/
diff --git a/Engine/Source/Runtime/CoreUObject/Private/UObject/RemoteObjectTransfer.cpp b/Engine/Source/Runtime/CoreUObject/Private/UObject/RemoteObjectTransfer.cpp
--- a/Engine/Source/Runtime/CoreUObject/Private/UObject/RemoteObjectTransfer.cpp
+++ b/Engine/Source/Runtime/CoreUObject/Private/UObject/RemoteObjectTransfer.cpp
@@ -290,6 +290,16 @@ void MigrateObjectToRemoteServer(UObject* Object, FRemoteServerId DestinationSer
 	}
 }
 
+void RegisterSharedObject(UObject* Object)
+{
+	UE::RemoteObject::Private::RegisterSharedObject(Object);
+}
+
+void RegisterRemoteObjectId(FRemoteObjectId Id, FRemoteServerId OwnerServerId)
+{
+	UE::RemoteObject::Private::RegisterRemoteObjectId(Id, OwnerServerId);
+}
+
 void MigrateObjectToRemoteServer(FRemoteObjectId ObjectId, FRemoteServerId DestinationServerId)
 {
 	using namespace UE::RemoteObject::Handle;
diff --git a/Engine/Source/Runtime/CoreUObject/Public/UObject/RemoteObjectTransfer.h b/Engine/Source/Runtime/CoreUObject/Public/UObject/RemoteObjectTransfer.h
--- a/Engine/Source/Runtime/CoreUObject/Public/UObject/RemoteObjectTransfer.h
+++ b/Engine/Source/Runtime/CoreUObject/Public/UObject/RemoteObjectTransfer.h
@@ -65,6 +65,19 @@ namespace UE::RemoteObject::Transfer
 #endif
 	}
 
+	/**
+	* Registers object ID as known to be owned by another server, without migrating it
+	* @param ObjectId Remote object id
+	* @param OwnerServerId Id of a remote server that currently owns the object
+	*/
+	COREUOBJECT_API void RegisterRemoteObjectId(FRemoteObjectId Id, FRemoteServerId OwnerServerId);
+
+	/**
+	* Registers object for sharing, marking it as owned by the current server
+	* @param Object to make shareable
+	*/
+	COREUOBJECT_API void RegisterSharedObject(UObject* Object);
+
 	/**
 	* Establishes a synchronization barrier across all servers
 	*/
