- rework object migration to add an immediate mode to support dev/testing workflows
- add PrefetchRemoteObject api
- better hang detection so we don't sleep until we've been waiting for 10 seconds

diff --git a/Engine/Source/Runtime/CoreUObject/Private/UObject/RemoteObject.cpp b/Engine/Source/Runtime/CoreUObject/Private/UObject/RemoteObject.cpp
--- a/Engine/Source/Runtime/CoreUObject/Private/UObject/RemoteObject.cpp
+++ b/Engine/Source/Runtime/CoreUObject/Private/UObject/RemoteObject.cpp
@@ -18,12 +18,6 @@
 
 DEFINE_LOG_CATEGORY(LogRemoteObject);
 
-namespace UE::RemoteObject::Transfer
-{
-	extern FRemoteObjectId GNextRequestedObjectId;
-	extern UObject* GNextRequestedDestinationOuter;
-}
-
 namespace UE::RemoteObject::Private
 {
 
@@ -172,7 +166,6 @@ void MarkAsRemote(UObject* Object, FRemoteServerId DestinationServerId)
 	ObjectItem->SetFlags(EInternalObjectFlags::Remote);
 	ObjectItem->ClearFlags(EInternalObjectFlags_RootFlags | EInternalObjectFlags::RemoteReference);
 	ObjectMaps->FindOrAddRemoteObjectStub(Object, DestinationServerId);
-	ObjectMaps->UnmapObject(Object);
 }
 
 void MarkAsLocal(UObject* Object)
@@ -278,19 +271,14 @@ UObject* ResolveObject(const FRemoteObjectStub* Stub, ERemoteReferenceType RefTy
 
 		UObject* Outer = Object ? Object->GetOuter() : ObjectMaps->FindObject(Stub->OuterId);
 
-		// ensure the global writes are put into the Open so after we abort we can read them
-		UE_AUTORTFM_OPEN
-		{
-			GNextRequestedObjectId = Stub->Id;
-			GNextRequestedDestinationOuter = Outer;
+		MigrateObjectFromRemoteServer(Stub->Id, Stub->RemoteServerId, Outer);
 
-			UE_LOG(LogObj, Warning, TEXT("Aborting transaction, need obj id %s on server %s"),
-				*Stub->Id.ToString(),
-				*Stub->RemoteServerId.ToString());
-		};
+		// if running transactionally, we will have aborted and not reached here
 
-		AutoRTFM::AbortTransaction();
+		// if running non-transactionally, object migrated immediately, so we just re-resolve
+		return ResolveObject(Stub, RefType);
 	}
+
 	return Object;
 }
 
diff --git a/Engine/Source/Runtime/CoreUObject/Private/UObject/RemoteObjectTransfer.cpp b/Engine/Source/Runtime/CoreUObject/Private/UObject/RemoteObjectTransfer.cpp
--- a/Engine/Source/Runtime/CoreUObject/Private/UObject/RemoteObjectTransfer.cpp
+++ b/Engine/Source/Runtime/CoreUObject/Private/UObject/RemoteObjectTransfer.cpp
@@ -16,9 +16,6 @@
 
 namespace UE::RemoteObject::Transfer
 {
-FRemoteObjectId GNextRequestedObjectId;
-UObject* GNextRequestedDestinationOuter;
-
 const FRemoteServerId DatabaseId(0x3ff);
 
 TDelegate<void(FRemoteObjectId ObjectId, FRemoteServerId CurrentOwnerServerId, FRemoteServerId DestinationServerId)> RequestRemoteObjectDelegate;
@@ -101,7 +98,7 @@ public:
 			{
 				DestinationOuter = RequestedObjectsDestinationOuter[RequestedObjectIndex];
 				bFoundObject = true;
-			}			
+			}
 		}
 
 		if (bFoundObject)
@@ -215,6 +212,43 @@ UObject* MigrateObjectFromObjectData(FRemoteObjectData& ObjectData, FRemoteServe
 	return ReceivedObjects[0];
 } 
 
+void MigrateObjectFromRemoteServer(FRemoteObjectId ObjectId, FRemoteServerId CurrentOwnerServerId, UObject* DestinationOuter)
+{
+	// ensure the request is added in the Open so after we abort it is preserved
+	bool bAdded = false;
+
+	UE_AUTORTFM_OPEN
+	{
+		// it's possible a prefetch already added this to the RequestedObjects list
+		bAdded = (GTransferQueue.RequestedObjects.Find(ObjectId) == INDEX_NONE);
+
+		if (bAdded)
+		{
+			GTransferQueue.RequestedObjects.Add(ObjectId);
+			GTransferQueue.RequestedObjectsDestinationOuter.Add(DestinationOuter);
+		}
+
+		UE_LOG(LogRemoteObject, Verbose, TEXT("Aborting transaction, need obj id %s on server %s"),
+			*ObjectId.ToString(),
+			*CurrentOwnerServerId.ToString());
+	};
+
+	if (AutoRTFM::IsClosed())
+	{
+		// if we are in a transaction, abort so the outer handler can renegotiate object transfers and retry
+		AutoRTFM::AbortTransaction();
+	}
+	else
+	{
+		// if we are not in a transaction, request the object immediately
+		UE::RemoteObject::Handle::RequestObject(ObjectId);
+		GTransferQueue.RequestedObjects.Empty();
+		GTransferQueue.RequestedObjectsDestinationOuter.Empty();
+		GTransferQueue.LockedObjects.Empty();
+		GTransferQueue.PendingRequests = 0;
+	}
+}
+
 void MigrateObjectToRemoteServer(UObject* Object, FRemoteServerId DestinationServerId)
 {
 	using namespace UE::RemoteObject::Serialization;
@@ -309,7 +343,7 @@ void MigrateObjectToRemoteServer(FRemoteObjectId ObjectId, FRemoteServerId Desti
 	UE_LOG(LogRemoteObject, Verbose, TEXT(" >>> MigrateObjectToRemoteServer obj %s to %s"), *ObjectId.ToString(), *DestinationServerId.ToString());
 
 	UObject* Object = FindObject(ObjectId);
-	if (Object)
+	if (Object && !IsRemote(Object))
 	{
 		UE_LOG(LogRemoteObject, Verbose, TEXT("   > MigrateObjectToRemoteServer found obj %s"), *ObjectId.ToString());
 		MigrateObjectToRemoteServer(Object, DestinationServerId);
@@ -352,6 +386,36 @@ void OnObjectDataDenied(FRemoteObjectId ObjectId, FRemoteServerId RemoteServerId
 	GTransferQueue.DenyReceiveRequest(ObjectId);
 }
 
+void PrefetchRemoteObject(FRemoteObjectId ObjectId, UObject* DestinationOuter)
+{
+	bool bAdded = false;
+	
+	UE_AUTORTFM_OPEN
+	{
+		// it's possible a previous prefetch or strong fetch already added this to the RequestedObjects list
+		bAdded = (GTransferQueue.RequestedObjects.Find(ObjectId) == INDEX_NONE);
+
+		if (bAdded)
+		{
+			GTransferQueue.RequestedObjects.Add(ObjectId);
+			GTransferQueue.RequestedObjectsDestinationOuter.Add(DestinationOuter);
+
+			UE_LOG(LogRemoteObject, Verbose, TEXT("Prefetching obj id %s"),
+				*ObjectId.ToString());
+		}
+	};
+
+	if (bAdded && !AutoRTFM::IsClosed())
+	{
+		// if we are not in a transaction, request the object immediately
+		UE::RemoteObject::Handle::RequestObject(ObjectId);
+		GTransferQueue.RequestedObjects.Empty();
+		GTransferQueue.RequestedObjectsDestinationOuter.Empty();
+		GTransferQueue.LockedObjects.Empty();
+		GTransferQueue.PendingRequests = 0;
+	}
+}
+
 void ExecuteTransactional(const TFunction<void(void)>& Work)
 {
 	UE_LOG(LogRemoteObject, Verbose, TEXT("-- ExecuteTransactional -"));
@@ -361,6 +425,10 @@ void ExecuteTransactional(const TFunction<void(void)>& Work)
 	GTransferQueue.RequestedObjectsDestinationOuter.Empty();
 	GTransferQueue.PendingRequests = 0;
 
+	// we might run through this loop after the work succeeds in the case
+	// that the work prefetches objects that it doesn't access
+	bool bWorkSucceeded = false;
+
 	for (;;)
 	{
 		for (FRemoteObjectId ObjectId : GTransferQueue.LockedObjects)
@@ -386,28 +454,55 @@ void ExecuteTransactional(const TFunction<void(void)>& Work)
 
 		UE_LOG(LogRemoteObject, Verbose, TEXT("   : ObjectsToRequest: %d entries"), ObjectsToRequest.Num());
 
-		// send the requests...
+		// send the requests... slowly for now (to avoid straining the network subsystem)
 		int32 RequestsSent = 0;
-		for(FRemoteObjectId RequestedObjectId : ObjectsToRequest)
+		double HangProtectorBeginTime = FPlatformTime::Seconds();
+		for(int32 LoopCount = 0;; LoopCount++)
 		{
-			if (UE::RemoteObject::Handle::RequestObject(RequestedObjectId))
+			// do we need to, and can we send more requests?
+			while ((RequestsSent < ObjectsToRequest.Num()) && (GTransferQueue.PendingRequests < 5))
 			{
-				// true means we already have the object locally and didn't actually send a request for it
-				GTransferQueue.LockedObjects.Add(RequestedObjectId);
+				// yes, send one request
+				FRemoteObjectId RequestedObjectId = ObjectsToRequest[RequestsSent++];
+
+				if (UE::RemoteObject::Handle::RequestObject(RequestedObjectId))
+				{
+					// true means we already have the object locally and didn't actually send a request for it
+					GTransferQueue.LockedObjects.Add(RequestedObjectId);
+					UE_LOG(LogRemoteObject, Verbose, TEXT("      : %s is local"), *RequestedObjectId.ToString());
+				}
+				else
+				{
+					GTransferQueue.PendingRequests++;
+					UE_LOG(LogRemoteObject, Verbose, TEXT("      : %s is remote"), *RequestedObjectId.ToString());
+				}
 			}
-			else
+
+			// pump the network
+			FlushRemoteObjectRequestsDelegate.ExecuteIfBound();
+
+			// are all requests sent?
+			if (RequestsSent == ObjectsToRequest.Num())
 			{
-				RequestsSent++;
+				// yes, all requests are sent
+				// are they all complete?
+				if (GTransferQueue.PendingRequests == 0)
+				{
+					// yes, all done
+					break;
+				}
 			}
-		}
 
-		GTransferQueue.PendingRequests = RequestsSent;
+			if ((LoopCount % 100) == 0)
+			{
+				UE_LOG(LogRemoteObject, Verbose, TEXT("   : waiting for %d responses..."), GTransferQueue.PendingRequests);
+			}
 
-		// now tick the network until we get a response for each request made (either OK or DENIED)
-		while(GTransferQueue.PendingRequests > 0)
-		{
-			UE_LOG(LogRemoteObject, Verbose, TEXT("   : waiting for %d responses..."), GTransferQueue.PendingRequests);
-			FlushRemoteObjectRequestsDelegate.ExecuteIfBound();
+			// after 10 seconds, start sleeping in case we're deadlocked
+			if (FPlatformTime::Seconds() > HangProtectorBeginTime + 10.0)
+			{
+				FPlatformProcess::Sleep(0.5f);
+			}
 		}
 
 		// now we also might have migrated objects away, so check to see if LockedObjects == RequestedObjects
@@ -422,26 +517,37 @@ void ExecuteTransactional(const TFunction<void(void)>& Work)
 		UE_LOG(LogRemoteObject, Verbose, TEXT("   : READY"));
 
 		// all of the requested objects are locked, try to perform the work
-		
-		AutoRTFM::ETransactionResult WorkTransactResult = AutoRTFM::Transact([&Work](){ Work(); });
-
-		if (WorkTransactResult == AutoRTFM::ETransactionResult::AbortedByRequest)
-		{
-			UE_LOG(LogRemoteObject, Verbose, TEXT("   : work aborted, need object %s"), *GNextRequestedObjectId.ToString());
-
-			// the work aborted, insert a new requested object and retry
-			GTransferQueue.RequestedObjects.Add(GNextRequestedObjectId);
-			GTransferQueue.RequestedObjectsDestinationOuter.Add(GNextRequestedDestinationOuter);
-		}
-		else
+		if (!bWorkSucceeded)
 		{
-			UE_LOG(LogRemoteObject, Verbose, TEXT("   : work OK!"));
+			AutoRTFM::ETransactionResult WorkTransactResult = AutoRTFM::Transact([&Work](){ Work(); });
 
-			// all done!
-			GTransferQueue.RequestedObjects.Empty();
-			GTransferQueue.RequestedObjectsDestinationOuter.Empty();
-			GTransferQueue.LockedObjects.Empty();
-			break;
+			if (WorkTransactResult == AutoRTFM::ETransactionResult::AbortedByRequest)
+			{
+				// the work aborted, it should have inserted a new requested object, retry
+				UE_LOG(LogRemoteObject, Verbose, TEXT("   : work aborted"));
+			}
+			else
+			{
+				UE_LOG(LogRemoteObject, Verbose, TEXT("   : work OK!"));
+
+				bWorkSucceeded = true;
+
+				if (GTransferQueue.LockedObjects.Num() == GTransferQueue.RequestedObjects.Num())
+				{
+					// the work succeeded and no more object requests were made, we're all done!
+					GTransferQueue.RequestedObjects.Empty();
+					GTransferQueue.RequestedObjectsDestinationOuter.Empty();
+					GTransferQueue.LockedObjects.Empty();
+					GTransferQueue.PendingRequests = 0;
+					break;
+				}
+				else
+				{
+					// the work is complete, but it added object requests (prefetches), so we
+					// run through the loop again in order to fetch those objects
+					continue;
+				}
+			}
 		}
 	}
 }
diff --git a/Engine/Source/Runtime/CoreUObject/Public/UObject/RemoteObjectTransfer.h b/Engine/Source/Runtime/CoreUObject/Public/UObject/RemoteObjectTransfer.h
--- a/Engine/Source/Runtime/CoreUObject/Public/UObject/RemoteObjectTransfer.h
+++ b/Engine/Source/Runtime/CoreUObject/Public/UObject/RemoteObjectTransfer.h
@@ -16,6 +16,12 @@ namespace UE::RemoteObject::Transfer
 	*/
 	COREUOBJECT_API void ExecuteTransactional(const TFunction<void(void)>& Work);
 
+	/**
+	* (Advanced API) Adds the remote object to the current transaction's requested object list, causing
+	* the object to be migrated to this server. Does not abort the transaction if the object isn't yet local
+	*/
+	COREUOBJECT_API void PrefetchRemoteObject(FRemoteObjectId ObjectId, UObject* DestinationOuter);
+
 	/**
 	* Called when remote object data has been received from a remote server
 	* @param ObjectId Remote object id
@@ -52,6 +58,15 @@ namespace UE::RemoteObject::Transfer
 	* @param DestinationServerId Id of a remote server to receive the object
 	*/
 	COREUOBJECT_API void MigrateObjectToRemoteServer(FRemoteObjectId Id, FRemoteServerId DestinationServerId);
+	
+	/**
+	* Migrates an object from a remote server (temp function)
+	* @param ObjectId Remote object id
+	* @param CurrentOwnerServerId Id of a remote server that currently owns the object
+	* @param DestinationOuter Optional outer the migrated object should be reparented to
+	* @return Migrated object
+	*/
+	COREUOBJECT_API void MigrateObjectFromRemoteServer(FRemoteObjectId ObjectId, FRemoteServerId CurrentOwnerServerId, UObject* DestinationOuter = nullptr);
 
 	/**
 	* Utility that converts FRemoteObjectId to TObjectPtr of a given type. Will not resolve the object if it's not local.
