[AutoRTFM] Remove unnecessary open / onabort handlers

These are now transactionally safe, and mixing open with closed code is likely to cause memory corruption

diff --git a/Engine/Source/Runtime/CoreUObject/Public/UObject/UObjectAnnotation.h b/Engine/Source/Runtime/CoreUObject/Public/UObject/UObjectAnnotation.h
--- a/Engine/Source/Runtime/CoreUObject/Public/UObject/UObjectAnnotation.h
+++ b/Engine/Source/Runtime/CoreUObject/Public/UObject/UObjectAnnotation.h
@@ -103,15 +103,7 @@ private:
 				if (bAutoRemove)
 #endif
 				{
-					UE_AUTORTFM_OPEN
-					{
-						GUObjectArray.AddUObjectDeleteListener(this);
-					};
-
-					UE_AUTORTFM_ONABORT(this)
-					{
-						GUObjectArray.RemoveUObjectDeleteListener(this);
-					};
+					GUObjectArray.AddUObjectDeleteListener(this);
 				}
 			}
 		}
