[ALT/ReferenceFinder]
 - Hide objects still loading and serializing from the reference finder

#rnx

diff --git a/Engine/Source/Runtime/CoreUObject/Private/UObject/ReferencerFinder.cpp b/Engine/Source/Runtime/CoreUObject/Private/UObject/ReferencerFinder.cpp
--- a/Engine/Source/Runtime/CoreUObject/Private/UObject/ReferencerFinder.cpp
+++ b/Engine/Source/Runtime/CoreUObject/Private/UObject/ReferencerFinder.cpp
@@ -199,7 +199,8 @@ static TArray<UObject*> GetAllReferencersIncludingWeak(
 		for (int32 Index = FirstObjectIndex; Index < EndIndex; ++Index)
 		{
 			FUObjectItem& ObjectItem = GUObjectArray.GetObjectItemArrayUnsafe()[Index];
-			if (!ObjectItem.Object || ObjectItem.IsUnreachable())
+			// Skip any objects still being deserialized because it is not safe to access them until they are constructed correctly.
+			if (!ObjectItem.Object || ObjectItem.IsUnreachable() || ObjectItem.HasAnyFlags(EInternalObjectFlags::AsyncLoadingPhase1))
 			{
 				continue;
 			}
