Reference collectors will now use _ForGC variant of IsObjectHandleResolved to make them work properly with remote object handles

(IsObjectHandleResolved_ForGC behaves exactly like IsObjectHandleResolved for late resolve handles)

diff --git a/Engine/Source/Runtime/CoreUObject/Private/UObject/UObjectGlobals.cpp b/Engine/Source/Runtime/CoreUObject/Private/UObject/UObjectGlobals.cpp
--- a/Engine/Source/Runtime/CoreUObject/Private/UObject/UObjectGlobals.cpp
+++ b/Engine/Source/Runtime/CoreUObject/Private/UObject/UObjectGlobals.cpp
@@ -4972,7 +4972,7 @@ public:
 
 	virtual FArchive& operator<<(FObjectPtr& Object) override
 	{
-		if (IsObjectHandleResolved(Object.GetHandle()) && Object)
+		if (IsObjectHandleResolved_ForGC(Object.GetHandle()) && Object)
 		{
 			// NOTE: This is deliberately not triggering access tracking as that is an undesirable overhead
 			//		 during garbage collect and GC reference collection is not meant to be trackable.
@@ -5064,7 +5064,7 @@ void FReferenceCollector::HandleObjectReferences(FObjectPtr* InObjects, const in
 	for (int32 ObjectIndex = 0; ObjectIndex < ObjectNum; ++ObjectIndex)
 	{
 		FObjectPtr& Object = InObjects[ObjectIndex];
-		if (IsObjectHandleResolved(Object.GetHandle()))
+		if (IsObjectHandleResolved_ForGC(Object.GetHandle()))
 		{
 			HandleObjectReference(*reinterpret_cast<UObject**>(&Object), InReferencingObject, InReferencingProperty);
 		}
@@ -5285,7 +5285,7 @@ FORCEINLINE_DEBUGGABLE void CollectObjectReference(FReferenceCollector& Collecto
 	if constexpr (EnumHasAnyFlags(CollectFlags, EPropertyCollectFlags::NeedsReferencer))
 	{
 		// Sync reference processors will inspect Reference immediately so might as well avoid virtual call
-		if ((!!Reference) & IsObjectHandleResolved(*reinterpret_cast<FObjectHandle*>(Value))) //-V792
+		if ((!!Reference) & IsObjectHandleResolved_ForGC(*reinterpret_cast<FObjectHandle*>(Value))) //-V792
 		{
 			FReferenceCollector::AROPrivate::AddReferencedObject(Collector, Reference, Referencer, &Property);
 		}
