Fixed object properties with MustImplement metadata for a required interface pointing to an object that doesn't implement that interface.

diff --git a/Engine/Config/BaseEngine.ini b/Engine/Config/BaseEngine.ini
--- a/Engine/Config/BaseEngine.ini
+++ b/Engine/Config/BaseEngine.ini
@@ -116,6 +116,7 @@ ObjectReferenceReflectedUninitializedPropertyVerbosity=Error
 +MetaDataKeys=ActorClassToPoll
 +MetaDataKeys=ComponentClass
 +MetaDataKeys=MustImplement
++MetaDataKeys=ObjectMustImplement
 +MetaDataKeys=MetaClass
 +MetaDataKeys=MetaStruct
 +MetaDataKeys=RequiredAssetDataTags
diff --git a/Engine/Source/Editor/PropertyEditor/Private/PropertyCustomizationHelpers.cpp b/Engine/Source/Editor/PropertyEditor/Private/PropertyCustomizationHelpers.cpp
--- a/Engine/Source/Editor/PropertyEditor/Private/PropertyCustomizationHelpers.cpp
+++ b/Engine/Source/Editor/PropertyEditor/Private/PropertyCustomizationHelpers.cpp
@@ -242,7 +242,7 @@ namespace PropertyCustomizationHelpers
 			MetaClass = ObjectProperty->PropertyClass;
 
 			// Verse generated interfaces are not fully supported yet, for now they are stored in FObjectProperty metadata
-			RequiredInterface = ObjectProperty->GetOwnerProperty()->GetClassMetaData(TEXT("MustImplement"));
+			RequiredInterface = ObjectProperty->GetOwnerProperty()->GetClassMetaData(TEXT("ObjectMustImplement"));
 
 			if (!RequiredInterface && ObjectProperty->PropertyClass->HasAnyClassFlags(EClassFlags::CLASS_Interface))
 			{
diff --git a/Engine/Source/Editor/PropertyEditor/Private/UserInterface/PropertyEditor/SPropertyEditorEditInline.cpp b/Engine/Source/Editor/PropertyEditor/Private/UserInterface/PropertyEditor/SPropertyEditorEditInline.cpp
--- a/Engine/Source/Editor/PropertyEditor/Private/UserInterface/PropertyEditor/SPropertyEditorEditInline.cpp
+++ b/Engine/Source/Editor/PropertyEditor/Private/UserInterface/PropertyEditor/SPropertyEditorEditInline.cpp
@@ -275,7 +275,7 @@ TSharedRef<SWidget> SPropertyEditorEditInline::GenerateClassPicker(TSharedRef<IP
 	if (ClassFilter->ObjProperty)
 	{
 		// Verse generated interfaces are not fully supported by yet, for now they are stored in FObjectProperty metadata
-		ClassFilter->InterfaceThatMustBeImplemented = ClassFilter->ObjProperty->GetOwnerProperty()->GetClassMetaData(TEXT("MustImplement"));
+		ClassFilter->InterfaceThatMustBeImplemented = ClassFilter->ObjProperty->GetOwnerProperty()->GetClassMetaData(TEXT("ObjectMustImplement"));
 	}
 
 	if (ClassFilter->InterfaceThatMustBeImplemented != nullptr)
diff --git a/Engine/Source/Runtime/CoreUObject/Private/UObject/PropertyBaseObject.cpp b/Engine/Source/Runtime/CoreUObject/Private/UObject/PropertyBaseObject.cpp
--- a/Engine/Source/Runtime/CoreUObject/Private/UObject/PropertyBaseObject.cpp
+++ b/Engine/Source/Runtime/CoreUObject/Private/UObject/PropertyBaseObject.cpp
@@ -723,6 +723,42 @@ bool FObjectPropertyBase::AllowObjectTypeReinterpretationTo(const FObjectPropert
 	return false;
 }
 
+static bool CanPropertyHoldObject(const FObjectPropertyBase* Prop, TObjectPtr<const UObject> Obj)
+{
+	UClass* PropertyClass = Prop->PropertyClass;
+	checkSlow(PropertyClass);
+
+	UClass* ObjectClass = Obj.GetClass();
+	checkSlow(ObjectClass);
+
+#if WITH_METADATA
+	// Check if the object has metadata that states that the object must implement a particular interface
+	UClass* MetaDataClass = Prop->GetOwnerProperty()->GetClassMetaData(TEXT("ObjectMustImplement"));
+	if (MetaDataClass)
+	{
+		if (ObjectClass->ImplementsInterface(MetaDataClass))
+		{
+			return true;
+		}
+	}
+	else
+#endif // WITH_METADATA
+	{
+		// Check if the object is or is derived from the property class
+		if (ObjectClass->IsChildOf(PropertyClass) || ObjectClass->GetAuthoritativeClass()->IsChildOf(PropertyClass))
+		{
+			return true;
+		}
+		// Check if the object implements the interface - is it possible to have an object property with an interface class set?
+		if (ObjectClass->ImplementsInterface(PropertyClass))
+		{
+			return true;
+		}
+	}
+
+	return false;
+}
+
 UObject* FObjectPropertyBase::ConstructDefaultObjectValueIfNecessary(UObject* ExistingValue) const
 {
 	UObject* NewDefaultObjectValue = nullptr;
@@ -738,7 +774,7 @@ UObject* FObjectPropertyBase::ConstructDefaultObjectValueIfNecessary(UObject* Ex
 	{
 		UClass* ExistingValueClass = ExistingValue->GetClass();
 		// Sanity check to make sure the existing value class matches the property class
-		if (ExistingValueClass && (ExistingValueClass->IsChildOf(PropertyClass) || ExistingValueClass->GetAuthoritativeClass()->IsChildOf(PropertyClass)))
+		if (ExistingValueClass && CanPropertyHoldObject(this, ExistingValue))
 		{
 			if (ExistingValue->IsTemplate() && 	// Existing value is a template so we can construct a new value with it as the archetype
 				ExistingValue->GetOuter() != Outer) // Unless the template's Outer is the same as the new Outer in which case the template (ExistingValue) IS the object we can reuse
@@ -755,7 +791,12 @@ UObject* FObjectPropertyBase::ConstructDefaultObjectValueIfNecessary(UObject* Ex
 		}
 	}
 
-	if (!NewDefaultObjectValue && !PropertyClass->HasAnyClassFlags(CLASS_Abstract))
+	// We cannot know which object to create if the class is abstract or if there is a required interface
+	if (!NewDefaultObjectValue && !PropertyClass->HasAnyClassFlags(CLASS_Abstract)
+#if WITH_METADATA
+		&& !this->GetOwnerProperty()->GetClassMetaData(TEXT("ObjectMustImplement"))
+#endif // WITH_METADATA
+	)
 	{
 		// Existing value did not exist or it could not be used as a template
 		// Existing value may be null in case we were serializing an array of UObjects that failed to load (missing import class). Since the array is first pre-allocated with null values
@@ -804,9 +845,9 @@ void FObjectPropertyBase::CheckValidObject(void* ValueAddress, TObjectPtr<UObjec
 	auto IsDeferringValueLoad = [&]() { return false; };
 #endif // USE_CIRCULAR_DEPENDENCY_LOAD_DEFERRING
 
-	if ((PropertyClass != nullptr) && !ObjectClass->IsChildOf(PropertyClass) && !ObjectClass->GetAuthoritativeClass()->IsChildOf(PropertyClass) && !ObjectClass->ImplementsInterface(PropertyClass))
+	// Is a null class even possible here?
+	if ((PropertyClass != nullptr) && !CanPropertyHoldObject(this, Object))
 	{
-			
 		// we could be in the middle of replacing references to the 
 		// PropertyClass itself (in the middle of an FArchiveReplaceObjectRef 
 		// pass)... if this is the case, then we might have already replaced 
diff --git a/Engine/Source/Runtime/CoreUObject/Public/UObject/ObjectMacros.h b/Engine/Source/Runtime/CoreUObject/Public/UObject/ObjectMacros.h
--- a/Engine/Source/Runtime/CoreUObject/Public/UObject/ObjectMacros.h
+++ b/Engine/Source/Runtime/CoreUObject/Public/UObject/ObjectMacros.h
@@ -1406,6 +1406,9 @@ namespace UM
 		/// [PropertyMetadata] Used for Subclass and SoftClass properties. Indicates the selected class must implement a specific interface
 		MustImplement,
 
+		/// [PropertyMetadata] Used for object properties. Indicates the selected object must be of a class which implements a specific interface
+		ObjectMustImplement,
+
 		/// [PropertyMetadata] Used for numeric properties. Stipulates that the value must be a multiple of the metadata value.
 		Multiple,
 
