When leaking remote objects to database account for the fact that the leaked (sub)object may no longer be referenced by its parent hierarchy

+ Fixed leaking to database not working with AutoRTFM

diff --git a/Engine/Source/Runtime/CoreUObject/Private/UObject/RemoteObjectTransfer.cpp b/Engine/Source/Runtime/CoreUObject/Private/UObject/RemoteObjectTransfer.cpp
--- a/Engine/Source/Runtime/CoreUObject/Private/UObject/RemoteObjectTransfer.cpp
+++ b/Engine/Source/Runtime/CoreUObject/Private/UObject/RemoteObjectTransfer.cpp
@@ -50,7 +50,14 @@ public:
 	void RequestRemoteObject(FRemoteObjectId ObjectId, FRemoteServerId CurrentOwnerServerId)
 	{
 		// Request an object from remote server
-		RequestRemoteObjectDelegate.Execute(ObjectId, CurrentOwnerServerId, GetGlobalServerId());
+		if (CurrentOwnerServerId != DatabaseId)
+		{
+			RequestRemoteObjectDelegate.Execute(ObjectId, CurrentOwnerServerId, GetGlobalServerId());
+		}
+		else
+		{
+			RestoreRemoteObjectDataDelegate.Execute(ObjectId, DatabaseId, GetGlobalServerId());
+		}
 	}
 
 	/**
@@ -98,7 +105,7 @@ public:
 			{
 				DestinationOuter = RequestedObjectsDestinationOuter[RequestedObjectIndex];
 				bFoundObject = true;
-			}
+			}			
 		}
 
 		if (bFoundObject)
@@ -138,7 +145,7 @@ void RequestObjectData(FRemoteObjectId ObjectId, FRemoteServerId CurrentOwnerSer
 {
 	if (DestinationServerId == GetGlobalServerId())
 	{
-		GTransferQueue.RequestRemoteObject(ObjectId, CurrentOwnerServerId);
+		GTransferQueue.RequestRemoteObject(ObjectId, CurrentOwnerServerId);	
 	}
 	else
 	{
@@ -465,18 +472,18 @@ void ExecuteTransactional(const TFunction<void(void)>& Work)
 				// yes, send one request
 				FRemoteObjectId RequestedObjectId = ObjectsToRequest[RequestsSent++];
 
-				if (UE::RemoteObject::Handle::RequestObject(RequestedObjectId))
-				{
-					// true means we already have the object locally and didn't actually send a request for it
-					GTransferQueue.LockedObjects.Add(RequestedObjectId);
+			if (UE::RemoteObject::Handle::RequestObject(RequestedObjectId))
+			{
+				// true means we already have the object locally and didn't actually send a request for it
+				GTransferQueue.LockedObjects.Add(RequestedObjectId);
 					UE_LOG(LogRemoteObject, Verbose, TEXT("      : %s is local"), *RequestedObjectId.ToString());
-				}
-				else
-				{
+			}
+			else
+			{
 					GTransferQueue.PendingRequests++;
 					UE_LOG(LogRemoteObject, Verbose, TEXT("      : %s is remote"), *RequestedObjectId.ToString());
-				}
 			}
+		}
 
 			// pump the network
 			FlushRemoteObjectRequestsDelegate.ExecuteIfBound();
@@ -494,8 +501,8 @@ void ExecuteTransactional(const TFunction<void(void)>& Work)
 			}
 
 			if ((LoopCount % 100) == 0)
-			{
-				UE_LOG(LogRemoteObject, Verbose, TEXT("   : waiting for %d responses..."), GTransferQueue.PendingRequests);
+		{
+			UE_LOG(LogRemoteObject, Verbose, TEXT("   : waiting for %d responses..."), GTransferQueue.PendingRequests);
 			}
 
 			// after 10 seconds, start sleeping in case we're deadlocked
@@ -519,28 +526,28 @@ void ExecuteTransactional(const TFunction<void(void)>& Work)
 		// all of the requested objects are locked, try to perform the work
 		if (!bWorkSucceeded)
 		{
-			AutoRTFM::ETransactionResult WorkTransactResult = AutoRTFM::Transact([&Work](){ Work(); });
+		AutoRTFM::ETransactionResult WorkTransactResult = AutoRTFM::Transact([&Work](){ Work(); });
 
-			if (WorkTransactResult == AutoRTFM::ETransactionResult::AbortedByRequest)
-			{
+		if (WorkTransactResult == AutoRTFM::ETransactionResult::AbortedByRequest)
+		{
 				// the work aborted, it should have inserted a new requested object, retry
 				UE_LOG(LogRemoteObject, Verbose, TEXT("   : work aborted"));
-			}
-			else
-			{
-				UE_LOG(LogRemoteObject, Verbose, TEXT("   : work OK!"));
+		}
+		else
+		{
+			UE_LOG(LogRemoteObject, Verbose, TEXT("   : work OK!"));
 
 				bWorkSucceeded = true;
 
 				if (GTransferQueue.LockedObjects.Num() == GTransferQueue.RequestedObjects.Num())
 				{
 					// the work succeeded and no more object requests were made, we're all done!
-					GTransferQueue.RequestedObjects.Empty();
-					GTransferQueue.RequestedObjectsDestinationOuter.Empty();
-					GTransferQueue.LockedObjects.Empty();
+			GTransferQueue.RequestedObjects.Empty();
+			GTransferQueue.RequestedObjectsDestinationOuter.Empty();
+			GTransferQueue.LockedObjects.Empty();
 					GTransferQueue.PendingRequests = 0;
-					break;
-				}
+			break;
+		}
 				else
 				{
 					// the work is complete, but it added object requests (prefetches), so we
@@ -585,28 +592,40 @@ void StoreObjectToDatabase(UObject* Object)
 {
 	using namespace UE::RemoteObject::Serialization;
 	using namespace UE::RemoteObject::Private;
+	using namespace UE::RemoteObject::Transfer;
+
+	do
+	{
+		// Object could be a subobject of its outer that is going to be destroyed later (and have this function invoked on)
+		// so to make storing to database (disk) consistent with how objects are migrated (root object with its subobjects in one transfer), 
+		// find this object's outermost object that is also unreachable and marked as remote reference.
+		// If we re-entered this loop because Object has not been stored to disk with its OutermostUnreachableRemoteReference
+		// we still want to try and get an Outer that might also have not been stored. If such outer does not exist GetOutermostUnreachableRemoteReference(Object) will return Object itself.
+		UObject* OutermostUnreachableRemoteReference = GetOutermostUnreachableRemoteReference(Object);
 
-	// Object could a subobject of its outer that is going to be destroyed (and have this function invoked on) later so to make
-	// storing to database (disk) consistent with how objects are migrated (root object with its subobjects in one transfer), 
-	// find this object's outer that is also unreachable and marked as remote reference
-	UObject* OutermostUnreachableRemoteReference = GetOutermostUnreachableRemoteReference(Object);
+		TSet<UObject*> SentObjects;
+		{
+			TUniquePtr<FRemoteObjectData> ObjectData(SerializeObjectData(OutermostUnreachableRemoteReference, SentObjects));
+			StoreRemoteObjectDataDelegate.Execute(FRemoteObjectId(OutermostUnreachableRemoteReference), DatabaseId, *ObjectData);
+		}
 
-	TSet<UObject*> SentObjects;
-	TUniquePtr<FRemoteObjectData> ObjectData(SerializeObjectData(OutermostUnreachableRemoteReference, SentObjects));
-	
-	UE::RemoteObject::Transfer::StoreRemoteObjectDataDelegate.Execute(FRemoteObjectId(OutermostUnreachableRemoteReference), DatabaseId, *ObjectData);
+		checkf(SentObjects.Num() > 0, TEXT("PeerId:%s unable to store object %s to DatabaseId:%s"), *GetGlobalServerId().ToString(), *FRemoteObjectId(Object).ToString(), *DatabaseId.ToString());
 
-	checkf(SentObjects.Num() > 0, TEXT("PeerId:%s unable to store object %s to DatabaseId:%s"), *GetGlobalServerId().ToString(), *FRemoteObjectId(Object).ToString(), *DatabaseId.ToString());
+		// At this point the entire leaked object tree should be remote with the exception that its current owner server is the database
+		for (UObject* SentObject : SentObjects)
+		{
+			// MarkAsRemote also clears EInternalObjectFlags::RemoteReference flag
+			MarkAsRemote(SentObject, DatabaseId);
+		}
 
-	// At this point the entire leaked object tree is remote with the exception that its current owner server is the database
-	for (UObject* SentObject : SentObjects)
-	{
-		MarkAsRemote(SentObject, DatabaseId);
-	}
+		// It's possible that OutermostUnreachableRemoteReference no longer referenced this Object and Object hasn't been stored to disk
+		// in which case we need to repeat this loop until Object is marked as remote (and is no longer marked as RemoteReference).
+		
+		// Assert if we're about to enter an infinite loop.
+		checkf(!Object->HasAnyInternalFlags(EInternalObjectFlags::RemoteReference) || OutermostUnreachableRemoteReference != Object,
+			TEXT("PeerId:%s failed to store %s to database"), *GetGlobalServerId().ToString(), *Object->GetPathName());
 
-	checkf(!Object->HasAnyInternalFlags(EInternalObjectFlags::RemoteReference), 
-		TEXT("PeerId:%s has object %s still marked as remote reference after storing to database%s"), *FRemoteObjectId(Object).ToString(), *Object->GetPathName(),
-		OutermostUnreachableRemoteReference != Object ? *FString::Printf(TEXT(" when storing its outermost remote referenced object %s"), *OutermostUnreachableRemoteReference->GetPathName()) : TEXT("."));
+	} while (Object->HasAnyInternalFlags(EInternalObjectFlags::RemoteReference));
 }
 
 } // namespace UE::RemoteObject::Transfer::Private
\ No newline at end of file
