Use the new CascadingRetryTransaction to make a transactionally safe async loader that will abort and retry a transaction if a package actually has to be loaded.

diff --git a/Engine/Source/Programs/AutoRTFMTests/Private/Delegate.cpp b/Engine/Source/Programs/AutoRTFMTests/Private/Delegate.cpp
--- a/Engine/Source/Programs/AutoRTFMTests/Private/Delegate.cpp
+++ b/Engine/Source/Programs/AutoRTFMTests/Private/Delegate.cpp
@@ -3,28 +3,12 @@
 #include "Catch2Includes.h"
 #include <AutoRTFM/AutoRTFM.h>
 #include "Delegates/Delegate.h"
-
-// Temporarily changes the AutoRTFM retry mode for the lifetime of the FScopedRetry object.
-struct FScopedRetry
-{
-	FScopedRetry(AutoRTFM::ForTheRuntime::EAutoRTFMRetryTransactionState NewRetry)
-		: OldRetry(AutoRTFM::ForTheRuntime::GetRetryTransaction())
-	{
-		AutoRTFM::ForTheRuntime::SetRetryTransaction(NewRetry);
-	}
-
-	~FScopedRetry()
-	{
-		AutoRTFM::ForTheRuntime::SetRetryTransaction(OldRetry);
-	}
-
-	AutoRTFM::ForTheRuntime::EAutoRTFMRetryTransactionState OldRetry;
-};
+#include "Utils.h"
 
 TEST_CASE("Delegate.Broadcast")
 {
 	// Tests are sensitive to retries. Disable for these tests.
-	FScopedRetry Retry(AutoRTFM::ForTheRuntime::EAutoRTFMRetryTransactionState::NoRetry);
+	AutoRTFMTestUtils::FScopedRetry Retry(AutoRTFM::ForTheRuntime::EAutoRTFMRetryTransactionState::NoRetry);
 
 	TMulticastDelegate<void()> Delegate;
 
diff --git a/Engine/Source/Programs/AutoRTFMTests/Private/UECoreTests.cpp b/Engine/Source/Programs/AutoRTFMTests/Private/UECoreTests.cpp
--- a/Engine/Source/Programs/AutoRTFMTests/Private/UECoreTests.cpp
+++ b/Engine/Source/Programs/AutoRTFMTests/Private/UECoreTests.cpp
@@ -14,12 +14,15 @@
 #include "Serialization/CustomVersion.h"
 #include "UObject/NameTypes.h"
 #include "UObject/UObjectArray.h"
+#include "UObject/UObjectGlobals.h"
 #include "MyAutoRTFMTestObject.h"
+#include "Containers/Queue.h"
 #include "Misc/TransactionallySafeRWScopeLock.h"
 #include "Misc/TransactionallySafeMutex.h"
-#include "Containers/Queue.h"
 #include "Misc/ConfigCacheIni.h"
+#include "Misc/CoreDelegates.h"
 #include "StructUtils/PropertyBag.h"
+#include "Utils.h"
 
 DEFINE_LOG_CATEGORY_STATIC(LogAutoRTFMTests, Log, All)
 
@@ -2365,3 +2368,175 @@ TEST_CASE("UECore.UE_LOGFMT")
 		});
 	}
 }
+TEST_CASE("UECore.AsyncLoading")
+{
+	SECTION("LoadPackageAsync")
+	{
+		int32 RequestId = -1;
+
+		AutoRTFM::Testing::Commit([&]
+		{
+			FString Name(FString::Printf(TEXT("/AutoRTFMTestPackage%d"), __LINE__));
+			RequestId = LoadPackageAsync(Name);
+		});
+
+		FlushAsyncLoading(RequestId);
+	}
+
+	SECTION("FlushAsyncLoading")
+	{
+		FString Name(FString::Printf(TEXT("/AutoRTFMTestPackage%d"), __LINE__));
+		int32 RequestId = LoadPackageAsync(Name);
+
+		AutoRTFM::Testing::Commit([&]
+		{
+			FlushAsyncLoading(RequestId);
+		});
+	}
+
+	SECTION("CompletionDelegate is called closed")
+	{
+		AutoRTFM::Testing::Abort([&]
+		{
+			FString Name(FString::Printf(TEXT("/AutoRTFMTestPackage%d"), __LINE__));
+			FLoadPackageAsyncDelegate CompletionDelegate;
+			CompletionDelegate.BindLambda([&](const FName&, UPackage*, EAsyncLoadingResult::Type) { REQUIRE(AutoRTFM::IsClosed()); });
+			int RequestId = LoadPackageAsync(Name, CompletionDelegate);
+			FlushAsyncLoading(RequestId);
+			AutoRTFM::AbortTransaction();
+		});
+	}
+	
+	SECTION("CompletionDelegate aborts")
+	{
+		AutoRTFM::Testing::Abort([&]
+		{
+			FString Name(FString::Printf(TEXT("/AutoRTFMTestPackage%d"), __LINE__));
+			FLoadPackageAsyncDelegate CompletionDelegate;
+			CompletionDelegate.BindLambda([&](const FName&, UPackage*, EAsyncLoadingResult::Type) { AutoRTFM::AbortTransaction(); });
+			int RequestId = LoadPackageAsync(Name, CompletionDelegate);
+			FlushAsyncLoading(RequestId);
+			FAIL("Unreachable!");
+		});
+	}
+	
+	SECTION("FLoadPackageAsyncOptionalParams::CompletionDelegate is called closed")
+	{
+		AutoRTFM::Testing::Abort([&]
+		{
+			FString Name(FString::Printf(TEXT("/AutoRTFMTestPackage%d"), __LINE__));
+			FLoadPackageAsyncOptionalParams Params;
+			Params.CompletionDelegate.Reset(new FLoadPackageAsyncDelegate());
+			Params.CompletionDelegate->BindLambda([&](const FName&, UPackage*, EAsyncLoadingResult::Type) { REQUIRE(AutoRTFM::IsClosed()); });
+			int RequestId = LoadPackageAsync(Name, MoveTemp(Params));
+			FlushAsyncLoading(RequestId);
+			AutoRTFM::AbortTransaction();
+		});
+	}
+	
+	SECTION("FLoadPackageAsyncOptionalParams::CompletionDelegate aborts")
+	{
+		AutoRTFM::Testing::Abort([&]
+		{
+			FString Name(FString::Printf(TEXT("/AutoRTFMTestPackage%d"), __LINE__));
+			FLoadPackageAsyncOptionalParams Params;
+			Params.CompletionDelegate.Reset(new FLoadPackageAsyncDelegate());
+			Params.CompletionDelegate->BindLambda([&](const FName&, UPackage*, EAsyncLoadingResult::Type) { AutoRTFM::AbortTransaction(); });
+			int RequestId = LoadPackageAsync(Name, MoveTemp(Params));
+			FlushAsyncLoading(RequestId);
+			FAIL("Unreachable!");
+		});
+	}
+	
+	SECTION("FLoadPackageAsyncOptionalParams::CompletionDelegate creates UObject")
+	{
+		UMyAutoRTFMTestObject* OpenObject = nullptr;
+		UMyAutoRTFMTestObject* ClosedObject = nullptr;
+
+		AutoRTFM::Testing::Abort([&]
+		{
+			ClosedObject = NewObject<UMyAutoRTFMTestObject>();
+			FString Name(FString::Printf(TEXT("/AutoRTFMTestPackage%d"), __LINE__));
+			FLoadPackageAsyncOptionalParams Params;
+			Params.CompletionDelegate.Reset(new FLoadPackageAsyncDelegate());
+			Params.CompletionDelegate->BindLambda([&](const FName&, UPackage*, EAsyncLoadingResult::Type)
+			{
+				OpenObject = NewObject<UMyAutoRTFMTestObject>();
+				AutoRTFM::AbortTransaction();
+			});
+			int RequestId = LoadPackageAsync(Name, MoveTemp(Params));
+			FlushAsyncLoading(RequestId);
+			FAIL("Unreachable!");
+		});
+
+		REQUIRE(nullptr == ClosedObject);
+		REQUIRE(nullptr == OpenObject);
+	}
+	
+	SECTION("FLoadPackageAsyncOptionalParams::CompletionDelegate calls another LoadPackageAsync")
+	{
+		AutoRTFM::Testing::Abort([&]
+		{
+			FLoadPackageAsyncOptionalParams Params;
+			Params.CompletionDelegate.Reset(new FLoadPackageAsyncDelegate());
+			Params.CompletionDelegate->BindLambda([&](const FName&, UPackage*, EAsyncLoadingResult::Type)
+			{
+				int RequestId = LoadPackageAsync(FString::Printf(TEXT("/AutoRTFMTestPackage%d"), __LINE__), MoveTemp(Params));
+				AutoRTFM::AbortTransaction();
+			});
+			int RequestId = LoadPackageAsync(FString::Printf(TEXT("/AutoRTFMTestPackage%d"), __LINE__), MoveTemp(Params));
+			FlushAsyncLoading(RequestId);
+			FAIL("Unreachable!");
+		});
+	}
+
+	SECTION("Multiple retries because of multiple loads with commit")
+	{
+		AutoRTFMTestUtils::FScopedRetry Retry(AutoRTFM::ForTheRuntime::EAutoRTFMRetryTransactionState::NoRetry);
+
+		int NumCompletionCallbacks = 0;
+		AutoRTFM::Testing::Commit([&]
+		{
+			FLoadPackageAsyncDelegate CompletionDelegate;
+			CompletionDelegate.BindLambda([&](const FName&, UPackage*, EAsyncLoadingResult::Type) 
+			{
+				// Do this open so we can check how many retries occurred.
+				AutoRTFM::Open([&] { NumCompletionCallbacks++; });
+			});
+
+			TArray<int32> RequestIds;
+			RequestIds.Add(LoadPackageAsync(FString::Printf(TEXT("/AutoRTFMTestPackage%d"), __LINE__), CompletionDelegate));
+			RequestIds.Add(LoadPackageAsync(FString::Printf(TEXT("/AutoRTFMTestPackage%d"), __LINE__), CompletionDelegate));
+			RequestIds.Add(LoadPackageAsync(FString::Printf(TEXT("/AutoRTFMTestPackage%d"), __LINE__), CompletionDelegate));
+
+			FlushAsyncLoading(RequestIds);
+			REQUIRE(6 == NumCompletionCallbacks);
+		});
+	}
+
+	SECTION("Multiple retries because of multiple loads with abort")
+	{
+		AutoRTFMTestUtils::FScopedRetry Retry(AutoRTFM::ForTheRuntime::EAutoRTFMRetryTransactionState::NoRetry);
+
+		int NumCompletionCallbacks = 0;
+		AutoRTFM::Testing::Abort([&]
+		{
+			FLoadPackageAsyncDelegate CompletionDelegate;
+			CompletionDelegate.BindLambda([&](const FName&, UPackage*, EAsyncLoadingResult::Type) 
+			{
+				// Do this open so we can check how many retries occurred.
+				AutoRTFM::Open([&] { NumCompletionCallbacks++; });
+			});
+
+			TArray<int32> RequestIds;
+			RequestIds.Add(LoadPackageAsync(FString::Printf(TEXT("/AutoRTFMTestPackage%d"), __LINE__), CompletionDelegate));
+			RequestIds.Add(LoadPackageAsync(FString::Printf(TEXT("/AutoRTFMTestPackage%d"), __LINE__), CompletionDelegate));
+			RequestIds.Add(LoadPackageAsync(FString::Printf(TEXT("/AutoRTFMTestPackage%d"), __LINE__), CompletionDelegate));
+
+			FlushAsyncLoading(RequestIds);
+			REQUIRE(6 == NumCompletionCallbacks);
+
+			AutoRTFM::AbortTransaction();
+		});
+	}
+}
diff --git a/Engine/Source/Programs/AutoRTFMTests/Private/Utils.h b/Engine/Source/Programs/AutoRTFMTests/Private/Utils.h
--- a/Engine/Source/Programs/AutoRTFMTests/Private/Utils.h
+++ b/Engine/Source/Programs/AutoRTFMTests/Private/Utils.h
@@ -9,6 +9,23 @@
 namespace AutoRTFMTestUtils
 {
 
+// Temporarily changes the AutoRTFM retry mode for the lifetime of the FScopedRetry object.
+struct FScopedRetry
+{
+	FScopedRetry(AutoRTFM::ForTheRuntime::EAutoRTFMRetryTransactionState NewRetry)
+		: OldRetry(AutoRTFM::ForTheRuntime::GetRetryTransaction())
+	{
+		AutoRTFM::ForTheRuntime::SetRetryTransaction(NewRetry);
+	}
+
+	~FScopedRetry()
+	{
+		AutoRTFM::ForTheRuntime::SetRetryTransaction(OldRetry);
+	}
+
+	AutoRTFM::ForTheRuntime::EAutoRTFMRetryTransactionState OldRetry;
+};
+
 // A helper class that for the lifetime of the object, intercepts and records UE_LOG warnings.
 class FCaptureWarningContext : private FFeedbackContext
 {
diff --git a/Engine/Source/Runtime/CoreUObject/Private/Serialization/AsyncLoadingThread.h b/Engine/Source/Runtime/CoreUObject/Private/Serialization/AsyncLoadingThread.h
--- a/Engine/Source/Runtime/CoreUObject/Private/Serialization/AsyncLoadingThread.h
+++ b/Engine/Source/Runtime/CoreUObject/Private/Serialization/AsyncLoadingThread.h
@@ -12,7 +12,6 @@
 #include "Serialization/AsyncLoading.h"
 #include "Serialization/LoadTimeTracePrivate.h"
 #include "Misc/ScopeLock.h"
-#include "Misc/TransactionallySafeCriticalSection.h"
 #include "Misc/CommandLine.h"
 #include "Misc/App.h"
 #include "HAL/Runnable.h"
@@ -77,7 +76,7 @@ struct FAsyncLoadEvent
 /** [EDL] Event queue for the event driven loader */
 struct FAsyncLoadEventQueue
 {
-	//FTransactionallySafeCriticalSection AsyncLoadEventQueueCritical; //@todoio maybe this doesn't need to be protected by a critical section. 
+	//FCriticalSection AsyncLoadEventQueueCritical; //@todoio maybe this doesn't need to be protected by a critical section. 
 	int32 RunningSerialNumber;
 	TArray<FAsyncLoadEvent> EventQueue;
 
@@ -313,7 +312,7 @@ private:
 	TSet<int32> PendingRequests;
 #if THREADSAFE_UOBJECTS
 	/** Synchronization object for PendingRequests list */
-	FTransactionallySafeCriticalSection PendingRequestsCritical;
+	FCriticalSection PendingRequestsCritical;
 #endif
 
 	/** [ASYNC/GAME THREAD] Number of package load requests in the async loading queue */
diff --git a/Engine/Source/Runtime/CoreUObject/Private/Serialization/AsyncPackageLoader.cpp b/Engine/Source/Runtime/CoreUObject/Private/Serialization/AsyncPackageLoader.cpp
--- a/Engine/Source/Runtime/CoreUObject/Private/Serialization/AsyncPackageLoader.cpp
+++ b/Engine/Source/Runtime/CoreUObject/Private/Serialization/AsyncPackageLoader.cpp
@@ -5,6 +5,7 @@
 #include "Serialization/AsyncLoadingThread.h"
 #include "Serialization/AsyncLoading2.h"
 #include "Serialization/EditorPackageLoader.h"
+#include "Serialization/TransactionallySafeAsyncLoading.h"
 #include "UObject/GCObject.h"
 #include "UObject/LinkerLoad.h"
 #include "IO/PackageId.h"
@@ -41,7 +42,7 @@ FThreadSafeCounter IAsyncPackageLoader::NextPackageRequestId;
 
 int32 IAsyncPackageLoader::GetNextRequestId()
 {
-	return AutoRTFM::Open([] { return NextPackageRequestId.Increment(); });
+	return NextPackageRequestId.Increment();
 }
 
 class FEarlyRegistrationEventsRecorder
@@ -247,6 +248,12 @@ void InitAsyncThread()
 
 			GPackageLoader.Reset(MakeAsyncPackageLoader2(FIoDispatcher::Get()));
 		}
+
+#if UE_AUTORTFM
+		// If we are using AutoRTFM we need to layer a transactionally safe package loader over
+		// whatever loader was to be used.
+		GPackageLoader.Reset(MakeTransactionallySafeAsyncPackageLoader(GPackageLoader.Release()));
+#endif
 	}
 
 	FPlatformAtomics::InterlockedIncrement(&GIsLoaderCreated);
diff --git a/Engine/Source/Runtime/CoreUObject/Private/Serialization/TransactionallySafeAsyncLoading.cpp b/Engine/Source/Runtime/CoreUObject/Private/Serialization/TransactionallySafeAsyncLoading.cpp
new file mode 100644
--- /dev/null
+++ b/Engine/Source/Runtime/CoreUObject/Private/Serialization/TransactionallySafeAsyncLoading.cpp
@@ -0,0 +1,303 @@
+// Copyright Epic Games, Inc. All Rights Reserved.
+
+#include "Serialization/TransactionallySafeAsyncLoading.h"
+#include "Serialization/AsyncPackageLoader.h"
+#include "AutoRTFM/AutoRTFM.h"
+#include "Containers/Map.h"
+
+class FTransactionallySafeAsyncPackageLoader final : public IAsyncPackageLoader
+{
+public:
+	FTransactionallySafeAsyncPackageLoader(IAsyncPackageLoader* const InWrappedPackageLoader) : WrappedPackageLoader(InWrappedPackageLoader) {}
+
+	virtual ~FTransactionallySafeAsyncPackageLoader()
+	{
+		delete WrappedPackageLoader;
+	}
+
+	void InitializeLoading() override
+	{
+		AutoRTFM::UnreachableIfTransactional(UnreachableMessage);
+		WrappedPackageLoader->InitializeLoading();
+	}
+
+	void ShutdownLoading() override
+	{
+		AutoRTFM::UnreachableIfTransactional(UnreachableMessage);
+		WrappedPackageLoader->ShutdownLoading();
+	}
+
+	void StartThread() override
+	{
+		AutoRTFM::UnreachableIfTransactional(UnreachableMessage);
+		WrappedPackageLoader->StartThread();
+	}
+
+	bool ShouldAlwaysLoadPackageAsync(const FPackagePath& PackagePath) override
+	{
+		AutoRTFM::UnreachableIfTransactional(UnreachableMessage);
+		return WrappedPackageLoader->ShouldAlwaysLoadPackageAsync(PackagePath);
+	}
+
+	int32 LoadPackage(
+			const FPackagePath& PackagePath,
+			FName CustomPackageName,
+			FLoadPackageAsyncDelegate InCompletionDelegate,
+			EPackageFlags InPackageFlags,
+			int32 InPIEInstanceID,
+			int32 InPackagePriority,
+			const FLinkerInstancingContext* InInstancingContext,
+			uint32 InLoadFlags) override
+	{
+		if (AutoRTFM::IsClosed())
+		{
+			FString PackagePathAsString = PackagePath.GetDebugNameWithExtension();
+			FMapPayload* const Payload = RetriedPackageCache.Find(PackagePathAsString);
+			if (!Payload)
+			{
+				// We do not have the package cached and ready to be used in a
+				// transaction. So we need to do the expensive (but necessary)
+				// operation of retrying the transaction, and loading the
+				// package in the gap between the initial transaction aborting
+				// and the retried one beginning.
+				AutoRTFM::Open([&]
+				{
+					FPackagePath PackagePathClone = PackagePath;
+					FName CustomPackageNameClone = FName(CustomPackageName.ToString());
+
+					const AutoRTFM::EContextStatus Status = AutoRTFM::Close([&]
+					{
+						AutoRTFM::CascadingRetryTransaction([this, PackagePathClone, PackagePathAsString, CustomPackageNameClone, InPackageFlags, InPIEInstanceID, InPackagePriority, InInstancingContext, InLoadFlags]
+						{
+							FLoadPackageAsyncDelegate CompletionDelegate;
+							CompletionDelegate.BindLambda([this, PackagePathAsString](const FName& PackageName, UPackage* const Package, EAsyncLoadingResult::Type Result)
+							{
+								this->RetriedPackageCache.Add(PackagePathAsString, FMapPayload(PackageName, Package, Result));
+							});
+
+							const int32 RequestId = this->WrappedPackageLoader->LoadPackage(
+								PackagePathClone,
+								CustomPackageNameClone,
+								CompletionDelegate,
+								InPackageFlags,
+								InPIEInstanceID,
+								InPackagePriority,
+								InInstancingContext,
+								InLoadFlags);
+							this->WrappedPackageLoader->FlushLoading({ RequestId });
+
+							// The completion delegate should have triggered so double check that now.
+							FMapPayload* const Payload = this->RetriedPackageCache.Find(PackagePathAsString);
+							if (ensure(Payload))
+							{
+								Payload->RequestId = RequestId;
+								this->MaxFlushedRequestId = FMath::Max(this->MaxFlushedRequestId, RequestId);
+							}
+						});
+					});
+
+					ensure(Status == AutoRTFM::EContextStatus::AbortedByCascade);
+				});
+
+				return -1;
+			}
+			else
+			{
+				InCompletionDelegate.ExecuteIfBound(Payload->PackageName, Payload->LoadedPackage, Payload->Result);
+				return Payload->RequestId;
+			}
+		}
+		else
+		{
+			AutoRTFM::UnreachableIfTransactional(UnreachableMessage);
+			return WrappedPackageLoader->LoadPackage(PackagePath,
+				CustomPackageName,
+				InCompletionDelegate,
+				InPackageFlags,
+				InPIEInstanceID,
+				InPackagePriority,
+				InInstancingContext,
+				InLoadFlags);
+		}
+	}
+
+	int32 LoadPackage(const FPackagePath& PackagePath, FLoadPackageAsyncOptionalParams OptionalParams) override
+	{
+		if (AutoRTFM::IsClosed())
+		{
+			FLoadPackageAsyncDelegate CompletionDelegate;
+
+			if (OptionalParams.CompletionDelegate.IsValid())
+			{
+				CompletionDelegate = *OptionalParams.CompletionDelegate.Get();
+			}
+
+			return LoadPackage(PackagePath,
+				OptionalParams.CustomPackageName,
+				CompletionDelegate,
+				OptionalParams.PackageFlags,
+				OptionalParams.PIEInstanceID,
+				OptionalParams.PackagePriority,
+				OptionalParams.InstancingContext,
+				OptionalParams.LoadFlags);
+		}
+		else
+		{
+			AutoRTFM::UnreachableIfTransactional(UnreachableMessage);
+			return WrappedPackageLoader->LoadPackage(PackagePath, MoveTemp(OptionalParams));
+		}
+	}
+
+	EAsyncPackageState::Type ProcessLoading(bool bUseTimeLimit, bool bUseFullTimeLimit, double TimeLimit) override
+	{
+		AutoRTFM::UnreachableIfTransactional(UnreachableMessage);
+		return WrappedPackageLoader->ProcessLoading(bUseTimeLimit, bUseFullTimeLimit, TimeLimit);
+	}
+	
+	EAsyncPackageState::Type ProcessLoadingUntilComplete(TFunctionRef<bool()> CompletionPredicate, double TimeLimit) override
+	{
+		AutoRTFM::UnreachableIfTransactional(UnreachableMessage);
+		return WrappedPackageLoader->ProcessLoadingUntilComplete(CompletionPredicate, TimeLimit);
+	}
+
+	void CancelLoading() override
+	{
+		AutoRTFM::UnreachableIfTransactional(UnreachableMessage);
+		WrappedPackageLoader->CancelLoading();
+	}
+
+	void SuspendLoading() override
+	{
+		AutoRTFM::UnreachableIfTransactional(UnreachableMessage);
+		WrappedPackageLoader->SuspendLoading();
+	}
+
+	void ResumeLoading() override
+	{
+		AutoRTFM::UnreachableIfTransactional(UnreachableMessage);
+		WrappedPackageLoader->ResumeLoading();
+	}
+
+	void FlushLoading(TConstArrayView<int32> RequestIds) override
+	{
+		if (AutoRTFM::IsClosed())
+		{
+			AutoRTFM::Open([&]
+			{
+				for (const int32 RequestId : RequestIds)
+				{
+					if (RequestId > MaxFlushedRequestId)
+					{
+						TArray<int32> RequestIdsClone(RequestIds);
+						
+						const AutoRTFM::EContextStatus Status = AutoRTFM::Close([&]
+						{
+							AutoRTFM::CascadingRetryTransaction([this, RequestIdsClone]
+							{
+								WrappedPackageLoader->FlushLoading(RequestIdsClone);
+
+								MaxFlushedRequestId = FMath::Max(MaxFlushedRequestId, FMath::Max(RequestIdsClone));
+							});
+						});
+
+						ensure(Status == AutoRTFM::EContextStatus::AbortedByCascade);
+					}
+				}
+			});
+		}
+		else
+		{
+			AutoRTFM::UnreachableIfTransactional(UnreachableMessage);
+			WrappedPackageLoader->FlushLoading(RequestIds);
+		}
+	}
+
+	int32 GetNumQueuedPackages() override
+	{
+		return WrappedPackageLoader->GetNumQueuedPackages();
+	}
+
+	int32 GetNumAsyncPackages() override
+	{
+		return WrappedPackageLoader->GetNumAsyncPackages();
+	}
+
+	float GetAsyncLoadPercentage(const FName& PackageName) override
+	{
+		return WrappedPackageLoader->GetAsyncLoadPercentage(PackageName);
+	}
+
+	bool IsAsyncLoadingSuspended() override
+	{
+		return WrappedPackageLoader->IsAsyncLoadingSuspended();
+	}
+
+	bool IsInAsyncLoadThread() override
+	{
+		return WrappedPackageLoader->IsInAsyncLoadThread();
+	}
+
+	bool IsMultithreaded() override
+	{
+		return WrappedPackageLoader->IsMultithreaded();
+	}
+
+	bool IsAsyncLoadingPackages() override
+	{
+		return WrappedPackageLoader->IsAsyncLoadingPackages();
+	}
+
+	void NotifyConstructedDuringAsyncLoading(UObject* Object, bool bSubObject) override
+	{
+		AutoRTFM::UnreachableIfTransactional(UnreachableMessage);
+		WrappedPackageLoader->NotifyConstructedDuringAsyncLoading(Object, bSubObject);
+	}
+
+	void NotifyUnreachableObjects(const TArrayView<FUObjectItem*>& UnreachableObjects) override
+	{
+		AutoRTFM::UnreachableIfTransactional(UnreachableMessage);
+		WrappedPackageLoader->NotifyUnreachableObjects(UnreachableObjects);
+	}
+
+	void NotifyRegistrationEvent(const TCHAR* PackageName, const TCHAR* Name, ENotifyRegistrationType NotifyRegistrationType, ENotifyRegistrationPhase NotifyRegistrationPhase, UObject* (*InRegister)(), bool InbDynamic, UObject* FinishedObject) override
+	{
+		AutoRTFM::UnreachableIfTransactional(UnreachableMessage);
+		WrappedPackageLoader->NotifyRegistrationEvent(PackageName, Name, NotifyRegistrationType, NotifyRegistrationPhase, InRegister, InbDynamic, FinishedObject);
+	}
+
+	void NotifyRegistrationComplete() override
+	{
+		AutoRTFM::UnreachableIfTransactional(UnreachableMessage);
+		WrappedPackageLoader->NotifyRegistrationComplete();
+	}
+
+	ELoaderType GetLoaderType() const override
+	{
+		return WrappedPackageLoader->GetLoaderType();
+	}
+
+private:
+	static const TCHAR* const UnreachableMessage;
+	IAsyncPackageLoader* const WrappedPackageLoader;
+
+	struct FMapPayload final
+	{
+		FMapPayload(FName PackageName, UPackage* LoadedPackage, EAsyncLoadingResult::Type Result) :
+			PackageName(PackageName), LoadedPackage(LoadedPackage), Result(Result) {}
+
+		FName PackageName;
+		UPackage* LoadedPackage;
+		EAsyncLoadingResult::Type Result;
+		int32 RequestId = -1;
+	};
+
+	TMap<FString, FMapPayload> RetriedPackageCache;
+	int32 MaxFlushedRequestId = INT32_MIN;
+};
+
+const TCHAR* const FTransactionallySafeAsyncPackageLoader::UnreachableMessage = TEXT("Cannot call function within a transaction!");
+
+IAsyncPackageLoader* MakeTransactionallySafeAsyncPackageLoader(IAsyncPackageLoader* const InWrappedPackageLoader)
+{
+	return new FTransactionallySafeAsyncPackageLoader(InWrappedPackageLoader);
+}
diff --git a/Engine/Source/Runtime/CoreUObject/Public/Serialization/TransactionallySafeAsyncLoading.h b/Engine/Source/Runtime/CoreUObject/Public/Serialization/TransactionallySafeAsyncLoading.h
new file mode 100644
--- /dev/null
+++ b/Engine/Source/Runtime/CoreUObject/Public/Serialization/TransactionallySafeAsyncLoading.h
@@ -0,0 +1,14 @@
+// Copyright Epic Games, Inc. All Rights Reserved.
+
+#pragma once
+
+class IAsyncPackageLoader;
+
+/**
+ * Creates a new instance of the Transactionally Safe Async Package Loader.
+ *
+ * @param WrappedPackageLoader The underlying package loader to use when not in a transaction.
+ *
+ * @return The async package loader.
+ */
+IAsyncPackageLoader* MakeTransactionallySafeAsyncPackageLoader(IAsyncPackageLoader* WrappedPackageLoader = nullptr);
