Cleaned up use of FScopedIDOSerializationContext

Saving during cooking and saving subobjects in CPFUO now share a constructor with the others uses which improves maintainability.
Only enable bTrackInitializedProperties and bTrackSerializedProperties if the object supports tracking those states.
Set tracking flags during saving too to guarantee that tracking is off when it is expected to be.

#rnx

diff --git a/Engine/Source/Runtime/CoreUObject/Internal/UObject/PropertyBagRepository.h b/Engine/Source/Runtime/CoreUObject/Internal/UObject/PropertyBagRepository.h
--- a/Engine/Source/Runtime/CoreUObject/Internal/UObject/PropertyBagRepository.h
+++ b/Engine/Source/Runtime/CoreUObject/Internal/UObject/PropertyBagRepository.h
@@ -160,13 +160,11 @@ struct FScopedIDOSerializationContext
 {
 #if WITH_EDITORONLY_DATA
 	COREUOBJECT_API FScopedIDOSerializationContext(UObject* InObject, FArchive& Archive);
-	COREUOBJECT_API FScopedIDOSerializationContext(UObject* InObject, bool bImpersonate); // assumes save
-	COREUOBJECT_API explicit FScopedIDOSerializationContext(bool bImpersonate); // assumes save
 	COREUOBJECT_API ~FScopedIDOSerializationContext();
 
-	FArchive* Archive = nullptr;
-	UObject* const Object = nullptr;
-	const int64 PreSerializeOffset = 0;
+	FArchive& Archive;
+	UObject* const Object;
+	const int64 PreSerializeOffset;
 	UObject* SavedSerializedObject;
 	bool bSavedTrackSerializedPropertyPath;
 	bool bSavedTrackInitializedProperties;
@@ -174,17 +172,13 @@ struct FScopedIDOSerializationContext
 	bool bSavedTrackUnknownProperties;
 	bool bSavedTrackUnknownEnumNames;
 	bool bSavedImpersonateProperties;
-	bool bCreateIDO = false;
+	bool bCreateIDO;
 
 private:
-	void SaveSerializeContext(FUObjectSerializeContext* SerializeContext);
-	void RestoreSerializeContext(FUObjectSerializeContext* SerializeContext) const;
 	// if we're loading and IDO should be created, this will be called when the context falls out of scope
 	void FinishCreatingInstanceDataObject() const;
 #else
 	inline FScopedIDOSerializationContext(UObject* InObject, FArchive& Archive) {}
-	inline FScopedIDOSerializationContext(UObject* InObject, bool bImpersonate) {}
-	inline explicit FScopedIDOSerializationContext(bool bImpersonate) {}
 #endif
 
 	FScopedIDOSerializationContext(const FScopedIDOSerializationContext&) = delete;
diff --git a/Engine/Source/Runtime/CoreUObject/Private/UObject/PropertyBagRepository.cpp b/Engine/Source/Runtime/CoreUObject/Private/UObject/PropertyBagRepository.cpp
--- a/Engine/Source/Runtime/CoreUObject/Private/UObject/PropertyBagRepository.cpp
+++ b/Engine/Source/Runtime/CoreUObject/Private/UObject/PropertyBagRepository.cpp
@@ -871,53 +871,43 @@ void FPropertyBagRepository::CreateInstanceDataObjectUnsafe(UObject* Owner, FPro
 }
 
 FScopedIDOSerializationContext::FScopedIDOSerializationContext(UObject* InObject, FArchive& InArchive)
-	: Archive(&InArchive)
+	: Archive(InArchive)
 	, Object(InObject)
 	, PreSerializeOffset(InArchive.Tell())
 {
 	FUObjectSerializeContext* SerializeContext = FUObjectThreadContext::Get().GetSerializeContext();
-	SaveSerializeContext(SerializeContext);
 
-	const bool bHasIDOSupport = IsInstanceDataObjectSupportEnabled(Object);
+	// Save the parts of the serialize context that are modified below.
+	SavedSerializedObject = SerializeContext->SerializedObject;
+	bSavedTrackSerializedPropertyPath = SerializeContext->bTrackSerializedPropertyPath;
+	bSavedTrackInitializedProperties = SerializeContext->bTrackInitializedProperties;
+	bSavedTrackSerializedProperties = SerializeContext->bTrackSerializedProperties;
+	bSavedTrackUnknownProperties = SerializeContext->bTrackUnknownProperties;
+	bSavedTrackUnknownEnumNames = SerializeContext->bTrackUnknownEnumNames;
+	bSavedImpersonateProperties = SerializeContext->bImpersonateProperties;
 
-	if (Archive->IsLoading())
-	{
-		bCreateIDO = bHasIDOSupport && !SerializeContext->bImpersonateProperties &&
-			!InObject->GetClass()->HasAnyClassFlags(CLASS_NewerVersionExists);
-
-		// Enable creation of a property path name tree to track any property that does not match the current class schema,
-		// except when impersonation is enabled because that implies we are deserializing an IDO.
-		SerializeContext->bTrackSerializedPropertyPath = bCreateIDO;
-		SerializeContext->bTrackUnknownProperties = bCreateIDO;
-		SerializeContext->bTrackUnknownEnumNames = bCreateIDO;
-		SerializeContext->SerializedObject = Object;
-
-		// Enable tracking of initialized properties when loading an IDO, which is implied by impersonation being enabled.
-		const bool bLoadingIDO = bHasIDOSupport && SerializeContext->bImpersonateProperties;
-		// Track property initialization state if IDO feature not enabled
-		SerializeContext->bTrackInitializedProperties = IsInstanceDataObjectSupportEnabled() ? bLoadingIDO : true; 
-		SerializeContext->bTrackSerializedProperties = bLoadingIDO;
-	}
-	else
-	{
-		SerializeContext->bImpersonateProperties = bHasIDOSupport;
-	}
-}
+	// Disable if cooking because extra data within impersonated objects will not be understood outside of the editor.
+	const bool bHasIDOSupport = !InArchive.IsCooking() && IsInstanceDataObjectSupportEnabled(InObject);
 
-FScopedIDOSerializationContext::FScopedIDOSerializationContext(UObject* InObject, bool bImpersonate)
-{
-	FUObjectSerializeContext* SerializeContext = FUObjectThreadContext::Get().GetSerializeContext();
-	SaveSerializeContext(SerializeContext);
+	const bool bIsLoading = InArchive.IsLoading();
 
-	SerializeContext->bImpersonateProperties = bImpersonate && IsInstanceDataObjectSupportEnabled(InObject);
-}
+	bCreateIDO = bIsLoading && bHasIDOSupport && !SerializeContext->bImpersonateProperties &&
+		!InObject->GetClass()->HasAnyClassFlags(CLASS_NewerVersionExists);
 
-FScopedIDOSerializationContext::FScopedIDOSerializationContext(bool bImpersonate)
-{
-	FUObjectSerializeContext* SerializeContext = FUObjectThreadContext::Get().GetSerializeContext();
-	SaveSerializeContext(SerializeContext);
+	SerializeContext->SerializedObject = InObject;
 
-	SerializeContext->bImpersonateProperties = bImpersonate && IsInstanceDataObjectSupportEnabled();
+	// Enable tracking any property or enum that does not match the current schema when creating an IDO.
+	SerializeContext->bTrackSerializedPropertyPath = bCreateIDO;
+	SerializeContext->bTrackUnknownProperties = bCreateIDO;
+	SerializeContext->bTrackUnknownEnumNames = bCreateIDO;
+
+	SerializeContext->bTrackInitializedProperties = bIsLoading && FInitializedPropertyValueState(InObject).IsTracking();
+	SerializeContext->bTrackSerializedProperties = bIsLoading && FSerializedPropertyValueState(InObject).IsTracking();
+
+	if (!bIsLoading)
+	{
+		SerializeContext->bImpersonateProperties = bHasIDOSupport;
+	}
 }
 
 FScopedIDOSerializationContext::~FScopedIDOSerializationContext()
@@ -927,22 +917,8 @@ FScopedIDOSerializationContext::~FScopedIDOSerializationContext()
 		FinishCreatingInstanceDataObject();
 	}
 
-	RestoreSerializeContext(FUObjectThreadContext::Get().GetSerializeContext());
-}
-
-void FScopedIDOSerializationContext::SaveSerializeContext(FUObjectSerializeContext* SerializeContext)
-{
-	SavedSerializedObject = SerializeContext->SerializedObject;
-	bSavedTrackSerializedPropertyPath = SerializeContext->bTrackSerializedPropertyPath;
-	bSavedTrackInitializedProperties = SerializeContext->bTrackInitializedProperties;
-	bSavedTrackSerializedProperties = SerializeContext->bTrackSerializedProperties;
-	bSavedTrackUnknownProperties = SerializeContext->bTrackUnknownProperties;
-	bSavedTrackUnknownEnumNames = SerializeContext->bTrackUnknownEnumNames;
-	bSavedImpersonateProperties = SerializeContext->bImpersonateProperties;
-}
-
-void FScopedIDOSerializationContext::RestoreSerializeContext(FUObjectSerializeContext* SerializeContext) const
-{
+	// Restore the parts of the serialize context that the constructor modified.
+	FUObjectSerializeContext* SerializeContext = FUObjectThreadContext::Get().GetSerializeContext();
 	SerializeContext->SerializedObject = SavedSerializedObject;
 	SerializeContext->bTrackSerializedPropertyPath = bSavedTrackSerializedPropertyPath;
 	SerializeContext->bTrackInitializedProperties = bSavedTrackInitializedProperties;
@@ -954,28 +930,27 @@ void FScopedIDOSerializationContext::RestoreSerializeContext(FUObjectSerializeCo
 
 void FScopedIDOSerializationContext::FinishCreatingInstanceDataObject() const
 {
-	if (Archive == Object->GetLinker())
+	if (&Archive == Object->GetLinker())
 	{
 		// when using the linker, the repository will handle offsets
-		FPropertyBagRepository::Get().CreateInstanceDataObject(Object, Archive);
+		FPropertyBagRepository::Get().CreateInstanceDataObject(Object, &Archive);
 	}
 	// Check that we haven't already created an IDO for this object so we can skip further work.
 	// This double processing can happen at re-instancing time during CPFUO.
 	else if (FPropertyBagRepository::Get().FindInstanceDataObject(Object) == nullptr)
 	{
-		check(Archive);
-		const int64 PostSerializeOffset = Archive->Tell();
+		const int64 PostSerializeOffset = Archive.Tell();
 	
 		// CreateInstanceDataObject will re-call DstObject->Serialize(*this) so set the seek pointer back before DestObject in the archive
-		Archive->Seek(PreSerializeOffset);
+		Archive.Seek(PreSerializeOffset);
 			
-		FPropertyBagRepository::Get().CreateInstanceDataObject(Object, Archive);
+		FPropertyBagRepository::Get().CreateInstanceDataObject(Object, &Archive);
 
 		// make sure seek pointer is back to where it should be
-		if (!ensure(Archive->Tell() == PostSerializeOffset))
+		if (!ensure(Archive.Tell() == PostSerializeOffset))
 		{
 			// for some reason CreateInstanceDataObject read a different amount of data than expected... reset seek pointer back to where it should be
-			Archive->Seek(PostSerializeOffset);
+			Archive.Seek(PostSerializeOffset);
 		}
 	}
 }
diff --git a/Engine/Source/Runtime/CoreUObject/Private/UObject/SavePackage/PackageHarvester.cpp b/Engine/Source/Runtime/CoreUObject/Private/UObject/SavePackage/PackageHarvester.cpp
--- a/Engine/Source/Runtime/CoreUObject/Private/UObject/SavePackage/PackageHarvester.cpp
+++ b/Engine/Source/Runtime/CoreUObject/Private/UObject/SavePackage/PackageHarvester.cpp
@@ -7,6 +7,7 @@
 #include "Interfaces/ITargetPlatform.h"
 #include "Misc/PackageAccessTracking.h"
 #include "UObject/OverridableManager.h"
+#include "UObject/PropertyBagRepository.h"
 #include "UObject/SavePackage/SaveContext.h"
 #include "UObject/SavePackage/SavePackageUtilities.h"
 #include "UObject/UObjectGlobals.h"
@@ -466,6 +467,7 @@ void FPackageHarvester::ProcessExport(const FExportWithContext& InProcessContext
 {
 	UObject* Export = InProcessContext.Export;
 	FExportScope HarvesterScope(*this, InProcessContext);
+	UE::FScopedIDOSerializationContext IDOSaveContext(Export, *this);
 
 	// Harvest its class; warn if class is not included
 	UClass* Class = Export->GetClass();
diff --git a/Engine/Source/Runtime/CoreUObject/Private/UObject/SavePackage2.cpp b/Engine/Source/Runtime/CoreUObject/Private/UObject/SavePackage2.cpp
--- a/Engine/Source/Runtime/CoreUObject/Private/UObject/SavePackage2.cpp
+++ b/Engine/Source/Runtime/CoreUObject/Private/UObject/SavePackage2.cpp
@@ -2120,6 +2120,8 @@ ESavePackageResult WritePackageTextHeader(FStructuredArchive::FRecord& Structure
 		SCOPED_SAVETIMER(UPackage_Save_SaveExport);
 		SCOPED_SAVETIMER_TEXT(*WriteToString<128>(GetClassTraceScope(Export.Object), TEXT("_SaveSerialize")));
 
+		UE::FScopedIDOSerializationContext IDOSaveContext(Export.Object, Linker);
+
 		Export.SerialOffset = ExportsArchive.Tell();
 		Linker.CurrentlySavingExport = FPackageIndex::FromExport(ExportIndex);
 		Linker.CurrentlySavingExportObject = Export.Object;
@@ -2133,7 +2135,6 @@ ESavePackageResult WritePackageTextHeader(FStructuredArchive::FRecord& Structure
 		}
 		else
 		{
-			TGuardValue<UObject*> GuardSerializedObject(SerializeContext->SerializedObject, Export.Object);
 			Export.Object->Serialize(Linker);
 #if WITH_EDITOR
 			Export.Object->CookAdditionalFiles(SaveContext.GetFilename(), SaveContext.GetTargetPlatform(),
@@ -2172,8 +2173,6 @@ ESavePackageResult WriteExports(FStructuredArchive::FRecord& StructuredArchiveRo
 	FLinkerSave* Linker = SaveContext.GetLinker();
 	FScopedSlowTask SlowTask((float)Linker->ExportMap.Num(), FText(), SaveContext.IsUsingSlowTask());
 
-	FUObjectSerializeContext* SerializeContext = FUObjectThreadContext::Get().GetSerializeContext();
-
 	FStructuredArchive::FRecord ExportsRecord = StructuredArchiveRoot.EnterRecord(TEXT("Exports"));
 
 	// Save exports.
@@ -2191,6 +2190,8 @@ ESavePackageResult WriteExports(FStructuredArchive::FRecord& StructuredArchiveRo
 		{
 			SCOPED_SAVETIMER(UPackage_Save_SaveExport);
 
+			UE::FScopedIDOSerializationContext IDOSaveContext(Export.Object, ExportsRecord.GetUnderlyingArchive());
+
 			// Save the object data.
 			Export.SerialOffset = Linker->Tell();
 			Linker->CurrentlySavingExport = FPackageIndex::FromExport(i);
@@ -2221,8 +2222,6 @@ ESavePackageResult WriteExports(FStructuredArchive::FRecord& StructuredArchiveRo
 			}
 			else
 			{
-				TGuardValue<UObject*> GuardSerializedObject(SerializeContext->SerializedObject, Export.Object);
-
 				if (bSupportsText)
 				{
 					FStructuredArchive::FRecord ExportRecord = ExportSlot.EnterRecord();
@@ -3054,8 +3053,6 @@ ESavePackageResult InnerSave(FSaveContext& SaveContext)
 {
 	SaveContext.SetEDLCookChecker(&FEDLCookCheckerThreadState::Get());
 
-	UE::FScopedIDOSerializationContext IDOSaveContext(!SaveContext.IsCooking());
-
 	// Create slow task dialog if needed
 	const int32 TotalSaveSteps = 3;
 	FScopedSlowTask SlowTask(TotalSaveSteps, FText(), SaveContext.IsUsingSlowTask());
diff --git a/Engine/Source/Runtime/Engine/Private/UnrealEngine.cpp b/Engine/Source/Runtime/Engine/Private/UnrealEngine.cpp
--- a/Engine/Source/Runtime/Engine/Private/UnrealEngine.cpp
+++ b/Engine/Source/Runtime/Engine/Private/UnrealEngine.cpp
@@ -17975,11 +17975,28 @@ void UEngine::CopyPropertiesForUnrelatedObjects(UObject* OldObject, UObject* New
 			SavedInstances.Add(pRecord);
 			pRecord->OldInstance = OldInstance;
 			OldInstanceMap.Add(OldInstance->GetPathName(OldObject), SavedInstances.Num() - 1);
-			const uint32 AdditionalPortFlags = Params.bCopyDeprecatedProperties ? PPF_UseDeprecatedProperties : PPF_None;
-			
-			UE::FScopedIDOSerializationContext IDOSaveContext(OldInstance, /*bImpersonate*/ true);
-			
-			FObjectWriter SubObjWriter(OldInstance, pRecord->SavedProperties, true, true, Params.bDoDelta, AdditionalPortFlags);
+
+			FObjectWriter SubObjWriter(pRecord->SavedProperties);
+
+			SubObjWriter.ArIgnoreArchetypeRef = true;
+			SubObjWriter.ArIgnoreClassRef = true;
+			SubObjWriter.ArNoDelta = !Params.bDoDelta;
+			SubObjWriter.ArPortFlags = PPF_None;
+			if (Params.bCopyDeprecatedProperties)
+			{
+				SubObjWriter.ArPortFlags |= PPF_UseDeprecatedProperties;
+			}
+
+		#if USE_STABLE_LOCALIZATION_KEYS
+			if (GIsEditor && !(SubObjWriter.ArPortFlags & (PPF_DuplicateVerbatim | PPF_DuplicateForPIE)))
+			{
+				SubObjWriter.SetLocalizationNamespace(TextNamespaceUtil::EnsurePackageNamespace(OldInstance));
+			}
+		#endif // USE_STABLE_LOCALIZATION_KEYS
+
+			UE::FScopedIDOSerializationContext IDOSaveContext(OldInstance, SubObjWriter);
+
+			OldInstance->Serialize(SubObjWriter);
 		}
 	}
 
