Cooker: Blueprints: Fix "Content missing from cook" warnings for packages that refer to a target blueprint's template subobjects within its on-disk UBlueprintGeneratedClass that do not get recreated in the newly recompiled UBlueprintGeneratedClass for the target blueprint when it loads. Most of the work necessary to report the no-longer-used template subobjects was already handled by InheritableComponentHandler, but some were not, so we now report them explicitly in ReinstanceBatch.
#rnx

diff --git a/Engine/Source/Editor/Kismet/Private/BlueprintCompilationManager.cpp b/Engine/Source/Editor/Kismet/Private/BlueprintCompilationManager.cpp
--- a/Engine/Source/Editor/Kismet/Private/BlueprintCompilationManager.cpp
+++ b/Engine/Source/Editor/Kismet/Private/BlueprintCompilationManager.cpp
@@ -2578,13 +2578,16 @@ void FBlueprintCompilationManagerImpl::ReinstanceBatch(TArray<FReinstancingJob>&
 		{
 			BPGClass->UpdateCustomPropertyListForPostConstruction();
 
-			// patch new cdo into linker table:
-			if(OldCDO && ReinstancingJob.Reinstancer.IsValid())
+			// patch new class and cdo into linker table:
+			if (ReinstancingJob.Reinstancer.IsValid())
 			{
 				UBlueprint* CurrentBP = CastChecked<UBlueprint>(BPGClass->ClassGeneratedBy);
-				if(FLinkerLoad* CurrentLinker = CurrentBP->GetLinker())
+				UObject* NewCDO = CurrentBP->GeneratedClass->ClassDefaultObject;
+				FLinkerLoad* CurrentLinker = CurrentBP->GetLinker();
+				if (CurrentLinker)
 				{
 					int32 OldCDOIndex = INDEX_NONE;
+					int32 OldClassIndex = INDEX_NONE;
 
 					for (int32 i = 0; i < CurrentLinker->ExportMap.Num(); i++)
 					{
@@ -2597,16 +2600,27 @@ void FBlueprintCompilationManagerImpl::ReinstanceBatch(TArray<FReinstancingJob>&
 						// But in a Control Rig BP, multiple CDO can be in a package because custom UClasses are used to describe 
 						// the layout of RigVM Memory Storage. These UClasses and their CDO are serialzed with the package as well.
 						// Thus, we have to ensure that the export is not just a CDO, but also a CDO with the matching name
-						if ((ThisExport.ObjectFlags & RF_ClassDefaultObject) && (ThisExport.ObjectName == CurrentBP->GeneratedClass->ClassDefaultObject->GetFName()))
+						if ((ThisExport.ObjectFlags & RF_ClassDefaultObject) && (ThisExport.ObjectName == NewCDO->GetFName()))
 						{
 							OldCDOIndex = i;
-							break;
+						}
+						// Likewise we have to search by name for the class.
+						if (ThisExport.ObjectName == BPGClass->GetFName()
+							&& ThisExport.OuterIndex.IsNull() /* The class has to be a direct subobject of the package */
+							)
+						{
+							OldClassIndex = i;
 						}
 					}
 
-					if(OldCDOIndex != INDEX_NONE)
+					if (OldCDOIndex != INDEX_NONE)
+					{
+						CurrentLinker->PRIVATE_PatchNewObjectIntoExport(OldCDOIndex, NewCDO, InLoadContext);
+						NewCDO->CheckDefaultSubobjects();
+					}
+					if (OldClassIndex != INDEX_NONE)
 					{
-						FBlueprintEditorUtils::PatchNewCDOIntoLinker(CurrentBP->GeneratedClass->ClassDefaultObject, CurrentLinker, OldCDOIndex, InLoadContext);
+						CurrentLinker->PRIVATE_PatchNewObjectIntoExport(OldClassIndex, BPGClass, InLoadContext);
 					}
 				}
 			}
diff --git a/Engine/Source/Editor/KismetCompiler/Private/KismetCompiler.cpp b/Engine/Source/Editor/KismetCompiler/Private/KismetCompiler.cpp
--- a/Engine/Source/Editor/KismetCompiler/Private/KismetCompiler.cpp
+++ b/Engine/Source/Editor/KismetCompiler/Private/KismetCompiler.cpp
@@ -5105,20 +5105,21 @@ void FKismetCompilerContext::CompileFunctions(EInternalCompilerFlags InternalFla
 
 			if( !Blueprint->HasAnyFlags(RF_BeingRegenerated) )
 			{
-				// Propagate the old CDO's properties to the new
 				if( OldCDO )
 				{
+					// Patch our export table. This also fixes up load flags
 					if (OldLinker && OldGenLinkerIdx != INDEX_NONE)
 					{
-						// If we have a list of objects that are loading, patch our export table. This also fixes up load flags
-						FBlueprintEditorUtils::PatchNewCDOIntoLinker(Blueprint->GeneratedClass->GetDefaultObject(), OldLinker, OldGenLinkerIdx, nullptr);
+						OldLinker->PRIVATE_PatchNewObjectIntoExport(OldGenLinkerIdx, Blueprint->GeneratedClass->GetDefaultObject());
 					}
+					// Also patch (or invalidate if no longer existing) subobjects of the class
+					FLinkerLoad::PRIVATE_PatchNewObjectIntoExport(TargetClass, TargetClass);
 
+					// Propagate the old CDO's properties to the new
 					UEditorEngine::FCopyPropertiesForUnrelatedObjectsParams CopyDetails;
 					CopyDetails.bCopyDeprecatedProperties = Blueprint->bIsRegeneratingOnLoad;
 					CopyDetails.bNotifyObjectReplacement = true; 
 					UEditorEngine::CopyPropertiesForUnrelatedObjects(OldCDO, NewCDO, CopyDetails);
-					FBlueprintEditorUtils::PatchCDOSubobjectsIntoExport(OldCDO, NewCDO);
 				}
 				else
 				{
diff --git a/Engine/Source/Editor/UnrealEd/Private/Kismet2/BlueprintEditorUtils.cpp b/Engine/Source/Editor/UnrealEd/Private/Kismet2/BlueprintEditorUtils.cpp
--- a/Engine/Source/Editor/UnrealEd/Private/Kismet2/BlueprintEditorUtils.cpp
+++ b/Engine/Source/Editor/UnrealEd/Private/Kismet2/BlueprintEditorUtils.cpp
@@ -833,39 +833,10 @@ void FBlueprintEditorUtils::PreloadConstructionScript(USimpleConstructionScript*
 
 void FBlueprintEditorUtils::PatchNewCDOIntoLinker(UObject* CDO, FLinkerLoad* Linker, int32 ExportIndex, FUObjectSerializeContext* InLoadContext)
 {
-	if( (CDO != nullptr) && (Linker != nullptr) && (ExportIndex != INDEX_NONE) && Linker->ExportMap.Num() != 0 )
+	if (CDO != nullptr && Linker != nullptr)
 	{
-		// Get rid of the old thing that was in its place
-		UObject* OldCDO = Linker->ExportMap[ExportIndex].Object;
-		if( OldCDO != nullptr )
-		{
-			EObjectFlags OldObjectFlags = OldCDO->GetFlags();
-			OldCDO->ClearFlags(RF_NeedLoad|RF_NeedPostLoad|RF_NeedPostLoadSubobjects);
-			OldCDO->SetLinker(nullptr, INDEX_NONE);
-			
-			// Copy flags from the old CDO.
-			CDO->SetFlags(OldObjectFlags);
-
-			FUObjectSerializeContext* LoadContext = InLoadContext ? InLoadContext : FUObjectThreadContext::Get().GetSerializeContext();
-
-			// Make sure the new CDO gets PostLoad called on it, so either add it to ObjLoaded list, or replace it if already present.
-			if (LoadContext && !LoadContext->PRIVATE_PatchNewObjectIntoExport(OldCDO, CDO))
-			{
-				if (OldObjectFlags & RF_NeedPostLoad)
-				{
-					LoadContext->AddLoadedObject(CDO);
-				}
-			}
-		}
-
-		// Patch the new CDO in, and update the Export.Object
-		CDO->SetLinker(Linker, ExportIndex);
-		Linker->ExportMap[ExportIndex].Object = CDO;
-
-		PatchCDOSubobjectsIntoExport(OldCDO, CDO);
-
-		// This was set to true when the trash class was invalidated, but now we have a valid object
-		Linker->ExportMap[ExportIndex].bExportLoadFailed = false;
+		Linker->PRIVATE_PatchNewObjectIntoExport(ExportIndex, CDO, InLoadContext);
+		CDO->CheckDefaultSubobjects();
 	}
 }
 
@@ -1444,9 +1415,6 @@ void FBlueprintEditorUtils::PatchCDOSubobjectsIntoExport(UObject* PreviousCDO, U
 				{
 					UObject* NewSubObj = *NewSubObjPtr;
 					FLinkerLoad::PRIVATE_PatchNewObjectIntoExport(OldSubObj, NewSubObj);
-
-					// Recursively patch all nested subobjects into the export table as well.
-					PatchSubObjects(OldSubObj, NewSubObj);
 				}
 				else
 				{
diff --git a/Engine/Source/Editor/UnrealEd/Public/Kismet2/BlueprintEditorUtils.h b/Engine/Source/Editor/UnrealEd/Public/Kismet2/BlueprintEditorUtils.h
--- a/Engine/Source/Editor/UnrealEd/Public/Kismet2/BlueprintEditorUtils.h
+++ b/Engine/Source/Editor/UnrealEd/Public/Kismet2/BlueprintEditorUtils.h
@@ -214,9 +214,7 @@ public:
 	 */
 	static UNREALED_API void LinkExternalDependencies(UBlueprint* Blueprint);
 
-	/**
-	 * Replace subobjects of CDO in linker
-	 */
+	UE_DEPRECATED(5.6, "Use FLinkerLoad::PRIVATE_PatchNewObjectIntoExport instead, which handles the subobjects. And note the comment on that function: both this deprecated function and the LinkerLoad function are coupled to the linker in complicated ways and should only be used in cases of blueprint reinstancing.")
 	static UNREALED_API void PatchCDOSubobjectsIntoExport(UObject* PreviousCDO, UObject* NewCDO);
 
 	/** Recreates class meta data */
diff --git a/Engine/Source/Runtime/CoreUObject/Private/UObject/LinkerLoad.cpp b/Engine/Source/Runtime/CoreUObject/Private/UObject/LinkerLoad.cpp
--- a/Engine/Source/Runtime/CoreUObject/Private/UObject/LinkerLoad.cpp
+++ b/Engine/Source/Runtime/CoreUObject/Private/UObject/LinkerLoad.cpp
@@ -637,6 +637,19 @@ FLinkerLoad* FLinkerLoad::FindExistingLinkerForImport(int32 Index) const
 	return FoundLinker;
 }
 
+
+void FLinkerLoad::PRIVATE_PatchNewObjectIntoExport(UObject* OldObject, UObject* NewObject,
+	FUObjectSerializeContext* InLoadContext, bool bHideGarbageObjects)
+{
+	FLinkerLoad* OldObjectLinker = OldObject->GetLinker();
+	// if this thing doesn't have a linker, then it wasn't loaded off disk and all of this is moot
+	if (OldObjectLinker)
+	{
+		OldObjectLinker->PRIVATE_PatchNewObjectIntoExport(OldObject->GetLinkerIndex(), NewObject, InLoadContext,
+			bHideGarbageObjects);
+	}
+}
+
 /**
  * !!!!!!!!!!!!!!!!!!!!!!!!!!!!!!!!!!!!!!!!!!!!!!!!!!!!!!!!!!!!!!!!!!!!!!!!!!!!!!!!!!!!!!!!!!!!!!!!!!!!!!!!!!!!!!!!!!!!!!!!!!!!!!!!!!
  * CAUTION:  This function is potentially DANGEROUS.  Should only be used when you're really, really sure you know what you're doing.
@@ -650,36 +663,112 @@ FLinkerLoad* FLinkerLoad::FindExistingLinkerForImport(int32 Index) const
  * WARNING!!!	This function is potentially very dangerous!  It should only be used at very specific times, and in very specific cases.
  *				If you're unsure, DON'T TRY TO USE IT!!!
  */
-void FLinkerLoad::PRIVATE_PatchNewObjectIntoExport(UObject* OldObject, UObject* NewObject)
+void FLinkerLoad::PRIVATE_PatchNewObjectIntoExport(int32 OldExportIndex, UObject* NewObject,
+	FUObjectSerializeContext* InLoadContext, bool bHideGarbageObjects)
 {
-	// Cache off the old object's linker and export index.  We'll slide the new object in here.
-	FLinkerLoad* OldObjectLinker = OldObject->GetLinker();
-	// if this thing doesn't have a linker, then it wasn't loaded off disk and all of this is moot
-	if (OldObjectLinker)
+	if (!ExportMap.IsValidIndex(OldExportIndex))
 	{
-		const int32 CachedLinkerIndex = OldObject->GetLinkerIndex();
-		FObjectExport& ObjExport = OldObjectLinker->ExportMap[CachedLinkerIndex];
-		
-		// Since we don't copy the internal flags, the mirrored flags need can't be set on the new object as well
-		const EObjectFlags OldObjectFlags = OldObject->GetFlags() & ~(RF_MirroredGarbage);
+		return;
+	}
+
+	FObjectExport& ObjExport = ExportMap[OldExportIndex];
+	UObject* OldObject = ObjExport.Object;
+
+	// MarkAsGarbage on an object will make it return IsValidChecked == false, and it should not be used.
+	// SavePackage will refuse to save it in this case, so we likewise should not give a pointer to it to
+	// imports from other packages; it should be treated as if it was deleted and InvalidateExport called
+	// on it instead. This is important during cook for imports of archetype objects that were deleted during
+	// blueprint compile.
+	bool bHideNewObject = bHideGarbageObjects && !IsValidChecked(NewObject);
 
-		// Detach the old object to make room for the new
-		OldObject->ClearFlags(RF_NeedLoad|RF_NeedPostLoad|RF_NeedPostLoadSubobjects);
-		OldObject->SetLinker(nullptr, INDEX_NONE, true);
+	if (OldObject != NewObject)
+	{
+		if (OldObject)
+		{
+			// Since we don't copy the internal flags from Old to New, the mirrored flags for internal flags
+			// need to not be copied as well.
+			const EObjectFlags OldObjectFlags = OldObject->GetFlags() & ~(RF_MirroredGarbage);
+
+			// Detach the old object to make room for the new
+			OldObject->ClearFlags(RF_NeedLoad | RF_NeedPostLoad | RF_NeedPostLoadSubobjects);
+			OldObject->SetLinker(nullptr, INDEX_NONE, true /* bShouldDetachExisting */);
+
+			// Copy flags from the old CDO.
+			NewObject->SetFlags(OldObjectFlags);
+
+			// If the object was in the ObjLoaded queue (exported, but not yet serialized), swap out for our new object
+			FUObjectSerializeContext* LoadContext = InLoadContext ? InLoadContext
+				: FUObjectThreadContext::Get().GetSerializeContext();
+			if (LoadContext)
+			{
+				if (!LoadContext->PRIVATE_PatchNewObjectIntoExport(OldObject, NewObject))
+				{
+					// Make sure the new object gets PostLoad called on it if it needs it:
+					// it wasn't in the queue so add it ObjLoaded list
+					if (OldObjectFlags & RF_NeedPostLoad)
+					{
+						LoadContext->AddLoadedObject(NewObject);
+					}
+				}
+			}
+		}
 
-		// Copy flags from the old CDO.
-		NewObject->SetFlags(OldObjectFlags);
+		// bExportLoadFailed = true may have been previously set to true, and we now have a valid object, so should we
+		// set it to false? Answer: No. Anything reading bExportLoadFailed will have first checked ObjExport.Object and
+		// just returned that if non-null. So bExportLoadFailed is only used if garbage collection happens and the
+		// NewObject is garbage collected. In that case, reloading the old object would be invalid; we need to reexecute
+		// blueprint compile or whatever else figured out how to construct the NewObject, rather than using the data
+		// that was serialized to disk for the old object.
 
 		// Move the new object into the old object's slot, so any references to this object will now reference the new
-		NewObject->SetLinker(OldObjectLinker, CachedLinkerIndex);
-		ObjExport.Object = NewObject;
+		if (!bHideNewObject)
+		{
+			NewObject->SetLinker(this, OldExportIndex);
+			ObjExport.Object = NewObject;
+		}
+	}
+	if (bHideNewObject)
+	{
+		NewObject->ClearFlags(RF_NeedLoad | RF_NeedPostLoad | RF_NeedPostLoadSubobjects);
+		NewObject->SetLinker(nullptr, INDEX_NONE, true /* bShouldDetachExisting */);
+		ObjExport.bExportLoadFailed = true;
+	}
 
-		// If the object was in the ObjLoaded queue (exported, but not yet serialized), swap out for our new object
-		FUObjectThreadContext::Get().GetSerializeContext()->PRIVATE_PatchNewObjectIntoExport(OldObject, NewObject);
+	// Recursively call PRIVATE_PatchNewObjectIntoExport for every old child object in the linker map that matches
+	// a new child object, and recursively InvalidateExport for every old child object that doesn't have a match.
+	// We need to do this because a change in identity for the outer object implies a change in identity for the
+	// child objects.
+	// We have to use children in the linker map rather than children in the in-memory OldObject because children
+	// of the in-memory OldObject may have been renamed out of OldObject by e.g. blueprint reinstancing.
+	TArray<int32> OldChildExports;
+	FindDirectChildExportsFromExportTable(OldExportIndex, OldChildExports);
+	TMap<FName, UObject*> NewChildMap;
+	if (!bHideNewObject)
+	{
+		TArray<UObject*> NewChildren;
+		GetObjectsWithOuter(NewObject, NewChildren, false /* bIncludeNestedObjects */);
+		NewChildMap.Reserve(NewChildren.Num());
+		for (UObject* NewChild : NewChildren)
+		{
+			NewChildMap.Add(NewChild->GetFName(), NewChild);
+		}
+	}
+	for (int32 OldChildIndex : OldChildExports)
+	{
+		FObjectExport& OldChildExport = ExportMap[OldChildIndex];
+		UObject** NewChildPtr = NewChildMap.Find(OldChildExport.ObjectName);
+		if (NewChildPtr)
+		{
+			PRIVATE_PatchNewObjectIntoExport(OldChildIndex, *NewChildPtr, InLoadContext, bHideGarbageObjects);
+		}
+		else
+		{
+			InvalidateExportIndex(OldChildIndex, bHideGarbageObjects);
+		}
 	}
 }
 
-void FLinkerLoad::InvalidateExport(UObject* OldObject)
+void FLinkerLoad::InvalidateExport(UObject* OldObject, bool bHideGarbageObjects)
 {
 	FLinkerLoad* OldObjectLinker = OldObject->GetLinker();
 	const int32 CachedLinkerIndex = OldObject->GetLinkerIndex();
@@ -688,24 +777,64 @@ void FLinkerLoad::InvalidateExport(UObject* OldObject)
 	{
 		// Prevent any further loading as this export is now invalid
 		OldObject->ClearFlags(RF_NeedLoad | RF_NeedPostLoad | RF_NeedPostLoadSubobjects);
+		bool bHideObject = bHideGarbageObjects && !IsValidChecked(OldObject);
+		if (bHideObject)
+		{
+			OldObject->SetLinker(nullptr, INDEX_NONE, true /* bShouldDetachExisting */);
+		}
 
 		FObjectExport& ObjExport = OldObjectLinker->ExportMap[CachedLinkerIndex];
 		ObjExport.bExportLoadFailed = true;
+
+		// Recursively call InvalidateExport for every child object in the linker map.
+		// We need to do this because an invalidation of an outer implies an invalidation of the child objects.
+		// We have to use children in the linker map rather than children in the in-memory OldObject because children
+		// of the in-memory object may have been renamed out of the object by e.g. blueprint reinstancing.
+		TArray<int32> ChildExports;
+		OldObjectLinker->FindDirectChildExportsFromExportTable(CachedLinkerIndex, ChildExports);
+		for (int32 ChildExport : ChildExports)
+		{
+			OldObjectLinker->InvalidateExportIndex(ChildExport, bHideGarbageObjects);
+		}
+	}
+}
+
+void FLinkerLoad::FindDirectChildExportsFromExportTable(int32 ExportIndex, TArray<int32>& OutChildExports)
+{
+	int32 NumExports = ExportMap.Num();
+	for (int32 ChildIndex = 0; ChildIndex < NumExports; ++ChildIndex)
+	{
+		const FObjectExport& ChildExport = ExportMap[ChildIndex];
+		if (ChildExport.OuterIndex.IsExport() && ChildExport.OuterIndex.ToExport() == ExportIndex)
+		{
+			OutChildExports.Add(ChildIndex);
+		}
 	}
+}
 
-	ForEachObjectWithOuter(OldObject, [OldObjectLinker](UObject* SubObject)
+void FLinkerLoad::InvalidateExportIndex(int32 ExportIndex, bool bHideGarbageObjects)
+{
+	FObjectExport& ObjExport = ExportMap[ExportIndex];
+	UObject* OldObject = ObjExport.Object;
+	if (OldObject)
+	{
+		// Prevent any further loading as this export is now invalid
+		OldObject->ClearFlags(RF_NeedLoad | RF_NeedPostLoad | RF_NeedPostLoadSubobjects);
+		bool bHideObject = bHideGarbageObjects && !IsValidChecked(OldObject);
+		if (bHideObject)
 		{
-			FLinkerLoad* OldSubObjectLinker = SubObject->GetLinker();
-			if (OldSubObjectLinker == OldObjectLinker)
-			{
-				const int32 SubObjectCachedLinkerIndex = SubObject->GetLinkerIndex();
-				if (OldObjectLinker->ExportMap.IsValidIndex(SubObjectCachedLinkerIndex))
-				{
-					FObjectExport& ObjExport = OldObjectLinker->ExportMap[SubObjectCachedLinkerIndex];
-					ObjExport.bExportLoadFailed = true;
-				}
-			}
-		});
+			OldObject->SetLinker(nullptr, INDEX_NONE, true /* bShouldDetachExisting */);
+		}
+	}
+	ObjExport.bExportLoadFailed = true;
+
+	// Recursively call InvalidateExport for every child object in the linker map; see comment in InvalidateExport.
+	TArray<int32> ChildExports;
+	FindDirectChildExportsFromExportTable(ExportIndex, ChildExports);
+	for (int32 ChildExport : ChildExports)
+	{
+		InvalidateExportIndex(ChildExport, bHideGarbageObjects);
+	}
 }
 
 /**
diff --git a/Engine/Source/Runtime/CoreUObject/Public/UObject/LinkerLoad.h b/Engine/Source/Runtime/CoreUObject/Public/UObject/LinkerLoad.h
--- a/Engine/Source/Runtime/CoreUObject/Public/UObject/LinkerLoad.h
+++ b/Engine/Source/Runtime/CoreUObject/Public/UObject/LinkerLoad.h
@@ -788,7 +788,10 @@ public:
 	 * WARNING!!!	This function is potentially very dangerous!  It should only be used at very specific times, and in very specific cases.
 	 *				If you're unsure, DON'T TRY TO USE IT!!!
 	 */
-	COREUOBJECT_API static void PRIVATE_PatchNewObjectIntoExport(UObject* OldObject, UObject* NewObject);
+	COREUOBJECT_API static void PRIVATE_PatchNewObjectIntoExport(UObject* OldObject, UObject* NewObject,
+		FUObjectSerializeContext* InLoadContext = nullptr, bool bHideGarbageObjects=true);
+	COREUOBJECT_API void PRIVATE_PatchNewObjectIntoExport(int32 OldExportIndex, UObject* NewObject,
+		FUObjectSerializeContext* InLoadContext = nullptr, bool bHideGarbageObjects = true);
 
 	/**
 	 * Wraps a call to the package linker's ResolveAllImports().
@@ -806,7 +809,7 @@ public:
 	 * Invalidates the future loading of a specific object, so that subsequent loads will fail
 	 * This is used to invalidate sub objects of a replaced object that may no longer be valid
 	 */
-	COREUOBJECT_API static void InvalidateExport(UObject* OldObject);
+	COREUOBJECT_API static void InvalidateExport(UObject* OldObject, bool bHideGarbageObjects = true);
 
 #if WITH_EDITOR
 	COREUOBJECT_API static bool GetPreloadingEnabled();
@@ -859,7 +862,16 @@ private:
 	void ReplaceExportIndexes(const FPackageIndex& OldIndex, const FPackageIndex& NewIndex);
 
 #endif // WITH_EDITOR
-
+	/**
+	 * Helper for recursion in PRIVATE_ForceLoadAllDependencies; searches the Export table to find
+	 * all children of the given ExportIndex. Does not search recursively - does not find children of children.
+	 */
+	void FindDirectChildExportsFromExportTable(int32 ExportIndex, TArray<int32>& OutChildExports);
+	/**
+	 * Recursive helper for InvalidateExport. Invalidates an export even if it does not have a UObject* in memory,
+	 * and recursively invalidates child exports.
+	 */
+	void InvalidateExportIndex(int32 ExportIndex, bool bHideGarbageObjects);
 
 	/**
 	 * Validate if the current linker may reference a certain package
