Fix crash when saving by caching metadata so we don't have look for it during save proper

#rnx

diff --git a/Engine/Source/Runtime/CoreUObject/Private/UObject/SavePackage2.cpp b/Engine/Source/Runtime/CoreUObject/Private/UObject/SavePackage2.cpp
--- a/Engine/Source/Runtime/CoreUObject/Private/UObject/SavePackage2.cpp
+++ b/Engine/Source/Runtime/CoreUObject/Private/UObject/SavePackage2.cpp
@@ -259,6 +259,8 @@ void PreSavePackage(FSaveContext& SaveContext)
 	{
 		Package->SetPersistentGuid(FGuid::NewGuid());
 	}
+	// Create or cache metadata on the Package so it isn't searched for once we enter the save scope where it is illegal
+	Package->GetMetaData();
 #endif //WITH_EDITOR
 }
 
