Add missing dependencies for classes added in the cooked package but not present in the editor package:

- FSaveContext::AddExportedClassesToDependencies : find all the classes in the cooked packages and add them as dependencies if they are not in the editor package.

- HarvestPackage : if cooking, call the new function AddExportedClassesToDependencies.

#rnx

diff --git a/Engine/Source/Runtime/CoreUObject/Private/UObject/SavePackage/SaveContext.cpp b/Engine/Source/Runtime/CoreUObject/Private/UObject/SavePackage/SaveContext.cpp
--- a/Engine/Source/Runtime/CoreUObject/Private/UObject/SavePackage/SaveContext.cpp
+++ b/Engine/Source/Runtime/CoreUObject/Private/UObject/SavePackage/SaveContext.cpp
@@ -363,6 +363,44 @@ void FSaveContext::UpdateEditorRealmPackageBuildDependencies()
 #endif
 }
 
+void FSaveContext::AddExportedClassesToDependencies()
+{
+#if WITH_EDITOR
+
+	//first find the list of the existing class in the editor package
+	FName PackageName = GetPackage()->GetFName();
+	FAssetPackageData AssetPackageData;
+	IAssetRegistryInterface::GetPtr()->TryGetAssetPackageData(PackageName, AssetPackageData);
+
+	//Retrieve all the UObjects saved in the package
+	const TSet<FTaggedExport>& ExportsSet = GetHarvestedRealm(ESaveRealm::Game).GetExports();
+
+	//Make a list of unique class saved in the cooked package but not present in the editor package.
+	//Those are the one I want to add as dependencies.
+	TStringBuilder<256> ClassPathName;
+	TSet<FName> ImportedClassesSet(AssetPackageData.ImportedClasses);
+	TSet<const UClass*> UniqueClassesToAdd;
+	for (const FTaggedExport& Export : ExportsSet)
+	{
+		const UClass* ClassToExport = Export.Obj->GetClass();
+
+		ClassPathName.Reset();
+		ClassToExport->GetPathName(nullptr, ClassPathName);
+		bool AlreadyInEditorPackage = ImportedClassesSet.Find(FName(ClassPathName)) != nullptr;
+		if (!AlreadyInEditorPackage)
+		{
+			UniqueClassesToAdd.Add(ClassToExport);
+		}
+	}
+
+	//Finally add those missing class to the dependencies
+	TArray<UE::Cook::FCookDependency>& CookDependencies = GetObjectSaveContext().CookBuildDependencies;
+	for (const UClass* Class : UniqueClassesToAdd)
+		CookDependencies.Add(UE::Cook::FCookDependency::NativeClass(Class));
+
+#endif // WITH_EDITOR
+}
+
 const TCHAR* LexToString(ESaveableStatus Status)
 {
 	static_assert(static_cast<int32>(ESaveableStatus::__Count) == 10);
diff --git a/Engine/Source/Runtime/CoreUObject/Private/UObject/SavePackage/SaveContext.h b/Engine/Source/Runtime/CoreUObject/Private/UObject/SavePackage/SaveContext.h
--- a/Engine/Source/Runtime/CoreUObject/Private/UObject/SavePackage/SaveContext.h
+++ b/Engine/Source/Runtime/CoreUObject/Private/UObject/SavePackage/SaveContext.h
@@ -958,6 +958,8 @@ public:
 	*/
 	void UpdateEditorRealmPackageBuildDependencies();
 
+	void AddExportedClassesToDependencies();
+
 	const TMap<TObjectPtr<UObject>, TArray<FName>>& GetSearchableNamesObjectMap() const
 	{
 		return GetHarvestedRealm().GetSearchableNamesObjectMap();
diff --git a/Engine/Source/Runtime/CoreUObject/Private/UObject/SavePackage2.cpp b/Engine/Source/Runtime/CoreUObject/Private/UObject/SavePackage2.cpp
--- a/Engine/Source/Runtime/CoreUObject/Private/UObject/SavePackage2.cpp
+++ b/Engine/Source/Runtime/CoreUObject/Private/UObject/SavePackage2.cpp
@@ -460,6 +460,11 @@ ESavePackageResult HarvestPackage(FSaveContext& SaveContext)
 	SaveContext.SetCustomVersions(Harvester.GetCustomVersions());
 	SaveContext.SetTransientPropertyOverrides(Harvester.ReleaseTransientPropertyOverrides());
 
+	if (SaveContext.IsCooking())
+	{
+		SaveContext.AddExportedClassesToDependencies();
+	}
+
 	// Contractually all CookBuildDependencies (which come from PreSave or from Serialize) must be declared by this point.
 	// Copy PackageBuildDependencies from CookBuildDependencies and implement our contract that PackageBuildDependencies
 	// are in the name map.
