[AutoRTFM] Various fixes for package loader calls in transactional code

* Fix the logic in InitAsyncThread() so that the FTransactionallySafeAsyncPackageLoader is correctly constructed even when GPackageLoader is not null (move outside of if-block).
* Remove AutoRTFM::UnreachableIfTransactional() in ShouldAlwaysLoadPackageAsync(). This is safe to call in a transaction.
* Remove a number of open blocks around calls into the loader. The FTransactionallySafeAsyncPackageLoader now handles these, and these opens can corrupt memory as they modify memory also written by the closed transaction.

diff --git a/Engine/Source/Runtime/CoreUObject/Private/Serialization/AsyncPackageLoader.cpp b/Engine/Source/Runtime/CoreUObject/Private/Serialization/AsyncPackageLoader.cpp
--- a/Engine/Source/Runtime/CoreUObject/Private/Serialization/AsyncPackageLoader.cpp
+++ b/Engine/Source/Runtime/CoreUObject/Private/Serialization/AsyncPackageLoader.cpp
@@ -248,13 +248,13 @@ void InitAsyncThread()
 
 			GPackageLoader.Reset(MakeAsyncPackageLoader2(FIoDispatcher::Get()));
 		}
+	}
 
 #if UE_AUTORTFM
-		// If we are using AutoRTFM we need to layer a transactionally safe package loader over
-		// whatever loader was to be used.
-		GPackageLoader.Reset(MakeTransactionallySafeAsyncPackageLoader(GPackageLoader.Release()));
+	// If we are using AutoRTFM we need to layer a transactionally safe package loader over
+	// whatever loader was to be used.
+	GPackageLoader.Reset(MakeTransactionallySafeAsyncPackageLoader(GPackageLoader.Release()));
 #endif
-	}
 
 	FPlatformAtomics::InterlockedIncrement(&GIsLoaderCreated);
 
diff --git a/Engine/Source/Runtime/CoreUObject/Private/Serialization/TransactionallySafeAsyncLoading.cpp b/Engine/Source/Runtime/CoreUObject/Private/Serialization/TransactionallySafeAsyncLoading.cpp
--- a/Engine/Source/Runtime/CoreUObject/Private/Serialization/TransactionallySafeAsyncLoading.cpp
+++ b/Engine/Source/Runtime/CoreUObject/Private/Serialization/TransactionallySafeAsyncLoading.cpp
@@ -35,7 +35,6 @@ public:
 
 	bool ShouldAlwaysLoadPackageAsync(const FPackagePath& PackagePath) override
 	{
-		AutoRTFM::UnreachableIfTransactional(UnreachableMessage);
 		return WrappedPackageLoader->ShouldAlwaysLoadPackageAsync(PackagePath);
 	}
 
diff --git a/Engine/Source/Runtime/CoreUObject/Private/UObject/SoftObjectPath.cpp b/Engine/Source/Runtime/CoreUObject/Private/UObject/SoftObjectPath.cpp
--- a/Engine/Source/Runtime/CoreUObject/Private/UObject/SoftObjectPath.cpp
+++ b/Engine/Source/Runtime/CoreUObject/Private/UObject/SoftObjectPath.cpp
@@ -628,10 +628,7 @@ UObject* FSoftObjectPath::TryLoad(FUObjectSerializeContext* InLoadContext) const
 			}
 #endif
 
-			UE_AUTORTFM_OPEN
-			{
-				LoadedObject = StaticLoadObject(UObject::StaticClass(), nullptr, *PathString, nullptr, LOAD_None, nullptr, true);
-			};
+			LoadedObject = StaticLoadObject(UObject::StaticClass(), nullptr, *PathString, nullptr, LOAD_None, nullptr, true);
 
 			// Look at core redirects if we didn't find the object
 			if (!LoadedObject && SoftObjectPath::bResolveCoreRedirects)
@@ -639,10 +636,7 @@ UObject* FSoftObjectPath::TryLoad(FUObjectSerializeContext* InLoadContext) const
 				FSoftObjectPath FixupObjectPath = *this;
 				if (FixupObjectPath.FixupCoreRedirects())
 				{
-					UE_AUTORTFM_OPEN
-					{
-						LoadedObject = LoadObject<UObject>(nullptr, *FixupObjectPath.ToString());
-					};
+					LoadedObject = LoadObject<UObject>(nullptr, *FixupObjectPath.ToString());
 				}
 			}
 
diff --git a/Engine/Source/Runtime/CoreUObject/Private/UObject/UObjectGlobals.cpp b/Engine/Source/Runtime/CoreUObject/Private/UObject/UObjectGlobals.cpp
--- a/Engine/Source/Runtime/CoreUObject/Private/UObject/UObjectGlobals.cpp
+++ b/Engine/Source/Runtime/CoreUObject/Private/UObject/UObjectGlobals.cpp
@@ -1287,22 +1287,16 @@ bool ResolveName2(UObject*& InPackage, FStringBuilderBase& InOutName, bool Creat
 			
 			if (!bIsScriptPackage && !InPackage)
 			{
-				UE_AUTORTFM_OPEN
-				{
-					InPackage = LoadPackage(Cast<UPackage>(InPackage), *PartialName, LoadFlags, nullptr, InstancingContext);
-				};
+				InPackage = LoadPackage(Cast<UPackage>(InPackage), *PartialName, LoadFlags, nullptr, InstancingContext);
 			}
 
 			if (!InPackage)
 			{
-				UE_AUTORTFM_OPEN
+				InPackage = CreatePackage(*PartialName);
+				if (bIsScriptPackage)
 				{
-					InPackage = CreatePackage(*PartialName);
-					if (bIsScriptPackage)
-					{
-						Cast<UPackage>(InPackage)->SetPackageFlags(PKG_CompiledIn);
-					}
-				};
+					Cast<UPackage>(InPackage)->SetPackageFlags(PKG_CompiledIn);
+				}
 			}
 
 			check(InPackage);
