Marking remote reference patcher archive as MigratingRemoteObjects to prevent it from accidently resolving remote objects

diff --git a/Engine/Source/Runtime/CoreUObject/Private/UObject/RemoteObjectSerialization.cpp b/Engine/Source/Runtime/CoreUObject/Private/UObject/RemoteObjectSerialization.cpp
--- a/Engine/Source/Runtime/CoreUObject/Private/UObject/RemoteObjectSerialization.cpp
+++ b/Engine/Source/Runtime/CoreUObject/Private/UObject/RemoteObjectSerialization.cpp
@@ -584,6 +584,7 @@ public:
 		ArIsObjectReferenceCollector = true;
 		SetIsPersistent(false);
 		InitObjectMap(DeserializedObjects);
+		SetMigratingRemoteObjects(true);
 	}
 
 	virtual FArchive& operator<<(FObjectPtr& Value) override
