Changing some LogRemoteObject logs from Verbose to VeryVerbose because they occur multiple times per frame and are non-exceptional.

diff --git a/Engine/Source/Runtime/CoreUObject/Private/UObject/RemoteObjectTransfer.cpp b/Engine/Source/Runtime/CoreUObject/Private/UObject/RemoteObjectTransfer.cpp
--- a/Engine/Source/Runtime/CoreUObject/Private/UObject/RemoteObjectTransfer.cpp
+++ b/Engine/Source/Runtime/CoreUObject/Private/UObject/RemoteObjectTransfer.cpp
@@ -425,7 +425,7 @@ void PrefetchRemoteObject(FRemoteObjectId ObjectId, UObject* DestinationOuter)
 
 void ExecuteTransactional(const TFunction<void(void)>& Work)
 {
-	UE_LOG(LogRemoteObject, Verbose, TEXT("-- ExecuteTransactional -"));
+	UE_LOG(LogRemoteObject, VeryVerbose, TEXT("-- ExecuteTransactional --"));
 
 	GTransferQueue.LockedObjects.Empty();
 	GTransferQueue.RequestedObjects.Empty();
@@ -440,12 +440,12 @@ void ExecuteTransactional(const TFunction<void(void)>& Work)
 	{
 		for (FRemoteObjectId ObjectId : GTransferQueue.LockedObjects)
 		{
-			UE_LOG(LogRemoteObject, Verbose, TEXT("   : locked obj : %s"), *ObjectId.ToString());
+			UE_LOG(LogRemoteObject, VeryVerbose, TEXT("   : locked obj : %s"), *ObjectId.ToString());
 		}
 
 		for (FRemoteObjectId ObjectId : GTransferQueue.RequestedObjects)
 		{
-			UE_LOG(LogRemoteObject, Verbose, TEXT("   : requested obj : %s"), *ObjectId.ToString());
+			UE_LOG(LogRemoteObject, VeryVerbose, TEXT("   : requested obj : %s"), *ObjectId.ToString());
 		}
 
 		TArray<FRemoteObjectId> ObjectsToRequest;
@@ -459,7 +459,7 @@ void ExecuteTransactional(const TFunction<void(void)>& Work)
 			}
 		}
 
-		UE_LOG(LogRemoteObject, Verbose, TEXT("   : ObjectsToRequest: %d entries"), ObjectsToRequest.Num());
+		UE_LOG(LogRemoteObject, VeryVerbose, TEXT("   : ObjectsToRequest: %d entries"), ObjectsToRequest.Num());
 
 		// send the requests... slowly for now (to avoid straining the network subsystem)
 		int32 RequestsSent = 0;
@@ -472,18 +472,18 @@ void ExecuteTransactional(const TFunction<void(void)>& Work)
 				// yes, send one request
 				FRemoteObjectId RequestedObjectId = ObjectsToRequest[RequestsSent++];
 
-			if (UE::RemoteObject::Handle::RequestObject(RequestedObjectId))
-			{
-				// true means we already have the object locally and didn't actually send a request for it
-				GTransferQueue.LockedObjects.Add(RequestedObjectId);
-					UE_LOG(LogRemoteObject, Verbose, TEXT("      : %s is local"), *RequestedObjectId.ToString());
-			}
-			else
-			{
+				if (UE::RemoteObject::Handle::RequestObject(RequestedObjectId))
+				{
+					// true means we already have the object locally and didn't actually send a request for it
+					GTransferQueue.LockedObjects.Add(RequestedObjectId);
+					UE_LOG(LogRemoteObject, VeryVerbose, TEXT("      : %s is local"), *RequestedObjectId.ToString());
+				}
+				else
+				{
 					GTransferQueue.PendingRequests++;
-					UE_LOG(LogRemoteObject, Verbose, TEXT("      : %s is remote"), *RequestedObjectId.ToString());
+					UE_LOG(LogRemoteObject, VeryVerbose, TEXT("      : %s is remote"), *RequestedObjectId.ToString());
+				}
 			}
-		}
 
 			// pump the network
 			FlushRemoteObjectRequestsDelegate.ExecuteIfBound();
@@ -501,8 +501,8 @@ void ExecuteTransactional(const TFunction<void(void)>& Work)
 			}
 
 			if ((LoopCount % 100) == 0)
-		{
-			UE_LOG(LogRemoteObject, Verbose, TEXT("   : waiting for %d responses..."), GTransferQueue.PendingRequests);
+			{
+				UE_LOG(LogRemoteObject, VeryVerbose, TEXT("   : waiting for %d responses..."), GTransferQueue.PendingRequests);
 			}
 
 			// after 10 seconds, start sleeping in case we're deadlocked
@@ -515,39 +515,39 @@ void ExecuteTransactional(const TFunction<void(void)>& Work)
 		// now we also might have migrated objects away, so check to see if LockedObjects == RequestedObjects
 		if(GTransferQueue.LockedObjects.Num() != GTransferQueue.RequestedObjects.Num())
 		{
-			UE_LOG(LogRemoteObject, Verbose, TEXT("   : did we give up locked objects? trying again..."));
+			UE_LOG(LogRemoteObject, VeryVerbose, TEXT("   : did we give up locked objects? trying again..."));
 
 			// try again!
 			continue;
 		}
 
-		UE_LOG(LogRemoteObject, Verbose, TEXT("   : READY"));
+		UE_LOG(LogRemoteObject, VeryVerbose, TEXT("   : READY"));
 
 		// all of the requested objects are locked, try to perform the work
 		if (!bWorkSucceeded)
 		{
-		AutoRTFM::ETransactionResult WorkTransactResult = AutoRTFM::Transact([&Work](){ Work(); });
+			AutoRTFM::ETransactionResult WorkTransactResult = AutoRTFM::Transact([&Work](){ Work(); });
 
-		if (WorkTransactResult == AutoRTFM::ETransactionResult::AbortedByRequest)
-		{
-				// the work aborted, it should have inserted a new requested object, retry
-				UE_LOG(LogRemoteObject, Verbose, TEXT("   : work aborted"));
-		}
-		else
-		{
-			UE_LOG(LogRemoteObject, Verbose, TEXT("   : work OK!"));
+			if (WorkTransactResult == AutoRTFM::ETransactionResult::AbortedByRequest)
+			{
+					// the work aborted, it should have inserted a new requested object, retry
+				UE_LOG(LogRemoteObject, VeryVerbose, TEXT("   : work aborted"));
+			}
+			else
+			{
+				UE_LOG(LogRemoteObject, VeryVerbose, TEXT("   : work OK!"));
 
 				bWorkSucceeded = true;
 
 				if (GTransferQueue.LockedObjects.Num() == GTransferQueue.RequestedObjects.Num())
 				{
 					// the work succeeded and no more object requests were made, we're all done!
-			GTransferQueue.RequestedObjects.Empty();
-			GTransferQueue.RequestedObjectsDestinationOuter.Empty();
-			GTransferQueue.LockedObjects.Empty();
+					GTransferQueue.RequestedObjects.Empty();
+					GTransferQueue.RequestedObjectsDestinationOuter.Empty();
+					GTransferQueue.LockedObjects.Empty();
 					GTransferQueue.PendingRequests = 0;
-			break;
-		}
+					break;
+				}
 				else
 				{
 					// the work is complete, but it added object requests (prefetches), so we
