Fix for "Illegal call to StaticFindObjectFast() while serializing object data!" while saving.

diff --git a/Engine/Source/Runtime/CoreUObject/Private/UObject/PropertyObject.cpp b/Engine/Source/Runtime/CoreUObject/Private/UObject/PropertyObject.cpp
--- a/Engine/Source/Runtime/CoreUObject/Private/UObject/PropertyObject.cpp
+++ b/Engine/Source/Runtime/CoreUObject/Private/UObject/PropertyObject.cpp
@@ -180,7 +180,7 @@ void FObjectProperty::SerializeItem(FStructuredArchive::FSlot Slot, void* Value,
 			Slot << (*ObjectPtr);
 		}
 #if !(UE_BUILD_TEST || UE_BUILD_SHIPPING) 
-		if (!UnderlyingArchive.IsSaving())
+		if (UnderlyingArchive.IsLoading() || UnderlyingArchive.IsModifyingWeakAndStrongReferences())
 		{
 			CheckValidObject(Value, CurrentValue);
 		}
@@ -223,8 +223,11 @@ void FObjectProperty::SerializeItem(FStructuredArchive::FSlot Slot, void* Value,
 				//        to accommodate this (as it depends on finding itself as the set value)
 #endif // USE_CIRCULAR_DEPENDENCY_LOAD_DEFERRING
 
-				TObjectPtr<UObject> CurrentValue(*ObjectPtr);
-				CheckValidObject(Value, CurrentValue);
+				if (UnderlyingArchive.IsLoading() || UnderlyingArchive.IsModifyingWeakAndStrongReferences())
+				{
+					TObjectPtr<UObject> CurrentValue(*ObjectPtr);
+					CheckValidObject(Value, CurrentValue);
+				}
 			}
 		}
 		else
@@ -321,7 +324,10 @@ void FObjectProperty::PostSerializeObjectItem(FArchive& SerializingArchive, void
 		//        to accommodate this (as it depends on finding itself as the set value)
 #endif // USE_CIRCULAR_DEPENDENCY_LOAD_DEFERRING
 
-		CheckValidObject(Value, ObjectValue);
+		if (SerializingArchive.IsLoading() || SerializingArchive.IsModifyingWeakAndStrongReferences())
+		{
+			CheckValidObject(Value, ObjectValue);
+		}
 	}
 }
 
