UObject::Rename will now create redirectors in the original object's package instead of its outer.
This will be required when renaming public entities saved in their own package.

diff --git a/Engine/Source/Runtime/CoreUObject/Private/UObject/Obj.cpp b/Engine/Source/Runtime/CoreUObject/Private/UObject/Obj.cpp
--- a/Engine/Source/Runtime/CoreUObject/Private/UObject/Obj.cpp
+++ b/Engine/Source/Runtime/CoreUObject/Private/UObject/Obj.cpp
@@ -398,7 +398,7 @@ bool UObject::Rename(const TCHAR* InName, UObject* NewOuter, ERenameFlags Flags)
 		if ( Redirector == NULL )
 		{
 			// create a UObjectRedirector with the same name as the old object we are redirecting
-			Redirector = NewObject<UObjectRedirector>(OldOuter, OldName, RF_Standalone | RF_Public);
+			Redirector = NewObject<UObjectRedirector>(OldOuter, OldName, RF_Standalone | RF_Public, nullptr, false, nullptr, GetExternalPackage());
 		}
 
 		// point the redirector object to this object
diff --git a/Engine/Source/Runtime/CoreUObject/Public/UObject/UObjectGlobals.h b/Engine/Source/Runtime/CoreUObject/Public/UObject/UObjectGlobals.h
--- a/Engine/Source/Runtime/CoreUObject/Public/UObject/UObjectGlobals.h
+++ b/Engine/Source/Runtime/CoreUObject/Public/UObject/UObjectGlobals.h
@@ -1799,14 +1799,14 @@ COREUOBJECT_API void CheckIsClassChildOf_Internal(const UClass* Parent, const UC
  * @param	Template	the object to use for initializing the new object.  If not specified, the class's default object will be used
  * @param	bCopyTransientsFromClassDefaults	if true, copy transient from the class defaults instead of the pass in archetype ptr (often these are the same)
  * @param	InInstanceGraph						contains the mappings of instanced objects and components to their templates
- * @param	ExternalPackage						Assign an external Package to the created object if non-null
+ * @param	InExternalPackage					Assign an external Package to the created object if non-null
  * @param 	RemoteId	Globally unique id for this object
  *
  * @return	a pointer of type T to a new object of the specified class
  */
 template< class T >
 FUNCTION_NON_NULL_RETURN_START
-	T* NewObject(UObject* Outer, const UClass* Class, FName Name = NAME_None, EObjectFlags Flags = RF_NoFlags, UObject* Template = nullptr, bool bCopyTransientsFromClassDefaults = false, FObjectInstancingGraph* InInstanceGraph = nullptr, UPackage* ExternalPackage = nullptr)
+	T* NewObject(UObject* Outer, const UClass* Class, FName Name = NAME_None, EObjectFlags Flags = RF_NoFlags, UObject* Template = nullptr, bool bCopyTransientsFromClassDefaults = false, FObjectInstancingGraph* InInstanceGraph = nullptr, UPackage* InExternalPackage = nullptr)
 FUNCTION_NON_NULL_RETURN_END
 {
 	if (Name == NAME_None)
@@ -1826,7 +1826,7 @@ FUNCTION_NON_NULL_RETURN_END
 	Params.Template = Template;
 	Params.bCopyTransientsFromClassDefaults = bCopyTransientsFromClassDefaults;
 	Params.InstanceGraph = InInstanceGraph;
-	Params.ExternalPackage = ExternalPackage;
+	Params.ExternalPackage = InExternalPackage;
 
 	T* Result = static_cast<T*>(StaticConstructObject_Internal(Params));
 	return Result;
@@ -1849,7 +1849,7 @@ FUNCTION_NON_NULL_RETURN_END
 
 template< class T >
 FUNCTION_NON_NULL_RETURN_START
-	T* NewObject(UObject* Outer, FName Name, EObjectFlags Flags = RF_NoFlags, UObject* Template = nullptr, bool bCopyTransientsFromClassDefaults = false, FObjectInstancingGraph* InInstanceGraph = nullptr)
+	T* NewObject(UObject* Outer, FName Name, EObjectFlags Flags = RF_NoFlags, UObject* Template = nullptr, bool bCopyTransientsFromClassDefaults = false, FObjectInstancingGraph* InInstanceGraph = nullptr, UPackage* InExternalPackage = nullptr)
 FUNCTION_NON_NULL_RETURN_END
 {
 	if (Name == NAME_None)
@@ -1864,6 +1864,7 @@ FUNCTION_NON_NULL_RETURN_END
 	Params.Template = Template;
 	Params.bCopyTransientsFromClassDefaults = bCopyTransientsFromClassDefaults;
 	Params.InstanceGraph = InInstanceGraph;
+	Params.ExternalPackage = InExternalPackage;
 
 	T* Result = static_cast<T*>(StaticConstructObject_Internal(Params));
 	return Result;
