[CoreUObject]
- Adds new "PluginManager.LeakedAssetTrace.Severity.Editor" cvar allowing editor builds to configure their leaked asset reporting severity.

If not set, PluginManager.LeakedAssetTrace.Severity.Editor == PluginManager.LeakedAssetTrace.Severity.

PluginManager.LeakedAssetTrace.Severity.Shipping overrides PluginManager.LeakedAssetTrace.Severity.Editor as it does PluginManager.LeakedAssetTrace.Severity already

[RN] Adds new "PluginManager.LeakedAssetTrace.Severity.Editor" CVar allowing Editor builds to configure the severity of leaked asset reporting from GameFeaturePlugins. If not set, "PluginManager.LeakedAssetTrace.Severity.Editor" == "PluginManager.LeakedAssetTrace.Severity". "PluginManager.LeakedAssetTrace.Severity.Shipping" overrides "PluginManager.LeakedAssetTrace.Severity.Editor" as it does "PluginManager.LeakedAssetTrace.Severity" already

diff --git a/Engine/Source/Runtime/CoreUObject/Private/Misc/CoreUObjectPluginManager.cpp b/Engine/Source/Runtime/CoreUObject/Private/Misc/CoreUObjectPluginManager.cpp
--- a/Engine/Source/Runtime/CoreUObject/Private/Misc/CoreUObjectPluginManager.cpp
+++ b/Engine/Source/Runtime/CoreUObject/Private/Misc/CoreUObjectPluginManager.cpp
@@ -31,7 +31,9 @@ namespace UE::CoreUObject::Private
 #if WITH_LOW_LEVEL_TESTS
 	bool bEnsureOnLeakedPackages = false;
 #endif
+
 	static int32 GLeakedAssetTrace_Severity = 2;
+	static int32 GLeakedAssetTrace_EditorSeverity = -1;
 	static FAutoConsoleVariableRef CVarLeakedAssetTrace_Severity(
 #if UE_BUILD_SHIPPING
 		TEXT("PluginManager.LeakedAssetTrace.Severity.Shipping"),
@@ -47,6 +49,17 @@ namespace UE::CoreUObject::Private
 		ECVF_Default
 	);
 
+	static FAutoConsoleVariableRef CVarLeakedAssetTrace_EditorSeverity(
+		TEXT("PluginManager.LeakedAssetTrace.Severity.Editor"),
+		GLeakedAssetTrace_EditorSeverity,
+		TEXT("Controls severity of logging in Editor when the engine detects that assets from an Game Feature Plugin were leaked during unloading or unmounting. Overrides PluginManager.LeakedAssetTrace.Severity if set.\n")
+		TEXT("0 - all reference tracing and logging is disabled\n")
+		TEXT("1 - logs an error\n")
+		TEXT("2 - ensure\n")
+		TEXT("3 - fatal error\n"),
+		ECVF_Default
+	);
+
 	static FAutoConsoleVariableRef CVarLeakedAssetTrace_SeverityOldName(
 #if UE_BUILD_SHIPPING
 		TEXT("GameFeaturePlugin.LeakedAssetTrace.Severity.Shipping"),
@@ -134,7 +147,13 @@ namespace UE::CoreUObject::Private
 	{
 		TRACE_CPUPROFILER_EVENT_SCOPE(HandlePossibleAssetLeaks);
 
-		const bool bFindLeakedPackages = GVerifyUnload && (GLeakedAssetTrace_Severity != 0 || GRenameLeakedPackages);
+		// Allow the editor to customize it's own severity during development.
+#if WITH_EDITOR && !UE_BUILD_SHIPPING
+		const int32 LeakedAssetSeverity = GLeakedAssetTrace_EditorSeverity >= 0 ? GLeakedAssetTrace_EditorSeverity : GLeakedAssetTrace_Severity;
+#else
+		const int32 LeakedAssetSeverity = GLeakedAssetTrace_Severity;
+#endif
+		const bool bFindLeakedPackages = GVerifyUnload && (LeakedAssetSeverity != 0 || GRenameLeakedPackages);
 		if (!bFindLeakedPackages)
 		{
 			return;
@@ -179,10 +198,10 @@ namespace UE::CoreUObject::Private
 		}
 #endif
 
-		if (GLeakedAssetTrace_Severity != 0)
+		if (LeakedAssetSeverity != 0)
 		{
 			EPrintStaleReferencesOptions Options = EPrintStaleReferencesOptions::None;
-			switch (GLeakedAssetTrace_Severity)
+			switch (LeakedAssetSeverity)
 			{
 			case 3:
 				Options = EPrintStaleReferencesOptions::Fatal;
