remoteobject: re-resolve in place rather than recursively, so we can emit a better error message if it fails

diff --git a/Engine/Source/Runtime/CoreUObject/Private/UObject/RemoteObject.cpp b/Engine/Source/Runtime/CoreUObject/Private/UObject/RemoteObject.cpp
--- a/Engine/Source/Runtime/CoreUObject/Private/UObject/RemoteObject.cpp
+++ b/Engine/Source/Runtime/CoreUObject/Private/UObject/RemoteObject.cpp
@@ -264,7 +264,7 @@ UObject* ResolveObject(const FRemoteObjectStub* Stub, ERemoteReferenceType RefTy
 		return Object;
 	}
 
-	const bool bRemoteObject = !Object || IsRemote(Object);
+	bool bRemoteObject = !Object || IsRemote(Object);
 	if (bRemoteObject)
 	{
 		checkf(!GIsGarbageCollecting, TEXT("Resolving remote objects while collecting garbage is not allowed (trying to resolve object %s (%s)"), *Stub->Id.ToString(), *Stub->Name.ToString());
@@ -276,7 +276,9 @@ UObject* ResolveObject(const FRemoteObjectStub* Stub, ERemoteReferenceType RefTy
 		// if running transactionally, we will have aborted and not reached here
 
 		// if running non-transactionally, object migrated immediately, so we just re-resolve
-		return ResolveObject(Stub, RefType);
+		Object = ObjectMaps->FindObject(Stub->Id);
+		bRemoteObject = !Object || IsRemote(Object);
+		checkf(!bRemoteObject, TEXT("Failed to resolve remote object %s, either this code is not running in a transaction and should be, or the transaction failed to abort"), *Stub->Id.ToString());
 	}
 
 	return Object;
