[CoreUObject]
* Added ApplicationCore to private include paths in order to be able to remove ..\..\.. paths (clang vfsoverlay does not like those because they go outside root of vfs

diff --git a/Engine/Source/Runtime/CoreUObject/CoreUObject.Build.cs b/Engine/Source/Runtime/CoreUObject/CoreUObject.Build.cs
--- a/Engine/Source/Runtime/CoreUObject/CoreUObject.Build.cs
+++ b/Engine/Source/Runtime/CoreUObject/CoreUObject.Build.cs
@@ -33,6 +33,7 @@ public class CoreUObject : ModuleRules
 			new string[]
 			{
 				"TargetPlatform",
+				"ApplicationCore",
 			}
 		);
 
diff --git a/Engine/Source/Runtime/CoreUObject/Private/UObject/NoExportTypes.cpp b/Engine/Source/Runtime/CoreUObject/Private/UObject/NoExportTypes.cpp
--- a/Engine/Source/Runtime/CoreUObject/Private/UObject/NoExportTypes.cpp
+++ b/Engine/Source/Runtime/CoreUObject/Private/UObject/NoExportTypes.cpp
@@ -1,8 +1,8 @@
 // Copyright Epic Games, Inc. All Rights Reserved.
 
 #include "UObject/NoExportTypes.h"
-#include "../../../ApplicationCore/Public/GenericPlatform/ICursor.h"
-#include "../../../ApplicationCore/Public/GenericPlatform/IInputInterface.h"
-#include "../../../ApplicationCore/Public/GenericPlatform/InputDeviceMappingPolicy.h"
+#include "GenericPlatform/ICursor.h"
+#include "GenericPlatform/IInputInterface.h"
+#include "GenericPlatform/InputDeviceMappingPolicy.h"
 
 #include UE_INLINE_GENERATED_CPP_BY_NAME(NoExportTypes)
