[CoreUObject]
* Fixed compile error when building without application core

diff --git a/Engine/Source/Runtime/CoreUObject/CoreUObject.Build.cs b/Engine/Source/Runtime/CoreUObject/CoreUObject.Build.cs
--- a/Engine/Source/Runtime/CoreUObject/CoreUObject.Build.cs
+++ b/Engine/Source/Runtime/CoreUObject/CoreUObject.Build.cs
@@ -33,7 +33,6 @@ public class CoreUObject : ModuleRules
 			new string[]
 			{
 				"TargetPlatform",
-				"ApplicationCore",
 			}
 		);
 
@@ -72,6 +71,11 @@ public class CoreUObject : ModuleRules
 			PrivateDependencyModuleNames.Add("DerivedDataCache");
 		}
 
+		if (Target.bCompileAgainstApplicationCore)
+		{
+			PrivateIncludePathModuleNames.Add("ApplicationCore");
+		}
+
 		PrivateDefinitions.Add("UE_DEFINE_LEGACY_MATH_CONSTANT_MACRO_NAMES=0");
 	}
 }
diff --git a/Engine/Source/Runtime/CoreUObject/Private/UObject/NoExportTypes.cpp b/Engine/Source/Runtime/CoreUObject/Private/UObject/NoExportTypes.cpp
--- a/Engine/Source/Runtime/CoreUObject/Private/UObject/NoExportTypes.cpp
+++ b/Engine/Source/Runtime/CoreUObject/Private/UObject/NoExportTypes.cpp
@@ -1,8 +1,15 @@
 // Copyright Epic Games, Inc. All Rights Reserved.
 
 #include "UObject/NoExportTypes.h"
+
+#if WITH_APPLICATION_CORE
 #include "GenericPlatform/ICursor.h"
 #include "GenericPlatform/IInputInterface.h"
 #include "GenericPlatform/InputDeviceMappingPolicy.h"
+#else
+#include "../../../ApplicationCore/Public/GenericPlatform/ICursor.h" // TODO, Bad! will not work with clang vfs overlay
+#include "../../../ApplicationCore/Public/GenericPlatform/IInputInterface.h" // TODO, Bad! will not work with clang vfs overlay
+#include "../../../ApplicationCore/Public/GenericPlatform/InputDeviceMappingPolicy.h" // TODO, Bad! will not work with clang vfs overlay
+#endif
 
 #include UE_INLINE_GENERATED_CPP_BY_NAME(NoExportTypes)
