[StructUtils] Made FStructTracker's serializability optional, still "on" by default

#rnx

diff --git a/Engine/Source/Runtime/CoreUObject/Private/StructUtils/StructTypeBitSet.cpp b/Engine/Source/Runtime/CoreUObject/Private/StructUtils/StructTypeBitSet.cpp
--- a/Engine/Source/Runtime/CoreUObject/Private/StructUtils/StructTypeBitSet.cpp
+++ b/Engine/Source/Runtime/CoreUObject/Private/StructUtils/StructTypeBitSet.cpp
@@ -92,19 +92,22 @@ int32 FStructTracker::Register(const UStruct& InStructType)
 	checkSlow(StructTypesList.Num() == NewIndex);
 	StructTypesList.Add(&InStructType);
 
-	// first-time SerializationHash initialization
-	if (SerializationHash == 0)
+	if (bIsSerializable)
 	{
-		ensure(StructTypesList.Num() == 1);
-		UStruct* BaseType = BaseStructGetter();
-		CA_ASSUME(BaseType);
-		SerializationHash = GetTypeHash(BaseType->GetFullName());
-	}
-	SerializationHash = HashCombine(SerializationHash, GetTypeHash(InStructType.GetFullName()));
+		// first-time SerializationHash initialization
+		if (SerializationHash == 0)
+		{
+			ensure(StructTypesList.Num() == 1);
+			UStruct* BaseType = BaseStructGetter();
+			CA_ASSUME(BaseType);
+			SerializationHash = GetTypeHash(BaseType->GetFullName());
+		}
+		SerializationHash = HashCombine(SerializationHash, GetTypeHash(InStructType.GetFullName()));
 
-	// it's worth pointing out that we're registering a given tracker for all the hashes created along the way.
-	// This will help with loading bitsets from serialized data.
-	FBitSetSerializationDictionary::RegisterHash(SerializationHash, *this);
+		// it's worth pointing out that we're registering a given tracker for all the hashes created along the way.
+		// This will help with loading bitsets from serialized data.
+		FBitSetSerializationDictionary::RegisterHash(SerializationHash, *this);
+	}
 
 #if WITH_STRUCTUTILS_DEBUG
 	DebugStructTypeNamesList.Add(InStructType.GetFName());
@@ -132,7 +135,13 @@ void FStructTracker::Serialize(FArchive& Ar, FStructTypeBitSet::FBitSetContainer
 		Current = Last - 1
 	};
 
-	uint8 Version = uint8(EVersion::Current);
+	if (!bIsSerializable && Ar.IsSaving())
+	{
+		// nothing to do if trying to save. If loading it means it's some old data and we at least need to consume it.
+		return;
+	}
+
+	uint8 Version = static_cast<uint8>(EVersion::Current);
 	Ar << Version;
 
 	if (Ar.IsSaving())
@@ -163,7 +172,7 @@ void FStructTracker::Serialize(FArchive& Ar, FStructTypeBitSet::FBitSetContainer
 
 		// Size of the serialized memory
 		Ar.Seek(SizeOffset);	
-		SerialSize = (int32)(FinalDataOffset - InitialDataOffset);
+		SerialSize = static_cast<int32>(FinalDataOffset - InitialDataOffset);
 		Ar << SerialSize;
 		// switch back to the end position
 		Ar.Seek(FinalDataOffset);
@@ -191,7 +200,7 @@ void FStructTracker::Serialize(FArchive& Ar, FStructTypeBitSet::FBitSetContainer
 		{
 			FTopLevelAssetPath BaseStructPath;
 			Ar << BaseStructPath;
-			// first, verify that the base type of the tracker matches - otherwise we can accidentally storing 
+			// first, verify that the base type of the tracker matches - otherwise we can be accidentally storing 
 			// the wrong types creating the Mapping (via FindOrAddStructTypeIndex below)
 			UStruct* SerializedBaseStructType = FindObject<UStruct>(BaseStructPath);
 
diff --git a/Engine/Source/Runtime/CoreUObject/Public/StructUtils/StructTypeBitSet.h b/Engine/Source/Runtime/CoreUObject/Public/StructUtils/StructTypeBitSet.h
--- a/Engine/Source/Runtime/CoreUObject/Public/StructUtils/StructTypeBitSet.h
+++ b/Engine/Source/Runtime/CoreUObject/Public/StructUtils/StructTypeBitSet.h
@@ -287,6 +287,7 @@ struct COREUOBJECT_API FStructTracker
 	TSet<uint32> StructTypeToIndexSet;
 	TArray<TWeakObjectPtr<const UStruct>, TInlineAllocator<64>> StructTypesList;
 	uint32 SerializationHash = 0;
+	uint32 bIsSerializable : 1 = true;
 	const FBaseStructGetter BaseStructGetter;
 };
 
