Fixes for IDO serialization

- Removed filtering from serialization of the IDO. The whole object needs to be serialized to properly track initialized and serialized state, and the CopyTaggedProperties afterward was messing up that state.
- Added tracking of initialized property state to the save in CopyTaggedProperties to filter out uninitialized properties. This is now only used by PostLoadInstanceDataObject and DuplicateInstanceDataObject.
- Made the archive required for CreateInstanceDataObject to remove a possible failure case, and added the pre-serialize offset to keep the seeking logic together.
- Made serialized property state tracking consistent with initialized property state tracking in tagged property serialization, and added it for text assets.
- Disabled unknown enum tracking during CopyTaggedProperties, in case it was active in the calling scope.

#rnx

diff --git a/Engine/Source/Runtime/CoreUObject/Internal/UObject/PropertyBagRepository.h b/Engine/Source/Runtime/CoreUObject/Internal/UObject/PropertyBagRepository.h
--- a/Engine/Source/Runtime/CoreUObject/Internal/UObject/PropertyBagRepository.h
+++ b/Engine/Source/Runtime/CoreUObject/Internal/UObject/PropertyBagRepository.h
@@ -57,14 +57,20 @@ public:
 	void ShrinkMaps();
 
 	// Future version for reworked InstanceDataObjects - track InstanceDataObject rather than bag (directly):
+
 	/**
-	 * Instantiate an InstanceDataObject object representing all fields within the bag, tracked against the owner object.
-	 * @param Owner			- Associated in world object.
-	 * @param Archive		- used to read value of new archive from. Leave this set to nullptr to use the object's linker or copy Owner
-	 * @return				- Custom InstanceDataObject object, UClass derived from associated bag.
+	 * Instantiate an InstanceDataObject object for the owner and serialize it from the archive.
+	 *
+	 * Creates an InstanceDataObject class containing the union of the fields in the owner and its associated
+	 * unknown property tree.
+	 *
+	 * @param Owner		Object from which to construct an InstanceDataObject.
+	 * @param Archive	Archive to serialize the IDO from. May be the linker for the Owner.
+	 * @param Offset	Offset to seek Archive to before serializing the object. Ignored if Archive is the linker.
+	 * @return			The InstanceDataObject object, UClass derived from associated unknown property tree.
 	 */
-	COREUOBJECT_API UObject* CreateInstanceDataObject(UObject* Owner, FArchive* Archive = nullptr);
-	COREUOBJECT_API UObject* DuplicateInstanceDataObject(UObject* SourceOwner, UObject* DestOwner);
+	UObject* CreateInstanceDataObject(UObject* Owner, FArchive& Archive, int64 Offset);
+	UObject* DuplicateInstanceDataObject(UObject* SourceOwner, UObject* DestOwner);
 
 	// called at the end of postload to copy data from Owner to its IDO
 	COREUOBJECT_API void PostLoadInstanceDataObject(const UObject* Owner);
@@ -147,9 +153,6 @@ private:
 
 	// Delete owner reference and disassociate all data. Returns success.
 	bool RemoveAssociationUnsafe(const UObject* Owner);
-
-	// Instantiate InstanceDataObject within BagData. Returns InstanceDataObject object. 
-	void CreateInstanceDataObjectUnsafe(UObject* Owner, FPropertyBagAssociationData& BagData, FArchive* Archive = nullptr);
 };
 
 #endif // WITH_EDITORONLY_DATA
@@ -159,9 +162,10 @@ private:
 struct FScopedIDOSerializationContext
 {
 #if WITH_EDITORONLY_DATA
-	COREUOBJECT_API FScopedIDOSerializationContext(UObject* InObject, FArchive& Archive);
+	COREUOBJECT_API FScopedIDOSerializationContext(UObject* Object, FArchive& Archive);
 	COREUOBJECT_API ~FScopedIDOSerializationContext();
 
+private:
 	FArchive& Archive;
 	UObject* const Object;
 	const int64 PreSerializeOffset;
@@ -173,12 +177,8 @@ struct FScopedIDOSerializationContext
 	bool bSavedTrackUnknownEnumNames;
 	bool bSavedImpersonateProperties;
 	bool bCreateIDO;
-
-private:
-	// if we're loading and IDO should be created, this will be called when the context falls out of scope
-	void FinishCreatingInstanceDataObject() const;
 #else
-	inline FScopedIDOSerializationContext(UObject* InObject, FArchive& Archive) {}
+	inline FScopedIDOSerializationContext(UObject* Object, FArchive& Archive) {}
 #endif
 
 	FScopedIDOSerializationContext(const FScopedIDOSerializationContext&) = delete;
diff --git a/Engine/Source/Runtime/CoreUObject/Private/UObject/Class.cpp b/Engine/Source/Runtime/CoreUObject/Private/UObject/Class.cpp
--- a/Engine/Source/Runtime/CoreUObject/Private/UObject/Class.cpp
+++ b/Engine/Source/Runtime/CoreUObject/Private/UObject/Class.cpp
@@ -1392,19 +1392,14 @@ void UStruct::LoadTaggedPropertiesFromText(FStructuredArchive::FSlot Slot, uint8
 				Tag.ArrayIndex = ItemIndex;
 				Tag.Name = PropertyName;
 
-				const auto SetInitialized = [SerializeContext](const UStruct* Struct, void* Data, const FProperty* Property, int32 ArrayIndex)
+				if (Tag.SerializeType == EPropertyTagSerializeType::Skipped)
 				{
 				#if WITH_EDITORONLY_DATA
 					if (SerializeContext->bTrackInitializedProperties)
 					{
-						UE::FInitializedPropertyValueState(Struct, Data).Set(Property, ArrayIndex);
+						UE::FInitializedPropertyValueState(this, Data).Set(Property, ItemIndex);
 					}
 				#endif
-				};
-
-				if (Tag.SerializeType == EPropertyTagSerializeType::Skipped)
-				{
-					SetInitialized(this, Data, Property, ItemIndex);
 					continue;
 				}
 
@@ -1421,11 +1416,13 @@ void UStruct::LoadTaggedPropertiesFromText(FStructuredArchive::FSlot Slot, uint8
 					continue;
 				}
 
+				bool bSerialized = false;
+
 				switch (Property->ConvertFromType(Tag, ItemSlot.GetValue(), Data, DefaultsStruct, Defaults))
 				{
 				case EConvertFromTypeResult::Converted:
 				case EConvertFromTypeResult::Serialized:
-					SetInitialized(this, Data, Property, ItemIndex);
+					bSerialized = true;
 					break;
 
 				case EConvertFromTypeResult::UseSerializeItem:
@@ -1443,7 +1440,7 @@ void UStruct::LoadTaggedPropertiesFromText(FStructuredArchive::FSlot Slot, uint8
 
 						// This property is ok.
 						Tag.SerializeTaggedProperty(ItemSlot.GetValue(), Property, DestAddress, DefaultsFromParent);
-						SetInitialized(this, Data, Property, ItemIndex);
+						bSerialized = true;
 					}
 					break;
 
@@ -1453,6 +1450,20 @@ void UStruct::LoadTaggedPropertiesFromText(FStructuredArchive::FSlot Slot, uint8
 				default:
 					check(false);
 				}
+
+			#if WITH_EDITORONLY_DATA
+				if (bSerialized)
+				{
+					if (SerializeContext->bTrackInitializedProperties)
+					{
+						UE::FInitializedPropertyValueState(this, Data).Set(Property, ItemIndex);
+					}
+					if (SerializeContext->bTrackSerializedProperties)
+					{
+						UE::FSerializedPropertyValueState(this, Data).Set(Property, ItemIndex);
+					}
+				}
+			#endif
 			}
 		}
 	}
@@ -1759,13 +1770,6 @@ void UStruct::SerializeVersionedTaggedProperties(FStructuredArchive::FSlot Slot,
 
 				Tag.SetProperty(Property);
 
-#if WITH_EDITORONLY_DATA
-				if (SerializeContext->bTrackSerializedProperties && Property)
-				{
-					UE::FSerializedPropertyValueState(this, Data).Set(Property, Tag.ArrayIndex);
-				}
-#endif
-
 				bool bTryStoreUnknownPropertyPath = false;
 
 				if (Property)
@@ -1923,12 +1927,15 @@ void UStruct::SerializeVersionedTaggedProperties(FStructuredArchive::FSlot Slot,
 						*Tag.Name.ToString(), *GetName(), *WriteToString<64>(Tag.GetType()), Loaded, Tag.Size, *UnderlyingArchive.GetArchiveName(),
 						Property ? *Property->GetName() : TEXT(""), *GetFullName());
 
-
 				#if WITH_EDITORONLY_DATA
 					if (SerializeContext->bTrackInitializedProperties && Property)
 					{
 						UE::FInitializedPropertyValueState(this, Data).Set(Property, Tag.ArrayIndex);
 					}
+					if (SerializeContext->bTrackSerializedProperties && Property)
+					{
+						UE::FSerializedPropertyValueState(this, Data).Set(Property, Tag.ArrayIndex);
+					}
 				#endif
 				}
 				else if (Tag.Size != Loaded)
diff --git a/Engine/Source/Runtime/CoreUObject/Private/UObject/InstanceDataObjectUtils.cpp b/Engine/Source/Runtime/CoreUObject/Private/UObject/InstanceDataObjectUtils.cpp
--- a/Engine/Source/Runtime/CoreUObject/Private/UObject/InstanceDataObjectUtils.cpp
+++ b/Engine/Source/Runtime/CoreUObject/Private/UObject/InstanceDataObjectUtils.cpp
@@ -833,9 +833,13 @@ namespace UE
 	{
 		FUObjectSerializeContext* SerializeContext = FUObjectThreadContext::Get().GetSerializeContext();
 		TGuardValue<bool> ImpersonatePropertiesScope(SerializeContext->bImpersonateProperties, true);
-		// don't mark properties as set by serialization when performing copy
+
+		// Track only initialized properties when copying. This is required to skip uninitialized properties
+		// during saving and to mark initialized properties during loading.
+		TGuardValue<bool> ScopedTrackInitializedProperties(SerializeContext->bTrackInitializedProperties, true);
 		TGuardValue<bool> ScopedTrackSerializedProperties(SerializeContext->bTrackSerializedProperties, false);
 		TGuardValue<bool> ScopedTrackUnknownProperties(SerializeContext->bTrackUnknownProperties, false);
+		TGuardValue<bool> ScopedTrackUnknownEnumNames(SerializeContext->bTrackUnknownEnumNames, false);
 
 		TArray<uint8> Buffer;
 		Buffer.Reserve(Source->GetClass()->GetStructureSize());
@@ -845,8 +849,6 @@ namespace UE
 		Writer.ArNoDelta = true;
 		Source->GetClass()->SerializeTaggedProperties(Writer, (uint8*)Source, Source->GetClass(), nullptr);
 
-		TGuardValue<bool> ScopedTrackInitializedProperties(SerializeContext->bTrackInitializedProperties, true);
-
 		FObjectReader Reader(Buffer);
 		FSerializingDefaultsScope ReaderDefaultsScope(Reader, Dest);
 		Reader.ArMergeOverrides = true;
diff --git a/Engine/Source/Runtime/CoreUObject/Private/UObject/LinkerLoad.cpp b/Engine/Source/Runtime/CoreUObject/Private/UObject/LinkerLoad.cpp
--- a/Engine/Source/Runtime/CoreUObject/Private/UObject/LinkerLoad.cpp
+++ b/Engine/Source/Runtime/CoreUObject/Private/UObject/LinkerLoad.cpp
@@ -1118,7 +1118,6 @@ FLinkerLoad::FLinkerLoad(UPackage* InParent, const FPackagePath& InPackagePath,
 , TemplateForGetArchetypeFromLoader(nullptr)
 , bForceSimpleIndexToObject(false)
 , bLockoutLegacyOperations(false)
-, bSkipKnownProperties(false)
 , bIsAsyncLoader(false)
 , bIsDestroyingLoader(false)
 #if WITH_EDITOR
@@ -7374,22 +7373,6 @@ bool FLinkerLoad::SerializeBulkData(FBulkData& BulkData, const FBulkDataSerializ
 	return true;
 }
 
-bool FLinkerLoad::ShouldSkipProperty(const FProperty* InProperty) const
-{
-	if (bSkipKnownProperties)
-	{
-#if WITH_METADATA
-		static const FName NAME_IsLooseMetadata(ANSITEXTVIEW("IsLoose"));
-		static const FName NAME_ContainsLoosePropertiesMetadata(ANSITEXTVIEW("ContainsLooseProperties"));
-		if (!InProperty->GetBoolMetaData(NAME_IsLooseMetadata) && !InProperty->GetBoolMetaData(NAME_ContainsLoosePropertiesMetadata))
-#endif // WITH_METADATA
-		{
-			return true;
-		}
-	}
-	return FArchiveUObject::ShouldSkipProperty(InProperty);
-}
-
 void FLinkerLoad::SerializeBulkMeta(UE::BulkData::Private::FBulkMetaData& Meta, FBulkDataCookedIndex& CookedIndex, int64& DuplicateSerialOffset, int32 ElementSize)
 {
 	using namespace UE::BulkData::Private;
diff --git a/Engine/Source/Runtime/CoreUObject/Private/UObject/PropertyBagRepository.cpp b/Engine/Source/Runtime/CoreUObject/Private/UObject/PropertyBagRepository.cpp
--- a/Engine/Source/Runtime/CoreUObject/Private/UObject/PropertyBagRepository.cpp
+++ b/Engine/Source/Runtime/CoreUObject/Private/UObject/PropertyBagRepository.cpp
@@ -617,17 +617,6 @@ void FPropertyBagRepository::PostEditChangeChainProperty(const UObject* Object,
 #endif
 }
 
-UObject* FPropertyBagRepository::CreateInstanceDataObject(UObject* Owner, FArchive* Archive)
-{
-	FPropertyBagRepositoryLock LockRepo(this);
-	FPropertyBagAssociationData& BagData = AssociatedData.FindOrAdd(Owner);
-	if(!BagData.InstanceDataObject)
-	{
-		CreateInstanceDataObjectUnsafe(Owner, BagData, Archive);
-	}
-	return BagData.InstanceDataObject;
-}
-
 UObject* FPropertyBagRepository::DuplicateInstanceDataObject(UObject* SourceOwner, UObject* DestOwner)
 {
 	FPropertyBagRepositoryLock LockRepo(this);
@@ -807,9 +796,15 @@ FString FPropertyBagRepository::GetReferencerName() const
 	return TEXT("FPropertyBagRepository");
 }
 
-void FPropertyBagRepository::CreateInstanceDataObjectUnsafe(UObject* Owner, FPropertyBagAssociationData& BagData, FArchive* Archive)
+UObject* FPropertyBagRepository::CreateInstanceDataObject(UObject* Owner, FArchive& Archive, int64 PreSerializeOffset)
 {
-	check(!BagData.InstanceDataObject);	// No repeated calls
+	FPropertyBagRepositoryLock LockRepo(this);
+	FPropertyBagAssociationData& BagData = AssociatedData.FindOrAdd(Owner);
+	if (BagData.InstanceDataObject)
+	{
+		return BagData.InstanceDataObject;
+	}
+
 	TSharedPtr<FPropertyPathNameTree> PropertyTree = FUnknownPropertyTree(Owner).Find();
 	const FUnknownEnumNames EnumNames(Owner);
 	const FUnknownEnumNames* EnumNamesPtr = !EnumNames.IsEmpty() ? &EnumNames : nullptr;
@@ -836,38 +831,33 @@ void FPropertyBagRepository::CreateInstanceDataObjectUnsafe(UObject* Owner, FPro
 	// enable impersonation so that the IDO gets loaded instead of Owner
 	TGuardValue<bool> ScopedImpersonateProperties(LoadContext->bImpersonateProperties, true);
 
-	FLinkerLoad* Linker = Owner->GetLinker();
-	if (Archive && Archive != Linker)
-	{
-		// re-deserialize Owner but redirect it into the IDO instead using impersonation
-		{
-			FGuardValue_Bitfield(Archive->ArMergeOverrides, true);
-			Owner->Serialize(*Archive);
-		}
-		
-		// copy data from owner to IDO
-		CopyTaggedProperties(Owner, BagData.InstanceDataObject);
-	}
-	else if (Linker)
+	if (FLinkerLoad* Linker = Owner->GetLinker())
 	{
 		Owner->SetFlags(RF_NeedLoad);
 		{
-			TGuardValue<bool> ScopedSkipKnownProperties(Linker->bSkipKnownProperties, true);
 			FGuardValue_Bitfield(Linker->ArMergeOverrides, true);
 			Linker->Preload(Owner);
 		}
-
-		// copy data from owner to IDO
-		CopyTaggedProperties(Owner, BagData.InstanceDataObject);
 	}
 	else
 	{
-		ensureMsgf(PropertyTree == nullptr,
-			TEXT("Linker missing when generating IDO for an object with unknown properties. The unknown properties will be lost. Path: %s"),
-			*Owner->GetPathName());
-		// copy data from owner to IDO
-		CopyTaggedProperties(Owner, BagData.InstanceDataObject);
+		const int64 PostSerializeOffset = Archive.Tell();
+		Archive.Seek(PreSerializeOffset);
+
+		{
+			FGuardValue_Bitfield(Archive.ArMergeOverrides, true);
+			Owner->Serialize(Archive);
+		}
+
+		if (!ensureMsgf(Archive.Tell() == PostSerializeOffset,
+						TEXT("Serializing %s into its IDO consumed %" INT64_FMT " bytes when %" INT64_FMT " bytes were expected."),
+						*Owner->GetPathName(), Archive.Tell() - PreSerializeOffset, PostSerializeOffset - PreSerializeOffset))
+		{
+			Archive.Seek(PostSerializeOffset);
+		}
 	}
+
+	return BagData.InstanceDataObject;
 }
 
 FScopedIDOSerializationContext::FScopedIDOSerializationContext(UObject* InObject, FArchive& InArchive)
@@ -914,7 +904,7 @@ FScopedIDOSerializationContext::~FScopedIDOSerializationContext()
 {
 	if (bCreateIDO)
 	{
-		FinishCreatingInstanceDataObject();
+		FPropertyBagRepository::Get().CreateInstanceDataObject(Object, Archive, PreSerializeOffset);
 	}
 
 	// Restore the parts of the serialize context that the constructor modified.
@@ -928,33 +918,6 @@ FScopedIDOSerializationContext::~FScopedIDOSerializationContext()
 	SerializeContext->bImpersonateProperties = bSavedImpersonateProperties;
 }
 
-void FScopedIDOSerializationContext::FinishCreatingInstanceDataObject() const
-{
-	if (&Archive == Object->GetLinker())
-	{
-		// when using the linker, the repository will handle offsets
-		FPropertyBagRepository::Get().CreateInstanceDataObject(Object, &Archive);
-	}
-	// Check that we haven't already created an IDO for this object so we can skip further work.
-	// This double processing can happen at re-instancing time during CPFUO.
-	else if (FPropertyBagRepository::Get().FindInstanceDataObject(Object) == nullptr)
-	{
-		const int64 PostSerializeOffset = Archive.Tell();
-	
-		// CreateInstanceDataObject will re-call DstObject->Serialize(*this) so set the seek pointer back before DestObject in the archive
-		Archive.Seek(PreSerializeOffset);
-			
-		FPropertyBagRepository::Get().CreateInstanceDataObject(Object, &Archive);
-
-		// make sure seek pointer is back to where it should be
-		if (!ensure(Archive.Tell() == PostSerializeOffset))
-		{
-			// for some reason CreateInstanceDataObject read a different amount of data than expected... reset seek pointer back to where it should be
-			Archive.Seek(PostSerializeOffset);
-		}
-	}
-}
-
 // Not sure this is necessary.
 void FPropertyBagRepository::ShrinkMaps()
 {
diff --git a/Engine/Source/Runtime/CoreUObject/Public/UObject/LinkerLoad.h b/Engine/Source/Runtime/CoreUObject/Public/UObject/LinkerLoad.h
--- a/Engine/Source/Runtime/CoreUObject/Public/UObject/LinkerLoad.h
+++ b/Engine/Source/Runtime/CoreUObject/Public/UObject/LinkerLoad.h
@@ -168,10 +168,6 @@ public:
 	UObject*				TemplateForGetArchetypeFromLoader;
 	bool					bForceSimpleIndexToObject;
 	bool					bLockoutLegacyOperations;
-	
-	/** Set to true during IDO generation to skip non-loose properties during serialization.
-	 * In any other context, this would end up skipping *every* property */
-	UE_INTERNAL bool bSkipKnownProperties;
 
 	/** True if Loader is FAsyncArchive  */
 	bool					bIsAsyncLoader;
@@ -410,8 +406,6 @@ public:
 
 	virtual bool SerializeBulkData(FBulkData& BulkData, const FBulkDataSerializationParams& Params) override;
 
-	virtual bool ShouldSkipProperty(const FProperty* InProperty) const override;
-
 private:
 
 	void SerializeBulkMeta(UE::BulkData::Private::FBulkMetaData& Meta, FBulkDataCookedIndex& CookedIndex, int64& DuplicateSerialOffset, int32 ElementSize);
