- Added ability to control how a dynamic layer will reload subclasses and classes of instances

diff --git a/Engine/Source/Runtime/Core/Private/Misc/ConfigCacheIni.cpp b/Engine/Source/Runtime/Core/Private/Misc/ConfigCacheIni.cpp
--- a/Engine/Source/Runtime/Core/Private/Misc/ConfigCacheIni.cpp
+++ b/Engine/Source/Runtime/Core/Private/Misc/ConfigCacheIni.cpp
@@ -6780,6 +6780,30 @@ class FIniExec : public FSelfRegisteringExec
 			return true;
 		}
 
+		if (FParse::Command(&Cmd, TEXT("AddTrackedDyn")))
+		{
+			TCHAR BranchName[256];
+			TCHAR Filename[1024];
+
+			if (FParse::Token(Cmd, BranchName, UE_ARRAY_COUNT(BranchName), true) &&
+				FParse::Token(Cmd, Filename, UE_ARRAY_COUNT(Filename), true))
+			{
+				FConfigBranch* Branch = GConfig->FindBranch(NAME_None, BranchName);
+				UE::DynamicConfig::PerformDynamicConfig("TrackedDyn", [Branch, Filename](FConfigModificationTracker* ChangeTracker)
+					{
+						FDynamicLayerInfo Info { Filename, "TrackedDyn", (uint16)DynamicLayerPriority::Unknown };
+						Branch->AddDynamicLayersToHierarchy({ Info });
+					});
+			}
+		}
+
+		if (FParse::Command(&Cmd, TEXT("RemoveTrackedDyn")))
+		{
+			UE::DynamicConfig::PerformDynamicConfig("TrackedDyn", [](FConfigModificationTracker* ChangeTracker)
+				{
+					FConfigCacheIni::RemoveTagFromAllBranches("TrackedDyn", ChangeTracker);
+				});
+		}
 
 		if (FParse::Command(&Cmd, TEXT("Diff")))
 		{
@@ -6834,7 +6858,6 @@ class FIniExec : public FSelfRegisteringExec
 			TCHAR FileName[256];
 			if (FParse::Token(Cmd, FileName, UE_ARRAY_COUNT(FileName), true))
 			{
-				FConfigModificationTracker ChangeTracker;
 				FString FilenameBase = FPaths::GetBaseFilename(FileName);
 				FConfigBranch* Branch = GConfig->FindBranch(*FilenameBase, FilenameBase);
 				
diff --git a/Engine/Source/Runtime/Core/Private/Misc/ConfigUtilities.cpp b/Engine/Source/Runtime/Core/Private/Misc/ConfigUtilities.cpp
--- a/Engine/Source/Runtime/Core/Private/Misc/ConfigUtilities.cpp
+++ b/Engine/Source/Runtime/Core/Private/Misc/ConfigUtilities.cpp
@@ -56,8 +56,22 @@ void PerformDynamicConfig(FName Tag, TFunction<void(class FConfigModificationTra
 			}
 		}
 
-		// reload any config uobjects that were updated (assuming the uobject system is in use)
-		ReloadObjects.Broadcast(&ChangeTracker);
+		if (ReloadObjects.IsBound())
+		{
+			// check for any object-reload control in each branch that was modified
+			for (const auto& BranchPair : ChangeTracker.ModifiedSectionsPerBranch)
+			{
+				// @todo: if this ever needs to look in other config systems for some reason, then we'd have to move this
+				// code probably into AddDynamicLayerStringToHierarchy/AddDynamicLayersToHierarchy (they could call a helper
+				// in ConfigTracker and pass it the branch)
+				FString IniNameString = BranchPair.Key.ToString();
+				GConfig->GetArray(TEXT("DynamicLayerClassReloadControl"), TEXT("SkipReloadInstances"), ChangeTracker.ClassesToSkipInstances, IniNameString);
+				GConfig->GetArray(TEXT("DynamicLayerClassReloadControl"), TEXT("SkipReloadSubclasses"), ChangeTracker.ClassesToSkipSubclasses, IniNameString);
+			}
+
+			// reload any config uobjects that were updated (assuming the uobject system is in use)
+			ReloadObjects.Broadcast(&ChangeTracker);
+		}
 		
 		// if we updated GDeviceProfilesIni, then fixup active DP if needed
 		TSet<FString>* DPSections = ChangeTracker.ModifiedSectionsPerBranch.Find(*GDeviceProfilesIni);
diff --git a/Engine/Source/Runtime/Core/Public/Misc/ConfigCacheIni.h b/Engine/Source/Runtime/Core/Public/Misc/ConfigCacheIni.h
--- a/Engine/Source/Runtime/Core/Public/Misc/ConfigCacheIni.h
+++ b/Engine/Source/Runtime/Core/Public/Misc/ConfigCacheIni.h
@@ -508,6 +508,10 @@ public:
 	
 	// section -> tracker
 	TMap<FString, FCVarTracker> CVars;
+
+	// sometimes reloading all instances/subclasses of a class causes trouble, so we 
+	TArray<FString> ClassesToSkipSubclasses;
+	TArray<FString> ClassesToSkipInstances;
 };
 
 class FConfigCommandStreamSection : public FConfigSectionMap
diff --git a/Engine/Source/Runtime/CoreUObject/Private/UObject/UObjectBaseUtility.cpp b/Engine/Source/Runtime/CoreUObject/Private/UObject/UObjectBaseUtility.cpp
--- a/Engine/Source/Runtime/CoreUObject/Private/UObject/UObjectBaseUtility.cpp
+++ b/Engine/Source/Runtime/CoreUObject/Private/UObject/UObjectBaseUtility.cpp
@@ -623,13 +623,40 @@ void UObjectBaseUtility::ReloadObjectsFromModifiedConfigSections(const FConfigMo
 	// Now that we have a list of classes to update, we can iterate objects and reload
 	for (const UClass* ClassToReload : ClassesToReload)
 	{
+		bool bReloadInstances = true;
+		bool bReloadSubclasses = true;
+
+		if (ChangeTracker->ClassesToSkipSubclasses.Num() || ChangeTracker->ClassesToSkipInstances.Num())
+		{
+			FString ClassPath = ClassToReload->GetPathName();
+			bReloadInstances = !ChangeTracker->ClassesToSkipInstances.Contains(ClassPath);
+			bReloadSubclasses = !ChangeTracker->ClassesToSkipSubclasses.Contains(ClassPath);
+		}
+
 		TArray<UObject*> ClassObjectsToReload;
-		GetObjectsOfClass(ClassToReload, ClassObjectsToReload, true, RF_NoFlags);
-		for (UObject* ReloadObject : ClassObjectsToReload)
+		if (bReloadInstances)
 		{
-			if (IsValid(ReloadObject))
+			GetObjectsOfClass(ClassToReload, ClassObjectsToReload, bReloadSubclasses, RF_NoFlags);
+			for (UObject* ReloadObject : ClassObjectsToReload)
 			{
-				ReloadObjectImpl(ReloadObject);
+				if (IsValid(ReloadObject))
+				{
+					ReloadObjectImpl(ReloadObject);
+				}
+			}
+		}
+		else
+		{
+			// if not the instances, then just reload the CDO of the class
+			ReloadObjectImpl(ClassToReload->GetDefaultObject());
+			if (bReloadSubclasses)
+			{
+				TArray<UClass*> Results;
+				GetDerivedClasses(ClassToReload, Results, true);
+				for (UClass* Subclass : Results)
+				{
+					ReloadObjectImpl(Subclass->GetDefaultObject());
+				}
 			}
 		}
 	}
