[CoreUObject]
* Added private include to Tests folder

diff --git a/Engine/Source/Runtime/CoreUObject/CoreUObject.Build.cs b/Engine/Source/Runtime/CoreUObject/CoreUObject.Build.cs
--- a/Engine/Source/Runtime/CoreUObject/CoreUObject.Build.cs
+++ b/Engine/Source/Runtime/CoreUObject/CoreUObject.Build.cs
@@ -76,6 +76,7 @@ public class CoreUObject : ModuleRules
 			PrivateIncludePathModuleNames.Add("ApplicationCore");
 		}
 
+		PrivateIncludePaths.Add(System.IO.Path.Combine(ModuleDirectory, "Tests"));
 		PrivateDefinitions.Add("UE_DEFINE_LEGACY_MATH_CONSTANT_MACRO_NAMES=0");
 	}
 }
