[AutoRTFM] Use a transactionally-safe read-write lock in FCoreRedirects::FRWLockWithExclusiveRecursion.

`FRWWithExclusiveRecursionScopeLockForRead` was previously adapted for AutoRTFM by simply opening around the lock, and registering an OnAbort handler to undo the lock; its adaptation predates the existence of transactionally-safe locks. `FRWWithExclusiveRecursionScopeLockForWrite` was not adapted for AutoRTFM at all, but can be reached in transactional code.

We now use a transactionally-safe read-write lock here.

diff --git a/Engine/Source/Runtime/CoreUObject/Private/UObject/CoreRedirects.cpp b/Engine/Source/Runtime/CoreUObject/Private/UObject/CoreRedirects.cpp
--- a/Engine/Source/Runtime/CoreUObject/Private/UObject/CoreRedirects.cpp
+++ b/Engine/Source/Runtime/CoreUObject/Private/UObject/CoreRedirects.cpp
@@ -1280,22 +1280,18 @@ public:
 	: InternalLock(InLock)
 	, NeedsUnlock(true)
 	{
-		// Memory validation disabled as this is hot code.
-		UE_AUTORTFM_OPEN_NO_VALIDATION { InternalLock.ReadLock(); };
-		AutoRTFM::PushOnAbortHandler(this, [this] { this->InternalLock.ReadUnlock(); });
+		InternalLock.ReadLock();
 	}
 
 	~FRWWithExclusiveRecursionScopeLockForRead()
 	{
 		if (NeedsUnlock)
 		{
-			AutoRTFM::PopOnAbortHandler(this);
-			// Memory validation disabled as this is hot code.
-			UE_AUTORTFM_OPEN_NO_VALIDATION { InternalLock.ReadUnlock(); };
+			InternalLock.ReadUnlock();
 		}
 	}
-protected:
 
+protected:
 	enum class EInitFlag
 	{
 		InitFlag
@@ -3450,9 +3446,40 @@ void FCoreRedirects::RegisterNativeRedirectsUnderWriteLock(const FCoreRedirector
 }
 #endif // UE_WITH_CORE_REDIRECTS
 
+uint32 FCoreRedirects::FRWLockWithExclusiveRecursion::LoadWriteLockOwnerThreadIdRelaxed() const
+{
+	uint32 ThreadId;
+	// Memory validation is disabled, as this is a hot code path.
+	UE_AUTORTFM_OPEN_NO_VALIDATION
+	{
+		ThreadId = WriteLockOwnerThreadId.load(std::memory_order_relaxed);
+	};
+	return ThreadId;
+}
+
+void FCoreRedirects::FRWLockWithExclusiveRecursion::StoreWriteLockOwnerThreadId(uint32 ThreadId, std::memory_order Order)
+{
+	if (AutoRTFM::IsClosed())
+	{
+		// Memory validation is disabled, as this is a hot code path.
+		// If an abort occurs, we want to restore the variable's previous value, so we use
+		// `RecordOpenWrite` here so that AutoRTFM can undo our work. (We aren't overly
+		// concerned with atomicity as all loads are done in `memory_order_relaxed`.)
+		UE_AUTORTFM_OPEN_NO_VALIDATION
+		{
+			AutoRTFM::RecordOpenWrite(&ThreadId);
+			WriteLockOwnerThreadId.store(ThreadId, Order);
+		};
+	}
+	else
+	{
+		WriteLockOwnerThreadId.store(ThreadId, Order);
+	}
+}
+
 void FCoreRedirects::FRWLockWithExclusiveRecursion::ReadLock()
 {
-	int32 InitialLockOwner = WriteLockOwnerThreadId.load(std::memory_order_relaxed);
+	uint32 InitialLockOwner = LoadWriteLockOwnerThreadIdRelaxed();
 	
 	// Avoid calling GetCurrentThreadId in the common case where we start out unlocked
 	// If we are unowned at the start, then there are three cases to consider:
@@ -3486,22 +3513,22 @@ void FCoreRedirects::FRWLockWithExclusiveRecursion::ReadLock()
 void FCoreRedirects::FRWLockWithExclusiveRecursion::WriteLock()
 {
 	uint32 CurrentThreadId = FPlatformTLS::GetCurrentThreadId();
-	if (WriteLockOwnerThreadId.load(std::memory_order_relaxed) != CurrentThreadId)
+	if (LoadWriteLockOwnerThreadIdRelaxed() != CurrentThreadId)
 	{
 		InternalLock.WriteLock();
-		WriteLockOwnerThreadId = CurrentThreadId;
+		StoreWriteLockOwnerThreadId(CurrentThreadId);
 	}
 	RecursionCount++;
 }
 
 void FCoreRedirects::FRWLockWithExclusiveRecursion::WriteUnlock()
 {
-	ensure(WriteLockOwnerThreadId.load(std::memory_order_relaxed) == FPlatformTLS::GetCurrentThreadId());
+	ensure(LoadWriteLockOwnerThreadIdRelaxed() == FPlatformTLS::GetCurrentThreadId());
 	ensure(RecursionCount > 0);
 	RecursionCount--;
 	if (RecursionCount == 0)
 	{
-		WriteLockOwnerThreadId.store(0, std::memory_order_relaxed);
+		StoreWriteLockOwnerThreadId(0, std::memory_order_relaxed);
 		InternalLock.WriteUnlock();
 	}
 }
@@ -3510,7 +3537,7 @@ void FCoreRedirects::FRWLockWithExclusiveRecursion::ReadUnlock()
 {
 	if (RecursionCount > 0)
 	{
-		ensureMsgf(WriteLockOwnerThreadId.load(std::memory_order_relaxed) == FPlatformTLS::GetCurrentThreadId(), 
+		ensureMsgf(LoadWriteLockOwnerThreadIdRelaxed() == FPlatformTLS::GetCurrentThreadId(),
 			TEXT("Called ReadUnlock() on a lock held exclusively by another thread."));
 		RecursionCount--;
 	}
diff --git a/Engine/Source/Runtime/CoreUObject/Public/UObject/CoreRedirects.h b/Engine/Source/Runtime/CoreUObject/Public/UObject/CoreRedirects.h
--- a/Engine/Source/Runtime/CoreUObject/Public/UObject/CoreRedirects.h
+++ b/Engine/Source/Runtime/CoreUObject/Public/UObject/CoreRedirects.h
@@ -13,7 +13,7 @@
 #include "Containers/UnrealString.h"
 #include "CoreMinimal.h"
 #include "CoreTypes.h"
-#include "HAL/CriticalSection.h"
+#include "Misc/TransactionallySafeRWLock.h"
 #include "Misc/EnumClassFlags.h"
 #include "UObject/NameTypes.h"
 #include "UObject/UnrealNames.h"
@@ -515,7 +515,10 @@ private:
 		void ReadUnlock();
 
 	private:
-		FRWLock InternalLock;
+		uint32 LoadWriteLockOwnerThreadIdRelaxed() const;
+		void StoreWriteLockOwnerThreadId(uint32 ThreadId, std::memory_order Order = std::memory_order_seq_cst);
+
+		FTransactionallySafeRWLock InternalLock;
 		std::atomic<uint32> WriteLockOwnerThreadId = 0;
 		int32 RecursionCount = 0;
 	};
