Clear EInternalObjectFlags::Garbage when migrating remote objects back

diff --git a/Engine/Source/Runtime/CoreUObject/Private/UObject/RemoteObjectSerialization.cpp b/Engine/Source/Runtime/CoreUObject/Private/UObject/RemoteObjectSerialization.cpp
--- a/Engine/Source/Runtime/CoreUObject/Private/UObject/RemoteObjectSerialization.cpp
+++ b/Engine/Source/Runtime/CoreUObject/Private/UObject/RemoteObjectSerialization.cpp
@@ -654,6 +654,8 @@ UObject* ConstructeRemoteObject(UObject* DestinationOuter, const FRemoteObjectHe
 
 	MapObject(Object);
 	MarkAsLocal(Object);
+	// It's possible that one of the migrated objects was still in memory on this server and it was marked as Garbage in PostMigrate so clear the flag in case we migrated it back
+	Object->ClearGarbage();
 
 	return Object;
 }
