Added remote object delegates for AutoRTFM transaction completion/failure
- The intention is to only use these for statistics reporting

diff --git a/Engine/Source/Runtime/CoreUObject/Private/UObject/RemoteObjectTransfer.cpp b/Engine/Source/Runtime/CoreUObject/Private/UObject/RemoteObjectTransfer.cpp
--- a/Engine/Source/Runtime/CoreUObject/Private/UObject/RemoteObjectTransfer.cpp
+++ b/Engine/Source/Runtime/CoreUObject/Private/UObject/RemoteObjectTransfer.cpp
@@ -36,6 +36,9 @@ TDelegate<void(FRemoteObjectId ObjectId, FRemoteServerId DatabaseId, FRemoteServ
 TDelegate<void()> FlushRemotePhysicsQueryDelegate;
 TDelegate<void()> ClearRemotePhysicsQueryQueueDelegate;
 
+TMulticastDelegate<void()> OnTransactionCompletedDelegate;
+TMulticastDelegate<void()> OnTransactionAbortedDelegate;
+	
 UObject* MigrateObjectFromObjectData(FRemoteObjectData& ObjectData, FRemoteServerId RemoteServerId, UObject* DestinationOuter);
 
 /**
@@ -537,12 +540,14 @@ void ExecuteTransactional(const TFunction<void(void)>& Work)
 			{
 					// the work aborted, it should have inserted a new requested object, retry
 				UE_LOG(LogRemoteObject, VeryVerbose, TEXT("   : work aborted"));
+				OnTransactionAbortedDelegate.Broadcast();
 			}
 			else
 			{
 				UE_LOG(LogRemoteObject, VeryVerbose, TEXT("   : work OK!"));
 
 				bWorkSucceeded = true;
+				OnTransactionCompletedDelegate.Broadcast();
 
 				ClearRemotePhysicsQueryQueueDelegate.ExecuteIfBound();
 
diff --git a/Engine/Source/Runtime/CoreUObject/Public/UObject/RemoteObjectTransfer.h b/Engine/Source/Runtime/CoreUObject/Public/UObject/RemoteObjectTransfer.h
--- a/Engine/Source/Runtime/CoreUObject/Public/UObject/RemoteObjectTransfer.h
+++ b/Engine/Source/Runtime/CoreUObject/Public/UObject/RemoteObjectTransfer.h
@@ -145,5 +145,11 @@ namespace UE::RemoteObject::Transfer
 	/** Delegate that marks the end of a successful transaction. Allows the physics query queue to be cleared. */
 	extern COREUOBJECT_API TDelegate<void()> ClearRemotePhysicsQueryQueueDelegate;
 
+	/** Delegate executed when a ExecuteTransactional observes a completed AutoRTFM transaction */
+	extern COREUOBJECT_API TMulticastDelegate<void()> OnTransactionCompletedDelegate;
+
+	/** Delegate executed when a ExecuteTransactional observes an aborted AutoRTFM transaction */
+	extern COREUOBJECT_API TMulticastDelegate<void()> OnTransactionAbortedDelegate;
+
 } // namespace UE::RemoveObject::Transfer
 
