Unset GIsSavingPackage during save when invalid references are found, so that calls can be made to StaticFindObject during reference detection.  Saving is no longer occurring at this point and so should be safe.

diff --git a/Engine/Source/Runtime/CoreUObject/Private/UObject/SavePackage2.cpp b/Engine/Source/Runtime/CoreUObject/Private/UObject/SavePackage2.cpp
--- a/Engine/Source/Runtime/CoreUObject/Private/UObject/SavePackage2.cpp
+++ b/Engine/Source/Runtime/CoreUObject/Private/UObject/SavePackage2.cpp
@@ -762,6 +762,12 @@ ESavePackageResult ValidateIllegalReferences(FSaveContext& SaveContext, TArray<U
 
 	if (ObjectsToCheck)
 	{
+		// Override GIsSavingPackage because we're no longer saving packages at this point and some serialization code can trigger a
+		// a StaticFindObject during referencing finding, which is disallowed during package saving.
+		// 
+		// This will temporarily override the value of GIsSavingPackage by FScopedSavingFlag in InnerSave.
+		TGuardValue<bool> IsSavingPackageGuard(GIsSavingPackage, false);
+
 		UObject* MostLikelyCulprit = nullptr;
 		UObject* Referencer = nullptr;
 		const FProperty* ReferencerProperty = nullptr;
@@ -3082,7 +3088,10 @@ ESavePackageResult InnerSave(FSaveContext& SaveContext)
 	}
 
 	// @todo: Need to adjust GIsSavingPackage to properly prevent generating reference once, package harvesting is done
-	// GIsSavingPackage is too harsh however, since it should be scope only to the current package
+	// GIsSavingPackage is too harsh however, since it should be scope only to the current package.
+	//
+	// Note: ValidateIllegalReferences will temporarily override FScopedSavingFlag's setting of GIsSavingPackage
+	// when illegal references are detected.
 	FScopedSavingFlag IsSavingFlag(SaveContext.IsConcurrent(), SaveContext.GetPackage());
 
 	// Split the save context into its harvested contexts,
