dstm:
- delay the resolve of the DestinationOuter pointer until after the log message
- clean up some of the logging

diff --git a/Engine/Source/Runtime/CoreUObject/Private/UObject/RemoteObjectTransfer.cpp b/Engine/Source/Runtime/CoreUObject/Private/UObject/RemoteObjectTransfer.cpp
--- a/Engine/Source/Runtime/CoreUObject/Private/UObject/RemoteObjectTransfer.cpp
+++ b/Engine/Source/Runtime/CoreUObject/Private/UObject/RemoteObjectTransfer.cpp
@@ -102,7 +102,7 @@ public:
 	*/
 	void FulfillReceiveRequest(FRemoteObjectId ObjectId, FRemoteServerId RemoteServerId, FRemoteObjectData* ObjectData)
 	{
-		UObject* DestinationOuter = nullptr;
+		TObjectPtr<UObject> DestinationOuter;
 
 		bool bFoundObject = false;
 		for(int32 RequestedObjectIndex = 0; RequestedObjectIndex < RequestedObjects.Num(); RequestedObjectIndex++)
@@ -127,6 +127,8 @@ public:
 		checkf(UE::CoreUObject::Private::FObjectHandleUtils::GetRemoteId(Result) == ObjectId, TEXT("PeerId:%s requested object %s migration from PeerId:%s but received %s"),
 			*GetGlobalServerId().ToString(), *ObjectId.ToString(), *RemoteServerId.ToString(), *UE::CoreUObject::Private::FObjectHandleUtils::GetRemoteId(Result).ToString());
 
+		UE_LOG(LogRemoteObject, Verbose, TEXT("     : FulfillReceiveRequest %s DONE"), *ObjectId.ToString());
+
 		if (bFoundObject)
 		{
 			PendingRequests--;
@@ -195,9 +197,9 @@ UObject* MigrateObjectFromObjectData(FRemoteObjectData& ObjectData, FRemoteServe
 			for(int32 SubObjectIndex = 1; SubObjectIndex < ReceivedObjects.Num(); SubObjectIndex++)
 			{
 				FRemoteObjectId SubObjectId = UE::CoreUObject::Private::FObjectHandleUtils::GetRemoteId(ReceivedObjects[SubObjectIndex]);
-				UE_LOG(LogRemoteObject, Verbose, TEXT("         (subobject %d) : %s"), SubObjectIndex, *SubObjectId.ToString());
+				UE_LOG(LogRemoteObject, VeryVerbose, TEXT("         (subobject %d) : %s"), SubObjectIndex, *SubObjectId.ToString());
 			}
-			UE_LOG(LogRemoteObject, Verbose, TEXT("         <done listing subobjects>"));
+			UE_LOG(LogRemoteObject, VeryVerbose, TEXT("         <done listing subobjects>"));
 		}
 		else
 		{
@@ -206,9 +208,9 @@ UObject* MigrateObjectFromObjectData(FRemoteObjectData& ObjectData, FRemoteServe
 			for (int32 SubObjectIndex = 1; SubObjectIndex < ReceivedObjects.Num(); SubObjectIndex++)
 			{
 				FRemoteObjectId SubObjectId = UE::CoreUObject::Private::FObjectHandleUtils::GetRemoteId(ReceivedObjects[SubObjectIndex]);
-				UE_LOG(LogRemoteObject, Verbose, TEXT("         (subobject %d) : %s"), SubObjectIndex, *SubObjectId.ToString());
+				UE_LOG(LogRemoteObject, VeryVerbose, TEXT("         (subobject %d) : %s"), SubObjectIndex, *SubObjectId.ToString());
 			}
-			UE_LOG(LogRemoteObject, Verbose, TEXT("         <done listing subobjects>"));
+			UE_LOG(LogRemoteObject, VeryVerbose, TEXT("         <done listing subobjects>"));
 		}
 	}
 
@@ -269,7 +271,7 @@ void MigrateObjectToRemoteServer(UObject* Object, FRemoteServerId DestinationSer
 
 	FRemoteObjectId ObjectId = UE::CoreUObject::Private::FObjectHandleUtils::GetRemoteId(Object);
 
-	UE_LOG(LogRemoteObject, Verbose, TEXT("     : MigrateObjectToRemoteServer sending obj %s (isremote %d) (%d locked objs)"), *ObjectId.ToString(), Handle::IsRemote(Object), GTransferQueue.LockedObjects.Num());
+	UE_LOG(LogRemoteObject, Verbose, TEXT("     : MigrateObjectToRemoteServer sending obj %s [%s] (isremote %d) (%d locked objs)"), *ObjectId.ToString(), *Object->GetName(), Handle::IsRemote(Object), GTransferQueue.LockedObjects.Num());
 
 	// is the object locked?
 	bool bIsObjectLocked = false;
@@ -309,12 +311,12 @@ void MigrateObjectToRemoteServer(UObject* Object, FRemoteServerId DestinationSer
 		TSet<UObject*> SentObjects;
 		FRemoteObjectData* ObjectData = SerializeObjectData(Object, SentObjects);
 
-		UE_LOG(LogRemoteObject, Verbose, TEXT("     : ... MigrateObjectToRemoteServer serialization complete obj %s (isremote %d)"), *ObjectId.ToString(), Handle::IsRemote(Object));
+		UE_LOG(LogRemoteObject, VeryVerbose, TEXT("     : ... MigrateObjectToRemoteServer serialization complete obj %s (isremote %d)"), *ObjectId.ToString(), Handle::IsRemote(Object));
 
 		for(UObject* SO : SentObjects)
 		{
 			FRemoteObjectId SubObjectId = UE::CoreUObject::Private::FObjectHandleUtils::GetRemoteId(SO);
-			UE_LOG(LogRemoteObject, Verbose, TEXT("         sent subobject : %s"), *SubObjectId.ToString());
+			UE_LOG(LogRemoteObject, VeryVerbose, TEXT("         sent subobject : %s"), *SubObjectId.ToString());
 		}
 
 		// Notify that object data has been migrated (actually will be but SendObjectData does a MoveTemp on the data)
@@ -431,7 +433,7 @@ void PrefetchRemoteObject(FRemoteObjectId ObjectId, UObject* DestinationOuter)
 
 void ExecuteTransactional(const TFunction<void(void)>& Work)
 {
-	UE_LOG(LogRemoteObject, VeryVerbose, TEXT("-- ExecuteTransactional --"));
+	UE_LOG(LogRemoteObject, Verbose, TEXT("-- ExecuteTransactional --"));
 
 	GTransferQueue.LockedObjects.Empty();
 	GTransferQueue.RequestedObjects.Empty();
@@ -470,7 +472,8 @@ void ExecuteTransactional(const TFunction<void(void)>& Work)
 		// send the requests... slowly for now (to avoid straining the network subsystem)
 		int32 RequestsSent = 0;
 		double HangProtectorBeginTime = FPlatformTime::Seconds();
-		for(int32 LoopCount = 0;; LoopCount++)
+		double LastWaitingPrintTime = FPlatformTime::Seconds();
+		for(;;)
 		{
 			// do we need to, and can we send more requests?
 			while ((RequestsSent < ObjectsToRequest.Num()) && (GTransferQueue.PendingRequests < 5))
@@ -506,9 +509,10 @@ void ExecuteTransactional(const TFunction<void(void)>& Work)
 				}
 			}
 
-			if ((LoopCount % 100) == 0)
+			if (FPlatformTime::Seconds() > LastWaitingPrintTime + 1.0)
 			{
 				UE_LOG(LogRemoteObject, VeryVerbose, TEXT("   : waiting for %d responses..."), GTransferQueue.PendingRequests);
+				LastWaitingPrintTime = FPlatformTime::Seconds();
 			}
 
 			// after 10 seconds, start sleeping in case we're deadlocked
@@ -569,11 +573,13 @@ void ExecuteTransactional(const TFunction<void(void)>& Work)
 			}
 		}
 	}
+
+	UE_LOG(LogRemoteObject, Verbose, TEXT("-- ExecuteTransactional DONE --"));
 }
 
 void WaitForServersToSync()
 {
-	UE_LOG(LogRemoteObject, Verbose, TEXT("-- WaitForServersToSync -"));
+	UE_LOG(LogRemoteObject, Verbose, TEXT("-- WaitForServersToSync --"));
 	WaitForServersToSyncDelegate.ExecuteIfBound();
 }
 
