When resolving remote object FWeakObjectPtr will also look for the object in the local objects map

diff --git a/Engine/Source/Runtime/CoreUObject/Private/UObject/WeakObjectPtr.cpp b/Engine/Source/Runtime/CoreUObject/Private/UObject/WeakObjectPtr.cpp
--- a/Engine/Source/Runtime/CoreUObject/Private/UObject/WeakObjectPtr.cpp
+++ b/Engine/Source/Runtime/CoreUObject/Private/UObject/WeakObjectPtr.cpp
@@ -184,32 +184,45 @@ FUObjectItem* FWeakObjectPtr::TryResolveRemoteObject() const
 	using namespace UE::RemoteObject::Handle;
 	using namespace UE::RemoteObject::Private;
 
-	if (FRemoteObjectStub* Stub = FindRemoteObjectStub(ObjectRemoteId))
+	UObject* ResolvedObject = FindObject(ObjectRemoteId);
+
+	if (ResolvedObject)
+	{		
+		if (IsRemote(ResolvedObject))
+		{
+			// Object memory is still on this server but it's marked as remote so we need to resolve it
+			ResolvedObject = ResolveObject(ResolvedObject, ERemoteReferenceType::Weak);
+		}
+	}
+	else if (FRemoteObjectStub* Stub = FindRemoteObjectStub(ObjectRemoteId))
+	{
+		ResolvedObject = ResolveObject(Stub, ERemoteReferenceType::Weak);
+	}
+
+	if (ResolvedObject)
 	{
-		if (UObject* ResolvedObject = ResolveObject(Stub))
+		int32 NewIndex = GUObjectArray.ObjectToIndex(ResolvedObject);
+		FUObjectItem* ObjectItem = GUObjectArray.IndexToObject(NewIndex);
+		checkf(ObjectItem->GetRemoteId() == ObjectRemoteId, TEXT("Remote object %s has a different id (%s) than the weak object pointer that resolved it (%s)"),
+			*ResolvedObject->GetPathName(), *ObjectItem->GetRemoteId().ToString(), *ObjectRemoteId.ToString());
+		if (ObjectIndex != NewIndex)
 		{
-			int32 NewIndex = GUObjectArray.ObjectToIndex(ResolvedObject);
-			FUObjectItem* ObjectItem = GUObjectArray.IndexToObject(NewIndex);
-			checkf(ObjectItem->GetRemoteId() == ObjectRemoteId, TEXT("Remote object %s has a different id (%s) than the weak object pointer that resolved it (%s)"),
-				*ResolvedObject->GetPathName(), *ObjectItem->GetRemoteId().ToString(), *ObjectRemoteId.ToString());
-			if (ObjectIndex != NewIndex)
-			{
-				// Currently we can't just rely on object id to resolve weak object pointers and we don't want to heep hitting the resolve path
-				// if a remote object is resolved and exists on this server so (sadly) we need to update the index, yuck.
-				// In general ObjectIndex would not be required at all if all objects were hashed with RemoteId in UObjectHash.cpp but currently we don't.
-				// ObjectIndex also serves as an optimization because it's faster than hash table lookup.
-				// With RemoteId neither ObjectIndex nor SerialNumber is used for FWeakObjectPtr comparisons/hash value calculation
-				FPlatformAtomics::AtomicStore(&const_cast<FWeakObjectPtr*>(this)->ObjectIndex, NewIndex);
-			}
-			if (ObjectSerialNumber != ObjectItem->GetSerialNumber())
-			{
-				checkf(ObjectSerialNumber == 0, TEXT("Attempting to change existing and valid serial number %d to %d when resolving remote object %s (%s)"),
-					ObjectSerialNumber, ObjectItem->GetSerialNumber(), *ResolvedObject->GetPathName(), *ObjectRemoteId.ToString());
-				FPlatformAtomics::AtomicStore(&const_cast<FWeakObjectPtr*>(this)->ObjectSerialNumber, ObjectItem->GetSerialNumber());
-			}
-			return ObjectItem;
+			// Currently we can't just rely on object id to resolve weak object pointers and we don't want to heep hitting the resolve path
+			// if a remote object is resolved and exists on this server so (sadly) we need to update the index, yuck.
+			// In general ObjectIndex would not be required at all if all objects were hashed with RemoteId in UObjectHash.cpp but currently we don't.
+			// ObjectIndex also serves as an optimization because it's faster than hash table lookup.
+			// With RemoteId neither ObjectIndex nor SerialNumber is used for FWeakObjectPtr comparisons/hash value calculation
+			FPlatformAtomics::AtomicStore(&const_cast<FWeakObjectPtr*>(this)->ObjectIndex, NewIndex);
 		}
+		if (ObjectSerialNumber != ObjectItem->GetSerialNumber())
+		{
+			checkf(ObjectSerialNumber == 0, TEXT("Attempting to change existing and valid serial number %d to %d when resolving remote object %s (%s)"),
+				ObjectSerialNumber, ObjectItem->GetSerialNumber(), *ResolvedObject->GetPathName(), *ObjectRemoteId.ToString());
+			FPlatformAtomics::AtomicStore(&const_cast<FWeakObjectPtr*>(this)->ObjectSerialNumber, ObjectItem->GetSerialNumber());
+		}
+		return ObjectItem;
 	}
+
 	return nullptr;
 }
 
