Renamed recently introduced MakeGloballyUniqueObjectName to MakeUniversallyUniqueObjectName as the concept of globally unique name already exists for legacy ANY_PACKAGE replacement.
Also, MakeUniqueObjectName can now accept EUniqueObjectNameOptions::UniversallyUnique to route into MakeUniversallyUniqueObjectName instead.

diff --git a/Engine/Source/Runtime/CoreUObject/Private/UObject/UObjectGlobals.cpp b/Engine/Source/Runtime/CoreUObject/Private/UObject/UObjectGlobals.cpp
--- a/Engine/Source/Runtime/CoreUObject/Private/UObject/UObjectGlobals.cpp
+++ b/Engine/Source/Runtime/CoreUObject/Private/UObject/UObjectGlobals.cpp
@@ -2818,6 +2818,11 @@ FName MakeUniqueObjectName(UObject* Parent, const UClass* Class, FName InBaseNam
 	check(Class);
 	const FName BaseName = InBaseName.IsNone() ? Class->GetFName() : InBaseName;
 
+	if (!!(Options & EUniqueObjectNameOptions::UniversallyUnique))
+	{
+		return MakeUniversallyUniqueObjectName(Parent, BaseName);
+	}
+
 #if TRY_REUSE_NAMES
 	if (FName Result = NameReuse::MakeUniqueObjectNameReusingNumber(Parent, BaseName, Options); !Result.IsNone())
 	{
@@ -2902,7 +2907,7 @@ FName MakeUniqueObjectName(UObject* Parent, const UClass* Class, FName InBaseNam
 	return TestName;
 }
 
-FName MakeGloballyUniqueObjectName(UObject* Outer, FName BaseName)
+FName MakeUniversallyUniqueObjectName(UObject* Outer, FName BaseName)
 {
 	// Generate a unique generation context for this session.
 	static struct FContext
diff --git a/Engine/Source/Runtime/CoreUObject/Public/UObject/UObjectGlobals.h b/Engine/Source/Runtime/CoreUObject/Public/UObject/UObjectGlobals.h
--- a/Engine/Source/Runtime/CoreUObject/Public/UObject/UObjectGlobals.h
+++ b/Engine/Source/Runtime/CoreUObject/Public/UObject/UObjectGlobals.h
@@ -958,7 +958,8 @@ COREUOBJECT_API void IncrementalPurgeGarbage( bool bUseTimeLimit, double TimeLim
 enum class EUniqueObjectNameOptions
 {
 	None = 0,
-	GloballyUnique = 1 << 1, // Whether to make the object name unique globally (across all objects that currently exist)
+	GloballyUnique		= 1 << 1,	// Whether to make the object name unique globally (across all objects that currently exist in the current session)
+	UniversallyUnique	= 1 << 2	// Whether to make the object name unique universally (across all objects that currently exist in all sessions)
 };
 ENUM_CLASS_FLAGS(EUniqueObjectNameOptions);
 
@@ -977,15 +978,15 @@ ENUM_CLASS_FLAGS(EUniqueObjectNameOptions);
 COREUOBJECT_API FName MakeUniqueObjectName( UObject* Outer, const UClass* Class, FName BaseName = NAME_None, EUniqueObjectNameOptions Options = EUniqueObjectNameOptions::None);
 
 /**
- * Create a globally unique name by combining a base name, a GUID and a timestamp.
+ * Create a universally unique name by combining a base name, a GUID and a timestamp.
  * The object name returned is guaranteed not to locally and globally exist.
  *
  * @param	Outer		the outer for the object that needs to be named
  * @param	BaseName	optional base name to use when generating the globally unique object name.
  * 
- * @return	name is the form BaseName_##, where ## a globally unique identifier making sure that name is globally unique.
+ * @return	name is the form BaseName_##, where ## a universally unique identifier making sure that name is universally unique.
   */
-COREUOBJECT_API FName MakeGloballyUniqueObjectName(UObject* Outer, FName InBaseName = NAME_None);
+COREUOBJECT_API FName MakeUniversallyUniqueObjectName(UObject* Outer, FName InBaseName = NAME_None);
 
 /**
  * Given a display label string, generates an FName slug that is a valid FName for that label.
