Cache/apply cooked metadata for inner fields on container/optional property types.

diff --git a/Engine/Source/Runtime/CoreUObject/Private/UObject/CookedMetaData.cpp b/Engine/Source/Runtime/CoreUObject/Private/UObject/CookedMetaData.cpp
--- a/Engine/Source/Runtime/CoreUObject/Private/UObject/CookedMetaData.cpp
+++ b/Engine/Source/Runtime/CoreUObject/Private/UObject/CookedMetaData.cpp
@@ -3,6 +3,7 @@
 #include "UObject/CookedMetaData.h"
 #include "UObject/Package.h"
 #include "UObject/MetaData.h"
+#include "UObject/PropertyOptional.h"
 #include "UObject/UnrealType.h"
 
 #include UE_INLINE_GENERATED_CPP_BY_NAME(CookedMetaData)
@@ -38,8 +39,37 @@ void PostLoadCookedMetaData(CookedMetaDataType* CookedMetaDataPtr)
 #endif // WITH_METADATA
 }
 
+void GetInnerFields(const FField* OuterField, TArray<FField*, TInlineAllocator<2>>& OutInnerFields)
+{
+	OutInnerFields.Reset();
+
+	if (const FArrayProperty* ArrayProperty = CastField<FArrayProperty>(OuterField))
+	{
+		OutInnerFields.Emplace(ArrayProperty->Inner);
+	}
+	else if (const FSetProperty* SetProperty = CastField<FSetProperty>(OuterField))
+	{
+		OutInnerFields.Emplace(SetProperty->ElementProp);
+	}
+	else if (const FMapProperty* MapProperty = CastField<FMapProperty>(OuterField))
+	{
+		OutInnerFields.Emplace(MapProperty->KeyProp);
+		OutInnerFields.Emplace(MapProperty->ValueProp);
+	}
+	else if (const FOptionalProperty* OptionalProperty = CastField<FOptionalProperty>(OuterField))
+	{
+		OutInnerFields.Emplace(OptionalProperty->GetValueProperty());
+	}
+}
+
 } // namespace CookedMetaDataUtil::Internal
 
+FFieldCookedMetaDataKey::FFieldCookedMetaDataKey()
+{
+	// Shouldn't get more than 2 levels deep on inner field paths, so reserve that up front.
+	// Note: We can't use an inline allocator for that b/c this is reflected via UPROPERTY().
+	FieldPath.Reserve(2);
+}
 
 bool FObjectCookedMetaDataStore::HasMetaData() const
 {
@@ -80,29 +110,73 @@ void FObjectCookedMetaDataStore::ApplyMetaData(UObject* TargetObject) const
 
 bool FFieldCookedMetaDataStore::HasMetaData() const
 {
-	return FieldMetaData.Num() > 0;
+	return FieldMetaData.Num() > 0 || SubFieldMetaData.Num() > 0;
 }
 
 void FFieldCookedMetaDataStore::CacheMetaData(const FField* SourceField)
 {
 	FieldMetaData.Reset();
+	SubFieldMetaData.Reset();
+
+#if WITH_METADATA
+	FFieldCookedMetaDataKey SubFieldMapKey;
+	CacheFieldMetaDataInternal(SourceField, FieldMetaData, SubFieldMapKey);
+#endif // WITH_METADATA
+}
+
+void FFieldCookedMetaDataStore::ApplyMetaData(FField* TargetField) const
+{
+#if WITH_METADATA
+	FFieldCookedMetaDataKey SubFieldMapKey;
+	ApplyFieldMetaDataInternal(TargetField, FieldMetaData, SubFieldMapKey);
+#endif // WITH_METADATA
+}
 
+void FFieldCookedMetaDataStore::CacheFieldMetaDataInternal(const FField* SourceField, TMap<FName, FString>& TargetFieldMetaData, FFieldCookedMetaDataKey& SubFieldMapKey)
+{
 #if WITH_METADATA
 	if (const TMap<FName, FString>* SourceFieldMetaData = SourceField->GetMetaDataMap())
 	{
-		FieldMetaData = *SourceFieldMetaData;
+		TargetFieldMetaData = *SourceFieldMetaData;
+	}
+
+	TArray<FField*, TInlineAllocator<2>> InnerFields;
+	CookedMetaDataUtil::Internal::GetInnerFields(SourceField, InnerFields);
+	for (FField* InnerField : InnerFields)
+	{
+		const TMap<FName, FString>* InnerFieldMetaData = InnerField->GetMetaDataMap();
+		if (!InnerFieldMetaData
+			|| InnerFieldMetaData->IsEmpty())
+		{
+			continue;
+		}
+
+		FFieldCookedMetaDataKey InnerFieldMapKey(SubFieldMapKey);
+		InnerFieldMapKey.FieldPath.Emplace(InnerField->GetFName());
+		CacheFieldMetaDataInternal(InnerField, SubFieldMetaData.Emplace(InnerFieldMapKey).MetaData, InnerFieldMapKey);
 	}
 #endif // WITH_METADATA
 }
 
-void FFieldCookedMetaDataStore::ApplyMetaData(FField* TargetField) const
+void FFieldCookedMetaDataStore::ApplyFieldMetaDataInternal(FField* TargetField, const TMap<FName, FString>& SourceFieldMetaData, FFieldCookedMetaDataKey& SubFieldMapKey) const
 {
 #if WITH_METADATA
-	TargetField->AppendMetaData(FieldMetaData);
+	TargetField->AppendMetaData(SourceFieldMetaData);
+
+	TArray<FField*, TInlineAllocator<2>> InnerFields;
+	CookedMetaDataUtil::Internal::GetInnerFields(TargetField, InnerFields);
+	for (FField* InnerField : InnerFields)
+	{
+		FFieldCookedMetaDataKey InnerFieldMapKey(SubFieldMapKey);
+		InnerFieldMapKey.FieldPath.Emplace(InnerField->GetFName());
+		if (const FFieldCookedMetaDataValue* InnerFieldMetaDataValue = SubFieldMetaData.Find(InnerFieldMapKey))
+		{
+			ApplyFieldMetaDataInternal(InnerField, InnerFieldMetaDataValue->MetaData, InnerFieldMapKey);
+		}
+	}
 #endif // WITH_METADATA
 }
 
-
 bool FStructCookedMetaDataStore::HasMetaData() const
 {
 	return ObjectMetaData.HasMetaData()
diff --git a/Engine/Source/Runtime/CoreUObject/Public/UObject/CookedMetaData.h b/Engine/Source/Runtime/CoreUObject/Public/UObject/CookedMetaData.h
--- a/Engine/Source/Runtime/CoreUObject/Public/UObject/CookedMetaData.h
+++ b/Engine/Source/Runtime/CoreUObject/Public/UObject/CookedMetaData.h
@@ -26,7 +26,50 @@ private:
 };
 
 /**
- * Cooked meta-data for a FField.
+ * Encapsulates the subfield metadata map key in a UScriptStruct (for UHT).
+ */
+USTRUCT()
+struct FFieldCookedMetaDataKey
+{
+public:
+	GENERATED_BODY()
+
+	UPROPERTY()
+	TArray<FName> FieldPath;
+
+	FFieldCookedMetaDataKey();
+
+	inline bool operator==(const FFieldCookedMetaDataKey& Other) const
+	{
+		return Other.FieldPath == FieldPath;
+	}
+
+	inline bool operator!=(const FFieldCookedMetaDataKey& Other) const
+	{
+		return Other.FieldPath != FieldPath;
+	}
+
+	friend uint32 GetTypeHash(const FFieldCookedMetaDataKey& Key)
+	{
+		return GetTypeHash(Key.FieldPath);
+	}
+};
+
+/**
+ * Encapsulates the subfield metadata map value in a UScriptStruct (for UHT).
+ */
+USTRUCT()
+struct FFieldCookedMetaDataValue
+{
+public:
+	GENERATED_BODY()
+
+	UPROPERTY()
+	TMap<FName, FString> MetaData;
+};
+
+/**
+ * Cooked meta-data for a FField, including its nested FField data.
  */
 USTRUCT()
 struct FFieldCookedMetaDataStore
@@ -38,9 +81,16 @@ public:
 	COREUOBJECT_API void CacheMetaData(const FField* SourceField);
 	COREUOBJECT_API void ApplyMetaData(FField* TargetField) const;
 
+protected:
+	void CacheFieldMetaDataInternal(const FField* SourceField, TMap<FName, FString>& TargetFieldMetaData, FFieldCookedMetaDataKey& SubFieldMapKey);
+	void ApplyFieldMetaDataInternal(FField* TargetField, const TMap<FName, FString>& SourceFieldMetaData, FFieldCookedMetaDataKey& SubFieldMapKey) const;
+
 private:
 	UPROPERTY()
 	TMap<FName, FString> FieldMetaData;
+
+	UPROPERTY()
+	TMap<FFieldCookedMetaDataKey, FFieldCookedMetaDataValue> SubFieldMetaData;
 };
 
 /**
