Adding the ability to construct FObjectPtr with FRemoteObjectId

diff --git a/Engine/Source/Runtime/CoreUObject/Private/UObject/RemoteObjectSerialization.cpp b/Engine/Source/Runtime/CoreUObject/Private/UObject/RemoteObjectSerialization.cpp
--- a/Engine/Source/Runtime/CoreUObject/Private/UObject/RemoteObjectSerialization.cpp
+++ b/Engine/Source/Runtime/CoreUObject/Private/UObject/RemoteObjectSerialization.cpp
@@ -374,7 +374,7 @@ public:
 		const FRemoteObjectId WeakPtrId = Value.GetRemoteId();
 		if (UE::RemoteObject::Handle::IsRemote(WeakPtrId))
 		{
-			Ptr = UE::RemoteObject::Transfer::RemoteObjectIdToObjectPtr(WeakPtrId);
+			Ptr = FObjectPtr(WeakPtrId);
 		}
 		else
 #endif
diff --git a/Engine/Source/Runtime/CoreUObject/Public/UObject/ObjectPtr.h b/Engine/Source/Runtime/CoreUObject/Public/UObject/ObjectPtr.h
--- a/Engine/Source/Runtime/CoreUObject/Public/UObject/ObjectPtr.h
+++ b/Engine/Source/Runtime/CoreUObject/Public/UObject/ObjectPtr.h
@@ -100,6 +100,14 @@ public:
 	}
 #endif
 	
+	explicit FORCEINLINE FObjectPtr(FRemoteObjectId ObjectId)
+#if UE_WITH_REMOTE_OBJECT_HANDLE
+		: FObjectPtr(UE::CoreUObject::Private::FRemoteObjectHandlePrivate::FromIdNoResolve(ObjectId))
+#else
+		: FObjectPtr()
+#endif
+	{
+	}
 
 	FORCEINLINE UObject* Get() const
 	{
diff --git a/Engine/Source/Runtime/CoreUObject/Public/UObject/RemoteObjectTransfer.h b/Engine/Source/Runtime/CoreUObject/Public/UObject/RemoteObjectTransfer.h
--- a/Engine/Source/Runtime/CoreUObject/Public/UObject/RemoteObjectTransfer.h
+++ b/Engine/Source/Runtime/CoreUObject/Public/UObject/RemoteObjectTransfer.h
@@ -71,6 +71,7 @@ namespace UE::RemoteObject::Transfer
 	/**
 	* Utility that converts FRemoteObjectId to TObjectPtr of a given type. Will not resolve the object if it's not local.
 	*/
+	UE_DEPRECATED(5.6, "RemoteObjectIdToObjectPtr is deprecated use FObjectPtr(FRemoteObjectId) instead.")
 	FORCEINLINE FObjectPtr RemoteObjectIdToObjectPtr(FRemoteObjectId RemoteId)
 	{
 #if UE_WITH_REMOTE_OBJECT_HANDLE
diff --git a/Engine/Source/Runtime/CoreUObject/Public/UObject/RemoteObjectTypes.h b/Engine/Source/Runtime/CoreUObject/Public/UObject/RemoteObjectTypes.h
--- a/Engine/Source/Runtime/CoreUObject/Public/UObject/RemoteObjectTypes.h
+++ b/Engine/Source/Runtime/CoreUObject/Public/UObject/RemoteObjectTypes.h
@@ -62,11 +62,6 @@ public:
 
 	FRemoteObjectId() = default;
 
-	explicit FRemoteObjectId(uint64 InId)
-		: Id(InId)
-	{
-	}
-
 	FRemoteObjectId(FRemoteServerId InServerId, uint64 InSerialNumber)
 		: SerialNumber(InSerialNumber)
 		, ServerId(InServerId.Id)
