dstm:
- threaded a RequestToken through the various messages for extra debug info and also to make the PendingRequest counter more accurate, so that it doesn't decrease when we get a forced migration that isn't actually a response to our request
- another cleanup pass on the logging, added LogRemoteObjectD for extended debug info
- removed tracking of object Outer since the value passed in was unused

diff --git a/Engine/Source/Runtime/CoreUObject/Private/UObject/RemoteObject.cpp b/Engine/Source/Runtime/CoreUObject/Private/UObject/RemoteObject.cpp
--- a/Engine/Source/Runtime/CoreUObject/Private/UObject/RemoteObject.cpp
+++ b/Engine/Source/Runtime/CoreUObject/Private/UObject/RemoteObject.cpp
@@ -17,6 +17,7 @@
 #include <atomic>
 
 DEFINE_LOG_CATEGORY(LogRemoteObject);
+DEFINE_LOG_CATEGORY(LogRemoteObjectD);
 
 namespace UE::RemoteObject::Private
 {
diff --git a/Engine/Source/Runtime/CoreUObject/Private/UObject/RemoteObjectPrivate.h b/Engine/Source/Runtime/CoreUObject/Private/UObject/RemoteObjectPrivate.h
--- a/Engine/Source/Runtime/CoreUObject/Private/UObject/RemoteObjectPrivate.h
+++ b/Engine/Source/Runtime/CoreUObject/Private/UObject/RemoteObjectPrivate.h
@@ -9,6 +9,7 @@
 #include "UObject/Object.h"
 
 DECLARE_LOG_CATEGORY_EXTERN(LogRemoteObject, Display, All);
+DECLARE_LOG_CATEGORY_EXTERN(LogRemoteObjectD, Display, All);
 
 namespace UE::RemoteObject::Private
 {
diff --git a/Engine/Source/Runtime/CoreUObject/Private/UObject/RemoteObjectSerialization.cpp b/Engine/Source/Runtime/CoreUObject/Private/UObject/RemoteObjectSerialization.cpp
--- a/Engine/Source/Runtime/CoreUObject/Private/UObject/RemoteObjectSerialization.cpp
+++ b/Engine/Source/Runtime/CoreUObject/Private/UObject/RemoteObjectSerialization.cpp
@@ -63,7 +63,7 @@ FString GenerateRemoteObjectFilename(FRemoteObjectId ObjectId, FRemoteServerId D
 	return FPaths::Combine(*FPaths::ProjectSavedDir(), *(FString::Printf(TEXT("%s_%s.remote"), *ObjectId.ToString(), *DestinationServerId.ToString())));
 }
 
-void LoadObjectFromDisk(FRemoteObjectId ObjectId, FRemoteServerId CurrentOwnerServerId, FRemoteServerId DestinationServerId)
+void LoadObjectFromDisk(uint32 RequestToken, FRemoteObjectId ObjectId, FRemoteServerId CurrentOwnerServerId, FRemoteServerId DestinationServerId)
 {
 	const FString Filename(GenerateRemoteObjectFilename(ObjectId, CurrentOwnerServerId));
 	TUniquePtr<FArchive> FileReader{ IFileManager::Get().CreateFileReader(*Filename) };
@@ -75,10 +75,10 @@ void LoadObjectFromDisk(FRemoteObjectId ObjectId, FRemoteServerId CurrentOwnerSe
 	FileReader->Close();
 	IFileManager::Get().Delete(*Filename, false, true, true);
 
-	UE::RemoteObject::Transfer::OnObjectDataReceived(ObjectId, CurrentOwnerServerId, ObjectData);
+	UE::RemoteObject::Transfer::OnObjectDataReceived(RequestToken, ObjectId, CurrentOwnerServerId, ObjectData);
 }
 
-void SaveObjectToDisk(FRemoteObjectId ObjectId, FRemoteServerId DestinationServerId, FRemoteObjectData& ObjectData)
+void SaveObjectToDisk(uint32 RequestToken, FRemoteObjectId ObjectId, FRemoteServerId DestinationServerId, FRemoteObjectData& ObjectData)
 {
 	TUniquePtr<FArchive> FileWriter{ IFileManager::Get().CreateFileWriter(*GenerateRemoteObjectFilename(ObjectId, DestinationServerId)) };
 	checkf(FileWriter.IsValid(), TEXT("Unable to create file writer for remote object %s"), *ObjectId.ToString());
@@ -615,7 +615,7 @@ FRemoteObjectConstructionOverrides::FRemoteObjectConstructionOverrides(const TAr
 	}
 }
 
-UObject* ConstructeRemoteObject(UObject* DestinationOuter, const FRemoteObjectHeader& Header, FRemoteObjectConstructionOverrides* ConstructionOverrides)
+UObject* ConstructRemoteObject(const FRemoteObjectHeader& Header, FRemoteObjectConstructionOverrides* ConstructionOverrides)
 {
 	using namespace UE::RemoteObject::Handle;
 	using namespace UE::RemoteObject::Private;
@@ -730,7 +730,7 @@ FRemoteObjectData* SerializeObjectData(UObject* Object, TSet<UObject*>& OutObjec
 	return ObjectData;
 }
 
-void DeserializeObjectData(FRemoteObjectData& ObjectData, UObject* DestinationOuter, TArray<UObject*>& OutObjects)
+void DeserializeObjectData(FRemoteObjectData& ObjectData, TArray<UObject*>& OutObjects)
 {
 	using namespace UE::RemoteObject::Handle;
 	using namespace UE::RemoteObject::Private;
@@ -776,8 +776,7 @@ void DeserializeObjectData(FRemoteObjectData& ObjectData, UObject* DestinationOu
 		OutObjects.Reserve(ObjectHeaders.Num());
 		for (int32 ObjectIndex = 0; ObjectIndex < ObjectHeaders.Num(); ++ObjectIndex)
 		{
-			// Only pass the destination outer to the first object, which is always the root object of the entire hierarchy that was serialized
-			OutObjects.Add(ConstructeRemoteObject(ObjectIndex == 0 ? DestinationOuter : nullptr, ObjectHeaders[ObjectIndex], &ConstructionOverrides));
+			OutObjects.Add(ConstructRemoteObject(ObjectHeaders[ObjectIndex], &ConstructionOverrides));
 		}
 
 		// Deserialize all objects
diff --git a/Engine/Source/Runtime/CoreUObject/Private/UObject/RemoteObjectSerialization.h b/Engine/Source/Runtime/CoreUObject/Private/UObject/RemoteObjectSerialization.h
--- a/Engine/Source/Runtime/CoreUObject/Private/UObject/RemoteObjectSerialization.h
+++ b/Engine/Source/Runtime/CoreUObject/Private/UObject/RemoteObjectSerialization.h
@@ -54,13 +54,13 @@ namespace UE::RemoteObject::Serialization
 	};
 
 	FRemoteObjectData* SerializeObjectData(UObject* Object, TSet<UObject*>& OutObjects);
-	void DeserializeObjectData(FRemoteObjectData& ObjectData, UObject* DestinationOuter, TArray<UObject*>& OutReceivedObjects);
+	void DeserializeObjectData(FRemoteObjectData& ObjectData, TArray<UObject*>& OutReceivedObjects);
 
 } // namespace UE::RemoveObject::Serialization
 
 
 namespace UE::RemoteObject::Serialization::Disk
 {
-	void LoadObjectFromDisk(FRemoteObjectId ObjectId, FRemoteServerId CurrentOwnerServerId, FRemoteServerId DestinationServerId);
-	void SaveObjectToDisk(FRemoteObjectId ObjectId, FRemoteServerId ServerId, FRemoteObjectData& ObjectData);
+	void LoadObjectFromDisk(uint32 RequestToken, FRemoteObjectId ObjectId, FRemoteServerId CurrentOwnerServerId, FRemoteServerId DestinationServerId);
+	void SaveObjectToDisk(uint32 RequestToken, FRemoteObjectId ObjectId, FRemoteServerId ServerId, FRemoteObjectData& ObjectData);
 }
\ No newline at end of file
diff --git a/Engine/Source/Runtime/CoreUObject/Private/UObject/RemoteObjectTransfer.cpp b/Engine/Source/Runtime/CoreUObject/Private/UObject/RemoteObjectTransfer.cpp
--- a/Engine/Source/Runtime/CoreUObject/Private/UObject/RemoteObjectTransfer.cpp
+++ b/Engine/Source/Runtime/CoreUObject/Private/UObject/RemoteObjectTransfer.cpp
@@ -13,14 +13,15 @@
 #include "UObject/ObjectHandlePrivate.h"
 #include "HAL/CriticalSection.h"
 #include "AutoRTFM/AutoRTFM.h"
+#include "Templates/Atomic.h"
 
 namespace UE::RemoteObject::Transfer
 {
 const FRemoteServerId DatabaseId(0x3ff);
 
-TDelegate<void(FRemoteObjectId ObjectId, FRemoteServerId CurrentOwnerServerId, FRemoteServerId DestinationServerId)> RequestRemoteObjectDelegate;
-TDelegate<void(FRemoteObjectId ObjectId, FRemoteServerId DestinationServerId, FRemoteObjectData& ObjectData)> RemoteObjectTransferDelegate;
-TDelegate<void(FRemoteObjectId ObjectId, FRemoteServerId DestinationServerId)> RemoteObjectDeniedTransferDelegate;
+TDelegate<void(uint32 RequestToken, FRemoteObjectId ObjectId, FRemoteServerId CurrentOwnerServerId, FRemoteServerId DestinationServerId)> RequestRemoteObjectDelegate;
+TDelegate<void(uint32 RequestToken, FRemoteObjectId ObjectId, FRemoteServerId DestinationServerId, FRemoteObjectData& ObjectData)> RemoteObjectTransferDelegate;
+TDelegate<void(uint32 RequestToken, FRemoteObjectId ObjectId, FRemoteServerId DestinationServerId)> RemoteObjectDeniedTransferDelegate;
 TDelegate<void()> FlushRemoteObjectRequestsDelegate;
 TDelegate<void()> WaitForServersToSyncDelegate;
 
@@ -30,8 +31,8 @@ TMulticastDelegate<void(const FRemoteObjectData& ObjectData, FRemoteServerId Rem
 TMulticastDelegate<void(FRemoteObjectId ObjectId, FRemoteServerId RemoteServerId)> OnObjectDataDeniedDelegate;
 TMulticastDelegate<void(const TArray<UObject*>& MigratedObjects, FRemoteServerId RemoteServerId)> OnObjectsReceivedDelegate;
 
-TDelegate<void(FRemoteObjectId ObjectId, FRemoteServerId DatabaseId, FRemoteObjectData& ObjectData)> StoreRemoteObjectDataDelegate;
-TDelegate<void(FRemoteObjectId ObjectId, FRemoteServerId DatabaseId, FRemoteServerId DestinationServerId)> RestoreRemoteObjectDataDelegate;
+TDelegate<void(uint32 RequestToken, FRemoteObjectId ObjectId, FRemoteServerId DatabaseId, FRemoteObjectData& ObjectData)> StoreRemoteObjectDataDelegate;
+TDelegate<void(uint32 RequestToken, FRemoteObjectId ObjectId, FRemoteServerId DatabaseId, FRemoteServerId DestinationServerId)> RestoreRemoteObjectDataDelegate;
 
 TDelegate<void()> FlushRemotePhysicsQueryDelegate;
 TDelegate<void()> ClearRemotePhysicsQueryQueueDelegate;
@@ -39,7 +40,7 @@ TDelegate<void()> ClearRemotePhysicsQueryQueueDelegate;
 TMulticastDelegate<void()> OnTransactionCompletedDelegate;
 TMulticastDelegate<void()> OnTransactionAbortedDelegate;
 	
-UObject* MigrateObjectFromObjectData(FRemoteObjectData& ObjectData, FRemoteServerId RemoteServerId, UObject* DestinationOuter);
+UObject* MigrateObjectFromObjectData(FRemoteObjectData& ObjectData, FRemoteServerId RemoteServerId);
 
 /**
  * Remote object transfer queue. 
@@ -50,33 +51,56 @@ class FRemoteObjectTransferQueue
 public:
 	TArray<FRemoteObjectId> LockedObjects;
 	TArray<FRemoteObjectId> RequestedObjects;
-	TArray<TObjectPtr<UObject>> RequestedObjectsDestinationOuter;
-	int32 PendingRequests;
+	TAtomic<uint32> ActiveRequestToken;
+	TAtomic<uint32> NextActiveRequestToken;
+	int32 PendingRequests = 0;
 
-	void RequestRemoteObject(FRemoteObjectId ObjectId, FRemoteServerId CurrentOwnerServerId)
+	void BeginRequest()
+	{
+		// we expect to not have an active request token
+		check(ActiveRequestToken.Load() == 0);
+
+		// we expect the object lists to be empty
+		check(LockedObjects.Num() == 0);
+		check(RequestedObjects.Num() == 0);
+
+		// we expect to not be currently waiting for any requests
+		check(PendingRequests == 0);
+
+		uint32 ExpectedToken = 0u;
+		bool TokenExchangeSuccess = ActiveRequestToken.CompareExchange(ExpectedToken, (NextActiveRequestToken.Load() % 32000) + 1000);
+		check(TokenExchangeSuccess);
+
+		NextActiveRequestToken++;
+	}
+
+	void EndRequest()
+	{
+		// we expect to have an active request token
+		uint32 OldToken = ActiveRequestToken.Exchange(0u);
+		check(OldToken != 0);
+
+		// we expect to have no pending network requests
+		check(PendingRequests == 0);
+
+		// clear the object lists
+		LockedObjects.Empty();
+		RequestedObjects.Empty();
+	}
+
+	void RequestRemoteObject(uint32 RequestToken, FRemoteObjectId ObjectId, FRemoteServerId CurrentOwnerServerId, FRemoteServerId DestinationServerId)
 	{
 		// Request an object from remote server
 		if (CurrentOwnerServerId != DatabaseId)
 		{
-			RequestRemoteObjectDelegate.Execute(ObjectId, CurrentOwnerServerId, GetGlobalServerId());
+			RequestRemoteObjectDelegate.Execute(RequestToken, ObjectId, CurrentOwnerServerId, DestinationServerId);
 		}
 		else
 		{
-			RestoreRemoteObjectDataDelegate.Execute(ObjectId, DatabaseId, GetGlobalServerId());
+			RestoreRemoteObjectDataDelegate.Execute(RequestToken, ObjectId, DatabaseId, GetGlobalServerId());
 		}
 	}
 
-	/**
-	* Forwards receive request to a remote server
-	* @param ObjectId Id of the object being requested
-	* @param CurrentOwnerServerId Remote server id that currently owns the reqested object
-	* @param DestinationServerId Remote server id that requests the object ownership
-	*/
-	void ForwardRemoteObjectRequest(FRemoteObjectId ObjectId, FRemoteServerId CurrentOwnerServerId, FRemoteServerId DestinationServerId)
-	{
-		RequestRemoteObjectDelegate.Execute(ObjectId, CurrentOwnerServerId, DestinationServerId);
-	}
-
 	/**
 	* Sends an object to a remote server (adds a send request)
 	* @param ObjectId Id of the object being requested
@@ -84,14 +108,14 @@ public:
 	* @param ObjectData Object data to send. Send request takes ownership of the data
 	* @return Future that's going to receive the object data
 	*/
-	void SendRemoteObject(FRemoteObjectId ObjectId, FRemoteServerId DestinationServerId, FRemoteObjectData* ObjectData)
+	void SendRemoteObject(uint32 RequestToken, FRemoteObjectId ObjectId, FRemoteServerId DestinationServerId, FRemoteObjectData* ObjectData)
 	{		
-		RemoteObjectTransferDelegate.Execute(ObjectId, DestinationServerId, *ObjectData);
+		RemoteObjectTransferDelegate.Execute(RequestToken, ObjectId, DestinationServerId, *ObjectData);
 	}
 
-	void SendDeniedRemoteObject(FRemoteObjectId ObjectId, FRemoteServerId DestinationServerId)
+	void SendDeniedRemoteObject(uint32 RequestToken, FRemoteObjectId ObjectId, FRemoteServerId DestinationServerId)
 	{
-		RemoteObjectDeniedTransferDelegate.Execute(ObjectId, DestinationServerId);
+		RemoteObjectDeniedTransferDelegate.Execute(RequestToken, ObjectId, DestinationServerId);
 	}
 
 	/**
@@ -100,46 +124,45 @@ public:
 	* @param Data Object data. Data ownership is transferred to the receive request (if it exists)
 	* @return True if a receive request for the specified object exists (and object data ownership was transferred), false otherwise
 	*/
-	void FulfillReceiveRequest(FRemoteObjectId ObjectId, FRemoteServerId RemoteServerId, FRemoteObjectData* ObjectData)
+	void FulfillReceiveRequest(uint32 RequestToken, FRemoteObjectId ObjectId, FRemoteServerId RemoteServerId, FRemoteObjectData* ObjectData)
 	{
-		TObjectPtr<UObject> DestinationOuter;
+		bool bMatchesActiveRequest = (RequestToken == ActiveRequestToken.Load());
+		bool bFoundRequest = (RequestedObjects.Find(ObjectId) != INDEX_NONE);
+		bool bAlreadyLocked = (LockedObjects.Find(ObjectId) != INDEX_NONE);
 
-		bool bFoundObject = false;
-		for(int32 RequestedObjectIndex = 0; RequestedObjectIndex < RequestedObjects.Num(); RequestedObjectIndex++)
-		{
-			if(RequestedObjects[RequestedObjectIndex] == ObjectId)
-			{
-				DestinationOuter = RequestedObjectsDestinationOuter[RequestedObjectIndex];
-				bFoundObject = true;
-			}			
-		}
+		// locked objects are ones that we believe to be local - it should be impossible for
+		// another server to send us an object that we already have
+		check(!bAlreadyLocked);
 
-		if (bFoundObject)
-		{
-			UE_LOG(LogRemoteObject, Verbose, TEXT("     : FulfillReceiveRequest %s (was requested)"), *ObjectId.ToString());
-		}
-		else
-		{
-			UE_LOG(LogRemoteObject, Verbose, TEXT("     : FulfillReceiveRequest %s (was not requested)"), *ObjectId.ToString());
-		}
+		UE_LOG(LogRemoteObject, Verbose, TEXT("     : FulfillReceiveRequest(%u) %s (was requested: %d, token matches: %d)"), RequestToken, *ObjectId.ToString(), bFoundRequest, bMatchesActiveRequest);
 
-		UObject* Result = MigrateObjectFromObjectData(*ObjectData, RemoteServerId, DestinationOuter);
+		UObject* Result = MigrateObjectFromObjectData(*ObjectData, RemoteServerId);
 		checkf(UE::CoreUObject::Private::FObjectHandleUtils::GetRemoteId(Result) == ObjectId, TEXT("PeerId:%s requested object %s migration from PeerId:%s but received %s"),
 			*GetGlobalServerId().ToString(), *ObjectId.ToString(), *RemoteServerId.ToString(), *UE::CoreUObject::Private::FObjectHandleUtils::GetRemoteId(Result).ToString());
 
-		UE_LOG(LogRemoteObject, Verbose, TEXT("     : FulfillReceiveRequest %s DONE"), *ObjectId.ToString());
+		UE_LOG(LogRemoteObject, Verbose, TEXT("     : FulfillReceiveRequest(%u) %s DONE"), RequestToken, *ObjectId.ToString());
 
-		if (bFoundObject)
+		if (bFoundRequest)
+		{
+			// we received an object that we requested (was either push migrated or directly responded to),
+			// so now consider it locked
+			LockedObjects.Add(ObjectId);
+		}
+
+		if (bMatchesActiveRequest && bFoundRequest)
 		{
 			PendingRequests--;
 		}
 	}
 
-	void DenyReceiveRequest(FRemoteObjectId ObjectId)
+	void DenyReceiveRequest(uint32 RequestToken, FRemoteObjectId ObjectId)
 	{
-		UE_LOG(LogRemoteObject, Verbose, TEXT("     : DenyReceiveRequest %s"), *ObjectId.ToString());
+		bool bMatchesActiveRequest = (RequestToken == ActiveRequestToken.Load());
+		bool bFoundRequest = (RequestedObjects.Find(ObjectId) != INDEX_NONE);
+
+		UE_LOG(LogRemoteObject, Verbose, TEXT("     : DenyReceiveRequest(%u) %s, (was requested: %d, token matches: %d)"), RequestToken, *ObjectId.ToString(), bFoundRequest, bMatchesActiveRequest);
 
-		if (RequestedObjects.Find(ObjectId) != INDEX_NONE)
+		if (bMatchesActiveRequest && bFoundRequest)
 		{
 			PendingRequests--;
 		}
@@ -148,25 +171,13 @@ public:
 
 FRemoteObjectTransferQueue GTransferQueue;
 
-
 void RequestObjectData(FRemoteObjectId ObjectId, FRemoteServerId CurrentOwnerServerId, FRemoteServerId DestinationServerId)
 {
-	if (DestinationServerId == GetGlobalServerId())
-	{
-		GTransferQueue.RequestRemoteObject(ObjectId, CurrentOwnerServerId);	
-	}
-	else
-	{
-		GTransferQueue.ForwardRemoteObjectRequest(ObjectId, CurrentOwnerServerId, DestinationServerId);
-	}
+	check(GTransferQueue.ActiveRequestToken.Load() != 0);
+	GTransferQueue.RequestRemoteObject(GTransferQueue.ActiveRequestToken.Load(), ObjectId, CurrentOwnerServerId, DestinationServerId);
 }
 
-void SendObjectData(UObject* Object, FRemoteServerId DestinationServerId, FRemoteObjectData* ObjectData)
-{
-	GTransferQueue.SendRemoteObject(UE::CoreUObject::Private::FObjectHandleUtils::GetRemoteId(Object), DestinationServerId, ObjectData);
-}
-
-UObject* MigrateObjectFromObjectData(FRemoteObjectData& ObjectData, FRemoteServerId RemoteServerId, UObject* DestinationOuter)
+UObject* MigrateObjectFromObjectData(FRemoteObjectData& ObjectData, FRemoteServerId RemoteServerId)
 {
 	using namespace UE::RemoteObject::Serialization;
 
@@ -174,44 +185,18 @@ UObject* MigrateObjectFromObjectData(FRemoteObjectData& ObjectData, FRemoteServe
 	OnObjectDataReceivedDelegate.Broadcast(ObjectData, RemoteServerId);
 
 	TArray<UObject*> ReceivedObjects;
-	DeserializeObjectData(ObjectData, DestinationOuter, ReceivedObjects);
+	DeserializeObjectData(ObjectData, ReceivedObjects);
 
 	checkf(ReceivedObjects.Num() > 0, TEXT("PeerId:%s unable to deserialize object data (%d bytes)"), *GetGlobalServerId().ToString(), ObjectData.Bytes.Num());
 
 	if (ReceivedObjects.Num() > 0)
 	{
-		FRemoteObjectId RootObjectId = UE::CoreUObject::Private::FObjectHandleUtils::GetRemoteId(ReceivedObjects[0]);
-
-		if (GTransferQueue.RequestedObjects.Find(RootObjectId) != INDEX_NONE)
-		{
-			if (GTransferQueue.LockedObjects.Find(RootObjectId) == INDEX_NONE)
-			{
-				UE_LOG(LogRemoteObject, Verbose, TEXT("     : MigrateObjectFromObjectData received object %s"), *RootObjectId.ToString());
-				GTransferQueue.LockedObjects.Add(RootObjectId);
-			}
-			else
-			{
-				UE_LOG(LogRemoteObject, Verbose, TEXT("     : MigrateObjectFromObjectData ???? received an object that was already locked %s"), *RootObjectId.ToString());
-			}
-
-			for(int32 SubObjectIndex = 1; SubObjectIndex < ReceivedObjects.Num(); SubObjectIndex++)
-			{
-				FRemoteObjectId SubObjectId = UE::CoreUObject::Private::FObjectHandleUtils::GetRemoteId(ReceivedObjects[SubObjectIndex]);
-				UE_LOG(LogRemoteObject, VeryVerbose, TEXT("         (subobject %d) : %s"), SubObjectIndex, *SubObjectId.ToString());
-			}
-			UE_LOG(LogRemoteObject, VeryVerbose, TEXT("         <done listing subobjects>"));
-		}
-		else
+		for(int32 SubObjectIndex = 1; SubObjectIndex < ReceivedObjects.Num(); SubObjectIndex++)
 		{
-			// got an object that we didn't request - ownership migration?
-			UE_LOG(LogRemoteObject, Verbose, TEXT("     : MigrateObjectFromObjectData received object %s (was not requested)"), *RootObjectId.ToString());
-			for (int32 SubObjectIndex = 1; SubObjectIndex < ReceivedObjects.Num(); SubObjectIndex++)
-			{
-				FRemoteObjectId SubObjectId = UE::CoreUObject::Private::FObjectHandleUtils::GetRemoteId(ReceivedObjects[SubObjectIndex]);
-				UE_LOG(LogRemoteObject, VeryVerbose, TEXT("         (subobject %d) : %s"), SubObjectIndex, *SubObjectId.ToString());
-			}
-			UE_LOG(LogRemoteObject, VeryVerbose, TEXT("         <done listing subobjects>"));
+			FRemoteObjectId SubObjectId = UE::CoreUObject::Private::FObjectHandleUtils::GetRemoteId(ReceivedObjects[SubObjectIndex]);
+			UE_LOG(LogRemoteObjectD, Verbose, TEXT("         subobject[%2d] : %s"), SubObjectIndex, *SubObjectId.ToString());
 		}
+		UE_LOG(LogRemoteObjectD, Verbose, TEXT("         <done listing subobjects>"));
 	}
 
 #if UE_WITH_REMOTE_OBJECT_HANDLE
@@ -240,7 +225,6 @@ void MigrateObjectFromRemoteServer(FRemoteObjectId ObjectId, FRemoteServerId Cur
 		if (bAdded)
 		{
 			GTransferQueue.RequestedObjects.Add(ObjectId);
-			GTransferQueue.RequestedObjectsDestinationOuter.Add(DestinationOuter);
 		}
 
 		UE_LOG(LogRemoteObject, Verbose, TEXT("Aborting transaction, need obj id %s on server %s"),
@@ -258,20 +242,19 @@ void MigrateObjectFromRemoteServer(FRemoteObjectId ObjectId, FRemoteServerId Cur
 		// if we are not in a transaction, request the object immediately
 		UE::RemoteObject::Handle::RequestObject(ObjectId);
 		GTransferQueue.RequestedObjects.Empty();
-		GTransferQueue.RequestedObjectsDestinationOuter.Empty();
 		GTransferQueue.LockedObjects.Empty();
 		GTransferQueue.PendingRequests = 0;
 	}
 }
 
-void MigrateObjectToRemoteServer(UObject* Object, FRemoteServerId DestinationServerId)
+void MigrateObjectToRemoteServer(uint32 RequestToken, UObject* Object, FRemoteServerId DestinationServerId)
 {
 	using namespace UE::RemoteObject::Serialization;
 	using namespace UE::RemoteObject::Private;
 
 	FRemoteObjectId ObjectId = UE::CoreUObject::Private::FObjectHandleUtils::GetRemoteId(Object);
 
-	UE_LOG(LogRemoteObject, Verbose, TEXT("     : MigrateObjectToRemoteServer sending obj %s [%s] (isremote %d) (%d locked objs)"), *ObjectId.ToString(), *Object->GetName(), Handle::IsRemote(Object), GTransferQueue.LockedObjects.Num());
+	UE_LOG(LogRemoteObject, Verbose, TEXT("     : MigrateObjectToRemoteServer(%u) sending obj %s [%s] (isremote %d) (%d locked objs)"), RequestToken, *ObjectId.ToString(), *Object->GetName(), Handle::IsRemote(Object), GTransferQueue.LockedObjects.Num());
 
 	// is the object locked?
 	bool bIsObjectLocked = false;
@@ -288,12 +271,12 @@ void MigrateObjectToRemoteServer(UObject* Object, FRemoteServerId DestinationSer
 				// we are higher priority, send a deny response
 				bIsObjectLocked = true;
 
-				UE_LOG(LogRemoteObject, Verbose, TEXT("     : MigrateObjectToRemoteServer DENYING request for %s because it's locked"), *LockedObjectId.ToString());
-				RemoteObjectDeniedTransferDelegate.Execute(LockedObjectId, DestinationServerId);
+				UE_LOG(LogRemoteObject, Verbose, TEXT("     : MigrateObjectToRemoteServer(%u) DENYING request for %s because it's locked"), RequestToken, *LockedObjectId.ToString());
+				RemoteObjectDeniedTransferDelegate.Execute(RequestToken, LockedObjectId, DestinationServerId);
 			}
 			else
 			{
-				UE_LOG(LogRemoteObject, Verbose, TEXT("     : MigrateObjectToRemoteServer GIVING UP our locked %s"), *LockedObjectId.ToString());
+				UE_LOG(LogRemoteObject, Verbose, TEXT("     : MigrateObjectToRemoteServer(%u) GIVING UP our locked %s"), RequestToken, *LockedObjectId.ToString());
 
 				// unlock the object from ourselves
 				GTransferQueue.LockedObjects.RemoveAt(LockedObjectIndex);
@@ -305,24 +288,24 @@ void MigrateObjectToRemoteServer(UObject* Object, FRemoteServerId DestinationSer
 
 	if (!bIsObjectLocked)
 	{
-		UE_LOG(LogRemoteObject, Verbose, TEXT("     : ... MigrateObjectToRemoteServer serializing obj %s (isremote %d)"), *ObjectId.ToString(), Handle::IsRemote(Object));
+		UE_LOG(LogRemoteObject, Verbose, TEXT("     : MigrateObjectToRemoteServer(%u) serializing obj %s (isremote %d)"), RequestToken, *ObjectId.ToString(), Handle::IsRemote(Object));
 
 		// send the object
 		TSet<UObject*> SentObjects;
 		FRemoteObjectData* ObjectData = SerializeObjectData(Object, SentObjects);
 
-		UE_LOG(LogRemoteObject, VeryVerbose, TEXT("     : ... MigrateObjectToRemoteServer serialization complete obj %s (isremote %d)"), *ObjectId.ToString(), Handle::IsRemote(Object));
+		UE_LOG(LogRemoteObject, VeryVerbose, TEXT("     : MigrateObjectToRemoteServer(%u) serialization complete obj %s (isremote %d)"), RequestToken, *ObjectId.ToString(), Handle::IsRemote(Object));
 
 		for(UObject* SO : SentObjects)
 		{
 			FRemoteObjectId SubObjectId = UE::CoreUObject::Private::FObjectHandleUtils::GetRemoteId(SO);
-			UE_LOG(LogRemoteObject, VeryVerbose, TEXT("         sent subobject : %s"), *SubObjectId.ToString());
+			UE_LOG(LogRemoteObjectD, Verbose, TEXT("         sent subobject : %s"), *SubObjectId.ToString());
 		}
 
 		// Notify that object data has been migrated (actually will be but SendObjectData does a MoveTemp on the data)
 		OnObjectDataSentDelegate.Broadcast(*ObjectData, DestinationServerId);
 
-		SendObjectData(Object, DestinationServerId, ObjectData);
+		GTransferQueue.SendRemoteObject(RequestToken, ObjectId, DestinationServerId, ObjectData);
 
 		checkf(SentObjects.Num() > 0, TEXT("PeerId:%s unable to migrate object %s to PeerId:%s"), *GetGlobalServerId().ToString(), *UE::CoreUObject::Private::FObjectHandleUtils::GetRemoteId(Object).ToString(), *DestinationServerId.ToString());
 
@@ -349,23 +332,23 @@ void RegisterRemoteObjectId(FRemoteObjectId Id, FRemoteServerId OwnerServerId)
 	UE::RemoteObject::Private::RegisterRemoteObjectId(Id, OwnerServerId);
 }
 
-void MigrateObjectToRemoteServer(FRemoteObjectId ObjectId, FRemoteServerId DestinationServerId)
+void MigrateObjectToRemoteServer(uint32 RequestToken, FRemoteObjectId ObjectId, FRemoteServerId DestinationServerId)
 {
 	using namespace UE::RemoteObject::Handle;
 	using namespace UE::RemoteObject::Serialization;
 	using namespace UE::RemoteObject::Private;
 
-	UE_LOG(LogRemoteObject, Verbose, TEXT(" >>> MigrateObjectToRemoteServer obj %s to %s"), *ObjectId.ToString(), *DestinationServerId.ToString());
+	UE_LOG(LogRemoteObject, Verbose, TEXT("MigrateObjectToRemoteServer(%u) obj %s to %s"), RequestToken, *ObjectId.ToString(), *DestinationServerId.ToString());
 
 	UObject* Object = FindObject(ObjectId);
 	if (Object && !IsRemote(Object))
 	{
-		UE_LOG(LogRemoteObject, Verbose, TEXT("   > MigrateObjectToRemoteServer found obj %s"), *ObjectId.ToString());
-		MigrateObjectToRemoteServer(Object, DestinationServerId);
+		UE_LOG(LogRemoteObject, VeryVerbose, TEXT("    : MigrateObjectToRemoteServer(%u) found obj %s"), RequestToken, *ObjectId.ToString());
+		MigrateObjectToRemoteServer(RequestToken, Object, DestinationServerId);
 	}
 	else
 	{
-		UE_LOG(LogRemoteObject, Verbose, TEXT("   > MigrateObjectToRemoteServer did not find obj %s"), *ObjectId.ToString());
+		UE_LOG(LogRemoteObject, VeryVerbose, TEXT("    : MigrateObjectToRemoteServer(%u) did not find obj %s"), RequestToken,*ObjectId.ToString());
 
 		FRemoteObjectStub* RemoteObjectStub = FindRemoteObjectStub(ObjectId);
 		FRemoteServerId CurrentOwnerServerId = RemoteObjectStub ? RemoteObjectStub->RemoteServerId : ObjectId.GetServerId();
@@ -374,31 +357,25 @@ void MigrateObjectToRemoteServer(FRemoteObjectId ObjectId, FRemoteServerId Desti
 
 		if (CurrentOwnerServerId != DestinationServerId)
 		{
-			UE_LOG(LogRemoteObject, Verbose, TEXT("     : MigrateObjectToRemoteServer --FORWARDING REQUEST--"));
-
-			// Forward the request to CurrentOwnerServerId
-			RequestObjectData(ObjectId, CurrentOwnerServerId, DestinationServerId);
+			UE_LOG(LogRemoteObject, Verbose, TEXT("    : MigrateObjectToRemoteServer(%u) forwarding request (we believe %s owns the Object)"), RequestToken, *CurrentOwnerServerId.ToString());
+			GTransferQueue.RequestRemoteObject(RequestToken, ObjectId, CurrentOwnerServerId, DestinationServerId);
 		}
 		else
 		{
-			UE_LOG(LogRemoteObject, Verbose, TEXT("     : MigrateObjectToRemoteServer --SENDING DENY-- %s"), *ObjectId.ToString());
-			RemoteObjectDeniedTransferDelegate.Execute(ObjectId, DestinationServerId);
+			UE_LOG(LogRemoteObject, Verbose, TEXT("    : MigrateObjectToRemoteServer(%u) denying request %s (we believe %s already has the Object)"), RequestToken, *ObjectId.ToString(), *CurrentOwnerServerId.ToString());
+			RemoteObjectDeniedTransferDelegate.Execute(RequestToken, ObjectId, DestinationServerId);
 		}
 	}
 }
 
-void OnObjectDataReceived(FRemoteObjectId ObjectId, FRemoteServerId RemoteServerId, FRemoteObjectData* Data)
+void OnObjectDataReceived(uint32 RequestToken, FRemoteObjectId ObjectId, FRemoteServerId RemoteServerId, FRemoteObjectData* Data)
 {
-	UE_LOG(LogRemoteObject, Verbose, TEXT("     : OnObjectDataReceived %s from %s"), *ObjectId.ToString(), *RemoteServerId.ToString());
-
-	GTransferQueue.FulfillReceiveRequest(ObjectId, RemoteServerId, Data);
+	GTransferQueue.FulfillReceiveRequest(RequestToken, ObjectId, RemoteServerId, Data);
 }
 
-void OnObjectDataDenied(FRemoteObjectId ObjectId, FRemoteServerId RemoteServerId)
+void OnObjectDataDenied(uint32 RequestToken, FRemoteObjectId ObjectId, FRemoteServerId RemoteServerId)
 {
-	UE_LOG(LogRemoteObject, Verbose, TEXT("     : OnObjectDataDenied %s from %s"), *ObjectId.ToString(), *RemoteServerId.ToString());
-
-	GTransferQueue.DenyReceiveRequest(ObjectId);
+	GTransferQueue.DenyReceiveRequest(RequestToken, ObjectId);
 }
 
 void PrefetchRemoteObject(FRemoteObjectId ObjectId, UObject* DestinationOuter)
@@ -413,10 +390,7 @@ void PrefetchRemoteObject(FRemoteObjectId ObjectId, UObject* DestinationOuter)
 		if (bAdded)
 		{
 			GTransferQueue.RequestedObjects.Add(ObjectId);
-			GTransferQueue.RequestedObjectsDestinationOuter.Add(DestinationOuter);
-
-			UE_LOG(LogRemoteObject, Verbose, TEXT("Prefetching obj id %s"),
-				*ObjectId.ToString());
+			UE_LOG(LogRemoteObject, Verbose, TEXT("PrefetchRemoteObject(%u) : obj %s"), GTransferQueue.ActiveRequestToken.Load(), *ObjectId.ToString());
 		}
 	};
 
@@ -425,7 +399,6 @@ void PrefetchRemoteObject(FRemoteObjectId ObjectId, UObject* DestinationOuter)
 		// if we are not in a transaction, request the object immediately
 		UE::RemoteObject::Handle::RequestObject(ObjectId);
 		GTransferQueue.RequestedObjects.Empty();
-		GTransferQueue.RequestedObjectsDestinationOuter.Empty();
 		GTransferQueue.LockedObjects.Empty();
 		GTransferQueue.PendingRequests = 0;
 	}
@@ -433,12 +406,9 @@ void PrefetchRemoteObject(FRemoteObjectId ObjectId, UObject* DestinationOuter)
 
 void ExecuteTransactional(const TFunction<void(void)>& Work)
 {
-	UE_LOG(LogRemoteObject, Verbose, TEXT("-- ExecuteTransactional --"));
+	GTransferQueue.BeginRequest();
 
-	GTransferQueue.LockedObjects.Empty();
-	GTransferQueue.RequestedObjects.Empty();
-	GTransferQueue.RequestedObjectsDestinationOuter.Empty();
-	GTransferQueue.PendingRequests = 0;
+	UE_LOG(LogRemoteObject, VeryVerbose, TEXT("-- ExecuteTransactional : request %u --"), GTransferQueue.ActiveRequestToken.Load());
 
 	// we might run through this loop after the work succeeds in the case
 	// that the work prefetches objects that it doesn't access
@@ -446,14 +416,16 @@ void ExecuteTransactional(const TFunction<void(void)>& Work)
 
 	for (;;)
 	{
+		UE_LOG(LogRemoteObjectD, Verbose, TEXT("   : %d locked objs :"), GTransferQueue.LockedObjects.Num());
 		for (FRemoteObjectId ObjectId : GTransferQueue.LockedObjects)
 		{
-			UE_LOG(LogRemoteObject, VeryVerbose, TEXT("   : locked obj : %s"), *ObjectId.ToString());
+			UE_LOG(LogRemoteObjectD, Verbose, TEXT("       %s"), *ObjectId.ToString());
 		}
 
+		UE_LOG(LogRemoteObjectD, Verbose, TEXT("   : %d requested objs :"), GTransferQueue.RequestedObjects.Num());
 		for (FRemoteObjectId ObjectId : GTransferQueue.RequestedObjects)
 		{
-			UE_LOG(LogRemoteObject, VeryVerbose, TEXT("   : requested obj : %s"), *ObjectId.ToString());
+			UE_LOG(LogRemoteObjectD, Verbose, TEXT("       %s"), *ObjectId.ToString());
 		}
 
 		TArray<FRemoteObjectId> ObjectsToRequest;
@@ -467,11 +439,10 @@ void ExecuteTransactional(const TFunction<void(void)>& Work)
 			}
 		}
 
-		UE_LOG(LogRemoteObject, VeryVerbose, TEXT("   : ObjectsToRequest: %d entries"), ObjectsToRequest.Num());
+		UE_LOG(LogRemoteObjectD, Verbose, TEXT("   : ObjectsToRequest: %d entries"), ObjectsToRequest.Num());
 
 		// send the requests... slowly for now (to avoid straining the network subsystem)
 		int32 RequestsSent = 0;
-		double HangProtectorBeginTime = FPlatformTime::Seconds();
 		double LastWaitingPrintTime = FPlatformTime::Seconds();
 		for(;;)
 		{
@@ -485,12 +456,12 @@ void ExecuteTransactional(const TFunction<void(void)>& Work)
 				{
 					// true means we already have the object locally and didn't actually send a request for it
 					GTransferQueue.LockedObjects.Add(RequestedObjectId);
-					UE_LOG(LogRemoteObject, VeryVerbose, TEXT("      : %s is local"), *RequestedObjectId.ToString());
+					UE_LOG(LogRemoteObjectD, Verbose, TEXT("   : marking %s as locked, it is already local"), *RequestedObjectId.ToString());
 				}
 				else
 				{
 					GTransferQueue.PendingRequests++;
-					UE_LOG(LogRemoteObject, VeryVerbose, TEXT("      : %s is remote"), *RequestedObjectId.ToString());
+					UE_LOG(LogRemoteObjectD, Verbose, TEXT("   : sent request for %s (%d pending requests)"), *RequestedObjectId.ToString(), GTransferQueue.PendingRequests);
 				}
 			}
 
@@ -514,18 +485,12 @@ void ExecuteTransactional(const TFunction<void(void)>& Work)
 				UE_LOG(LogRemoteObject, VeryVerbose, TEXT("   : waiting for %d responses..."), GTransferQueue.PendingRequests);
 				LastWaitingPrintTime = FPlatformTime::Seconds();
 			}
-
-			// after 10 seconds, start sleeping in case we're deadlocked
-			if (FPlatformTime::Seconds() > HangProtectorBeginTime + 10.0)
-			{
-				FPlatformProcess::Sleep(0.5f);
-			}
 		}
 
 		// now we also might have migrated objects away, so check to see if LockedObjects == RequestedObjects
 		if(GTransferQueue.LockedObjects.Num() != GTransferQueue.RequestedObjects.Num())
 		{
-			UE_LOG(LogRemoteObject, VeryVerbose, TEXT("   : did we give up locked objects? trying again..."));
+			UE_LOG(LogRemoteObject, VeryVerbose, TEXT("   : locked count %d mismatches requests %d, did we give up locked objects? trying again..."), GTransferQueue.LockedObjects.Num(), GTransferQueue.RequestedObjects.Num());
 
 			// try again!
 			continue;
@@ -533,7 +498,7 @@ void ExecuteTransactional(const TFunction<void(void)>& Work)
 
 		FlushRemotePhysicsQueryDelegate.ExecuteIfBound();
 
-		UE_LOG(LogRemoteObject, VeryVerbose, TEXT("   : READY"));
+		UE_LOG(LogRemoteObject, VeryVerbose, TEXT("   : READY: all %d requested objects are locked"), GTransferQueue.RequestedObjects.Num());
 
 		// all of the requested objects are locked, try to perform the work
 		if (!bWorkSucceeded)
@@ -543,12 +508,12 @@ void ExecuteTransactional(const TFunction<void(void)>& Work)
 			if (WorkTransactResult == AutoRTFM::ETransactionResult::AbortedByRequest)
 			{
 					// the work aborted, it should have inserted a new requested object, retry
-				UE_LOG(LogRemoteObject, VeryVerbose, TEXT("   : work aborted"));
+				UE_LOG(LogRemoteObject, VeryVerbose, TEXT("   : work ABORTED"));
 				OnTransactionAbortedDelegate.Broadcast();
 			}
 			else
 			{
-				UE_LOG(LogRemoteObject, VeryVerbose, TEXT("   : work OK!"));
+				UE_LOG(LogRemoteObject, VeryVerbose, TEXT("   : work COMPLETED!"));
 
 				bWorkSucceeded = true;
 				OnTransactionCompletedDelegate.Broadcast();
@@ -558,10 +523,6 @@ void ExecuteTransactional(const TFunction<void(void)>& Work)
 				if (GTransferQueue.LockedObjects.Num() == GTransferQueue.RequestedObjects.Num())
 				{
 					// the work succeeded and no more object requests were made, we're all done!
-					GTransferQueue.RequestedObjects.Empty();
-					GTransferQueue.RequestedObjectsDestinationOuter.Empty();
-					GTransferQueue.LockedObjects.Empty();
-					GTransferQueue.PendingRequests = 0;
 					break;
 				}
 				else
@@ -574,12 +535,13 @@ void ExecuteTransactional(const TFunction<void(void)>& Work)
 		}
 	}
 
-	UE_LOG(LogRemoteObject, Verbose, TEXT("-- ExecuteTransactional DONE --"));
+	UE_LOG(LogRemoteObject, VeryVerbose, TEXT("-- ExecuteTransactional %u DONE --"), GTransferQueue.ActiveRequestToken.Load());
+	GTransferQueue.EndRequest();
 }
 
 void WaitForServersToSync()
 {
-	UE_LOG(LogRemoteObject, Verbose, TEXT("-- WaitForServersToSync --"));
+	UE_LOG(LogRemoteObject, VeryVerbose, TEXT("-- WaitForServersToSync --"));
 	WaitForServersToSyncDelegate.ExecuteIfBound();
 }
 
@@ -624,7 +586,7 @@ void StoreObjectToDatabase(UObject* Object)
 		TSet<UObject*> SentObjects;
 		{
 			TUniquePtr<FRemoteObjectData> ObjectData(SerializeObjectData(OutermostUnreachableRemoteReference, SentObjects));
-			StoreRemoteObjectDataDelegate.Execute(FRemoteObjectId(OutermostUnreachableRemoteReference), DatabaseId, *ObjectData);
+			StoreRemoteObjectDataDelegate.Execute(PushMigrationRequestToken, FRemoteObjectId(OutermostUnreachableRemoteReference), DatabaseId, *ObjectData);
 		}
 
 		checkf(SentObjects.Num() > 0, TEXT("PeerId:%s unable to store object %s to DatabaseId:%s"), *GetGlobalServerId().ToString(), *FRemoteObjectId(Object).ToString(), *DatabaseId.ToString());
diff --git a/Engine/Source/Runtime/CoreUObject/Public/UObject/RemoteObjectTransfer.h b/Engine/Source/Runtime/CoreUObject/Public/UObject/RemoteObjectTransfer.h
--- a/Engine/Source/Runtime/CoreUObject/Public/UObject/RemoteObjectTransfer.h
+++ b/Engine/Source/Runtime/CoreUObject/Public/UObject/RemoteObjectTransfer.h
@@ -10,6 +10,8 @@
 
 namespace UE::RemoteObject::Transfer
 {
+	constexpr uint32 PushMigrationRequestToken = 0xFFFFFFFFu;
+
 	/**
 	* Execute the Work under an transaction that can catch remote object requests
 	* as well as negotiating remote object transfers, retrying the transaction until it succeeds.
@@ -28,14 +30,14 @@ namespace UE::RemoteObject::Transfer
 	* @param RemoteServerId Server that sent object data
 	* @param Data Remote object data. Data ownership is transferred to the remote object queue
 	*/
-	COREUOBJECT_API void OnObjectDataReceived(FRemoteObjectId ObjectId, FRemoteServerId RemoteServerId, FRemoteObjectData* Data);
+	COREUOBJECT_API void OnObjectDataReceived(uint32 RequestToken, FRemoteObjectId ObjectId, FRemoteServerId RemoteServerId, FRemoteObjectData* Data);
 
 	/**
 	* Called when a remote object request was denied by a remote server
 	* @param ObjectId Remote object id
 	* @param RemoteServerId Server that owns the object data
 	*/
-	COREUOBJECT_API void OnObjectDataDenied(FRemoteObjectId ObjectId, FRemoteServerId RemoteServerId);
+	COREUOBJECT_API void OnObjectDataDenied(uint32 RequestToken, FRemoteObjectId ObjectId, FRemoteServerId RemoteServerId);
 
 	/**
 	* Requests remote object data to be transferred from Current Owner Server to Destination Server
@@ -50,14 +52,14 @@ namespace UE::RemoteObject::Transfer
 	* @param Object Local object to migrate
 	* @param DestinationServerId Id of a remote server to receive the object
 	*/
-	COREUOBJECT_API void MigrateObjectToRemoteServer(UObject* Object, FRemoteServerId DestinationServerId);
+	COREUOBJECT_API void MigrateObjectToRemoteServer(uint32 RequestToken, UObject* Object, FRemoteServerId DestinationServerId);
 
 	/**
 	* Migrates an object to a remote server
 	* @param ObjectId Remote object id
 	* @param DestinationServerId Id of a remote server to receive the object
 	*/
-	COREUOBJECT_API void MigrateObjectToRemoteServer(FRemoteObjectId Id, FRemoteServerId DestinationServerId);
+	COREUOBJECT_API void MigrateObjectToRemoteServer(uint32 RequestToken, FRemoteObjectId Id, FRemoteServerId DestinationServerId);
 	
 	/**
 	* Migrates an object from a remote server (temp function)
@@ -102,13 +104,13 @@ namespace UE::RemoteObject::Transfer
 	extern COREUOBJECT_API const FRemoteServerId DatabaseId;
 
 	/** Delegate that transfers object data to another server */
-	extern COREUOBJECT_API TDelegate<void(FRemoteObjectId /*ObjectId*/, FRemoteServerId /*DestinationServerId*/, FRemoteObjectData& /*ObjectData*/)> RemoteObjectTransferDelegate;
+	extern COREUOBJECT_API TDelegate<void(uint32 /*RequestToken*/, FRemoteObjectId /*ObjectId*/, FRemoteServerId /*DestinationServerId*/, FRemoteObjectData& /*ObjectData*/)> RemoteObjectTransferDelegate;
 
 	/** Delegate that handles an object request being denied */
-	extern COREUOBJECT_API TDelegate<void(FRemoteObjectId /*ObjectId*/, FRemoteServerId /*DestinationServerId*/)> RemoteObjectDeniedTransferDelegate;
+	extern COREUOBJECT_API TDelegate<void(uint32 /*RequestToken*/, FRemoteObjectId /*ObjectId*/, FRemoteServerId /*DestinationServerId*/)> RemoteObjectDeniedTransferDelegate;
 
 	/** Delegate that requests remote object data from CurrentOwnerServerId to be transferred to DestinationServerId */
-	extern COREUOBJECT_API TDelegate<void(FRemoteObjectId /*ObjectId*/, FRemoteServerId /*CurrentOwnerServerId*/, FRemoteServerId /*DestinationServerId*/)> RequestRemoteObjectDelegate;
+	extern COREUOBJECT_API TDelegate<void(uint32 /*RequestToken*/, FRemoteObjectId /*ObjectId*/, FRemoteServerId /*CurrentOwnerServerId*/, FRemoteServerId /*DestinationServerId*/)> RequestRemoteObjectDelegate;
 
 	/** Delegate that flushes remote object requests */
 	extern COREUOBJECT_API TDelegate<void()> FlushRemoteObjectRequestsDelegate;
@@ -135,10 +137,10 @@ namespace UE::RemoteObject::Transfer
 	extern COREUOBJECT_API TMulticastDelegate<void(const TSet<UObject*>& /*MigratedObjects*/, FRemoteServerId /*RemoteServerId*/)> OnObjectsSentDelegate;
 
 	/** Delegate that stores locally unrachable object data into a database. */
-	extern COREUOBJECT_API TDelegate<void(FRemoteObjectId /*ObjectId*/, FRemoteServerId /*DatabaseId*/, FRemoteObjectData& /*ObjectData*/)> StoreRemoteObjectDataDelegate;
+	extern COREUOBJECT_API TDelegate<void(uint32 /*RequestToken*/, FRemoteObjectId /*ObjectId*/, FRemoteServerId /*DatabaseId*/, FRemoteObjectData& /*ObjectData*/)> StoreRemoteObjectDataDelegate;
 
 	/** Delegate that restores object data from a database. */
-	extern COREUOBJECT_API TDelegate<void(FRemoteObjectId /*ObjectId*/, FRemoteServerId /*DatabaseId*/, FRemoteServerId /*DestinationServerId*/)> RestoreRemoteObjectDataDelegate;
+	extern COREUOBJECT_API TDelegate<void(uint32 /*RequestToken*/, FRemoteObjectId /*ObjectId*/, FRemoteServerId /*DatabaseId*/, FRemoteServerId /*DestinationServerId*/)> RestoreRemoteObjectDataDelegate;
 
 	/** Delegate that flushes pending physics query requests. */
 	extern COREUOBJECT_API TDelegate<void()> FlushRemotePhysicsQueryDelegate;
