dstm:
- fix bug where UnmapObject was removing the entry for a RemoteId after it was mapped to a new object. Now we check to make sure the value (pointer) matches before removing the mapping
- add checks to RemoteObjectTransfer functions to ensure we don't call RPCs or serialization code inside of a transaction
- some cleanup/fixes to make EngineTest tests work again

diff --git a/Engine/Source/Runtime/CoreUObject/Private/UObject/RemoteObject.cpp b/Engine/Source/Runtime/CoreUObject/Private/UObject/RemoteObject.cpp
--- a/Engine/Source/Runtime/CoreUObject/Private/UObject/RemoteObject.cpp
+++ b/Engine/Source/Runtime/CoreUObject/Private/UObject/RemoteObject.cpp
@@ -54,7 +54,15 @@ public:
 		FRemoteObjectId Id = UE::CoreUObject::Private::FObjectHandleUtils::GetRemoteId(Object);
 		{
 			UE::TScopeLock MapLock(ObjectMapCritical);
-			IdToObjectMap.Remove(Id);
+
+			// only erase the entry from the map if the pointer matches the object
+			// we're trying to unmap, since the remote id may since have been mapped
+			// to a new object, and we don't want to remove that still-valid mapping
+			auto Iter = IdToObjectMap.CreateKeyIterator(Id);
+			if (Iter && (Iter.Value() == Object))
+			{
+				Iter.RemoveCurrent();
+			}
 		}
 	}
 
diff --git a/Engine/Source/Runtime/CoreUObject/Private/UObject/RemoteObjectTransfer.cpp b/Engine/Source/Runtime/CoreUObject/Private/UObject/RemoteObjectTransfer.cpp
--- a/Engine/Source/Runtime/CoreUObject/Private/UObject/RemoteObjectTransfer.cpp
+++ b/Engine/Source/Runtime/CoreUObject/Private/UObject/RemoteObjectTransfer.cpp
@@ -57,6 +57,8 @@ public:
 
 	void BeginRequest()
 	{
+		check(!AutoRTFM::IsClosed());
+
 		// we expect to not have an active request token
 		check(ActiveRequestToken.Load() == 0);
 
@@ -76,6 +78,8 @@ public:
 
 	void EndRequest()
 	{
+		check(!AutoRTFM::IsClosed());
+
 		// we expect to have an active request token
 		uint32 OldToken = ActiveRequestToken.Exchange(0u);
 		check(OldToken != 0);
@@ -90,6 +94,8 @@ public:
 
 	void RequestRemoteObject(uint32 RequestToken, FRemoteObjectId ObjectId, FRemoteServerId CurrentOwnerServerId, FRemoteServerId DestinationServerId)
 	{
+		check(!AutoRTFM::IsClosed());
+
 		// Request an object from remote server
 		if (CurrentOwnerServerId != DatabaseId)
 		{
@@ -109,12 +115,14 @@ public:
 	* @return Future that's going to receive the object data
 	*/
 	void SendRemoteObject(uint32 RequestToken, FRemoteObjectId ObjectId, FRemoteServerId DestinationServerId, FRemoteObjectData* ObjectData)
-	{		
+	{
+		check(!AutoRTFM::IsClosed());
 		RemoteObjectTransferDelegate.Execute(RequestToken, ObjectId, DestinationServerId, *ObjectData);
 	}
 
 	void SendDeniedRemoteObject(uint32 RequestToken, FRemoteObjectId ObjectId, FRemoteServerId DestinationServerId)
 	{
+		check(!AutoRTFM::IsClosed());
 		RemoteObjectDeniedTransferDelegate.Execute(RequestToken, ObjectId, DestinationServerId);
 	}
 
@@ -126,6 +134,8 @@ public:
 	*/
 	void FulfillReceiveRequest(uint32 RequestToken, FRemoteObjectId ObjectId, FRemoteServerId RemoteServerId, FRemoteObjectData* ObjectData)
 	{
+		check(!AutoRTFM::IsClosed());
+
 		bool bMatchesActiveRequest = (RequestToken == ActiveRequestToken.Load());
 		bool bFoundRequest = (RequestedObjects.Find(ObjectId) != INDEX_NONE);
 		bool bAlreadyLocked = (LockedObjects.Find(ObjectId) != INDEX_NONE);
@@ -157,6 +167,8 @@ public:
 
 	void DenyReceiveRequest(uint32 RequestToken, FRemoteObjectId ObjectId)
 	{
+		check(!AutoRTFM::IsClosed());
+
 		bool bMatchesActiveRequest = (RequestToken == ActiveRequestToken.Load());
 		bool bFoundRequest = (RequestedObjects.Find(ObjectId) != INDEX_NONE);
 
@@ -173,6 +185,8 @@ FRemoteObjectTransferQueue GTransferQueue;
 
 void RequestObjectData(FRemoteObjectId ObjectId, FRemoteServerId CurrentOwnerServerId, FRemoteServerId DestinationServerId)
 {
+	check(!AutoRTFM::IsClosed());
+
 	check(GTransferQueue.ActiveRequestToken.Load() != 0);
 	GTransferQueue.RequestRemoteObject(GTransferQueue.ActiveRequestToken.Load(), ObjectId, CurrentOwnerServerId, DestinationServerId);
 }
@@ -180,6 +194,7 @@ void RequestObjectData(FRemoteObjectId ObjectId, FRemoteServerId CurrentOwnerSer
 UObject* MigrateObjectFromObjectData(FRemoteObjectData& ObjectData, FRemoteServerId RemoteServerId)
 {
 	using namespace UE::RemoteObject::Serialization;
+	check(!AutoRTFM::IsClosed());
 
 	// Notify that object data has been migrated from a remote server
 	OnObjectDataReceivedDelegate.Broadcast(ObjectData, RemoteServerId);
@@ -217,6 +232,13 @@ void MigrateObjectFromRemoteServer(FRemoteObjectId ObjectId, FRemoteServerId Cur
 	// ensure the request is added in the Open so after we abort it is preserved
 	bool bAdded = false;
 
+	if (!AutoRTFM::IsClosed())
+	{
+		// if we are not in a transaction (dev/test workflow), begin a
+		// temp request
+		GTransferQueue.BeginRequest();
+	}
+
 	UE_AUTORTFM_OPEN
 	{
 		// it's possible a prefetch already added this to the RequestedObjects list
@@ -239,11 +261,11 @@ void MigrateObjectFromRemoteServer(FRemoteObjectId ObjectId, FRemoteServerId Cur
 	}
 	else
 	{
-		// if we are not in a transaction, request the object immediately
+		// if we are not in a transaction, request the object 
+		// (which should immediately get it) and end the request
+		GTransferQueue.PendingRequests++;
 		UE::RemoteObject::Handle::RequestObject(ObjectId);
-		GTransferQueue.RequestedObjects.Empty();
-		GTransferQueue.LockedObjects.Empty();
-		GTransferQueue.PendingRequests = 0;
+		GTransferQueue.EndRequest();
 	}
 }
 
@@ -252,6 +274,8 @@ void MigrateObjectToRemoteServer(uint32 RequestToken, UObject* Object, FRemoteSe
 	using namespace UE::RemoteObject::Serialization;
 	using namespace UE::RemoteObject::Private;
 
+	check(!AutoRTFM::IsClosed());
+
 	FRemoteObjectId ObjectId = UE::CoreUObject::Private::FObjectHandleUtils::GetRemoteId(Object);
 
 	UE_LOG(LogRemoteObject, Verbose, TEXT("     : MigrateObjectToRemoteServer(%u) sending obj %s [%s] (isremote %d) (%d locked objs)"), RequestToken, *ObjectId.ToString(), *Object->GetName(), Handle::IsRemote(Object), GTransferQueue.LockedObjects.Num());
@@ -348,9 +372,17 @@ void MigrateObjectToRemoteServer(uint32 RequestToken, FRemoteObjectId ObjectId,
 	}
 	else
 	{
-		UE_LOG(LogRemoteObject, VeryVerbose, TEXT("    : MigrateObjectToRemoteServer(%u) did not find obj %s"), RequestToken,*ObjectId.ToString());
-
 		FRemoteObjectStub* RemoteObjectStub = FindRemoteObjectStub(ObjectId);
+
+		if (RemoteObjectStub)
+		{
+			UE_LOG(LogRemoteObject, Verbose, TEXT("    : MigrateObjectToRemoteServer(%u) did not find obj %s (we do have a stub)"), RequestToken,*ObjectId.ToString());
+		}
+		else
+		{
+			UE_LOG(LogRemoteObject, Verbose, TEXT("    : MigrateObjectToRemoteServer(%u) did not find obj %s (also no stub)"), RequestToken,*ObjectId.ToString());
+		}
+
 		FRemoteServerId CurrentOwnerServerId = RemoteObjectStub ? RemoteObjectStub->RemoteServerId : ObjectId.GetServerId();
 		checkf(CurrentOwnerServerId != GetGlobalServerId(), TEXT("PeerId:%s attempting to forward request for object %s to PeerId:%s but object location is unknown"), 
 			*GetGlobalServerId().ToString(), *ObjectId.ToString(), *DestinationServerId.ToString());
@@ -370,16 +402,25 @@ void MigrateObjectToRemoteServer(uint32 RequestToken, FRemoteObjectId ObjectId,
 
 void OnObjectDataReceived(uint32 RequestToken, FRemoteObjectId ObjectId, FRemoteServerId RemoteServerId, FRemoteObjectData* Data)
 {
+	check(!AutoRTFM::IsClosed());
 	GTransferQueue.FulfillReceiveRequest(RequestToken, ObjectId, RemoteServerId, Data);
 }
 
 void OnObjectDataDenied(uint32 RequestToken, FRemoteObjectId ObjectId, FRemoteServerId RemoteServerId)
 {
+	check(!AutoRTFM::IsClosed());
 	GTransferQueue.DenyReceiveRequest(RequestToken, ObjectId);
 }
 
 void PrefetchRemoteObject(FRemoteObjectId ObjectId, UObject* DestinationOuter)
 {
+	if (!AutoRTFM::IsClosed())
+	{
+		// if we are outside of a transaction (dev/test workflow)
+		// then prefetch is a no-op
+		return;
+	}
+
 	bool bAdded = false;
 	
 	UE_AUTORTFM_OPEN
@@ -393,19 +434,12 @@ void PrefetchRemoteObject(FRemoteObjectId ObjectId, UObject* DestinationOuter)
 			UE_LOG(LogRemoteObject, Verbose, TEXT("PrefetchRemoteObject(%u) : obj %s"), GTransferQueue.ActiveRequestToken.Load(), *ObjectId.ToString());
 		}
 	};
-
-	if (bAdded && !AutoRTFM::IsClosed())
-	{
-		// if we are not in a transaction, request the object immediately
-		UE::RemoteObject::Handle::RequestObject(ObjectId);
-		GTransferQueue.RequestedObjects.Empty();
-		GTransferQueue.LockedObjects.Empty();
-		GTransferQueue.PendingRequests = 0;
-	}
 }
 
 void ExecuteTransactional(const TFunction<void(void)>& Work)
 {
+	check(!AutoRTFM::IsClosed());
+
 	GTransferQueue.BeginRequest();
 
 	UE_LOG(LogRemoteObject, VeryVerbose, TEXT("-- ExecuteTransactional : request %u --"), GTransferQueue.ActiveRequestToken.Load());
