Undo an old attempt to make `FReadScopeLock` work under transactions.

This was done before we had come up with the transactionally-safe lock approach.

diff --git a/Engine/Source/Programs/AutoRTFMTests/Private/RWLock.cpp b/Engine/Source/Programs/AutoRTFMTests/Private/RWLock.cpp
deleted file mode 100644
diff --git a/Engine/Source/Runtime/Core/Private/Internationalization/TextKey.cpp b/Engine/Source/Runtime/Core/Private/Internationalization/TextKey.cpp
--- a/Engine/Source/Runtime/Core/Private/Internationalization/TextKey.cpp
+++ b/Engine/Source/Runtime/Core/Private/Internationalization/TextKey.cpp
@@ -14,6 +14,7 @@
 #include "Misc/Guid.h"
 #include "Misc/LazySingleton.h"
 #include "Misc/ScopeRWLock.h"
+#include "Misc/TransactionallySafeRWLock.h"
 #include "Modules/VisualizerDebuggingState.h"
 #include "Memory/LinearAllocator.h"
 
@@ -86,7 +87,7 @@ public:
 		// Read-only
 		int32 NumElementsOnRead = 0;
 		{
-			FReadScopeLock ScopeLock(DataRW);
+			UE::TReadScopeLock ScopeLock(DataRW);
 
 			if (const FString* FoundString = LegacyTCHARState.Find(InIndex))
 			{
@@ -98,7 +99,7 @@ public:
 
 		// Write
 		{
-			FWriteScopeLock ScopeLock(DataRW);
+			UE::TWriteScopeLock ScopeLock(DataRW);
 
 			if (LegacyTCHARState.Num() > NumElementsOnRead)
 			{
@@ -136,7 +137,7 @@ public:
 	{
 		check(InIndex != INDEX_NONE);
 
-		FReadScopeLock ScopeLock(DataRW);
+		UE::TReadScopeLock ScopeLock(DataRW);
 
 		const FKeyData& KeyData = KeyDataAllocations.Get(InIndex);
 		if (KeyData.IsStringView())
@@ -153,7 +154,7 @@ public:
 	{
 		check(InIndex != INDEX_NONE);
 
-		FReadScopeLock ScopeLock(DataRW);
+		UE::TReadScopeLock ScopeLock(DataRW);
 
 		const FKeyData& KeyData = KeyDataAllocations.Get(InIndex);
 		if (KeyData.IsStringView())
@@ -170,7 +171,7 @@ public:
 	{
 		check(InIndex != INDEX_NONE);
 
-		FReadScopeLock ScopeLock(DataRW);
+		UE::TReadScopeLock ScopeLock(DataRW);
 
 		const FKeyData& KeyData = KeyDataAllocations.Get(InIndex);
 		return KeyData.StrHash;
@@ -308,7 +309,7 @@ private:
 		// Read-only
 		int32 NumElementsOnRead = 0;
 		{
-			FReadScopeLock ScopeLock(DataRW);
+			UE::TReadScopeLock ScopeLock(DataRW);
 
 			if (int32 FoundIndex = KeyDataAllocations.Find(KeyData);
 				FoundIndex != INDEX_NONE)
@@ -321,7 +322,7 @@ private:
 
 		// Write
 		{
-			FWriteScopeLock ScopeLock(DataRW);
+			UE::TWriteScopeLock ScopeLock(DataRW);
 
 			if (KeyDataAllocations.Num() > NumElementsOnRead)
 			{
@@ -474,7 +475,7 @@ private:
 		}
 	};
 
-	mutable FRWLock DataRW;
+	mutable FTransactionallySafeRWLock DataRW;
 	FKeyDataAllocator KeyDataAllocations;
 
 	// Sparse TCHAR state; built on-demand by anything still using the deprecated FTextKey::GetChars function
diff --git a/Engine/Source/Runtime/Core/Public/Misc/ScopeRWLock.h b/Engine/Source/Runtime/Core/Public/Misc/ScopeRWLock.h
--- a/Engine/Source/Runtime/Core/Public/Misc/ScopeRWLock.h
+++ b/Engine/Source/Runtime/Core/Public/Misc/ScopeRWLock.h
@@ -5,8 +5,6 @@
 #include "CoreTypes.h"
 #include "Misc/AssertionMacros.h"
 #include "HAL/CriticalSection.h"
-#include "AutoRTFM/AutoRTFM.h"
-#include "Templates/Function.h"
 
 namespace UE
 {
@@ -97,14 +95,12 @@ public:
 	UE_NODISCARD_CTOR explicit FReadScopeLock(FRWLock& InLock)
 		: Lock(InLock)
 	{
-		UE_AUTORTFM_OPEN{ Lock.ReadLock(); };
-		AutoRTFM::PushOnAbortHandler(this, [this](){ this->Lock.ReadUnlock(); });
+		Lock.ReadLock();
 	}
 
 	~FReadScopeLock()
 	{
-		UE_AUTORTFM_OPEN{ Lock.ReadUnlock(); };
-		AutoRTFM::PopOnAbortHandler(this);
+		Lock.ReadUnlock();
 	}
 
 private:
diff --git a/Engine/Source/Runtime/CoreUObject/Private/Misc/PackageName.cpp b/Engine/Source/Runtime/CoreUObject/Private/Misc/PackageName.cpp
--- a/Engine/Source/Runtime/CoreUObject/Private/Misc/PackageName.cpp
+++ b/Engine/Source/Runtime/CoreUObject/Private/Misc/PackageName.cpp
@@ -31,6 +31,7 @@
 #include "Misc/PathViews.h"
 #include "Misc/CoreUObjectPluginManager.h"
 #include "Misc/ScopeRWLock.h"
+#include "Misc/TransactionallySafeRWLock.h"
 #include "Misc/StringBuilder.h"
 #include "Modules/ModuleManager.h"
 #include "ProfilingDebugging/LoadTimeTracker.h"
@@ -212,7 +213,7 @@ FMountPoint::FMountPoint(FString&& InRootPath, FString&& InContentPathRelative,
 
 struct FLongPackagePathsSingleton
 {
-	mutable FRWLock MountLock;
+	mutable FTransactionallySafeRWLock MountLock;
 
 	FString ConfigRootPath;
 	FString EngineRootPath;
@@ -301,7 +302,7 @@ struct FLongPackagePathsSingleton
 			ConstructMountPoint(RootPath, ContentPath, FMountPoint::EMountFlags::Removable)));
 		FMountPoint* MountPoint = MountPointOwner.Get();
 		{
-			FWriteScopeLock ScopeLock(MountLock);
+			UE::TWriteScopeLock ScopeLock(MountLock);
 			MountPoints.Insert(MoveTemp(MountPointOwner), 0);
 
 			RootPathTree.FindOrAdd(MountPoint->RootPath) = MountPoint;
@@ -337,7 +338,7 @@ struct FLongPackagePathsSingleton
 	{
 		bool bRemoved = false;
 		{
-			FWriteScopeLock ScopeLock(MountLock);
+			UE::TWriteScopeLock ScopeLock(MountLock);
 			for (int32 MountPointIndex = 0; MountPointIndex < MountPoints.Num(); /* Conditionally incremented in loop */)
 			{
 				FMountPoint& ExistingMount = *MountPoints[MountPointIndex];
@@ -420,7 +421,7 @@ struct FLongPackagePathsSingleton
 	// Checks whether the specific root path is a valid mount point.
 	bool MountPointExists(const FString& RootPath)
 	{
-		FReadScopeLock ScopeLock(MountLock);
+		UE::TReadScopeLock ScopeLock(MountLock);
 		return RootPathTree.Contains(RootPath) || RootPath == MemoryRootPath;
 	}
 
@@ -536,7 +537,7 @@ private:
 	{
 		UE_LOG(LogPackageName, Log, TEXT("Valid mount points:"));
 
-		FReadScopeLock ScopeLock(MountLock);
+		UE::TReadScopeLock ScopeLock(MountLock);
 		for (const TUniquePtr<FMountPoint>& MountPoint : MountPoints)
 		{
 			UE_LOG(LogPackageName, Log, TEXT("	'%s' -> '%s'"), *MountPoint->RootPath, *MountPoint->ContentPathRelative);
@@ -603,7 +604,7 @@ private:
 			
 			if (IsRootPath)
 			{
-				FReadScopeLock ScopeLock(MountLock);
+				UE::TReadScopeLock ScopeLock(MountLock);
 				
 				TArray<FString> ContentPathRelatives;
 				for (TUniquePtr<FMountPoint>& ExistingMount : Get().MountPoints)
@@ -633,7 +634,7 @@ private:
 			}
 			else
 			{
-				FReadScopeLock ScopeLock(MountLock);
+				UE::TReadScopeLock ScopeLock(MountLock);
 			
 				ContentPath = Path;
 				TArray<FString> RootPaths;
@@ -697,7 +698,7 @@ void FPackageName::InternalFilenameToLongPackageName(FStringView InFilename, FSt
 	// Convert to relative path if it's not already a long package name
 	bool bIsValidLongPackageName = false;
 	{
-		FReadScopeLock ScopeLock(Paths.MountLock);
+		UE::TReadScopeLock ScopeLock(Paths.MountLock);
 		const FMountPoint* const* MountPointPtr = Paths.RootPathTree.FindClosestValue(Filename);
 		if (MountPointPtr)
 		{
@@ -741,7 +742,7 @@ void FPackageName::InternalFilenameToLongPackageName(FStringView InFilename, FSt
 
 	auto TryGetLocalPathResult = [&Paths, &OutPackageName](FStringView SearchName)
 	{
-		FReadScopeLock ScopeLock(Paths.MountLock);
+		UE::TReadScopeLock ScopeLock(Paths.MountLock);
 		const FMountPoint* const* MountPointPtr = Paths.ContentPathTree.FindClosestValue(SearchName);
 		if (!MountPointPtr)
 		{
@@ -891,7 +892,7 @@ FString FPackageName::FilenameToLongPackageName(const FString& InFilename)
 		TArray<FString> ContentRootsArrayAbsolute;
 		{
 			const auto& Paths = FLongPackagePathsSingleton::Get();
-			FReadScopeLock ScopeLock(Paths.MountLock);
+			UE::TReadScopeLock ScopeLock(Paths.MountLock);
 			for (const TUniquePtr<FMountPoint>& MountPoint : Paths.MountPoints)
 			{
 				ContentRootsArrayRelative.Add(MountPoint->ContentPathRelative);
@@ -934,7 +935,7 @@ bool FPackageName::TryConvertLongPackageNameToFilename(const FString& InLongPack
 bool FPackageName::TryConvertLongPackageNameToFilename(FStringView InLongPackageName, FString& OutFilename, FStringView InExtension)
 {
 	const auto& Paths = FLongPackagePathsSingleton::Get();
-	FReadScopeLock ScopeLock(Paths.MountLock);
+	UE::TReadScopeLock ScopeLock(Paths.MountLock);
 	const FMountPoint* const* MountPointPtr = Paths.RootPathTree.FindClosestValue(InLongPackageName);
 	if (!MountPointPtr)
 	{
@@ -1036,7 +1037,7 @@ bool FPackageName::SplitLongPackageName(const FString& InLongPackageName, FStrin
 
 	FStringView PackageRelPath;
 	{
-		FReadScopeLock ScopeLock(Paths.MountLock);
+		UE::TReadScopeLock ScopeLock(Paths.MountLock);
 
 		// Check to see whether our package came from a valid root
 		OutPackageRoot.Empty();
@@ -1299,7 +1300,7 @@ bool FPackageName::IsValidLongPackageName(FStringView InLongPackageName, bool bI
 			{
 				const FLongPackagePathsSingleton& Paths = FLongPackagePathsSingleton::Get();
 
-				FReadScopeLock ScopeLock(Paths.MountLock);
+				UE::TReadScopeLock ScopeLock(Paths.MountLock);
 				TStringBuilder<2048> ValidRootsString;
 				for (const TUniquePtr<FMountPoint>& MountPoint : Paths.MountPoints)
 				{
@@ -1336,7 +1337,7 @@ bool FPackageName::IsValidLongPackageName(FStringView InLongPackageName, bool bI
 
 	// Check valid roots
 	const FLongPackagePathsSingleton& Paths = FLongPackagePathsSingleton::Get();
-	FReadScopeLock ScopeLock(Paths.MountLock);
+	UE::TReadScopeLock ScopeLock(Paths.MountLock);
 
 	const FMountPoint* const* MountPointPtr = Paths.RootPathTree.FindClosestValue(InLongPackageName);
 	if (MountPointPtr)
@@ -1437,7 +1438,7 @@ bool FPackageName::IsValidObjectPath(FStringView InObjectPath, FText* OutReason)
 bool FPackageName::IsValidPath(FStringView InPath)
 {
 	const FLongPackagePathsSingleton& Paths = FLongPackagePathsSingleton::Get();
-	FReadScopeLock ScopeLock(Paths.MountLock);
+	UE::TReadScopeLock ScopeLock(Paths.MountLock);
 	// Note that we return false for the root folder "/".
 	return Paths.RootPathTree.FindClosestValue(InPath) != nullptr;
 }
@@ -1461,7 +1462,7 @@ FName FPackageName::GetPackageMountPoint(const FString& InPackagePath, bool InWi
 {
 	FLongPackagePathsSingleton& Paths = FLongPackagePathsSingleton::Get();
 	
-	FReadScopeLock ScopeLock(Paths.MountLock);
+	UE::TReadScopeLock ScopeLock(Paths.MountLock);
 	int32 WithoutSlashes = InWithoutSlashes ? 1 : 0;
 	auto TryRootPath = [&InPackagePath, WithoutSlashes](FStringView RootPath)
 	{
@@ -1496,7 +1497,7 @@ FString FPackageName::GetContentPathForPackageRoot(FStringView InMountPoint)
 {
 	FLongPackagePathsSingleton& Paths = FLongPackagePathsSingleton::Get();
 
-	FReadScopeLock ScopeLock(Paths.MountLock);
+	UE::TReadScopeLock ScopeLock(Paths.MountLock);
 	const FMountPoint* const* MountPointPtr = Paths.RootPathTree.FindClosestValue(InMountPoint);
 	if (!MountPointPtr)
 	{
@@ -1662,7 +1663,7 @@ bool FPackageName::TryGetMountPointForPath(FStringView InFilePathOrPackageName,
 	FPathViews::ToAbsolutePath(InFilePathOrPackageName, PossibleAbsFilePath);
 	FStringView RelPath;
 	const FLongPackagePathsSingleton& Paths = FLongPackagePathsSingleton::Get();
-	FReadScopeLock ScopeLock(Paths.MountLock);
+	UE::TReadScopeLock ScopeLock(Paths.MountLock);
 	const FMountPoint* const* MountPointPtr = Paths.RootPathTree.FindClosestValue(InFilePathOrPackageName);
 	if (MountPointPtr)
 	{
@@ -2701,7 +2702,7 @@ void FPackageName::QueryRootContentPaths(TArray<FString>& OutRootContentPaths, b
 	const FLongPackagePathsSingleton& Paths = FLongPackagePathsSingleton::Get();
 
 	{
-		FReadScopeLock ScopeLock(Paths.MountLock);
+		UE::TReadScopeLock ScopeLock(Paths.MountLock);
 		OutRootContentPaths.Reserve(Paths.MountPoints.Num() + (bIncludeReadOnlyRoots ? 1 : 0));
 		// Maintain the legacy order: non-removable non-readonly, then removable, then non-removable readonly
 		for (const TUniquePtr<FMountPoint>& MountPoint : Paths.MountPoints)
@@ -2755,7 +2756,7 @@ TArray<FString> FPackageName::QueryMountPointLocalAbsPaths()
 	TArray<FString> OutAbsPaths;
 
 	{
-		FReadScopeLock ScopeLock(Paths.MountLock);
+		UE::TReadScopeLock ScopeLock(Paths.MountLock);
 		OutAbsPaths.Reserve(Paths.MountPoints.Num() + 1);
 		for (const TUniquePtr<FMountPoint>& MountPoint : Paths.MountPoints)
 		{
