When reconstructing UObjects on top of existing UObjects try to restore the old remote object id the same way it's done for the  serial number

diff --git a/Engine/Source/Runtime/CoreUObject/Private/UObject/UObjectGlobals.cpp b/Engine/Source/Runtime/CoreUObject/Private/UObject/UObjectGlobals.cpp
--- a/Engine/Source/Runtime/CoreUObject/Private/UObject/UObjectGlobals.cpp
+++ b/Engine/Source/Runtime/CoreUObject/Private/UObject/UObjectGlobals.cpp
@@ -3553,6 +3553,9 @@ UObject* StaticAllocateObject
 
 	int32 OldIndex = -1;
 	int32 OldSerialNumber = 0;
+#if UE_WITH_REMOTE_OBJECT_HANDLE
+	FRemoteObjectId OldRemoteId;
+#endif
 
 	if( Obj == nullptr )
 	{	
@@ -3611,6 +3614,9 @@ UObject* StaticAllocateObject
 		{
 			OldIndex = GUObjectArray.ObjectToIndex(Obj);
 			OldSerialNumber = GUObjectArray.GetSerialNumber(OldIndex);
+#if UE_WITH_REMOTE_OBJECT_HANDLE
+			OldRemoteId = GUObjectArray.GetRemoteId(OldIndex);
+#endif
 
 			// Destroy the object.
 			SCOPE_CYCLE_COUNTER(STAT_DestroyObject);
@@ -3691,6 +3697,8 @@ UObject* StaticAllocateObject
 		// We can't have two different, non-zero serial numbers
 		checkf(OldSerialNumber == 0 || SerialNumber == 0 || SerialNumber == OldSerialNumber, TEXT("Attempting to change a new object's (%s %s.%s) serial number (%d -> %d)"),
 			*InClass->GetPathName(), *GetPathNameSafe(InOuter), *InName.ToString(), OldSerialNumber, SerialNumber);
+		checkf(!OldRemoteId.IsValid() || !RemoteId.IsValid() || OldRemoteId == RemoteId, TEXT("Attempting to change a new object's (%s %s.%s) remote id (%s -> %s)"),
+			*InClass->GetPathName(), *GetPathNameSafe(InOuter), *InName.ToString(), *OldRemoteId.ToString(), *RemoteId.ToString());
 #endif // UE_WITH_REMOTE_OBJECT_HANDLE
 
 		// perform the UObjectBase construction in the open - we expect the GC to invoke
@@ -3699,7 +3707,11 @@ UObject* StaticAllocateObject
 		UE_AUTORTFM_OPEN_NO_VALIDATION
 		{
 			FMemory::Memzero((void *)Obj, TotalSize);
-			new ((void *)Obj) UObjectBase(const_cast<UClass*>(InClass), InFlags|RF_NeedInitialization, InternalSetFlags, InOuter, InName, OldIndex, OldSerialNumber == 0 ? SerialNumber : OldSerialNumber, RemoteId);
+			new ((void *)Obj) UObjectBase(const_cast<UClass*>(InClass), InFlags|RF_NeedInitialization, InternalSetFlags, InOuter, InName, OldIndex, OldSerialNumber == 0 ? SerialNumber : OldSerialNumber, 
+#if UE_WITH_REMOTE_OBJECT_HANDLE
+				OldRemoteId.IsValid() ? OldRemoteId :
+#endif
+				RemoteId);
 		};
 	}
 	else
