Added functions for resolving FPropertyPathName to a property value

Added a type FPropertyValueInContainer to represent a reference to a single property value.
Added a function to convert FPropertyPath to an array of FPropertyPathName.
Added a function to find a property by both name and type name.

diff --git a/Engine/Source/Editor/PropertyEditor/Private/PropertyPath.cpp b/Engine/Source/Editor/PropertyEditor/Private/PropertyPath.cpp
new file mode 100644
--- /dev/null
+++ b/Engine/Source/Editor/PropertyEditor/Private/PropertyPath.cpp
@@ -0,0 +1,34 @@
+// Copyright Epic Games, Inc. All Rights Reserved.
+
+#include "PropertyPath.h"
+
+#include "UObject/PropertyPathName.h"
+#include "UObject/PropertyTypeName.h"
+
+TArray<UE::FPropertyPathName> FPropertyPath::ToPropertyPathName() const
+{
+	using namespace UE;
+
+	TArray<FPropertyPathName> PathNames;
+	FPropertyPathName PathName;
+
+	for (const FPropertyInfo& Info : Properties)
+	{
+		FPropertyTypeNameBuilder TypeName;
+		Info.Property->SaveTypeName(TypeName);
+		PathName.Push({Info.Property->GetFName(), TypeName.Build(), Info.ArrayIndex});
+
+		if (Info.Property->IsA<FObjectPropertyBase>())
+		{
+			PathNames.Emplace(MoveTemp(PathName));
+			PathName.Reset();
+		}
+	}
+
+	if (!PathName.IsEmpty())
+	{
+		PathNames.Emplace(MoveTemp(PathName));
+	}
+
+	return PathNames;
+}
diff --git a/Engine/Source/Editor/PropertyEditor/Public/PropertyPath.h b/Engine/Source/Editor/PropertyEditor/Public/PropertyPath.h
--- a/Engine/Source/Editor/PropertyEditor/Public/PropertyPath.h
+++ b/Engine/Source/Editor/PropertyEditor/Public/PropertyPath.h
@@ -5,6 +5,8 @@
 #include "UObject/UnrealType.h"
 #include "UObject/WeakFieldPtr.h"
 
+namespace UE { class FPropertyPathName; }
+
 struct FPropertyInfo
 {
 public:
@@ -202,6 +204,9 @@ public:
 		return Result;
 	}
 
+	/** Converts this to an array of property path names that are split between objects in the path. */
+	PROPERTYEDITOR_API TArray< UE::FPropertyPathName > ToPropertyPathName() const;
+
 private:
 
 	TArray< FPropertyInfo > Properties;
diff --git a/Engine/Source/Runtime/CoreUObject/Internal/UObject/InstanceDataObjectUtils.h b/Engine/Source/Runtime/CoreUObject/Internal/UObject/InstanceDataObjectUtils.h
--- a/Engine/Source/Runtime/CoreUObject/Internal/UObject/InstanceDataObjectUtils.h
+++ b/Engine/Source/Runtime/CoreUObject/Internal/UObject/InstanceDataObjectUtils.h
@@ -46,6 +46,9 @@ UClass* CreateInstanceDataObjectClass(const FPropertyPathNameTree* PropertyTree,
 /** Copy whether each property was set by serialization from one IDO to another. */
 void CopyPropertyValueSerializedData(const FFieldVariant& OldField, void* OldDataPtr, const FFieldVariant& NewField, void* NewDataPtr);
 
+/** Returns the IDO for the object if it has one, otherwise returns the object. */
+UE_API UObject* ResolveInstanceDataObject(UObject* Object);
+
 } // UE
 
 #undef UE_API
diff --git a/Engine/Source/Runtime/CoreUObject/Private/UObject/InstanceDataObjectUtils.cpp b/Engine/Source/Runtime/CoreUObject/Private/UObject/InstanceDataObjectUtils.cpp
--- a/Engine/Source/Runtime/CoreUObject/Private/UObject/InstanceDataObjectUtils.cpp
+++ b/Engine/Source/Runtime/CoreUObject/Private/UObject/InstanceDataObjectUtils.cpp
@@ -295,7 +295,7 @@ namespace UE
 		}
 
 		// Property bag placeholder objects are always enabled for IDO support
-		if (UE::FPropertyBagRepository::IsPropertyBagPlaceholderObject(InObject))
+		if (FPropertyBagRepository::IsPropertyBagPlaceholderObject(InObject))
 		{
 			return true;
 		}
@@ -1012,6 +1012,12 @@ namespace UE
 		}
 	}
 
+	UObject* ResolveInstanceDataObject(UObject* Object)
+	{
+		UObject* InstanceDataObject = FPropertyBagRepository::Get().FindInstanceDataObject(Object);
+		return InstanceDataObject ? InstanceDataObject : Object;
+	}
+
 } // UE
 
 #endif // WITH_EDITORONLY_DATA
diff --git a/Engine/Source/Runtime/CoreUObject/Private/UObject/PropertyMap.cpp b/Engine/Source/Runtime/CoreUObject/Private/UObject/PropertyMap.cpp
--- a/Engine/Source/Runtime/CoreUObject/Private/UObject/PropertyMap.cpp
+++ b/Engine/Source/Runtime/CoreUObject/Private/UObject/PropertyMap.cpp
@@ -5,6 +5,7 @@
 #include "Misc/ScopeExit.h"
 #include "UObject/OverriddenPropertySet.h"
 #include "UObject/PropertyHelper.h"
+#include "UObject/PropertyPathFunctions.h"
 #include "UObject/PropertyPathName.h"
 #include "UObject/PropertyTag.h"
 #include "UObject/UnrealTypePrivate.h"
@@ -144,9 +145,6 @@ namespace UEMapProperty_Private
 	}
 }
 
-static const FName NAME_MapKey(ANSITEXTVIEW("Key"));
-static const FName NAME_MapValue(ANSITEXTVIEW("Value"));
-
 IMPLEMENT_FIELD(FMapProperty)
 
 FMapProperty::FMapProperty(FFieldVariant InOwner, const FName& InName, EObjectFlags InObjectFlags, EMapPropertyFlags InMapFlags)
@@ -321,12 +319,12 @@ void FMapProperty::SerializeItem(FStructuredArchive::FSlot Slot, void* Value, co
 						int32 Index = MapHelper.AddDefaultValue_Invalid_NeedsRehash();
 						uint8* PairPtr = MapHelper.GetPairPtr(Index);
 						{
-							UE::FSerializedPropertyPathScope SerializedPropertyPath(Context, {NAME_MapKey});
+							UE::FSerializedPropertyPathScope SerializedPropertyPath(Context, {UE::NAME_Key});
 							FSerializedPropertyScope SerializedProperty(UnderlyingArchive, KeyProp, this);
 							KeyProp->SerializeItem(EntryRecord.EnterField(TEXT("Key")), PairPtr);
 						}
 						{
-							UE::FSerializedPropertyPathScope SerializedPropertyPath(Context, {NAME_MapValue});
+							UE::FSerializedPropertyPathScope SerializedPropertyPath(Context, {UE::NAME_Value});
 							FSerializedPropertyScope SerializedProperty(UnderlyingArchive, ValueProp, this);
 							ValueProp->SerializeItem(EntryRecord.EnterField(TEXT("Value")), PairPtr + MapLayout.ValueOffset);
 						}
@@ -362,7 +360,7 @@ void FMapProperty::SerializeItem(FStructuredArchive::FSlot Slot, void* Value, co
 						for (int32 i = 0; i < NumRemoved; ++i)
 						{
 							{
-								UE::FSerializedPropertyPathScope SerializedPropertyPath(Context, {NAME_MapKey});
+								UE::FSerializedPropertyPathScope SerializedPropertyPath(Context, {UE::NAME_Key});
 								FSerializedPropertyScope SerializedProperty(UnderlyingArchive, KeyProp, this);
 								KeyProp->SerializeItem(RemovedArray.EnterElement().EnterRecord().EnterField(TEXT("Key")), TempKeyValueStorage);
 							}
@@ -395,7 +393,7 @@ void FMapProperty::SerializeItem(FStructuredArchive::FSlot Slot, void* Value, co
 
 							// Read key into temporary storage
 							{
-								UE::FSerializedPropertyPathScope SerializedPropertyPath(Context, {NAME_MapKey});
+								UE::FSerializedPropertyPathScope SerializedPropertyPath(Context, {UE::NAME_Key});
 								FSerializedPropertyScope SerializedProperty(UnderlyingArchive, KeyProp, this);
 								KeyProp->SerializeItem(EntryRecord.EnterField(TEXT("Key")), TempKeyValueStorage);
 							}
@@ -405,7 +403,7 @@ void FMapProperty::SerializeItem(FStructuredArchive::FSlot Slot, void* Value, co
 
 							// Deserialize value into hash map-owned memory
 							{
-								UE::FSerializedPropertyPathScope SerializedPropertyPath(Context, {NAME_MapValue});
+								UE::FSerializedPropertyPathScope SerializedPropertyPath(Context, {UE::NAME_Value});
 								FSerializedPropertyScope SerializedProperty(UnderlyingArchive, ValueProp, this);
 								ValueProp->SerializeItem(EntryRecord.EnterField(TEXT("Value")), ValuePtr);
 							}
@@ -437,7 +435,7 @@ void FMapProperty::SerializeItem(FStructuredArchive::FSlot Slot, void* Value, co
 
 							// Read key into temporary storage
 							{
-								UE::FSerializedPropertyPathScope SerializedPropertyPath(Context, {NAME_MapKey});
+								UE::FSerializedPropertyPathScope SerializedPropertyPath(Context, {UE::NAME_Key});
 								FSerializedPropertyScope SerializedProperty(UnderlyingArchive, KeyProp, this);
 								KeyProp->SerializeItem(EntryRecord.EnterField(TEXT("Key")), TempKeyValueStorage);
 							}
@@ -446,7 +444,7 @@ void FMapProperty::SerializeItem(FStructuredArchive::FSlot Slot, void* Value, co
 
 							// Deserialize value into hash map-owned memory
 							{
-								UE::FSerializedPropertyPathScope SerializedPropertyPath(Context, {NAME_MapValue});
+								UE::FSerializedPropertyPathScope SerializedPropertyPath(Context, {UE::NAME_Value});
 								FSerializedPropertyScope SerializedProperty(UnderlyingArchive, ValueProp, this);
 								ValueProp->SerializeItem(EntryRecord.EnterField(TEXT("Value")), ValuePtr);
 							}
@@ -552,12 +550,12 @@ void FMapProperty::SerializeItem(FStructuredArchive::FSlot Slot, void* Value, co
 						FStructuredArchive::FRecord EntryRecord = ReplacedArray.EnterElement().EnterRecord();
 						uint8* PairPtr = MapHelper.GetPairPtr(It.GetInternalIndex());
 						{
-							UE::FSerializedPropertyPathScope SerializedPropertyPath(Context, {NAME_MapKey});
+							UE::FSerializedPropertyPathScope SerializedPropertyPath(Context, {UE::NAME_Key});
 							FSerializedPropertyScope SerializedProperty(UnderlyingArchive, KeyProp, this);
 							KeyProp->SerializeItem(EntryRecord.EnterField(TEXT("Key")), PairPtr);
 						}
 						{
-							UE::FSerializedPropertyPathScope SerializedPropertyPath(Context, {NAME_MapValue});
+							UE::FSerializedPropertyPathScope SerializedPropertyPath(Context, {UE::NAME_Value});
 							FSerializedPropertyScope SerializedProperty(UnderlyingArchive, ValueProp, this);
 							ValueProp->SerializeItem(EntryRecord.EnterField(TEXT("Value")), PairPtr + MapLayout.ValueOffset);
 						}
@@ -573,7 +571,7 @@ void FMapProperty::SerializeItem(FStructuredArchive::FSlot Slot, void* Value, co
 					for (int32 i : RemovedIndices)
 					{
 						FStructuredArchive::FRecord EntryRecord = RemovedArray.EnterElement().EnterRecord();
-						UE::FSerializedPropertyPathScope SerializedPropertyPath(Context, {NAME_MapKey});
+						UE::FSerializedPropertyPathScope SerializedPropertyPath(Context, {UE::NAME_Key});
 						FSerializedPropertyScope SerializedProperty(UnderlyingArchive, KeyProp, this);
 						KeyProp->SerializeItem(EntryRecord.EnterField(TEXT("Key")), DefaultsMapHelper.GetKeyPtr(i));
 					}
@@ -585,12 +583,12 @@ void FMapProperty::SerializeItem(FStructuredArchive::FSlot Slot, void* Value, co
 						FStructuredArchive::FRecord EntryRecord = ModifiedArray.EnterElement().EnterRecord();
 						uint8* PairPtr = MapHelper.GetPairPtr(i);
 						{
-							UE::FSerializedPropertyPathScope SerializedPropertyPath(Context, {NAME_MapKey});
+							UE::FSerializedPropertyPathScope SerializedPropertyPath(Context, {UE::NAME_Key});
 							FSerializedPropertyScope SerializedProperty(UnderlyingArchive, KeyProp, this);
 							KeyProp->SerializeItem(EntryRecord.EnterField(TEXT("Key")), PairPtr);
 						}
 						{
-							UE::FSerializedPropertyPathScope SerializedPropertyPath(Context, {NAME_MapValue});
+							UE::FSerializedPropertyPathScope SerializedPropertyPath(Context, {UE::NAME_Value});
 							FSerializedPropertyScope SerializedProperty(UnderlyingArchive, ValueProp, this);
 							ValueProp->SerializeItem(EntryRecord.EnterField(TEXT("Value")), PairPtr + MapLayout.ValueOffset);
 						}
@@ -604,12 +602,12 @@ void FMapProperty::SerializeItem(FStructuredArchive::FSlot Slot, void* Value, co
 						FStructuredArchive::FRecord EntryRecord = AddedArray.EnterElement().EnterRecord();
 						uint8* PairPtr = MapHelper.GetPairPtr(i);
 						{
-							UE::FSerializedPropertyPathScope SerializedPropertyPath(Context, {NAME_MapKey});
+							UE::FSerializedPropertyPathScope SerializedPropertyPath(Context, {UE::NAME_Key});
 							FSerializedPropertyScope SerializedProperty(UnderlyingArchive, KeyProp, this);
 							KeyProp->SerializeItem(EntryRecord.EnterField(TEXT("Key")), PairPtr);
 						}
 						{
-							UE::FSerializedPropertyPathScope SerializedPropertyPath(Context, {NAME_MapValue});
+							UE::FSerializedPropertyPathScope SerializedPropertyPath(Context, {UE::NAME_Value});
 							FSerializedPropertyScope SerializedProperty(UnderlyingArchive, ValueProp, this);
 							ValueProp->SerializeItem(EntryRecord.EnterField(TEXT("Value")), PairPtr + MapLayout.ValueOffset);
 						}
@@ -661,12 +659,12 @@ void FMapProperty::SerializeItem(FStructuredArchive::FSlot Slot, void* Value, co
 				FStructuredArchive::FRecord EntryRecord = EntriesArray.EnterElement().EnterRecord();
 				int32 Index = MapHelper.AddDefaultValue_Invalid_NeedsRehash();
 				{
-					UE::FSerializedPropertyPathScope SerializedPropertyPath(Context, {NAME_MapKey});
+					UE::FSerializedPropertyPathScope SerializedPropertyPath(Context, {UE::NAME_Key});
 					FSerializedPropertyScope SerializedProperty(UnderlyingArchive, KeyProp, this);
 					KeyProp->SerializeItem(EntryRecord.EnterField(TEXT("Key")), MapHelper.GetKeyPtr(Index));
 				}
 				{
-					UE::FSerializedPropertyPathScope SerializedPropertyPath(Context, {NAME_MapValue});
+					UE::FSerializedPropertyPathScope SerializedPropertyPath(Context, {UE::NAME_Value});
 					FSerializedPropertyScope SerializedProperty(UnderlyingArchive, ValueProp, this);
 					ValueProp->SerializeItem(EntryRecord.EnterField(TEXT("Value")), MapHelper.GetValuePtr(Index));
 				}
@@ -691,7 +689,7 @@ void FMapProperty::SerializeItem(FStructuredArchive::FSlot Slot, void* Value, co
 				TempKeyValueStorage = (uint8*)FMemory::Malloc(MapLayout.SetLayout.Size);
 				KeyProp->InitializeValue(TempKeyValueStorage);
 
-				UE::FSerializedPropertyPathScope SerializedPropertyPath(Context, {NAME_MapKey});
+				UE::FSerializedPropertyPathScope SerializedPropertyPath(Context, {UE::NAME_Key});
 				FSerializedPropertyScope SerializedProperty(UnderlyingArchive, KeyProp, this);
 				for (; NumKeysToRemove; --NumKeysToRemove)
 				{
@@ -723,7 +721,7 @@ void FMapProperty::SerializeItem(FStructuredArchive::FSlot Slot, void* Value, co
 
 				// Read key into temporary storage
 				{
-					UE::FSerializedPropertyPathScope SerializedPropertyPath(Context, {NAME_MapKey});
+					UE::FSerializedPropertyPathScope SerializedPropertyPath(Context, {UE::NAME_Key});
 					FSerializedPropertyScope SerializedProperty(UnderlyingArchive, KeyProp, this);
 					KeyProp->SerializeItem(EntryRecord.EnterField(TEXT("Key")), TempKeyValueStorage);
 				}
@@ -732,7 +730,7 @@ void FMapProperty::SerializeItem(FStructuredArchive::FSlot Slot, void* Value, co
 
 				// Deserialize value into hash map-owned memory
 				{
-					UE::FSerializedPropertyPathScope SerializedPropertyPath(Context, {NAME_MapValue});
+					UE::FSerializedPropertyPathScope SerializedPropertyPath(Context, {UE::NAME_Value});
 					FSerializedPropertyScope SerializedProperty(UnderlyingArchive, ValueProp, this);
 					ValueProp->SerializeItem(EntryRecord.EnterField(TEXT("Value")), ValuePtr);
 				}
@@ -797,12 +795,12 @@ void FMapProperty::SerializeItem(FStructuredArchive::FSlot Slot, void* Value, co
 				FStructuredArchive::FRecord EntryRecord = EntriesArray.EnterElement().EnterRecord();
 
 				{
-					UE::FSerializedPropertyPathScope SerializedPropertyPath(Context, {NAME_MapKey});
+					UE::FSerializedPropertyPathScope SerializedPropertyPath(Context, {UE::NAME_Key});
 					FSerializedPropertyScope SerializedProperty(UnderlyingArchive, KeyProp, this);
 					KeyProp->SerializeItem(EntryRecord.EnterField(TEXT("Key")), ValuePairPtr);
 				}
 				{
-					UE::FSerializedPropertyPathScope SerializedPropertyPath(Context, {NAME_MapValue});
+					UE::FSerializedPropertyPathScope SerializedPropertyPath(Context, {UE::NAME_Value});
 					FSerializedPropertyScope SerializedProperty(UnderlyingArchive, ValueProp, this);
 					ValueProp->SerializeItem(EntryRecord.EnterField(TEXT("Value")), ValuePairPtr + MapLayout.ValueOffset);
 				}
@@ -819,12 +817,12 @@ void FMapProperty::SerializeItem(FStructuredArchive::FSlot Slot, void* Value, co
 
 				uint8* ValuePairPtr = MapHelper.GetPairPtr(Iterator);
 				{
-					UE::FSerializedPropertyPathScope SerializedPropertyPath(Context, {NAME_MapKey});
+					UE::FSerializedPropertyPathScope SerializedPropertyPath(Context, {UE::NAME_Key});
 					FSerializedPropertyScope SerializedProperty(UnderlyingArchive, KeyProp, this);
 					KeyProp->SerializeItem(EntryRecord.EnterField(TEXT("Key")), ValuePairPtr);
 				}
 				{
-					UE::FSerializedPropertyPathScope SerializedPropertyPath(Context, {NAME_MapValue});
+					UE::FSerializedPropertyPathScope SerializedPropertyPath(Context, {UE::NAME_Value});
 					FSerializedPropertyScope SerializedProperty(UnderlyingArchive, ValueProp, this);
 					ValueProp->SerializeItem(EntryRecord.EnterField(TEXT("Value")), ValuePairPtr + MapLayout.ValueOffset);
 				}
@@ -1487,7 +1485,7 @@ EConvertFromTypeResult FMapProperty::ConvertFromType(const FPropertyTag& Tag, FS
 		TempKeyValueStorage = (uint8*)FMemory::Malloc(MapLayout.SetLayout.Size);
 		KeyProp->InitializeValue(TempKeyValueStorage);
 
-		if (SerializeOrConvert(bCanSerializeKey, KeyProp, KeyPropertyTag, NAME_MapKey, KeysToRemoveArray.EnterElement(), TempKeyValueStorage, DefaultsStruct))
+		if (SerializeOrConvert(bCanSerializeKey, KeyProp, KeyPropertyTag, UE::NAME_Key, KeysToRemoveArray.EnterElement(), TempKeyValueStorage, DefaultsStruct))
 		{
 			// If the key is in the map, remove it
 			int32 Found = MapHelper.FindMapIndexWithKey(TempKeyValueStorage);
@@ -1499,7 +1497,7 @@ EConvertFromTypeResult FMapProperty::ConvertFromType(const FPropertyTag& Tag, FS
 			// things are going fine, remove the rest of the keys:
 			for (int32 I = 1; I < NumKeysToRemove; ++I)
 			{
-				verify(SerializeOrConvert(bCanSerializeKey, KeyProp, KeyPropertyTag, NAME_MapKey, KeysToRemoveArray.EnterElement(), TempKeyValueStorage, DefaultsStruct));
+				verify(SerializeOrConvert(bCanSerializeKey, KeyProp, KeyPropertyTag, UE::NAME_Key, KeysToRemoveArray.EnterElement(), TempKeyValueStorage, DefaultsStruct));
 				Found = MapHelper.FindMapIndexWithKey(TempKeyValueStorage);
 				if (Found != INDEX_NONE)
 				{
@@ -1528,7 +1526,7 @@ EConvertFromTypeResult FMapProperty::ConvertFromType(const FPropertyTag& Tag, FS
 
 			FStructuredArchive::FRecord FirstPropertyRecord = EntriesArray.EnterElement().EnterRecord();
 
-			if (SerializeOrConvert(bCanSerializeKey, KeyProp, KeyPropertyTag, NAME_MapKey, FirstPropertyRecord.EnterField(TEXT("Key")), TempKeyValueStorage, DefaultsStruct))
+			if (SerializeOrConvert(bCanSerializeKey, KeyProp, KeyPropertyTag, UE::NAME_Key, FirstPropertyRecord.EnterField(TEXT("Key")), TempKeyValueStorage, DefaultsStruct))
 			{
 				// Add a new default value if the key doesn't currently exist in the map
 				bool bKeyAlreadyPresent = true;
@@ -1544,14 +1542,14 @@ EConvertFromTypeResult FMapProperty::ConvertFromType(const FPropertyTag& Tag, FS
 				KeyProp->CopyCompleteValue_InContainer(NextPairPtr, TempKeyValueStorage);
 
 				// Deserialize value
-				if (SerializeOrConvert(bCanSerializeValue, ValueProp, ValuePropertyTag, NAME_MapValue, FirstPropertyRecord.EnterField(TEXT("Value")), NextPairPtr, DefaultsStruct))
+				if (SerializeOrConvert(bCanSerializeValue, ValueProp, ValuePropertyTag, UE::NAME_Value, FirstPropertyRecord.EnterField(TEXT("Value")), NextPairPtr, DefaultsStruct))
 				{
 					// first entry went fine, convert the rest:
 					for (int32 I = 1; I < Num; ++I)
 					{
 						FStructuredArchive::FRecord PropertyRecord = EntriesArray.EnterElement().EnterRecord();
 
-						verify(SerializeOrConvert(bCanSerializeKey, KeyProp, KeyPropertyTag, NAME_MapKey, PropertyRecord.EnterField(TEXT("Key")), TempKeyValueStorage, DefaultsStruct));
+						verify(SerializeOrConvert(bCanSerializeKey, KeyProp, KeyPropertyTag, UE::NAME_Key, PropertyRecord.EnterField(TEXT("Key")), TempKeyValueStorage, DefaultsStruct));
 						NextPairIndex = MapHelper.FindMapIndexWithKey(TempKeyValueStorage);
 						if (NextPairIndex == INDEX_NONE)
 						{
@@ -1561,7 +1559,7 @@ EConvertFromTypeResult FMapProperty::ConvertFromType(const FPropertyTag& Tag, FS
 						NextPairPtr = MapHelper.GetPairPtrWithoutCheck(NextPairIndex);
 						// This copy is unnecessary when the key was already in the map:
 						KeyProp->CopyCompleteValue_InContainer(NextPairPtr, TempKeyValueStorage);
-						verify(SerializeOrConvert(bCanSerializeValue, ValueProp, ValuePropertyTag, NAME_MapValue, PropertyRecord.EnterField(TEXT("Value")), NextPairPtr, DefaultsStruct));
+						verify(SerializeOrConvert(bCanSerializeValue, ValueProp, ValuePropertyTag, UE::NAME_Value, PropertyRecord.EnterField(TEXT("Value")), NextPairPtr, DefaultsStruct));
 					}
 				}
 				else
diff --git a/Engine/Source/Runtime/CoreUObject/Private/UObject/PropertyPathFunctions.cpp b/Engine/Source/Runtime/CoreUObject/Private/UObject/PropertyPathFunctions.cpp
new file mode 100644
--- /dev/null
+++ b/Engine/Source/Runtime/CoreUObject/Private/UObject/PropertyPathFunctions.cpp
@@ -0,0 +1,164 @@
+// Copyright Epic Games, Inc. All Rights Reserved.
+
+#include "UObject/PropertyPathFunctions.h"
+
+#include "UObject/Class.h"
+#include "UObject/PropertyOptional.h"
+#include "UObject/PropertyPathName.h"
+#include "UObject/UnrealType.h"
+
+namespace UE
+{
+
+const FName NAME_Key(ANSITEXTVIEW("Key"));
+const FName NAME_Value(ANSITEXTVIEW("Value"));
+
+FProperty* FindPropertyByNameAndTypeName(const UStruct* Struct, FName Name, FPropertyTypeName TypeName)
+{
+	for (FProperty* Property = Struct->PropertyLink; Property; Property = Property->PropertyLinkNext)
+	{
+		if (Property->GetFName() == Name && Property->CanSerializeFromTypeName(TypeName))
+		{
+			return Property;
+		}
+	}
+	return nullptr;
+}
+
+inline static const UStruct* FindStructFromProperty(const FProperty* Property)
+{
+	if (const FStructProperty* StructProperty = CastField<FStructProperty>(Property))
+	{
+		return StructProperty->Struct;
+	}
+	return nullptr;
+}
+
+FPropertyValueInContainer TryResolvePropertyPath(const FPropertyPathName& Path, UObject* Object)
+{
+	const UStruct* NextStruct = Object ? Object->GetClass() : nullptr;
+	void* NextContainer = Object;
+
+	FPropertyValueInContainer Value;
+	for (int32 Index = 0, Count = Path.GetSegmentCount(); Index < Count; ++Index)
+	{
+		// Fail if the previous segment failed to resolve the struct or container for this segment.
+		if (!NextStruct || !NextContainer)
+		{
+			return {};
+		}
+
+		const FPropertyPathNameSegment Segment = Path.GetSegment(Index);
+		const FProperty* Property = FindPropertyByNameAndTypeName(NextStruct, Segment.Name, Segment.Type);
+
+		if (!Property)
+		{
+			return {};
+		}
+
+		Value.Property = Property;
+		Value.Struct = NextStruct;
+		Value.Container = NextContainer;
+		Value.ArrayIndex = 0;
+
+		// Check the bounds and assign the index for static arrays.
+		if (Property->ArrayDim > 1)
+		{
+			if (Segment.Index < 0 || Segment.Index >= Property->ArrayDim)
+			{
+				return {};
+			}
+			Value.ArrayIndex = Segment.Index;
+		}
+
+		// Resolve the struct and container for the next segment if there is one.
+		NextStruct = FindStructFromProperty(Property);
+		NextContainer = Property->ContainerPtrToValuePtr<uint8>(NextContainer, Value.ArrayIndex);
+
+		// Resolve optionals to the struct and container of their value if they have one.
+		if (const FOptionalProperty* OptionalProperty = CastField<FOptionalProperty>(Property))
+		{
+			FOptionalPropertyLayout OptionalHelper(OptionalProperty->GetValueProperty());
+			NextStruct = FindStructFromProperty(OptionalHelper.GetValueProperty());
+			NextContainer = OptionalHelper.GetValuePointerForReadOrReplaceIfSet(NextContainer);
+		}
+
+		// Scalar values and static containers are finished resolving.
+		if (Property->ArrayDim > 1 || Segment.Index == INDEX_NONE)
+		{
+			continue;
+		}
+
+		// Resolve dynamic containers, which have no struct when resolving directly to an element.
+		Value.Struct = nullptr;
+
+		if (const FArrayProperty* ArrayProperty = CastField<FArrayProperty>(Property))
+		{
+			FScriptArrayHelper ArrayHelper(ArrayProperty, NextContainer);
+			if (!ArrayHelper.IsValidIndex(Segment.Index))
+			{
+				return {};
+			}
+
+			NextStruct = FindStructFromProperty(ArrayProperty->Inner);
+			NextContainer = ArrayHelper.GetRawPtr(Segment.Index);
+			Value.Property = ArrayProperty->Inner;
+			Value.Container = NextContainer;
+		}
+		else if (const FSetProperty* SetProperty = CastField<FSetProperty>(Property))
+		{
+			FScriptSetHelper SetHelper(SetProperty, NextContainer);
+			if (!SetHelper.IsValidIndex(Segment.Index))
+			{
+				return {};
+			}
+
+			NextStruct = FindStructFromProperty(SetProperty->ElementProp);
+			NextContainer = SetHelper.GetElementPtr(Segment.Index);
+			Value.Property = SetProperty->ElementProp;
+			Value.Container = NextContainer;
+		}
+		else if (const FMapProperty* MapProperty = CastField<FMapProperty>(Property))
+		{
+			FScriptMapHelper MapHelper(MapProperty, NextContainer);
+			if (!MapHelper.IsValidIndex(Segment.Index) || ++Index == Count)
+			{
+				return {};
+			}
+
+			// A Key or Value segment with no type or index is required to distinguish which property to resolve.
+			const FPropertyPathNameSegment MapSegment = Path.GetSegment(Index);
+			if (!MapSegment.Type.IsEmpty() || MapSegment.Index != INDEX_NONE)
+			{
+				return {};
+			}
+
+			if (MapSegment.Name == NAME_Key)
+			{
+				NextStruct = FindStructFromProperty(MapProperty->KeyProp);
+				NextContainer = MapHelper.GetKeyPtr(Segment.Index);
+				Value.Property = MapProperty->KeyProp;
+			}
+			else if (MapSegment.Name == NAME_Value)
+			{
+				NextStruct = FindStructFromProperty(MapProperty->ValueProp);
+				NextContainer = MapHelper.GetValuePtr(Segment.Index);
+				Value.Property = MapProperty->ValueProp;
+			}
+			else
+			{
+				return {};
+			}
+
+			// The key and value property both have an offset relative to the pair.
+			Value.Container = MapHelper.GetPairPtr(Segment.Index);
+		}
+		else
+		{
+			return {};
+		}
+	}
+	return Value;
+}
+
+} // UE
diff --git a/Engine/Source/Runtime/CoreUObject/Public/UObject/PropertyPathFunctions.h b/Engine/Source/Runtime/CoreUObject/Public/UObject/PropertyPathFunctions.h
new file mode 100644
--- /dev/null
+++ b/Engine/Source/Runtime/CoreUObject/Public/UObject/PropertyPathFunctions.h
@@ -0,0 +1,81 @@
+// Copyright Epic Games, Inc. All Rights Reserved.
+
+#pragma once
+
+#include "CoreTypes.h"
+
+#define UE_API COREUOBJECT_API
+
+class FName;
+class FProperty;
+class UObject;
+class UStruct;
+
+namespace UE { class FPropertyPathName; }
+namespace UE { class FPropertyTypeName; }
+
+namespace UE
+{
+
+UE_API extern const FName NAME_Key;
+UE_API extern const FName NAME_Value;
+
+/**
+ * Find a property in Struct that matches both the name and the type.
+ *
+ * Type matching uses FProperty::CanSerializeFromTypeName.
+ *
+ * @return A matching property if found, otherwise null.
+ */
+UE_API FProperty* FindPropertyByNameAndTypeName(const UStruct* Struct, FName Name, FPropertyTypeName TypeName);
+
+/**
+ * A reference to a single property value in a container.
+ *
+ * An example of accessing the value from a valid reference:
+ * void* Data = Value.Property->ContainerPtrToValuePtr<void>(Value.Container, Value.ArrayIndex);
+ *
+ * An example of querying initialized property value state from a valid reference:
+ * bool bInitialized = !Value.Struct ||
+ *     FInitializedPropertyValueState(Value.Struct, Value.Container).IsSet(Value.Property, Value.ArrayIndex);
+ */
+struct FPropertyValueInContainer
+{
+	/** The referenced property. If Struct is set, Property is one of its properties. */
+	const FProperty* Property = nullptr;
+	/** The type of the container that contains the referenced property value. Null for a property in a container. */
+	const UStruct* Struct = nullptr;
+	/** The container that contains the referenced property value. If Struct is set, Container is an instance of it. */
+	void* Container = nullptr;
+	/** The static array index within the referenced property. Always non-negative in a valid reference. */
+	int32 ArrayIndex = INDEX_NONE;
+
+	inline explicit operator bool() const
+	{
+		return !!Property;
+	}
+
+	/** Returns a pointer to the property value that this references. */
+	template <typename ValueType>
+	inline ValueType* GetValuePtr() const
+	{
+		// A lambda template allows this to be defined with only a forward declaration of FProperty.
+		return [this]<typename PropertyType>(PropertyType* P)
+		{
+			return P->template ContainerPtrToValuePtr<ValueType>(Container, ArrayIndex);
+		}(Property);
+	}
+};
+
+/**
+ * Try to resolve the property path to a property value within the object.
+ *
+ * @param Path The non-empty property path to resolve.
+ * @param Object The object representing the root of the property path to resolve.
+ * @return A valid property value reference if resolved, otherwise an invalid reference.
+ */
+UE_API FPropertyValueInContainer TryResolvePropertyPath(const FPropertyPathName& Path, UObject* Object);
+
+} // UE
+
+#undef UE_API
diff --git a/Engine/Source/Runtime/CoreUObject/Tests/UObject/PropertyPathFunctionsTest.cpp b/Engine/Source/Runtime/CoreUObject/Tests/UObject/PropertyPathFunctionsTest.cpp
new file mode 100644
--- /dev/null
+++ b/Engine/Source/Runtime/CoreUObject/Tests/UObject/PropertyPathFunctionsTest.cpp
@@ -0,0 +1,322 @@
+// Copyright Epic Games, Inc. All Rights Reserved.
+
+#include "PropertyPathFunctionsTest.h"
+
+#if WITH_TESTS
+
+#include "Misc/StringBuilder.h"
+#include "Tests/TestHarnessAdapter.h"
+#include "UObject/PropertyPathFunctions.h"
+#include "UObject/PropertyPathName.h"
+
+namespace UE
+{
+
+TEST_CASE_NAMED(FPropertyPathFunctionsFindPropertyTest, "CoreUObject::PropertyPathFunctions::FindProperty", "[Core][UObject][SmokeFilter]")
+{
+	SECTION("Found")
+	{
+		const UStruct* TestType = UTestPropertyPathFunctionsClass::StaticClass();
+		for (const FProperty* Property = TestType->PropertyLink; Property; Property = Property->PropertyLinkNext)
+		{
+			FPropertyTypeNameBuilder TypeName;
+			Property->SaveTypeName(TypeName);
+			CHECK(Property == FindPropertyByNameAndTypeName(TestType, Property->GetFName(), TypeName.Build()));
+		}
+	}
+
+	SECTION("Missing")
+	{
+		const UStruct* TestType = UTestPropertyPathFunctionsClass::StaticClass();
+		const FPropertyTypeName FloatType = []
+		{
+			FPropertyTypeNameBuilder TypeName;
+			TypeName.AddName(NAME_FloatProperty);
+			return TypeName.Build();
+		}();
+		for (const FProperty* Property = TestType->PropertyLink; Property; Property = Property->PropertyLinkNext)
+		{
+			CHECK_FALSE(FindPropertyByNameAndTypeName(TestType, Property->GetFName(), FloatType));
+		}
+	}
+}
+
+TEST_CASE_NAMED(FPropertyPathFunctionsTryResolveTest, "CoreUObject::PropertyPathFunctions::TryResolve", "[Core][UObject][SmokeFilter]")
+{
+	UTestPropertyPathFunctionsClass* TestObject = NewObject<UTestPropertyPathFunctionsClass>();
+	const UClass* TestClass = StaticClass<UTestPropertyPathFunctionsClass>();
+	const UStruct* TestStruct = StaticStruct<FTestPropertyPathFunctionsStruct>();
+	const UStruct* TestStructKey = StaticStruct<FTestPropertyPathFunctionsStructKey>();
+
+	FTestPropertyPathFunctionsStruct& StructRef = TestObject->StructStaticArray[3];
+	StructRef.Int32 = 0x100;
+	StructRef.Int32StaticArray[5] = 0x101;
+	StructRef.Int32Array = {0x102, 0x103, 0x104, 0x105};
+	StructRef.Int32Set = {0x106, 0x107, 0x108, 0x109};
+	StructRef.Int32Map = {{0x10a, 0x10b}, {0x10c, 0x10d}, {0x10e, 0x10f}};
+	StructRef.Int32Optional = 0x110;
+
+	TestObject->StructArray.AddDefaulted(3);
+	TestObject->StructArray.Add(StructRef);
+	TestObject->StructSet = {{0x100}, {0x101}, {0x102}, {0x103}};
+	TestObject->StructMap = {{{0x104}, {}}, {{0x105}, {}}, {{0x106}, {}}, {{0x107}, StructRef}};
+	TestObject->StructOptional = StructRef;
+
+	struct FPropertyReference
+	{
+		const UStruct* Struct = nullptr;
+		const ANSICHAR* PropertyName = nullptr;
+		int32 ArrayIndex = INDEX_NONE;
+	};
+
+	const auto MakePath = [](std::initializer_list<FPropertyReference> Segments) -> FPropertyPathName
+	{
+		FPropertyPathName PathName;
+		for (const FPropertyReference& Segment : Segments)
+		{
+			FName PropertyName = Segment.PropertyName;
+			FPropertyTypeNameBuilder TypeName;
+			if (Segment.Struct)
+			{
+				FProperty* Property = Segment.Struct->FindPropertyByName(PropertyName);
+				CAPTURE(Segment.PropertyName);
+				CHECK(Property);
+				if (Property)
+				{
+					Property->SaveTypeName(TypeName);
+				}
+			}
+			PathName.Push({PropertyName, TypeName.Build(), Segment.ArrayIndex});
+		}
+		return PathName;
+	};
+
+	FPropertyValueInContainer Value;
+
+	// Resolve Empty Path
+	CHECK_FALSE((Value = TryResolvePropertyPath(MakePath({}), TestObject)));
+
+	// Resolve Null Object
+	CHECK_FALSE((Value = TryResolvePropertyPath(MakePath({{TestClass, "StructArray"}}), nullptr)));
+
+	// Resolve Int32
+	CHECK_FALSE((Value = TryResolvePropertyPath(MakePath({{TestStruct, "Int32"}}), TestObject)));
+
+	// Resolve StructStaticArray
+	CHECK_FALSE((Value = TryResolvePropertyPath(MakePath({{TestClass, "StructStaticArray"}}), TestObject)));
+
+	// Resolve StructStaticArray[3]
+	REQUIRE((Value = TryResolvePropertyPath(MakePath({{TestClass, "StructStaticArray", 3}}), TestObject)));
+	CHECK(Value.ArrayIndex == 3);
+	CHECK(Value.Struct == TestClass);
+	CHECK(Value.Property->GetFName() == "StructStaticArray");
+	CHECK(Value.GetValuePtr<void>() == &TestObject->StructStaticArray[3]);
+
+	// Resolve StructStaticArray[3] -> Int32
+	REQUIRE((Value = TryResolvePropertyPath(MakePath({{TestClass, "StructStaticArray", 3}, {TestStruct, "Int32"}}), TestObject)));
+	CHECK(Value.ArrayIndex == 0);
+	CHECK(Value.Struct == TestStruct);
+	CHECK(Value.Property->GetFName() == "Int32");
+	CHECK(Value.GetValuePtr<int32>() == &TestObject->StructStaticArray[3].Int32);
+	CHECK(*Value.GetValuePtr<int32>() == 0x100);
+
+	// Resolve StructStaticArray[3] -> Int32StaticArray[5]
+	REQUIRE((Value = TryResolvePropertyPath(MakePath({{TestClass, "StructStaticArray", 3}, {TestStruct, "Int32StaticArray", 5}}), TestObject)));
+	CHECK(Value.ArrayIndex == 5);
+	CHECK(Value.Struct == TestStruct);
+	CHECK(Value.Property->GetFName() == "Int32StaticArray");
+	CHECK(Value.GetValuePtr<int32>() == &TestObject->StructStaticArray[3].Int32StaticArray[5]);
+	CHECK(*Value.GetValuePtr<int32>() == 0x101);
+
+	// Resolve StructStaticArray[3] -> Int32Array[1]
+	REQUIRE((Value = TryResolvePropertyPath(MakePath({{TestClass, "StructStaticArray", 3}, {TestStruct, "Int32Array", 1}}), TestObject)));
+	CHECK(Value.ArrayIndex == 0);
+	CHECK(Value.Struct == nullptr);
+	CHECK(Value.Property->GetFName() == "Int32Array");
+	CHECK(Value.GetValuePtr<int32>() == &TestObject->StructStaticArray[3].Int32Array[1]);
+	CHECK(*Value.GetValuePtr<int32>() == 0x103);
+
+	// Resolve StructStaticArray[3] -> Int32Set[1]
+	REQUIRE((Value = TryResolvePropertyPath(MakePath({{TestClass, "StructStaticArray", 3}, {TestStruct, "Int32Set", 1}}), TestObject)));
+	CHECK(Value.ArrayIndex == 0);
+	CHECK(Value.Struct == nullptr);
+	CHECK(Value.Property->GetFName() == "Int32Set");
+	CHECK(Value.GetValuePtr<int32>() == &TestObject->StructStaticArray[3].Int32Set[FSetElementId::FromInteger(1)]);
+	CHECK(*Value.GetValuePtr<int32>() == 0x107);
+
+	// Resolve StructStaticArray[3] -> Int32Map[1]
+	CHECK_FALSE((Value = TryResolvePropertyPath(MakePath({{TestClass, "StructStaticArray", 3}, {TestStruct, "Int32Map", 1}}), TestObject)));
+
+	// Resolve StructStaticArray[3] -> Int32Map[1] -> Key
+	REQUIRE((Value = TryResolvePropertyPath(MakePath({{TestClass, "StructStaticArray", 3}, {TestStruct, "Int32Map", 1}, {nullptr, "Key"}}), TestObject)));
+	CHECK(Value.ArrayIndex == 0);
+	CHECK(Value.Struct == nullptr);
+	CHECK(Value.Property->GetFName() == "Int32Map_Key");
+	CHECK(Value.GetValuePtr<int32>() == &TestObject->StructStaticArray[3].Int32Map.Get(FSetElementId::FromInteger(1)).Key);
+	CHECK(*Value.GetValuePtr<int32>() == 0x10c);
+
+	// Resolve StructStaticArray[3] -> Int32Map[1] -> Value
+	REQUIRE((Value = TryResolvePropertyPath(MakePath({{TestClass, "StructStaticArray", 3}, {TestStruct, "Int32Map", 1}, {nullptr, "Value"}}), TestObject)));
+	CHECK(Value.ArrayIndex == 0);
+	CHECK(Value.Struct == nullptr);
+	CHECK(Value.Property->GetFName() == "Int32Map");
+	CHECK(Value.GetValuePtr<int32>() == &TestObject->StructStaticArray[3].Int32Map.Get(FSetElementId::FromInteger(1)).Value);
+	CHECK(*Value.GetValuePtr<int32>() == 0x10d);
+
+	// Resolve StructStaticArray[3] -> Int32Optional
+	REQUIRE((Value = TryResolvePropertyPath(MakePath({{TestClass, "StructStaticArray", 3}, {TestStruct, "Int32Optional"}}), TestObject)));
+	CHECK(Value.ArrayIndex == 0);
+	CHECK(Value.Struct == TestStruct);
+	CHECK(Value.Property->GetFName() == "Int32Optional");
+	CHECK(Value.GetValuePtr<TOptional<int32>>() == &TestObject->StructStaticArray[3].Int32Optional);
+	CHECK(*Value.GetValuePtr<TOptional<int32>>() == 0x110);
+
+	// Resolve StructStaticArray[3] -> Int32[0]
+	CHECK_FALSE((Value = TryResolvePropertyPath(MakePath({{TestClass, "StructStaticArray", 3}, {TestStruct, "Int32", 0}}), TestObject)));
+
+	// Resolve StructStaticArray[9]
+	CHECK_FALSE((Value = TryResolvePropertyPath(MakePath({{TestClass, "StructStaticArray", 9}}), TestObject)));
+
+	// Resolve StructArray
+	REQUIRE((Value = TryResolvePropertyPath(MakePath({{TestClass, "StructArray"}}), TestObject)));
+	CHECK(Value.ArrayIndex == 0);
+	CHECK(Value.Struct == TestClass);
+	CHECK(Value.Property->GetFName() == "StructArray");
+	CHECK(Value.GetValuePtr<void>() == &TestObject->StructArray);
+
+	// Resolve StructArray[3]
+	REQUIRE((Value = TryResolvePropertyPath(MakePath({{TestClass, "StructArray", 3}}), TestObject)));
+	CHECK(Value.ArrayIndex == 0);
+	CHECK(Value.Struct == nullptr);
+	CHECK(Value.Property->GetFName() == "StructArray");
+	CHECK(Value.GetValuePtr<void>() == &TestObject->StructArray[3]);
+
+	// Resolve StructArray[3] -> Int32StaticArray[5]
+	REQUIRE((Value = TryResolvePropertyPath(MakePath({{TestClass, "StructArray", 3}, {TestStruct, "Int32StaticArray", 5}}), TestObject)));
+	CHECK(Value.ArrayIndex == 5);
+	CHECK(Value.Struct == TestStruct);
+	CHECK(Value.Property->GetFName() == "Int32StaticArray");
+	CHECK(Value.GetValuePtr<int32>() == &TestObject->StructArray[3].Int32StaticArray[5]);
+	CHECK(*Value.GetValuePtr<int32>() == 0x101);
+
+	// Resolve StructArray[4]
+	CHECK_FALSE((Value = TryResolvePropertyPath(MakePath({{TestClass, "StructArray", 4}}), TestObject)));
+
+	// Resolve StructSet
+	REQUIRE((Value = TryResolvePropertyPath(MakePath({{TestClass, "StructSet"}}), TestObject)));
+	CHECK(Value.ArrayIndex == 0);
+	CHECK(Value.Struct == TestClass);
+	CHECK(Value.Property->GetFName() == "StructSet");
+	CHECK(Value.GetValuePtr<void>() == &TestObject->StructSet);
+
+	// Resolve StructSet[3]
+	REQUIRE((Value = TryResolvePropertyPath(MakePath({{TestClass, "StructSet", 3}}), TestObject)));
+	CHECK(Value.ArrayIndex == 0);
+	CHECK(Value.Struct == nullptr);
+	CHECK(Value.Property->GetFName() == "StructSet");
+	CHECK(Value.GetValuePtr<FTestPropertyPathFunctionsStructKey>() == &TestObject->StructSet[FSetElementId::FromInteger(3)]);
+	CHECK(Value.GetValuePtr<FTestPropertyPathFunctionsStructKey>()->Key == 0x103);
+
+	// Resolve StructSet[3] -> Key
+	REQUIRE((Value = TryResolvePropertyPath(MakePath({{TestClass, "StructSet", 3}, {TestStructKey, "Key"}}), TestObject)));
+	CHECK(Value.ArrayIndex == 0);
+	CHECK(Value.Struct == TestStructKey);
+	CHECK(Value.Property->GetFName() == "Key");
+	CHECK(Value.GetValuePtr<int32>() == &TestObject->StructSet[FSetElementId::FromInteger(3)].Key);
+	CHECK(*Value.GetValuePtr<int32>() == 0x103);
+
+	// Resolve StructSet[4]
+	CHECK_FALSE((Value = TryResolvePropertyPath(MakePath({{TestClass, "StructSet", 4}}), TestObject)));
+
+	// Resolve StructMap
+	REQUIRE((Value = TryResolvePropertyPath(MakePath({{TestClass, "StructMap"}}), TestObject)));
+	CHECK(Value.ArrayIndex == 0);
+	CHECK(Value.Struct == TestClass);
+	CHECK(Value.Property->GetFName() == "StructMap");
+	CHECK(Value.GetValuePtr<void>() == &TestObject->StructMap);
+
+	// Resolve StructMap[3]
+	CHECK_FALSE((Value = TryResolvePropertyPath(MakePath({{TestClass, "StructMap", 3}}), TestObject)));
+
+	// Resolve StructMap[3] -> Int32
+	CHECK_FALSE((Value = TryResolvePropertyPath(MakePath({{TestClass, "StructMap", 3}, {nullptr, "Int32"}}), TestObject)));
+
+	// Resolve StructMap[3] -> Int32(IntProperty)
+	CHECK_FALSE((Value = TryResolvePropertyPath(MakePath({{TestClass, "StructMap", 3}, {TestStruct, "Int32"}}), TestObject)));
+
+	// Resolve StructMap[3] -> Int32StaticArray[5]
+	CHECK_FALSE((Value = TryResolvePropertyPath(MakePath({{TestClass, "StructMap", 3}, {nullptr, "Int32StaticArray", 5}}), TestObject)));
+
+	// Resolve StructMap[3] -> Key
+	REQUIRE((Value = TryResolvePropertyPath(MakePath({{TestClass, "StructMap", 3}, {nullptr, "Key"}}), TestObject)));
+	CHECK(Value.ArrayIndex == 0);
+	CHECK(Value.Struct == nullptr);
+	CHECK(Value.Property->GetFName() == "StructMap_Key");
+	CHECK(Value.GetValuePtr<FTestPropertyPathFunctionsStructKey>() == &TestObject->StructMap.Get(FSetElementId::FromInteger(3)).Key);
+	CHECK(Value.GetValuePtr<FTestPropertyPathFunctionsStructKey>()->Key == 0x107);
+
+	// Resolve StructMap[3] -> Key -> Key
+	REQUIRE((Value = TryResolvePropertyPath(MakePath({{TestClass, "StructMap", 3}, {nullptr, "Key"}, {TestStructKey, "Key"}}), TestObject)));
+	CHECK(Value.ArrayIndex == 0);
+	CHECK(Value.Struct == TestStructKey);
+	CHECK(Value.Property->GetFName() == "Key");
+	CHECK(Value.GetValuePtr<int32>() == &TestObject->StructMap.Get(FSetElementId::FromInteger(3)).Key.Key);
+	CHECK(*Value.GetValuePtr<int32>() == 0x107);
+
+	// Resolve StructMap[3] -> Value
+	REQUIRE((Value = TryResolvePropertyPath(MakePath({{TestClass, "StructMap", 3}, {nullptr, "Value"}}), TestObject)));
+	CHECK(Value.ArrayIndex == 0);
+	CHECK(Value.Struct == nullptr);
+	CHECK(Value.Property->GetFName() == "StructMap");
+	CHECK(Value.GetValuePtr<FTestPropertyPathFunctionsStruct>() == &TestObject->StructMap.Get(FSetElementId::FromInteger(3)).Value);
+
+	// Resolve StructMap[3] -> Value -> Int32StaticArray[5]
+	REQUIRE((Value = TryResolvePropertyPath(MakePath({{TestClass, "StructMap", 3}, {nullptr, "Value"}, {TestStruct, "Int32StaticArray", 5}}), TestObject)));
+	CHECK(Value.ArrayIndex == 5);
+	CHECK(Value.Struct == TestStruct);
+	CHECK(Value.Property->GetFName() == "Int32StaticArray");
+	CHECK(Value.GetValuePtr<int32>() == &TestObject->StructMap.Get(FSetElementId::FromInteger(3)).Value.Int32StaticArray[5]);
+	CHECK(*Value.GetValuePtr<int32>() == 0x101);
+
+	// Resolve StructMap[3] -> Int32
+	REQUIRE((Value = TryResolvePropertyPath(MakePath({{TestClass, "StructMap", 3}, {nullptr, "Value"}, {TestStruct, "Int32"}}), TestObject)));
+	CHECK(Value.ArrayIndex == 0);
+	CHECK(Value.Struct == TestStruct);
+	CHECK(Value.Property->GetFName() == "Int32");
+	CHECK(Value.GetValuePtr<int32>() == &TestObject->StructMap.Get(FSetElementId::FromInteger(3)).Value.Int32);
+	CHECK(*Value.GetValuePtr<int32>() == 0x100);
+
+	// Resolve StructMap[4] -> Key
+	CHECK_FALSE((Value = TryResolvePropertyPath(MakePath({{TestClass, "StructMap", 4}, {nullptr, "Key"}}), TestObject)));
+
+	// Resolve StructOptional
+	REQUIRE((Value = TryResolvePropertyPath(MakePath({{TestClass, "StructOptional"}}), TestObject)));
+	CHECK(Value.ArrayIndex == 0);
+	CHECK(Value.Struct == TestClass);
+	CHECK(Value.Property->GetFName() == "StructOptional");
+	CHECK(Value.GetValuePtr<void>() == &TestObject->StructOptional);
+
+	// Resolve StructOptional -> Int32
+	REQUIRE((Value = TryResolvePropertyPath(MakePath({{TestClass, "StructOptional"}, {TestStruct, "Int32"}}), TestObject)));
+	CHECK(Value.ArrayIndex == 0);
+	CHECK(Value.Struct == TestStruct);
+	CHECK(Value.Property->GetFName() == "Int32");
+	CHECK(Value.GetValuePtr<int32>() == &TestObject->StructOptional->Int32);
+	CHECK(*Value.GetValuePtr<int32>() == 0x100);
+
+	TestObject->StructOptional.Reset();
+
+	// Resolve unset StructOptional
+	REQUIRE((Value = TryResolvePropertyPath(MakePath({{TestClass, "StructOptional"}}), TestObject)));
+	CHECK(Value.ArrayIndex == 0);
+	CHECK(Value.Struct == TestClass);
+	CHECK(Value.Property->GetFName() == "StructOptional");
+	CHECK(Value.GetValuePtr<void>() == &TestObject->StructOptional);
+
+	// Resolve unset StructOptional -> Int32
+	CHECK_FALSE((Value = TryResolvePropertyPath(MakePath({{TestClass, "StructOptional"}, {TestStruct, "Int32"}}), TestObject)));
+}
+
+} // UE
+
+#endif // WITH_TESTS
diff --git a/Engine/Source/Runtime/CoreUObject/Tests/UObject/PropertyPathFunctionsTest.h b/Engine/Source/Runtime/CoreUObject/Tests/UObject/PropertyPathFunctionsTest.h
new file mode 100644
--- /dev/null
+++ b/Engine/Source/Runtime/CoreUObject/Tests/UObject/PropertyPathFunctionsTest.h
@@ -0,0 +1,92 @@
+// Copyright Epic Games, Inc. All Rights Reserved.
+
+#pragma once
+
+#include "CoreTypes.h"
+#include "Containers/Array.h"
+#include "Containers/Map.h"
+#include "Containers/Set.h"
+#include "Misc/Optional.h"
+#include "Templates/TypeHash.h"
+#include "UObject/Object.h"
+#include "PropertyPathFunctionsTest.generated.h"
+
+USTRUCT()
+struct FTestPropertyPathFunctionsStructKey
+{
+	GENERATED_BODY()
+
+public:
+	UPROPERTY()
+	int32 Unused = -1;
+
+	UPROPERTY()
+	int32 Key = 0;
+
+	FTestPropertyPathFunctionsStructKey() = default;
+
+	inline FTestPropertyPathFunctionsStructKey(int32 InKey)
+		: Key(InKey)
+	{
+	}
+
+	inline bool operator==(const FTestPropertyPathFunctionsStructKey& Other) const
+	{
+		return Key == Other.Key;
+	}
+
+	friend inline uint32 GetTypeHash(const FTestPropertyPathFunctionsStructKey& Value)
+	{
+		return GetTypeHashHelper(Value.Key);
+	}
+};
+
+USTRUCT()
+struct FTestPropertyPathFunctionsStruct
+{
+	GENERATED_BODY()
+
+public:
+	UPROPERTY()
+	int32 Unused = -1;
+
+	UPROPERTY()
+	int32 Int32 = 0;
+
+	UPROPERTY()
+	int32 Int32StaticArray[8]{};
+
+	UPROPERTY()
+	TArray<int32> Int32Array;
+
+	UPROPERTY()
+	TSet<int32> Int32Set;
+
+	UPROPERTY()
+	TMap<int32, int32> Int32Map;
+
+	UPROPERTY()
+	TOptional<int32> Int32Optional;
+};
+
+UCLASS()
+class UTestPropertyPathFunctionsClass : public UObject
+{
+	GENERATED_BODY()
+
+public:
+	UPROPERTY()
+	FTestPropertyPathFunctionsStruct StructStaticArray[8];
+
+	UPROPERTY()
+	TArray<FTestPropertyPathFunctionsStruct> StructArray;
+
+	UPROPERTY()
+	TSet<FTestPropertyPathFunctionsStructKey> StructSet;
+
+	UPROPERTY()
+	TMap<FTestPropertyPathFunctionsStructKey, FTestPropertyPathFunctionsStruct> StructMap;
+
+	UPROPERTY()
+	TOptional<FTestPropertyPathFunctionsStruct> StructOptional;
+};
