[ZenLoader]
- Prevent hang detection when the loader is IO starved

#rnx

diff --git a/Engine/Source/Runtime/CoreUObject/Private/Serialization/AsyncLoading2.cpp b/Engine/Source/Runtime/CoreUObject/Private/Serialization/AsyncLoading2.cpp
--- a/Engine/Source/Runtime/CoreUObject/Private/Serialization/AsyncLoading2.cpp
+++ b/Engine/Source/Runtime/CoreUObject/Private/Serialization/AsyncLoading2.cpp
@@ -9950,7 +9950,7 @@ void FAsyncLoadingThread2::FlushLoading(TConstArrayView<int32> RequestIDs)
 							AsyncLoadingThreadState->ResetActivity();
 						}
 					}
-					else if (FPlatformTime::Seconds() - LastActivity > GStallDetectorTimeout && ++IdleLoopCount > GStallDetectorIdleLoops)
+					else if (GStallDetectorTimeout != 0.0f && PendingIoRequestsCounter.Load() == 0 && FPlatformTime::Seconds() - LastActivity > GStallDetectorTimeout && ++IdleLoopCount > GStallDetectorIdleLoops)
 					{
 #if !UE_BUILD_SHIPPING
 						{
