Disable the Loader checking the file system by default if we are Shipping/Test and not the Editor
This can be changed by a CVar if it is required

diff --git a/Engine/Source/Runtime/CoreUObject/Private/Serialization/AsyncLoading2.cpp b/Engine/Source/Runtime/CoreUObject/Private/Serialization/AsyncLoading2.cpp
--- a/Engine/Source/Runtime/CoreUObject/Private/Serialization/AsyncLoading2.cpp
+++ b/Engine/Source/Runtime/CoreUObject/Private/Serialization/AsyncLoading2.cpp
@@ -420,6 +420,15 @@ static FAutoConsoleVariableRef CVarGOnlyProcessRequiredPackagesWhenSyncLoading(
 	ECVF_Default
 );
 
+// If allowed in Shipping builds to load from the file system (Editor will always be allowed for this)
+static bool GAllowLoadingPackageFromFileSystemInShipping = false;
+static FAutoConsoleVariableRef CVarAllowLoadingPackagesFromFileSystem(
+	TEXT("s.AllowLoadingPackageFromFileSystemInShipping"),
+	GAllowLoadingPackageFromFileSystemInShipping,
+	TEXT("Allow the loader to try and load packages from the filesystem in Shipping configs"),
+	ECVF_Default
+	);
+
 #if WITH_EDITOR
 // This is important for the editor because the linker can end up recreating missing exports from its CreateImport function.
 // If we don't put the package in the zenloader loading queue, we can be left with objects in a RF_NeedLoad state that will never be actually loaded.
@@ -3912,6 +3921,13 @@ private:
 #if ALT2_ENABLE_LINKERLOAD_SUPPORT || WITH_EDITOR
 	bool TryGetPackagePathFromFileSystem(FName& PackageNameToLoad, FName& UPackageName, FPackagePath& OutPackagePath)
 	{
+		// If disabling loading from the file system in a shipping/test build and we are not the Editor avoid doing this look-up
+		// As this can be expensive to call and can cause contention on the GameThread and ASyncLoadingThread
+		if (!GAllowLoadingPackageFromFileSystemInShipping && !WITH_EDITOR && (UE_BUILD_SHIPPING || UE_BUILD_TEST))
+		{
+			return false;
+		}
+
 #if WITH_EDITORONLY_DATA
 		// In editor, set MatchCaseOnDisk=true so that we set the capitalization of the Package's FName to match the
 		// capitalization on disk. Different capitalizations can arise from imports of the package that were somehow
