Tracking the initialization state of Verse native bound objects in cooked builds.

#rnx

diff --git a/Engine/Source/Runtime/CoreUObject/Private/Blueprint/BlueprintSupport.cpp b/Engine/Source/Runtime/CoreUObject/Private/Blueprint/BlueprintSupport.cpp
--- a/Engine/Source/Runtime/CoreUObject/Private/Blueprint/BlueprintSupport.cpp
+++ b/Engine/Source/Runtime/CoreUObject/Private/Blueprint/BlueprintSupport.cpp
@@ -30,6 +30,10 @@
 #include "UObject/UObjectIterator.h"
 #endif
 
+#if WITH_EDITORONLY_DATA
+#include "UObject/PropertyStateTracking.h"
+#endif
+
 DEFINE_LOG_CATEGORY_STATIC(LogBlueprintSupport, Log, All);
 
 const FName FBlueprintTags::GeneratedClassPath(TEXT("GeneratedClass"));
@@ -2651,6 +2655,13 @@ bool FObjectInitializer::InitNonNativeProperty(FProperty* Property, UObject* Dat
 		if (!Property->HasAnyPropertyFlags(CPF_ZeroConstructor)) // this stuff is already zero
 		{
 			Property->InitializeValue_InContainer(Data);
+
+			#if WITH_EDITORONLY_DATA
+			if (Property->HasAnyPropertyFlags(CPF_RequiredParm))
+			{
+				UE::FInitializedPropertyValueState(Data->GetClass(), Data).Set(Property);
+			}
+			#endif
 		}
 		return true;
 	}
diff --git a/Engine/Source/Runtime/CoreUObject/Private/Serialization/UnversionedPropertySerialization.cpp b/Engine/Source/Runtime/CoreUObject/Private/Serialization/UnversionedPropertySerialization.cpp
--- a/Engine/Source/Runtime/CoreUObject/Private/Serialization/UnversionedPropertySerialization.cpp
+++ b/Engine/Source/Runtime/CoreUObject/Private/Serialization/UnversionedPropertySerialization.cpp
@@ -16,6 +16,7 @@
 #include "UObject/MetaData.h"
 #include "UObject/Package.h"
 #include "UObject/UObjectIterator.h"
+#include "UObject/PropertyStateTracking.h"
 #endif
 
 // Caches a property array per UStruct to avoid link-walking and touching FProperty data.
@@ -56,6 +57,9 @@ public:
 		, Offset(Property->GetOffset_ForInternal() + Property->GetElementSize() * InArrayIndex)
 		, bSerializeAsInteger(CanSerializeAsInteger(Property))
 		, bIsOptional(IsOptional(Property->GetClass()->GetCastFlags()))
+#if WITH_EDITORONLY_DATA
+		, bHasRequiredParam(Property->HasAnyPropertyFlags(CPF_RequiredParm))
+#endif
 		, IntType(GetIntType(Property->GetMinAlignment()))
 #else
 		, ArrayIndex(InArrayIndex)
@@ -173,6 +177,16 @@ public:
 		return Property->Identical(GetValue(Data), GetDefaultValue(Defaults), PortFlags);
 	}
 
+
+	bool HasRequiredParam() const
+	{
+#if WITH_EDITORONLY_DATA && CACHE_UNVERSIONED_PROPERTY_SCHEMA
+		return bHasRequiredParam;
+#else
+		return Property->HasAnyPropertyFlags(CPF_RequiredParm);
+#endif
+	}
+
 private:
 	enum class EIntegerType : uint8 { Uint8, Uint16, Uint32, Uint64 };
 
@@ -307,7 +321,10 @@ private:
 #if CACHE_UNVERSIONED_PROPERTY_SCHEMA
 	uint32 Offset;
 	bool bSerializeAsInteger;
-	bool bIsOptional;
+	bool bIsOptional : 1;
+#if WITH_EDITORONLY_DATA
+	bool bHasRequiredParam : 1;
+#endif
 	EIntegerType IntType;
 	uint8 FastZeroIntNum;
 #else
@@ -946,6 +963,16 @@ void SerializeUnversionedProperties(const UStruct* Struct, FStructuredArchive::F
 			{
 				FDefaultStruct Defaults(DefaultsData, DefaultsStruct);
 
+#if WITH_EDITORONLY_DATA
+				bool bTrackInitializedProperties = false;
+
+				// We cannot mark all properties as initialized when loading 'dense' property values.
+				if (UnderlyingArchive.DoDelta() && !UnderlyingArchive.IsTransacting())
+				{
+					bTrackInitializedProperties = UE::FInitializedPropertyValueState(Struct, Data).IsTracking();
+				}
+#endif
+
 				FStructuredArchive::FStream ValueStream = StructRecord.EnterStream(TEXT("Values"));
 				for (FUnversionedHeader::FIterator It(Header, Schema); It; It.Next())
 				{
@@ -954,6 +981,14 @@ void SerializeUnversionedProperties(const UStruct* Struct, FStructuredArchive::F
 #if WITH_EDITOR // Skip this scope to save time in the runtime; it is only needed for reference collection in editor
 						FSerializedPropertyScope SerializedProperty(UnderlyingArchive, It.GetSerializer().GetProperty());
 #endif
+
+#if WITH_EDITORONLY_DATA
+						if (bTrackInitializedProperties && It.GetSerializer().HasRequiredParam())
+						{
+							UE::FInitializedPropertyValueState(Struct, Data).Set(It.GetSerializer().GetProperty());
+						}
+#endif
+
 						It.GetSerializer().Serialize(ValueStream.EnterElement(), Data, Defaults);
 					}
 					else
