[ZenLoader]
- Fix loading sometimes getting tangled into transactions because of a missing scope in editor

[RN] Fix loading sometimes getting tangled into transactions because of a missing scope in editor

diff --git a/Engine/Source/Runtime/CoreUObject/Private/Serialization/AsyncLoading2.cpp b/Engine/Source/Runtime/CoreUObject/Private/Serialization/AsyncLoading2.cpp
--- a/Engine/Source/Runtime/CoreUObject/Private/Serialization/AsyncLoading2.cpp
+++ b/Engine/Source/Runtime/CoreUObject/Private/Serialization/AsyncLoading2.cpp
@@ -9863,6 +9863,12 @@ void FAsyncLoadingThread2::FlushLoading(TConstArrayView<int32> RequestIDs)
 			return;
 		}
 
+#if WITH_EDITOR
+		// In the editor loading cannot be part of a transaction as it cannot be undone, and may result in recording half-loaded objects. So we suppress any active transaction while in this stack, and set the editor loading flag
+		TGuardValue<ITransaction*> SuppressTransaction(GUndo, nullptr);
+		TGuardValueAccessors<bool> IsEditorLoadingPackageGuard(UE::GetIsEditorLoadingPackage, UE::SetIsEditorLoadingPackage, GIsEditor || UE::GetIsEditorLoadingPackage());
+#endif
+
 		const bool bIsFullFlush = RequestIDs.IsEmpty();
 		CONDITIONAL_SCOPE_CYCLE_COUNTER(STAT_FAsyncPackage_FlushAsyncLoadingGameThread, !bIsFullFlush);
 		CONDITIONAL_SCOPE_CYCLE_COUNTER(STAT_FAsyncPackage_FlushAllAsyncLoadingGameThread, bIsFullFlush);
