[TSAN]
- Replace TSAN_SAFE by an equivalent that doesn't prevent instrumentation

#rnx

diff --git a/Engine/Source/Runtime/Core/Private/UObject/UnrealNames.cpp b/Engine/Source/Runtime/Core/Private/UObject/UnrealNames.cpp
--- a/Engine/Source/Runtime/Core/Private/UObject/UnrealNames.cpp
+++ b/Engine/Source/Runtime/Core/Private/UObject/UnrealNames.cpp
@@ -563,10 +563,8 @@ public:
 	template<class ScopeLock>
 	FNameEntryHandle Create(FNameStringView Name, TOptional<FNameEntryId> ComparisonId, FNameEntryHeader Header)
 	{
-// Silence the TSAN warning as we don't care for prefetch if we fetch the old or new value
-#if !USING_THREAD_SANITISER
 		FPlatformMisc::Prefetch(Blocks[CurrentBlock]);
-#endif
+
 		FNameEntryHandle Handle = Allocate<ScopeLock>(FNameEntry::GetDataOffset() + Name.BytesWithoutTerminator());
 		FNameEntry& Entry = Resolve(Handle);
 
@@ -634,7 +632,7 @@ public:
 		return CurrentBlock + 1;
 	}
 
-	uint8** GetBlocksForDebugVisualizer() { return Blocks; }
+	uint8** GetBlocksForDebugVisualizer() { return (uint8**)Blocks; }
 
 	void DebugDump(TArray<const FNameEntry*>& Out) const
 	{
@@ -810,9 +808,9 @@ private:
 	}
 
 	mutable FRWLock Lock;
-	uint32 CurrentBlock = 0;
+	TSAN_ATOMIC(uint32) CurrentBlock = 0;
 	uint32 CurrentByteCursor = 0;
-	uint8* Blocks[FNameMaxBlocks] = {};
+	TSAN_ATOMIC(uint8*) Blocks[FNameMaxBlocks] = {};
 #if USE_FNAME_MMAP
 	TArray<uint8*> MmappedAddresses;
 	int32 MmapFile = -1;
diff --git a/Engine/Source/Runtime/CoreUObject/Public/UObject/UObjectArray.h b/Engine/Source/Runtime/CoreUObject/Public/UObject/UObjectArray.h
--- a/Engine/Source/Runtime/CoreUObject/Public/UObject/UObjectArray.h
+++ b/Engine/Source/Runtime/CoreUObject/Public/UObject/UObjectArray.h
@@ -451,15 +451,15 @@ namespace UE::UObjectArrayPrivate
 class FFixedUObjectArray
 {
 	/** Static primary table to chunks of pointers **/
-	FUObjectItem* Objects;
+	TSAN_ATOMIC(FUObjectItem*) Objects;
 	/** Number of elements we currently have **/
-	int32 MaxElements;
+	TSAN_ATOMIC(int32) MaxElements;
 	/** Current number of UObject slots */
-	int32 NumElements;
+	TSAN_ATOMIC(int32) NumElements;
 
 public:
 
-	FFixedUObjectArray() TSAN_SAFE
+	FFixedUObjectArray()
 		: Objects(nullptr)
 		, MaxElements(0)
 		, NumElements(0)
@@ -475,14 +475,14 @@ public:
 	* Expands the array so that Element[Index] is allocated. New pointers are all zero.
 	* @param Index The Index of an element we want to be sure is allocated
 	**/
-	void PreAllocate(int32 InMaxElements) TSAN_SAFE
+	void PreAllocate(int32 InMaxElements)
 	{
 		check(!Objects);
 		Objects = new FUObjectItem[InMaxElements];
 		MaxElements = InMaxElements;
 	}
 
-	int32 AddSingle() TSAN_SAFE
+	int32 AddSingle()
 	{
 		int32 Result = NumElements;
 		UE::UObjectArrayPrivate::CheckUObjectLimitReached(NumElements, MaxElements, 1);
@@ -493,7 +493,7 @@ public:
 		return Result;
 	}
 
-	int32 AddRange(int32 Count) TSAN_SAFE
+	int32 AddRange(int32 Count)
 	{
 		int32 Result = NumElements + Count - 1;
 		UE::UObjectArrayPrivate::CheckUObjectLimitReached(NumElements, MaxElements, Count);
@@ -504,13 +504,13 @@ public:
 		return Result;
 	}
 
-	FORCEINLINE FUObjectItem const* GetObjectPtr(int32 Index) const TSAN_SAFE
+	FORCEINLINE FUObjectItem const* GetObjectPtr(int32 Index) const
 	{
 		check(Index >= 0 && Index < NumElements);
 		return &Objects[Index];
 	}
 
-	FORCEINLINE FUObjectItem* GetObjectPtr(int32 Index) TSAN_SAFE
+	FORCEINLINE FUObjectItem* GetObjectPtr(int32 Index)
 	{
 		check(Index >= 0 && Index < NumElements);
 		return &Objects[Index];
@@ -521,7 +521,7 @@ public:
 	* Thread safe, but you know, someone might have added more elements before this even returns
 	* @return	the number of elements in the array
 	**/
-	FORCEINLINE int32 Num() const TSAN_SAFE
+	FORCEINLINE int32 Num() const
 	{
 		return NumElements;
 	}
@@ -531,7 +531,7 @@ public:
 	* Thread safe, but you know, someone might have added more elements before this even returns
 	* @return	the maximum number of elements in the array
 	**/
-	FORCEINLINE int32 Capacity() const TSAN_SAFE
+	FORCEINLINE int32 Capacity() const
 	{
 		return MaxElements;
 	}
@@ -594,19 +594,19 @@ class FChunkedFixedUObjectArray
 	/** If requested, a contiguous memory where all objects are allocated **/
 	FUObjectItem* PreAllocatedObjects;
 	/** Maximum number of elements **/
-	int32 MaxElements;
+	TSAN_ATOMIC(int32) MaxElements;
 	/** Number of elements we currently have **/
-	int32 NumElements;
+	TSAN_ATOMIC(int32) NumElements;
 	/** Maximum number of chunks **/
 	int32 MaxChunks;
 	/** Number of chunks we currently have **/
-	int32 NumChunks;
+	TSAN_ATOMIC(int32) NumChunks;
 
 
 	/**
 	* Allocates new chunk for the array
 	**/
-	void ExpandChunksToIndex(int32 Index) TSAN_SAFE
+	void ExpandChunksToIndex(int32 Index)
 	{
 		check(Index >= 0 && Index < MaxElements);
 		int32 ChunkIndex = Index / NumElementsPerChunk;
@@ -632,7 +632,7 @@ class FChunkedFixedUObjectArray
 public:
 
 	/** Constructor : Probably not thread safe **/
-	FChunkedFixedUObjectArray() TSAN_SAFE
+	FChunkedFixedUObjectArray()
 		: Objects(nullptr)
 		, PreAllocatedObjects(nullptr)
 		, MaxElements(0)
@@ -662,7 +662,7 @@ public:
 	* Expands the array so that Element[Index] is allocated. New pointers are all zero.
 	* @param Index The Index of an element we want to be sure is allocated
 	**/
-	void PreAllocate(int32 InMaxElements, bool bPreAllocateChunks) TSAN_SAFE
+	void PreAllocate(int32 InMaxElements, bool bPreAllocateChunks)
 	{
 		check(!Objects);
 		MaxChunks = InMaxElements / NumElementsPerChunk + 1;
@@ -696,7 +696,7 @@ public:
 	* Thread safe, but you know, someone might have added more elements before this even returns
 	* @return	the maximum number of elements in the array
 	**/
-	FORCEINLINE int32 Capacity() const TSAN_SAFE
+	FORCEINLINE int32 Capacity() const
 	{
 		return MaxElements;
 	}
@@ -716,30 +716,30 @@ public:
 	* Return a pointer to the pointer to a given element
 	* @param Index The Index of an element we want to retrieve the pointer-to-pointer for
 	**/
-	FORCEINLINE_DEBUGGABLE FUObjectItem const* GetObjectPtr(int32 Index) const TSAN_SAFE
+	FORCEINLINE_DEBUGGABLE FUObjectItem const* GetObjectPtr(int32 Index) const
 	{
 		const uint32 ChunkIndex = (uint32)Index / NumElementsPerChunk;
 		const uint32 WithinChunkIndex = (uint32)Index % NumElementsPerChunk;
 		checkf(IsValidIndex(Index), TEXT("IsValidIndex(%d)"), Index);
-		checkf(ChunkIndex < (uint32)NumChunks, TEXT("ChunkIndex (%d) < NumChunks (%d)"), ChunkIndex, NumChunks);
-		checkf(Index < MaxElements, TEXT("Index (%d) < MaxElements (%d)"), Index, MaxElements);
+		checkf(ChunkIndex < (uint32)NumChunks, TEXT("ChunkIndex (%d) < NumChunks (%d)"), ChunkIndex, (int32)NumChunks);
+		checkf(Index < MaxElements, TEXT("Index (%d) < MaxElements (%d)"), Index, (int32)MaxElements);
 		FUObjectItem* Chunk = Objects[ChunkIndex];
 		check(Chunk);
 		return Chunk + WithinChunkIndex;
 	}
-	FORCEINLINE_DEBUGGABLE FUObjectItem* GetObjectPtr(int32 Index) TSAN_SAFE
+	FORCEINLINE_DEBUGGABLE FUObjectItem* GetObjectPtr(int32 Index)
 	{
 		const uint32 ChunkIndex = (uint32)Index / NumElementsPerChunk;
 		const uint32 WithinChunkIndex = (uint32)Index % NumElementsPerChunk;
 		checkf(IsValidIndex(Index), TEXT("IsValidIndex(%d)"), Index);
-		checkf(ChunkIndex < (uint32)NumChunks, TEXT("ChunkIndex (%d) < NumChunks (%d)"), ChunkIndex, NumChunks);
-		checkf(Index < MaxElements, TEXT("Index (%d) < MaxElements (%d)"), Index, MaxElements);
+		checkf(ChunkIndex < (uint32)NumChunks, TEXT("ChunkIndex (%d) < NumChunks (%d)"), ChunkIndex, (int32)NumChunks);
+		checkf(Index < MaxElements, TEXT("Index (%d) < MaxElements (%d)"), Index, (int32)MaxElements);
 		FUObjectItem* Chunk = Objects[ChunkIndex];
 		check(Chunk);
 		return Chunk + WithinChunkIndex;
 	}
 
-	FORCEINLINE_DEBUGGABLE void PrefetchObjectPtr(int32 Index) const TSAN_SAFE
+	FORCEINLINE_DEBUGGABLE void PrefetchObjectPtr(int32 Index) const
 	{
 		const uint32 ChunkIndex = (uint32)Index / NumElementsPerChunk;
 		const uint32 WithinChunkIndex = (uint32)Index % NumElementsPerChunk;
@@ -766,7 +766,7 @@ public:
 		return *ItemPtr;
 	}
 
-	int32 AddRange(int32 NumToAdd) TSAN_SAFE
+	int32 AddRange(int32 NumToAdd)
 	{
 		int32 Result = NumElements;
 		UE::UObjectArrayPrivate::CheckUObjectLimitReached(Result, MaxElements, NumToAdd);
@@ -775,7 +775,7 @@ public:
 		return Result;
 	}
 
-	int32 AddSingle() TSAN_SAFE
+	int32 AddSingle()
 	{
 		return AddRange(1);
 	}
