[ZenLoader]
- Allow state dump on hang to be controlled by a define so it can be included in shipping if needed

#rnx

diff --git a/Engine/Source/Runtime/CoreUObject/Private/Serialization/AsyncLoading2.cpp b/Engine/Source/Runtime/CoreUObject/Private/Serialization/AsyncLoading2.cpp
--- a/Engine/Source/Runtime/CoreUObject/Private/Serialization/AsyncLoading2.cpp
+++ b/Engine/Source/Runtime/CoreUObject/Private/Serialization/AsyncLoading2.cpp
@@ -277,6 +277,10 @@ void FindAllRuntimeScriptPackages(TArray<UPackage*>& OutPackages)
 #define ALT2_LOG_VERBOSE DO_CHECK
 #endif
 
+#ifndef ALT2_DUMP_STATE_ON_HANG
+#define ALT2_DUMP_STATE_ON_HANG !UE_BUILD_SHIPPING
+#endif
+
 #ifndef ALT2_ENABLE_PACKAGE_DEPENDENCY_DEBUGGING
 #define ALT2_ENABLE_PACKAGE_DEPENDENCY_DEBUGGING 0
 #endif
@@ -3409,11 +3413,11 @@ private:
 	*/
 	void UpdateLoadPercentage();
 
-#if !UE_BUILD_SHIPPING
+#if ALT2_DUMP_STATE_ON_HANG
 	void DumpStateImpl(TSet<FAsyncPackage2*>& Set, int32 Indent = 0, TMultiMap<FEventLoadNode2*, FEventLoadNode2*>* MappedNodes = nullptr);
 #endif
 public:
-#if !UE_BUILD_SHIPPING
+#if ALT2_DUMP_STATE_ON_HANG
 	void DumpState();
 #endif
 
@@ -9974,7 +9978,7 @@ void FAsyncLoadingThread2::FlushLoading(TConstArrayView<int32> RequestIDs)
 					}
 					else if (GStallDetectorTimeout != 0.0f && PendingIoRequestsCounter.Load() == 0 && FPlatformTime::Seconds() - LastActivity > GStallDetectorTimeout && ++IdleLoopCount > GStallDetectorIdleLoops)
 					{
-#if !UE_BUILD_SHIPPING
+#if ALT2_DUMP_STATE_ON_HANG
 						{
 							FScopeLock ScopeLock(&AsyncPackagesCritical);
 
@@ -10012,7 +10016,7 @@ void FAsyncLoadingThread2::FlushLoading(TConstArrayView<int32> RequestIDs)
 								}
 							}
 						}
-#endif // !UE_BUILD_SHIPPING
+#endif // ALT2_DUMP_STATE_ON_HANG
 						UE_LOG(LogStreaming, Fatal, TEXT("Loading is stuck, flush will never finish"));
 					}
 				}
@@ -10139,7 +10143,7 @@ void FGlobalImportStore::FlushDeferredDeletePackagesQueue()
 	AsyncLoadingThread.ProcessDeferredDeletePackagesQueue();
 }
 
-#if !UE_BUILD_SHIPPING
+#if ALT2_DUMP_STATE_ON_HANG
 
 void FAsyncPackage2::DumpState()
 {
@@ -10323,7 +10327,7 @@ void FAsyncPackage2::DumpStateImpl(TSet<FAsyncPackage2*>& Set, int32 Indent, TMu
 	}
 }
 
-#endif // !UE_BUILD_SHIPPING
+#endif // ALT2_DUMP_STATE_ON_HANG
 
 IAsyncPackageLoader* MakeAsyncPackageLoader2(FIoDispatcher& InIoDispatcher, IAsyncPackageLoader* InUncookedPackageLoader)
 {
