[ZenLoader]
- New instrumentation integration for detecting race conditions during async loading

#rnx

diff --git a/Engine/Source/Runtime/CoreUObject/Private/Serialization/AsyncLoading2.cpp b/Engine/Source/Runtime/CoreUObject/Private/Serialization/AsyncLoading2.cpp
--- a/Engine/Source/Runtime/CoreUObject/Private/Serialization/AsyncLoading2.cpp
+++ b/Engine/Source/Runtime/CoreUObject/Private/Serialization/AsyncLoading2.cpp
@@ -69,6 +69,7 @@
 #include "ProfilingDebugging/CountersTrace.h"
 #include "ProfilingDebugging/AssetMetadataTrace.h"
 #include "Async/Async.h"
+#include "Sanitizer/RaceDetector.h"
 #include "Templates/GuardValueAccessors.h"
 #include "Async/ParallelFor.h"
 #include "Async/ManualResetEvent.h"
@@ -462,6 +463,16 @@ static FAutoConsoleVariableRef CVarGStallDetectorIdleLoops(
 	ECVF_Default
 );
 
+#if USING_INSTRUMENTATION
+static bool GDetectRaceDuringLoading = false;
+static FAutoConsoleVariableRef CVarDetectRaceDuringLoading(
+	TEXT("s.DetectRaceDuringLoading"),
+	GDetectRaceDuringLoading,
+	TEXT("Activate the race detector during loading periods"),
+	ECVF_Default
+);
+#endif // USING_INSTRUMENTATION
+
 CSV_DECLARE_CATEGORY_MODULE_EXTERN(CORE_API, Basic);
 CSV_DECLARE_CATEGORY_MODULE_EXTERN(CORE_API, FileIO);
 
@@ -8651,6 +8662,25 @@ uint32 FAsyncLoadingThread2::Run()
 		{
 			TRACE_CPUPROFILER_EVENT_SCOPE(AsyncLoadingTime);
 
+#if USING_INSTRUMENTATION
+			bool bInstrumentationActive = false;
+			if (GDetectRaceDuringLoading)
+			{
+				bInstrumentationActive = true;
+				UE::Sanitizer::RaceDetector::ToggleFilterOtherThreads(true);
+				UE::Sanitizer::RaceDetector::ToggleRaceDetection(true);
+			}
+
+			ON_SCOPE_EXIT
+			{
+				if (bInstrumentationActive)
+				{
+					UE::Sanitizer::RaceDetector::ToggleRaceDetection(false);
+					UE::Sanitizer::RaceDetector::ToggleFilterOtherThreads(false);
+				}
+			};
+#endif
+
 			bool bShouldSuspend = false;
 			bool bShouldWaitForExternalReads = false;
 			while (!bStopRequested.load(std::memory_order_relaxed))
