[TSAN/GC]
- Fix benign race on GIsGarbageCollecting
- Make sure instrumentation see the same fence as TSAN

#rnx

diff --git a/Engine/Source/Runtime/CoreUObject/Private/UObject/GarbageCollection.cpp b/Engine/Source/Runtime/CoreUObject/Private/UObject/GarbageCollection.cpp
--- a/Engine/Source/Runtime/CoreUObject/Private/UObject/GarbageCollection.cpp
+++ b/Engine/Source/Runtime/CoreUObject/Private/UObject/GarbageCollection.cpp
@@ -115,7 +115,7 @@ struct FGCTimingInfo
 };
 static FGCTimingInfo GTimingInfo;
 
-bool GIsGarbageCollecting = false;
+TSAN_ATOMIC(bool) GIsGarbageCollecting = false;
 
 /**
 * Call back into the async loading code to inform of the destruction of serialized objects
@@ -2536,9 +2536,9 @@ TArrayView<UObject*> FAROQueue::PopImpl()
 		// Atomic-fence synchronization of non-atomic stores via relaxed Head load
 		Fence();
 
-		// This is a temporary solution to make TSAN see the fence until the global fence
+		// This is a temporary solution to make TSAN and instrumentation see the fence until the global fence
 		// can be removed the next time we do a proper profiling pass of these bits on all platforms.
-		if constexpr (USING_THREAD_SANITISER && Fence == AcquireFence)
+		if constexpr ((USING_THREAD_SANITISER || USING_INSTRUMENTATION) && Fence == AcquireFence)
 		{
 			Head.load(std::memory_order_acquire);
 		}
@@ -4533,7 +4533,7 @@ void IncrementalPurgeGarbage(bool bUseTimeLimit, double TimeLimit)
 	else
 	{
 		// Set 'I'm garbage collecting' flag - might be checked inside various functions.
-		TGuardValue<bool> GuardIsGarbageCollecting(GIsGarbageCollecting, true);
+		TGuardValue GuardIsGarbageCollecting(GIsGarbageCollecting, true);
 
 		// Keep track of start time to enforce time limit unless bForceFullPurge is true;
 		GCStartTime = FPlatformTime::Seconds();
@@ -4632,7 +4632,7 @@ bool IncrementalDestroyGarbage(bool bUseTimeLimit, double TimeLimit)
 	const int32	TimeLimitEnforcementGranularityForDeletion = 100;
 
 	// Set 'I'm garbage collecting' flag - might be checked inside UObject::Destroy etc.
-	TGuardValue<bool> GuardIsGarbageCollecting(GIsGarbageCollecting, true);
+	TGuardValue GuardIsGarbageCollecting(GIsGarbageCollecting, true);
 
 	if( !GObjFinishDestroyHasBeenRoutedToAllObjects && !bTimeLimitReached )
 	{
diff --git a/Engine/Source/Runtime/CoreUObject/Public/UObject/GarbageCollection.h b/Engine/Source/Runtime/CoreUObject/Public/UObject/GarbageCollection.h
--- a/Engine/Source/Runtime/CoreUObject/Public/UObject/GarbageCollection.h
+++ b/Engine/Source/Runtime/CoreUObject/Public/UObject/GarbageCollection.h
@@ -152,7 +152,7 @@ struct FGCDirectReference
 };
 
 /** True if Garbage Collection is running. Use IsGarbageCollecting() functio n instead of using this variable directly */
-extern COREUOBJECT_API bool GIsGarbageCollecting;
+extern COREUOBJECT_API TSAN_ATOMIC(bool) GIsGarbageCollecting;
 
 /**
  * Gets the last time that the GC was run.
