[ZenLoader]
- Fix potential race with script objects being accessed from both GT and ALT

[RN] Fix potential race in script objects when async loading is active

diff --git a/Engine/Source/Runtime/CoreUObject/Private/Serialization/AsyncLoading2.cpp b/Engine/Source/Runtime/CoreUObject/Private/Serialization/AsyncLoading2.cpp
--- a/Engine/Source/Runtime/CoreUObject/Private/Serialization/AsyncLoading2.cpp
+++ b/Engine/Source/Runtime/CoreUObject/Private/Serialization/AsyncLoading2.cpp
@@ -1261,6 +1261,8 @@ private:
 	TMap<FPackageId, FLoadedPackageRef> Packages;
 	// All native script objects (structs, enums, classes, CDOs and their subobjects) from the initial load phase
 	TMap<FPackageObjectIndex, UObject*> ScriptObjects;
+	// We need to protect ScriptObjects as it is accessed from both GT and ALT.
+	mutable FRWLock ScriptObjectsLock;
 	// All currently loaded public export objects from any loaded package
 	TMap<int32, FPublicExportKey> ObjectIndexToPublicExport;
 
@@ -1282,11 +1284,13 @@ public:
 
 	int32 GetStoredScriptObjectsCount() const
 	{
+		FReadScopeLock ScopeLock(ScriptObjectsLock);
 		return ScriptObjects.Num();
 	}
 
 	SIZE_T GetStoredScriptObjectsAllocatedSize() const
 	{
+		FReadScopeLock ScopeLock(ScriptObjectsLock);
 		return ScriptObjects.GetAllocatedSize();
 	}
 
@@ -1678,6 +1682,8 @@ public:
 
 	inline UObject* FindScriptImportObject(FPackageObjectIndex GlobalIndex)
 	{
+		FReadScopeLock ScopeLock(ScriptObjectsLock);
+
 		check(GlobalIndex.IsScriptImport());
 		UObject* Object = nullptr;
 		Object = ScriptObjects.FindRef(GlobalIndex);
@@ -1743,6 +1749,8 @@ public:
 
 	void AddScriptObject(FStringView PackageName, FStringView Name, UObject* Object)
 	{
+		FWriteScopeLock ScopeLock(ScriptObjectsLock);
+
 		TStringBuilder<FName::StringBufferSize> FullName;
 		FPathViews::Append(FullName, PackageName);
 		FPathViews::Append(FullName, Name);
@@ -5272,6 +5280,8 @@ void FAsyncPackage2::SetupScriptDependencies()
 
 void FGlobalImportStore::FindAllScriptObjects(bool bVerifyOnly)
 {
+	FWriteScopeLock ScopeLock(ScriptObjectsLock);
+
 	TStringBuilder<FName::StringBufferSize> Name;
 	TArray<UPackage*> ScriptPackages;
 	TArray<UObject*> Objects;
@@ -5337,6 +5347,7 @@ void FGlobalImportStore::RegistrationComplete()
 		FindAllScriptObjects(/*bVerifyOnly*/true);
 	}
 #endif
+	FWriteScopeLock ScopeLock(ScriptObjectsLock);
 	ScriptObjects.Shrink();
 }
 
