[TSAN/UnversionedSerialization]
- Fix caching to be thread-safe on platforms with weak memory ordering
- Fix a potential leak under autortfm

#rnx

diff --git a/Engine/Source/Runtime/CoreUObject/Private/Serialization/UnversionedPropertySerialization.cpp b/Engine/Source/Runtime/CoreUObject/Private/Serialization/UnversionedPropertySerialization.cpp
--- a/Engine/Source/Runtime/CoreUObject/Private/Serialization/UnversionedPropertySerialization.cpp
+++ b/Engine/Source/Runtime/CoreUObject/Private/Serialization/UnversionedPropertySerialization.cpp
@@ -19,6 +19,8 @@
 #include "UObject/PropertyStateTracking.h"
 #endif
 
+#include <atomic>
+
 // Caches a property array per UStruct to avoid link-walking and touching FProperty data.
 //
 // As a reference point, this uses ~6MB of memory in an internal project and makes all
@@ -452,7 +454,7 @@ struct FUnversionedStructSchema
 #endif
 };
 
-static const FUnversionedStructSchema*& GetUnversionedSchema(const UStruct* Struct, bool bSkipEditorOnly)
+static std::atomic<const FUnversionedStructSchema*>& GetUnversionedSchema(const UStruct* Struct, bool bSkipEditorOnly)
 {
 #if WITH_EDITORONLY_DATA
 	return bSkipEditorOnly ? Struct->UnversionedGameSchema : Struct->UnversionedEditorSchema;
@@ -463,25 +465,43 @@ static const FUnversionedStructSchema*& GetUnversionedSchema(const UStruct* Stru
 
 const FUnversionedStructSchema& GetOrCreateUnversionedSchema(const UStruct* Struct, bool bSkipEditorOnly)
 {
-	if (const FUnversionedStructSchema* ExistingSchema = GetUnversionedSchema(Struct, bSkipEditorOnly))
+	std::atomic<const FUnversionedStructSchema*>* ExistingSchemaPtr = &GetUnversionedSchema(Struct, bSkipEditorOnly);
+	if (const FUnversionedStructSchema* ExistingSchema = ExistingSchemaPtr->load(std::memory_order_acquire))
 	{
 		return *ExistingSchema;
 	}
 
 	FUnversionedStructSchema* CreatedSchema = FUnversionedStructSchema::Create(Struct, bSkipEditorOnly);
 
-	// AutoRTFM cannot undo atomic swaps, so skip the caching if called from a transaction.
-	if (!AutoRTFM::IsClosed())
+	auto DedupSchema = [ExistingSchemaPtr, CreatedSchema]() -> const FUnversionedStructSchema*
 	{
-		void** CachedSchemaPtr = reinterpret_cast<void**>(const_cast<FUnversionedStructSchema**>(&GetUnversionedSchema(Struct, bSkipEditorOnly)));
-		if (const FUnversionedStructSchema* ExistingSchema = reinterpret_cast<const FUnversionedStructSchema*>(FPlatformAtomics::InterlockedCompareExchangePointer(CachedSchemaPtr, CreatedSchema, nullptr)))
+		const FUnversionedStructSchema* ExistingSchema = nullptr;
+		if (!ExistingSchemaPtr->compare_exchange_strong(ExistingSchema, CreatedSchema, std::memory_order_acq_rel, std::memory_order_acquire))
 		{
 			FUnversionedStructSchema::Delete(CreatedSchema);
-			return *ExistingSchema;
+			check(ExistingSchema);
+			return ExistingSchema;
 		}
+		return CreatedSchema;
+	};
+
+	if (AutoRTFM::IsClosed())
+	{
+		// AutoRTFM cannot undo atomic swaps, so skip the caching if called from a transaction.
+		// and get back to try to do the CAS on commit. 
+		// This implies that we might destroy the CreatedSchema that got used during the transaction
+		// but as long as nothing holds a pointer to it, all should be fine. If not, ASAN will catch it.
+		// OnCommit -> Dedup created schema using CAS.
+		// OnAbort  -> Created schema should get deleted.
+		// So we should not have any leaks remaining in any scenarios.
+		// We also assume that GC doesn't run during transactions and that the UStruct will remain valid until the commit.
+		AutoRTFM::OnCommit([DedupSchema]() { DedupSchema(); });
+
+		// For now return our created schema and we'll dedup it on commit.
+		return *CreatedSchema;
 	}
 
-	return *CreatedSchema;
+	return *DedupSchema();
 }
 
 using FUnversionedSchemaIterator = const FUnversionedPropertySerializer*;
@@ -928,10 +948,10 @@ bool CanUseUnversionedPropertySerialization(const ITargetPlatform* Target)
 void DestroyUnversionedSchema(const UStruct* Struct)
 {
 #if CACHE_UNVERSIONED_PROPERTY_SCHEMA
-	FUnversionedStructSchema::Delete(const_cast<FUnversionedStructSchema*>(Struct->UnversionedGameSchema));
+	FUnversionedStructSchema::Delete(const_cast<FUnversionedStructSchema*>(Struct->UnversionedGameSchema.load(std::memory_order_acquire)));
 	Struct->UnversionedGameSchema = nullptr;
 #if WITH_EDITORONLY_DATA
-	FUnversionedStructSchema::Delete(const_cast<FUnversionedStructSchema*>(Struct->UnversionedEditorSchema));
+	FUnversionedStructSchema::Delete(const_cast<FUnversionedStructSchema*>(Struct->UnversionedEditorSchema.load(std::memory_order_acquire)));
 	Struct->UnversionedEditorSchema = nullptr;
 #endif
 #endif
diff --git a/Engine/Source/Runtime/CoreUObject/Public/UObject/Class.h b/Engine/Source/Runtime/CoreUObject/Public/UObject/Class.h
--- a/Engine/Source/Runtime/CoreUObject/Public/UObject/Class.h
+++ b/Engine/Source/Runtime/CoreUObject/Public/UObject/Class.h
@@ -73,6 +73,7 @@
 #include "UObject/UnrealNames.h"
 #include "UObject/ObjectPtr.h"
 #include <type_traits>
+#include <atomic>
 
 class FArchive;
 class FEditPropertyChain;
@@ -449,10 +450,10 @@ public:
 #endif
 
 	/** Cached schema for optimized unversioned and filtereditoronly property serialization, owned by this. */
-	mutable const struct FUnversionedStructSchema* UnversionedGameSchema = nullptr;
+	mutable std::atomic<const struct FUnversionedStructSchema*> UnversionedGameSchema = nullptr;
 #if WITH_EDITORONLY_DATA
 	/** Cached schema for optimized unversioned property serialization, with editor data, owned by this. */
-	mutable const struct FUnversionedStructSchema* UnversionedEditorSchema = nullptr;
+	mutable std::atomic<const struct FUnversionedStructSchema*> UnversionedEditorSchema = nullptr;
 
 	/** Get the Schema Hash for this struct - the hash of its property names and types. */
 	COREUOBJECT_API const FBlake3Hash& GetSchemaHash(bool bSkipEditorOnly) const;
