Add new archive class FArchiveSavePackageCollector. It is used as a base class for FHarvestPackage and FPostSaveArchive. By default it is set with the same flag as FHarvestPackage.
Add new archive class FPostSaveArchive. It's not used yet. It will be used later on when implementing the PostSave step. It disables delta serialization to make sure objects will be serialized and not skipped. Apart from that it has the same flag as FHarvestPackage.

#rnx

diff --git a/Engine/Source/Runtime/CoreUObject/Private/UObject/SavePackage/PackageHarvester.cpp b/Engine/Source/Runtime/CoreUObject/Private/UObject/SavePackage/PackageHarvester.cpp
--- a/Engine/Source/Runtime/CoreUObject/Private/UObject/SavePackage/PackageHarvester.cpp
+++ b/Engine/Source/Runtime/CoreUObject/Private/UObject/SavePackage/PackageHarvester.cpp
@@ -193,6 +193,26 @@ bool DoesObjectNeedLoadForEditorGame(UObject* InObject)
 	return bNeedsLoadForEditorGame;
 }
 
+FArchiveSavePackageCollector::FArchiveSavePackageCollector(FSaveContext& SaveContext)
+	: FArchiveUObject()
+{
+	SetArchiveFlags(SaveContext);
+}
+
+void FArchiveSavePackageCollector::SetArchiveFlags(FSaveContext& SaveContext)
+{
+	this->SetIsSaving(true);
+	this->SetIsPersistent(true);
+	ArIsObjectReferenceCollector = true;
+	ArShouldSkipBulkData = true;
+	ArIgnoreClassGeneratedByRef = SaveContext.IsCooking();
+
+	this->SetPortFlags(SaveContext.GetPortFlags());
+	this->SetFilterEditorOnly(SaveContext.IsFilterEditorOnly());
+	this->SetSavePackageData(&SaveContext.GetArchiveSavePackageData());
+	this->SetUseUnversionedPropertySerialization(SaveContext.IsSaveUnversionedProperties());
+}
+
 FPackageHarvester::FExportScope::FExportScope(FPackageHarvester& InHarvester, const FExportWithContext& InToProcess)
 	: Harvester(InHarvester)
 	, PreviousExportHarvestingRealms(InHarvester.CurrentExportHarvestingRealms)
@@ -436,19 +456,9 @@ bool FPackageHarvester::HasAnyExportHarvestingRealms()
 }
 
 FPackageHarvester::FPackageHarvester(FSaveContext& InContext)
-	: SaveContext(InContext)
+	: FArchiveSavePackageCollector(InContext)
+	, SaveContext(InContext)
 {
-	this->SetIsSaving(true);
-	this->SetIsPersistent(true);
-	ArIsObjectReferenceCollector = true;
-	ArShouldSkipBulkData = true;
-	ArIgnoreClassGeneratedByRef = SaveContext.IsCooking();
-
-	this->SetPortFlags(SaveContext.GetPortFlags());
-	this->SetFilterEditorOnly(SaveContext.IsFilterEditorOnly());
-	this->SetSavePackageData(&SaveContext.GetArchiveSavePackageData());
-	this->SetUseUnversionedPropertySerialization(SaveContext.IsSaveUnversionedProperties());
-
 	ResolveOverrides();
 	// Clear the SaveContext's Saveable cache. It was used earlier during RoutePreSave, and the PreSave functions
 	// may have modified RF_Transient or other flags on UObjects and therefore invalidated the cached result. It
diff --git a/Engine/Source/Runtime/CoreUObject/Private/UObject/SavePackage/PackageHarvester.h b/Engine/Source/Runtime/CoreUObject/Private/UObject/SavePackage/PackageHarvester.h
--- a/Engine/Source/Runtime/CoreUObject/Private/UObject/SavePackage/PackageHarvester.h
+++ b/Engine/Source/Runtime/CoreUObject/Private/UObject/SavePackage/PackageHarvester.h
@@ -19,11 +19,20 @@ struct FWeakObjectPtr;
 struct FLazyObjectPtr;
 struct FSoftObjectPath;
 
+/* Base archive class for FPackageHarvester and FPostSaveArchive used to setup all the common flags. */
+class FArchiveSavePackageCollector : public FArchiveUObject
+{
+public:
+	FArchiveSavePackageCollector(FSaveContext& SaveContext);
+
+	void SetArchiveFlags(FSaveContext& SaveContext);
+};
+
 /**
  * Helper class used to gather all package exports, imports and dependencies to build linker tables when saving packages
  * Gather Export, Imports, Referenced Names, SoftObjectPaths, Custom Object Versions
  */
-class FPackageHarvester : public FArchiveUObject
+class FPackageHarvester : public FArchiveSavePackageCollector
 {
 public:
 	using FExportingRealmsArray = TArray<ESaveRealm, TInlineAllocator<2>>;
diff --git a/Engine/Source/Runtime/CoreUObject/Private/UObject/SavePackage2.cpp b/Engine/Source/Runtime/CoreUObject/Private/UObject/SavePackage2.cpp
--- a/Engine/Source/Runtime/CoreUObject/Private/UObject/SavePackage2.cpp
+++ b/Engine/Source/Runtime/CoreUObject/Private/UObject/SavePackage2.cpp
@@ -2749,6 +2749,22 @@ ESavePackageResult BeginCachePlatformCookedData(FSaveContext& SaveContext)
 #endif
 }
 
+/** 
+ * Archive used during post save serialization.
+ * It ignores serialization of primitive types. It is just used to provide a callback hook for UStruct::Serialize and UObject::Serialize to execute custom code.
+ */
+class FPostSaveArchive : public FArchiveSavePackageCollector
+{
+public:
+	FPostSaveArchive(FSaveContext& InContext)
+		: FArchiveSavePackageCollector(InContext)
+	{
+		//disable the delta serialization to make sure the objects are serialized no matter what.
+		ArNoDelta = true;
+		ArNoIntraPropertyDelta = true;
+	}
+};
+
 void PostSavePackage(FSaveContext& SaveContext)
 {
 	UPackage* Package = SaveContext.GetPackage();
