Add new PostSave serialization feature. When an object is being serialized in the harvest phase, it can request to be serialized again during the PostSave phase.
This request only works during the harvest phase of a cook. Otherwise the request is ignored.

FSaveContext : store the objects which requested a PostSaveSerialization.
FObjectSavePackageSerializeContext : new function to make a post save serialization request.
FObjectSaveContextData : new flag used to request a post save serialization. When this flag is true, the object will be stored in the list of objects to serialize during the post save phase.

FPackageHarvester::ProcessExport : save the current object in the list of post save serialization objects if requested.

PostSavePackage : Call serialization on the objects that requested it.

#rnx

diff --git a/Engine/Source/Runtime/CoreUObject/Private/UObject/SavePackage/PackageHarvester.cpp b/Engine/Source/Runtime/CoreUObject/Private/UObject/SavePackage/PackageHarvester.cpp
--- a/Engine/Source/Runtime/CoreUObject/Private/UObject/SavePackage/PackageHarvester.cpp
+++ b/Engine/Source/Runtime/CoreUObject/Private/UObject/SavePackage/PackageHarvester.cpp
@@ -569,6 +569,16 @@ void FPackageHarvester::ProcessExport(const FExportWithContext& InProcessContext
 		Export->Serialize(*this);
 	}
 
+	{
+		// Register the object for the post save serialization if requested
+		FObjectSaveContextData& ObjectSaveContext = SaveContext.GetObjectSaveContext();
+		if (ObjectSaveContext.bRequestPostSaveSerialization)
+		{
+			SaveContext.AddObjectToPostSaveSerialization(Export);
+			ObjectSaveContext.bRequestPostSaveSerialization = false;
+		}
+	}
+
 	// Gather object preload dependencies
 	if (SaveContext.IsCooking())
 	{
diff --git a/Engine/Source/Runtime/CoreUObject/Private/UObject/SavePackage/SaveContext.cpp b/Engine/Source/Runtime/CoreUObject/Private/UObject/SavePackage/SaveContext.cpp
--- a/Engine/Source/Runtime/CoreUObject/Private/UObject/SavePackage/SaveContext.cpp
+++ b/Engine/Source/Runtime/CoreUObject/Private/UObject/SavePackage/SaveContext.cpp
@@ -262,6 +262,16 @@ UE::SavePackageUtilities::EEditorOnlyObjectFlags FSaveContext::GetEditorOnlyObje
 			: EEditorOnlyObjectFlags::None);
 }
 
+void FSaveContext::AddObjectToPostSaveSerialization(UObject* Object)
+{
+	PostSaveObjectsToSerialize.Add(Object);
+}
+
+const TSet<UObject*>& FSaveContext::GetPostSaveObjectsToSerialize() const
+{
+	return PostSaveObjectsToSerialize;
+}
+
 namespace
 {
 	TArray<UClass*> AutomaticOptionalInclusionAssetTypeList;
diff --git a/Engine/Source/Runtime/CoreUObject/Private/UObject/SavePackage/SaveContext.h b/Engine/Source/Runtime/CoreUObject/Private/UObject/SavePackage/SaveContext.h
--- a/Engine/Source/Runtime/CoreUObject/Private/UObject/SavePackage/SaveContext.h
+++ b/Engine/Source/Runtime/CoreUObject/Private/UObject/SavePackage/SaveContext.h
@@ -1225,6 +1225,10 @@ public:
 #endif
 	UE::SavePackageUtilities::EEditorOnlyObjectFlags GetEditorOnlyObjectFlags() const;
 
+	void AddObjectToPostSaveSerialization(UObject* Object);
+
+	const TSet<UObject*>& GetPostSaveObjectsToSerialize() const;
+
 public:
 	ESavePackageResult Result;
 
@@ -1312,6 +1316,9 @@ private:
 	TArray<FName> PackageBuildDependencies;
 	// Empty list of FNames, used for functions that need to return a reference to an empty array.
 	TArray<FName> EmptyList;
+
+	// List of objects to serialize again during the post save phase.
+	TSet<UObject*> PostSaveObjectsToSerialize;
 };
 
 const TCHAR* LexToString(ESaveableStatus Status);
diff --git a/Engine/Source/Runtime/CoreUObject/Private/UObject/SavePackage/SavePackageUtilities.cpp b/Engine/Source/Runtime/CoreUObject/Private/UObject/SavePackage/SavePackageUtilities.cpp
--- a/Engine/Source/Runtime/CoreUObject/Private/UObject/SavePackage/SavePackageUtilities.cpp
+++ b/Engine/Source/Runtime/CoreUObject/Private/UObject/SavePackage/SavePackageUtilities.cpp
@@ -890,6 +890,18 @@ void FObjectSavePackageSerializeContext::RegisterDeterminismHelper(
 	}
 }
 
+void FObjectSavePackageSerializeContext::RequestPostSaveSerialization()
+{
+	if (GetPhase() != EObjectSaveContextPhase::Harvest || !IsCooking())
+	{
+		UE_LOG(LogSavePackage, Error,
+			TEXT("RequestPostSaveSerialization called when GetPhase() != EObjectSaveContextPhase::Harvest or !IsCooking(). This is invalid and will be ignored."));
+		FDebug::DumpStackTraceToLog(ELogVerbosity::Error);
+		return;
+	}
+	Data.bRequestPostSaveSerialization = true;
+}
+
 #endif
 
 void FObjectSaveContextData::Set(UPackage* Package, const ITargetPlatform* InTargetPlatform, const TCHAR* InTargetFilename, uint32 InSaveFlags)
diff --git a/Engine/Source/Runtime/CoreUObject/Private/UObject/SavePackage2.cpp b/Engine/Source/Runtime/CoreUObject/Private/UObject/SavePackage2.cpp
--- a/Engine/Source/Runtime/CoreUObject/Private/UObject/SavePackage2.cpp
+++ b/Engine/Source/Runtime/CoreUObject/Private/UObject/SavePackage2.cpp
@@ -2771,7 +2771,19 @@ void PostSavePackage(FSaveContext& SaveContext)
 	// restore initial packages flags since the save can currently mutate the package flags in an undesirable fashion, (i.e. clearing/adding editor only filtering)
 	Package->SetPackageFlagsTo(SaveContext.GetInitialPackageFlags());
 
-	// Do not run the rest of post save and send out events if the save wasn't succesful
+	// Call PostSave serialization
+	{
+		FPostSaveArchive PostSaveArchive(SaveContext);
+
+		const TSet<UObject*>& ObjectsToSerialize = SaveContext.GetPostSaveObjectsToSerialize();
+		for (UObject* Object : ObjectsToSerialize)
+		{
+			Object->Serialize(PostSaveArchive);
+		}
+	}
+
+	// If the save wasn't successful, do not run operations that update the in-memory package to match the on-disk package, and do not broadcast
+	// that the package was saved.
 	if (SaveContext.Result != ESavePackageResult::Success)
 	{
 		return;
diff --git a/Engine/Source/Runtime/CoreUObject/Public/UObject/ObjectSaveContext.h b/Engine/Source/Runtime/CoreUObject/Public/UObject/ObjectSaveContext.h
--- a/Engine/Source/Runtime/CoreUObject/Public/UObject/ObjectSaveContext.h
+++ b/Engine/Source/Runtime/CoreUObject/Public/UObject/ObjectSaveContext.h
@@ -160,6 +160,9 @@ struct FObjectSaveContextData
 	/** Call-site enforcement; records whether the base PreSave was called. */
 	bool bBaseClassCalled = false;
 
+	/** Set to true when the current object being serialized needs to call serialize again in PostSave phase. */
+	bool bRequestPostSaveSerialization = false;
+
 };
 
 /** Interface used by CollectSaveOverrides to access the save parameters. */
@@ -381,6 +384,8 @@ public:
 	 * their callback class for adding determinism diagnostics the package save.
 	 */
 	COREUOBJECT_API void RegisterDeterminismHelper(const TRefCountPtr<UE::Cook::IDeterminismHelper>& DeterminismHelper);
+	/** Called during the Harvest phase when the current object need to call Serialize again during the PostSave phase. */
+	COREUOBJECT_API void RequestPostSaveSerialization();
 #endif
 
 	/**
