Adding FRemoteObjectId UObject hash table and removing Id to UObject map from internal remote object maps.

Allows FWeakObjectPtr to consistently resolve remote and local objects with remote id.

diff --git a/Engine/Source/Runtime/CoreUObject/Private/UObject/Obj.cpp b/Engine/Source/Runtime/CoreUObject/Private/UObject/Obj.cpp
--- a/Engine/Source/Runtime/CoreUObject/Private/UObject/Obj.cpp
+++ b/Engine/Source/Runtime/CoreUObject/Private/UObject/Obj.cpp
@@ -1264,10 +1264,6 @@ bool UObject::ConditionalFinishDestroy()
 #endif
 		FinishDestroy();
 
-#if UE_WITH_REMOTE_OBJECT_HANDLE
-		UE::RemoteObject::Private::UnmapObject(this);
-#endif
-
 		// Make sure this object can't be accessed via weak pointers after it's been FinishDestroyed
 		GUObjectArray.ResetSerialNumber(this);
 
diff --git a/Engine/Source/Runtime/CoreUObject/Private/UObject/RemoteObject.cpp b/Engine/Source/Runtime/CoreUObject/Private/UObject/RemoteObject.cpp
--- a/Engine/Source/Runtime/CoreUObject/Private/UObject/RemoteObject.cpp
+++ b/Engine/Source/Runtime/CoreUObject/Private/UObject/RemoteObject.cpp
@@ -12,6 +12,7 @@
 #include "UObject/RemoteObjectSerialization.h"
 #include "UObject/RemoteObjectPrivate.h"
 #include "UObject/ObjectHandlePrivate.h"
+#include "UObject/UObjectHash.h"
 #include "Misc/CommandLine.h"
 #include "Misc/CoreDelegates.h"
 #include <atomic>
@@ -29,50 +30,9 @@ class FRemoteObjectMaps
 {
 	mutable FTransactionallySafeCriticalSection ObjectMapCritical;
 	TMap<FRemoteObjectId, UE::RemoteObject::Handle::FRemoteObjectStub*> RemoteObjects;
-	TMap<FRemoteObjectId, UObject*> IdToObjectMap;
 
 public:
 
-	bool IsMapped(FRemoteObjectId ObjectId) const
-	{
-		UE::TScopeLock MapLock(ObjectMapCritical);
-		return IdToObjectMap.Contains(ObjectId);
-	}
-
-	void MapObject(UObject* Object)
-	{
-		FRemoteObjectId Id = UE::CoreUObject::Private::FObjectHandleUtils::GetRemoteId(Object);
-		{
-			UE::TScopeLock MapLock(ObjectMapCritical);
-			UObject*& MappedObject = IdToObjectMap.FindOrAdd(Id);
-			MappedObject = Object;
-		}
-	}
-
-	void UnmapObject(UObject* Object)
-	{
-		FRemoteObjectId Id = UE::CoreUObject::Private::FObjectHandleUtils::GetRemoteId(Object);
-		{
-			UE::TScopeLock MapLock(ObjectMapCritical);
-
-			// only erase the entry from the map if the pointer matches the object
-			// we're trying to unmap, since the remote id may since have been mapped
-			// to a new object, and we don't want to remove that still-valid mapping
-			auto Iter = IdToObjectMap.CreateKeyIterator(Id);
-			if (Iter && (Iter.Value() == Object))
-			{
-				Iter.RemoveCurrent();
-			}
-		}
-	}
-
-	UObject* FindObject(FRemoteObjectId Id)
-	{
-		UE::TScopeLock MapLock(ObjectMapCritical);
-		UObject** MappedObject = IdToObjectMap.Find(Id);
-		return MappedObject ? *MappedObject : nullptr;
-	}
-
 	UE::RemoteObject::Handle::FRemoteObjectStub* FindRemoteObjectStub(FRemoteObjectId Id)
 	{
 		using namespace UE::RemoteObject::Handle;
@@ -164,7 +124,6 @@ void RegisterRemoteObjectId(FRemoteObjectId ObjectId, FRemoteServerId OwnerServe
 
 void RegisterSharedObject(UObject* Object)
 {
-	MapObject(Object);
 	MarkAsLocal(Object);
 }
 
@@ -184,22 +143,6 @@ void MarkAsLocal(UObject* Object)
 	ObjectItem->SetFlags(EInternalObjectFlags::RemoteReference);
 }
 
-void MapObject(UObject* Object)
-{
-	ObjectMaps->MapObject(Object);
-}
-
-void UnmapObject(UObject* Object)
-{
-	ObjectMaps->UnmapObject(Object);
-}
-
-UObject* FindObject(FRemoteObjectId ObjectId)
-{
-	UObject* Object = ObjectMaps->FindObject(ObjectId);
-	return Object;
-}
-
 UE::RemoteObject::Handle::FRemoteObjectStub* FindRemoteObjectStub(FRemoteObjectId ObjectId)
 {
 	return ObjectMaps->FindRemoteObjectStub(ObjectId);
@@ -237,7 +180,7 @@ bool IsRemote(FRemoteObjectId ObjectId)
 	{
 		return false;
 	}
-	if (UObject* Object = ObjectMaps->FindObject(ObjectId))
+	if (UObject* Object = StaticFindObjectFastInternal(ObjectId))
 	{
 		return IsRemote(Object);
 	}
@@ -263,7 +206,7 @@ UObject* ResolveObject(const FRemoteObjectStub* Stub, ERemoteReferenceType RefTy
 	using namespace UE::RemoteObject::Private;
 
 	// This is a slightly faster version of IsRemote(FRemoteObjectId) because we already know a Stub exists and we are going to re-use the Object pointer
-	UObject* Object = ObjectMaps->FindObject(Stub->Id);
+	UObject* Object = StaticFindObjectFastInternal(Stub->Id);
 	if (GIsGarbageCollecting && (Object || RefType == ERemoteReferenceType::Weak)) // Not a thread-safe test
 	{
 		// Begin/FinishDestroy overrides may attempt to access subobjects of objects that have been migrated in which case we
@@ -278,14 +221,14 @@ UObject* ResolveObject(const FRemoteObjectStub* Stub, ERemoteReferenceType RefTy
 	{
 		checkf(!GIsGarbageCollecting, TEXT("Resolving remote objects while collecting garbage is not allowed (trying to resolve object %s (%s)"), *Stub->Id.ToString(), *Stub->Name.ToString());
 
-		UObject* Outer = Object ? Object->GetOuter() : ObjectMaps->FindObject(Stub->OuterId);
+		UObject* Outer = Object ? Object->GetOuter() : StaticFindObjectFastInternal(Stub->OuterId);
 
 		MigrateObjectFromRemoteServer(Stub->Id, Stub->RemoteServerId, Outer);
 
 		// if running transactionally, we will have aborted and not reached here
 
 		// if running non-transactionally, object migrated immediately, so we just re-resolve
-		Object = ObjectMaps->FindObject(Stub->Id);
+		Object = StaticFindObjectFastInternal(Stub->Id);
 		bRemoteObject = !Object || IsRemote(Object);
 		checkf(!bRemoteObject, TEXT("Failed to resolve remote object %s, either this code is not running in a transaction and should be, or the transaction failed to abort"), *Stub->Id.ToString());
 	}
@@ -298,7 +241,7 @@ bool RequestObject(const FRemoteObjectStub* Stub)
 	using namespace UE::RemoteObject::Transfer;
 	using namespace UE::RemoteObject::Private;
 
-	UObject* Object = ObjectMaps->FindObject(Stub->Id);
+	UObject* Object = StaticFindObjectFastInternal(Stub->Id);
 	bool bRemoteObject = false;
 	if (Object)
 	{
@@ -355,7 +298,7 @@ bool CanResolveObject(FRemoteObjectId ObjectId)
 
 	// Note: this function needs to reflect the logic of ResolveObject(...) functions
 	
-	if (UObject* Object = ObjectMaps->FindObject(ObjectId))
+	if (UObject* Object = StaticFindObjectFastInternal(ObjectId))
 	{
 		// Object memory is local and even if it's already been migrated we can resolve it
 		return true;
@@ -380,13 +323,13 @@ namespace UE::CoreUObject::Private
 	void FObjectHandleUtils::ChangeRemoteId(UObjectBase* Object, FRemoteObjectId Id)
 	{
 		using namespace UE::RemoteObject::Private;
-		UnmapObject((UObject*)Object);
+		UnhashObject(Object);
 		FUObjectItem* ObjectItem = GUObjectArray.ObjectToObjectItem(Object);
 		// ObjectItem may not exist when the UObject system hasn't been initialized yet but theoretically this function should only get called when
 		// something attempts to re-construct a default subobject that already exists so ObjectItem should always be valid
 		checkf(ObjectItem, TEXT("Attempting to change remote ID for an object that does not exist in the global UObject array (it's possible GUObjectArray is not initialized yet, ObjectIndex=%d)"), GUObjectArray.ObjectToIndex(Object));
 		ObjectItem->RemoteId = Id;
-		MapObject((UObject*)Object);
+		HashObject(Object);
 	}
 
 	FRemoteObjectId FRemoteObjectHandlePrivate::GetRemoteId() const
@@ -416,7 +359,7 @@ namespace UE::CoreUObject::Private
 		UObject* Obj = nullptr;
 		if (ObjectId.IsValid())
 		{
-			Obj = ObjectMaps->FindObject(ObjectId);
+			Obj = StaticFindObjectFastInternal(ObjectId);
 			if (Obj && !Obj->HasAnyInternalFlags(EInternalObjectFlags::Remote))
 			{
 				return FRemoteObjectHandlePrivate(Obj);
diff --git a/Engine/Source/Runtime/CoreUObject/Private/UObject/RemoteObjectPrivate.h b/Engine/Source/Runtime/CoreUObject/Private/UObject/RemoteObjectPrivate.h
--- a/Engine/Source/Runtime/CoreUObject/Private/UObject/RemoteObjectPrivate.h
+++ b/Engine/Source/Runtime/CoreUObject/Private/UObject/RemoteObjectPrivate.h
@@ -28,22 +28,6 @@ namespace UE::RemoteObject::Private
 	*/
 	void MarkAsLocal(UObject* Object);
 
-	/**
-	* Maps an object with its id
-	*/
-	void MapObject(UObject* Object);
-
-	/**
-	* Unmaps object memory from its id
-	*/
-	void UnmapObject(UObject* Object);
-
-	/**
-	* Finds a local object by its unique id
-	* @returns Object if it exists in the current processes (and has been mapped with its id)
-	*/
-	UObject* FindObject(FRemoteObjectId ObjectId);
-
 	/**
 	* Registers a stub for a remote object that is known to be owned by a specific server
 	*/
diff --git a/Engine/Source/Runtime/CoreUObject/Private/UObject/RemoteObjectSerialization.cpp b/Engine/Source/Runtime/CoreUObject/Private/UObject/RemoteObjectSerialization.cpp
--- a/Engine/Source/Runtime/CoreUObject/Private/UObject/RemoteObjectSerialization.cpp
+++ b/Engine/Source/Runtime/CoreUObject/Private/UObject/RemoteObjectSerialization.cpp
@@ -201,15 +201,6 @@ class FArchiveRemoteObjectWriter : public FMemoryWriter
 			FNameIndexType PathNameIndex = AddPathNameToNameMap(Obj);
 			*this << PathNameIndex;
 		}
-
-		if (Obj)
-		{
-			// Map all references we encounter during serialization, especially if they're objects that are not going to be serialized
-			// because we may end up needing to resolve them later if a remote server requests them with their unique id.
-			// Note: we could just hash all objects with their FRemoteObjectId in UObjectHash.cpp but for now we can limit the memory impact by hashing 
-			// only those we encounter when serializing objects that are actually going to be migrated
-			UE::RemoteObject::Private::MapObject(Obj);
-		}
 	}
 
 public:
@@ -438,10 +429,6 @@ void SerializeRemoteObject(UObject* Object, FArchive& Ar, FRemoteObjectHeader& O
 	OutHeader.Outer = Outer;
 	OutHeader.Archetype = Object->GetArchetype();
 
-	MapObject(OutHeader.Class.Get());
-	MapObject(Outer);
-	MapObject(OutHeader.Archetype.Get());
-
 	Object->Serialize(Ar);
 }
 
@@ -652,7 +639,6 @@ UObject* ConstructRemoteObject(const FRemoteObjectHeader& Header, FRemoteObjectC
 
 	checkf(FRemoteObjectId(Object) == Header.RemoteId, TEXT("Created an object with a different ID:%s than requested:%s"), *FRemoteObjectId(Object).ToString(), *Header.RemoteId.ToString());
 
-	MapObject(Object);
 	MarkAsLocal(Object);
 	// It's possible that one of the migrated objects was still in memory on this server and it was marked as Garbage in PostMigrate so clear the flag in case we migrated it back
 	Object->ClearGarbage();
diff --git a/Engine/Source/Runtime/CoreUObject/Private/UObject/RemoteObjectTransfer.cpp b/Engine/Source/Runtime/CoreUObject/Private/UObject/RemoteObjectTransfer.cpp
--- a/Engine/Source/Runtime/CoreUObject/Private/UObject/RemoteObjectTransfer.cpp
+++ b/Engine/Source/Runtime/CoreUObject/Private/UObject/RemoteObjectTransfer.cpp
@@ -11,6 +11,7 @@
 #include "UObject/RemoteObject.h"
 #include "UObject/RemoteObjectPrivate.h"
 #include "UObject/ObjectHandlePrivate.h"
+#include "UObject/UObjectHash.h"
 #include "HAL/CriticalSection.h"
 #include "AutoRTFM/AutoRTFM.h"
 #include "Templates/Atomic.h"
@@ -364,7 +365,7 @@ void MigrateObjectToRemoteServer(uint32 RequestToken, FRemoteObjectId ObjectId,
 
 	UE_LOG(LogRemoteObject, Verbose, TEXT("MigrateObjectToRemoteServer(%u) obj %s to %s"), RequestToken, *ObjectId.ToString(), *DestinationServerId.ToString());
 
-	UObject* Object = FindObject(ObjectId);
+	UObject* Object = StaticFindObjectFastInternal(ObjectId);
 	if (Object && !IsRemote(Object))
 	{
 		UE_LOG(LogRemoteObject, VeryVerbose, TEXT("    : MigrateObjectToRemoteServer(%u) found obj %s"), RequestToken, *ObjectId.ToString());
diff --git a/Engine/Source/Runtime/CoreUObject/Private/UObject/UObjectHash.cpp b/Engine/Source/Runtime/CoreUObject/Private/UObject/UObjectHash.cpp
--- a/Engine/Source/Runtime/CoreUObject/Private/UObject/UObjectHash.cpp
+++ b/Engine/Source/Runtime/CoreUObject/Private/UObject/UObjectHash.cpp
@@ -403,6 +403,10 @@ public:
 	/** Map of object to their external package. */
 	TMap<UObjectBase*, UPackage*> ObjectToPackageMap;
 
+#if UE_WITH_REMOTE_OBJECT_HANDLE
+	TBucketMap<uint32> HashId;
+#endif
+
 	FUObjectHashTables()
 		: AllClassesVersion(0)
 		, NativeClassesVersion(0)
@@ -413,7 +417,8 @@ public:
 	{
 		const EParallelForFlags BaseFlags = FTaskGraphInterface::IsRunning() ? EParallelForFlags::None : EParallelForFlags::ForceSingleThread;
 		double StartTime = FPlatformTime::Seconds();
-		ParallelFor(7,
+		constexpr int32 NumHashTables = UE_WITH_REMOTE_OBJECT_HANDLE ? 8 : 7;
+		ParallelFor(NumHashTables,
 			[this](int32 Index)
 		{
 			switch (Index)
@@ -459,6 +464,11 @@ public:
 			case 6:
 				ObjectToPackageMap.Compact();
 				break;
+#if UE_WITH_REMOTE_OBJECT_HANDLE
+			case 7:
+				HashId.Compact();
+				break;
+#endif
 			}
 		}, BaseFlags | EParallelForFlags::Unbalanced);
 		UE_LOG(LogUObjectHash, Log, TEXT("Compacting FUObjectHashTables data took %6.2fms"), 1000.0f * float(FPlatformTime::Seconds() - StartTime));
@@ -989,6 +999,58 @@ UObject* StaticFindFirstObjectFastInternal(const UClass* ObjectClass, FName Obje
 	return Result;
 }
 
+UObject* StaticFindObjectFastInternalThreadSafe(FUObjectHashTables& ThreadHash, FRemoteObjectId RemoteId, EObjectFlags InExclusiveFlags, EInternalObjectFlags InExlusiveInternalFlags)
+{
+#if UE_WITH_REMOTE_OBJECT_HANDLE
+	const EInternalObjectFlags ExclusiveInternalFlags = EInternalObjectFlags::Unreachable | InExlusiveInternalFlags;
+	uint32 Hash = GetTypeHash(RemoteId);
+
+	FHashTableLock HashLock(ThreadHash);
+	FHashBucket* Bucket = ThreadHash.HashId.Find(Hash);
+	if (Bucket)
+	{
+		for (FHashBucketIterator It(*Bucket); It; ++It)
+		{
+			UObject* Object = (UObject*)*It;
+			if ((FRemoteObjectId(Object) == RemoteId)
+				/* Don't return objects that have any of the exclusive flags set */
+				&& !Object->HasAnyFlags(InExclusiveFlags)
+				&& !Object->HasAnyInternalFlags(ExclusiveInternalFlags)
+				)
+			{
+				checkf(!Object->IsUnreachable(), TEXT("%s"), *Object->GetFullName());
+				if (UE::GC::GIsIncrementalReachabilityPending)
+				{
+					UE::GC::MarkAsReachable(Object);
+				}
+				return Object;
+			}
+		}
+	}
+#else
+	UE_LOG(LogUObjectHash, Fatal, TEXT("StaticFindObjectFastInternal override that takes FRemoteObjectId can only be used with remote object handles enabled"));
+#endif // UE_WITH_REMOTE_OBJECT_HANDLE
+
+	return nullptr;
+}
+
+UObject* StaticFindObjectFastInternal(FRemoteObjectId RemoteId, EObjectFlags InExclusiveFlags, EInternalObjectFlags InExlusiveInternalFlags)
+{
+	UObject* Result = nullptr;
+
+	// Transactionally, a static find operation has to occur in the open as it touches shared state.
+	UE_AUTORTFM_OPEN
+	{
+		INC_DWORD_STAT(STAT_FindObjectFast);
+
+		FUObjectHashTables& ThreadHash = FUObjectHashTables::Get();
+		Result = StaticFindObjectFastInternalThreadSafe(ThreadHash, RemoteId, InExclusiveFlags, InExlusiveInternalFlags);
+	};
+
+	return Result;
+}
+
+
 // Assumes that ThreadHash's critical is already locked
 FORCEINLINE static void AddToOuterMap(FUObjectHashTables& ThreadHash, UObjectBase* Object)
 {
@@ -1046,6 +1108,15 @@ FORCEINLINE static UPackage* AssignExternalPackageToObject(FUObjectHashTables& T
 	return OldPackage;
 }
 
+#if UE_WITH_REMOTE_OBJECT_HANDLE
+// Assumes that ThreadHash's critical is already locked
+FORCEINLINE static void AddToRemoteIdMap(FUObjectHashTables& ThreadHash, UObjectBase* Object)
+{
+	FHashBucket& Bucket = ThreadHash.HashId.FindOrAdd(GetTypeHash(FRemoteObjectId(Object)));
+	checkSlow(!Bucket.Contains(Object)); // if it already exists, something is wrong with the external code
+	Bucket.Add(Object);
+}
+#endif // UE_WITH_REMOTE_OBJECT_HANDLE
 
 // Assumes that ThreadHash's critical is already locked
 FORCEINLINE static void RemoveFromOuterMap(FUObjectHashTables& ThreadHash, UObjectBase* Object)
@@ -1145,6 +1216,26 @@ FORCEINLINE static UPackage* UnassignExternalPackageFromObject(FUObjectHashTable
 	return OldPackage;
 }
 
+#if UE_WITH_REMOTE_OBJECT_HANDLE
+// Assumes that ThreadHash's critical is already locked
+FORCEINLINE static void RemoveFromRemoteIdMap(FUObjectHashTables& ThreadHash, UObjectBase* Object)
+{
+	const uint32 Hash = GetTypeHash(FRemoteObjectId(Object));
+	FHashBucket& Bucket = ThreadHash.HashId.FindOrAdd(Hash);
+	int32 NumRemoved = Bucket.Remove(Object);
+
+	if (!LIKELY(NumRemoved == 1))
+	{
+		OnHashFailure((UObjectBaseUtility*)Object, TEXT("HashId"), TEXT("remove miscount"));
+	}
+
+	if (!Bucket.Num())
+	{
+		ThreadHash.HashId.Remove(Hash);
+	}
+}
+#endif // UE_WITH_REMOTE_OBJECT_HANDLE
+
 void ShrinkUObjectHashTables()
 {
 	LLM_SCOPE_BYTAG(UObjectHash);
@@ -1619,6 +1710,10 @@ void HashObject(UObjectBase* Object)
 		}
 
 		AddToClassMap( ThreadHash, Object );
+
+#if UE_WITH_REMOTE_OBJECT_HANDLE
+		AddToRemoteIdMap(ThreadHash, Object);
+#endif
 	}
 }
 
@@ -1667,6 +1762,10 @@ void UnhashObject(UObjectBase* Object)
 		}
 
 		RemoveFromClassMap( ThreadHash, Object );
+
+#if UE_WITH_REMOTE_OBJECT_HANDLE
+		RemoveFromRemoteIdMap(ThreadHash, Object);
+#endif
 	}
 }
 
diff --git a/Engine/Source/Runtime/CoreUObject/Private/UObject/WeakObjectPtr.cpp b/Engine/Source/Runtime/CoreUObject/Private/UObject/WeakObjectPtr.cpp
--- a/Engine/Source/Runtime/CoreUObject/Private/UObject/WeakObjectPtr.cpp
+++ b/Engine/Source/Runtime/CoreUObject/Private/UObject/WeakObjectPtr.cpp
@@ -10,6 +10,7 @@
 #include "UObject/ObjectPtr.h"
 #include "UObject/RemoteObject.h"
 #include "UObject/RemoteObjectPrivate.h"
+#include "UObject/UObjectHash.h"
 
 DEFINE_LOG_CATEGORY_STATIC(LogWeakObjectPtr, Log, All);
 
@@ -184,7 +185,7 @@ FUObjectItem* FWeakObjectPtr::TryResolveRemoteObject() const
 	using namespace UE::RemoteObject::Handle;
 	using namespace UE::RemoteObject::Private;
 
-	UObject* ResolvedObject = FindObject(ObjectRemoteId);
+	UObject* ResolvedObject = StaticFindObjectFastInternal(ObjectRemoteId);
 
 	if (ResolvedObject)
 	{		
diff --git a/Engine/Source/Runtime/CoreUObject/Public/UObject/UObjectHash.h b/Engine/Source/Runtime/CoreUObject/Public/UObject/UObjectHash.h
--- a/Engine/Source/Runtime/CoreUObject/Public/UObject/UObjectHash.h
+++ b/Engine/Source/Runtime/CoreUObject/Public/UObject/UObjectHash.h
@@ -87,7 +87,7 @@ COREUOBJECT_API UObject* StaticFindFirstObjectFastInternal(const UClass* ObjectC
  * Variation of StaticFindObjectFast that uses explicit path.
  *
  * @param	ObjectClass		The to be found object's class
- * @param	ObjectName		The to be found object's class
+ * @param	ObjectName		The to be found object's name
  * @param	ObjectPathName	Full path name for the object to search for
  * @param	ExactClass		Whether to require an exact match with the passed in class
  * @param	ExclusiveFlags	Ignores objects that contain any of the specified exclusive flags
@@ -95,6 +95,16 @@ COREUOBJECT_API UObject* StaticFindFirstObjectFastInternal(const UClass* ObjectC
  */
 UObject* StaticFindObjectFastExplicit(const UClass* ObjectClass, FName ObjectName, const FString& ObjectPathName, bool bExactClass, EObjectFlags ExcludeFlags = RF_NoFlags);
 
+/**
+ * Variation of StaticFindObjectFastInternal that uses remote id.
+ *
+ * @param	RemoteId				The to be found object's remote id
+ * @param	InExclusiveFlags		Ignores objects that contain any of the specified exclusive flags
+ * @param	InExlusiveInternalFlags	Ignores objects that contain any of the specified internal exclusive flags
+ * @return	Returns a pointer to the found object or NULL if none could be found
+ */
+COREUOBJECT_API UObject* StaticFindObjectFastInternal(FRemoteObjectId RemoteId, EObjectFlags InExclusiveFlags = RF_NoFlags, EInternalObjectFlags InExlusiveInternalFlags = EInternalObjectFlags::None);
+
 /**
  * Return all objects with a given outer
  *
