[CoreRedirects]

- Moves FCoreRedirects global state into a FCoreRedirectsContext.
- Threads may now push/pop a thread local context for core redirects using FScopedCoreRedirectsContext. If no context is pushed a shared global context is used (as was the case previously).
- When creating a context to push you may inherit the state of the existing context. This copies the existing state protecting the original state from non-threadsafe mutation should it be shared.
-RedirectionSummary was moved into internal so it may be used by the context since it's part of the CoreRedirects global state.
- Private locking structures have been moved into LockedContext types for passing around rather than passing around a lock and a context.

#rnx

diff --git a/Engine/Source/Runtime/CoreUObject/Internal/UObject/CoreRedirects/CoreRedirectsContext.h b/Engine/Source/Runtime/CoreUObject/Internal/UObject/CoreRedirects/CoreRedirectsContext.h
new file mode 100644
--- /dev/null
+++ b/Engine/Source/Runtime/CoreUObject/Internal/UObject/CoreRedirects/CoreRedirectsContext.h
@@ -0,0 +1,239 @@
+// Copyright Epic Games, Inc. All Rights Reserved.
+
+#pragma once
+#include "Containers/Map.h"
+#include "Containers/UnrealString.h"
+#include "RedirectionSummary.h"
+#include "UObject/CoreRedirects.h"
+
+/**
+* Stores state required for FCoreRedirects API. The context encapsulates the necessary state for 
+* the FCoreRedirects API to function. Changine the context can be used to change the set of
+* CoreRedirects for the current thread without affecting other threads. It is up to the user to ensure
+* contexts are shared or not between threads. By default a single global context is shared amoung all
+* threads until a thread local context is set.
+*/
+struct COREUOBJECT_API FCoreRedirectsContext
+{
+	/*
+	* Returns the global context used by FCoreRedirects when no new context has been provided as an override.
+	*/
+	static FCoreRedirectsContext& GetGlobalContext();
+
+	/*
+	* Returns the current context used by FCoreRedirects for the current thread. If no new context has been applied,
+	* a global context is used for all threads.
+	*/
+	static FCoreRedirectsContext& GetThreadContext();
+
+	/*
+	* Changes the context for FCoreRedirects calls to use for the current thread. The passed in context must remain
+	* valid until PopThreadContext is called.
+	*/
+	static void SetThreadContext(FCoreRedirectsContext& NewContext);
+
+	/**
+	 * Run initialization steps that are needed before any data can be stored from FCoreRedirects calls.
+	 */
+	void InitializeContext();
+
+	/** Creates a context with a debug context name */
+	FCoreRedirectsContext(const FString& ContextName);
+
+	/** Creates a  context with a debug context name, and inherits the state from OtherContext.
+	The state from OtherContext is copied into the new context*/
+	FCoreRedirectsContext(const FString& ContextName, FCoreRedirectsContext& OtherContext);
+
+	/** Describes the context's intended purpose */
+	const FString& GetName() const;
+
+	/** Get whether this has been initialized */
+	const bool IsInitialized() const;
+
+	/** True if we are in debug mode that does extra validation */
+	const bool IsInDebugMode() const;
+
+	/** True if we have done our initial validation. After initial validation, each change to redirects will validate independently */
+	const bool HasValidated() const;
+
+	/** Map from config name to flag */
+	TMap<FName, ECoreRedirectFlags>& GetConfigKeyMap();
+
+	/** Map from name of thing being mapped to full list. List must be filtered further */
+	FCoreRedirects::FRedirectTypeMap& GetRedirectTypeMap();
+
+#if WITH_EDITOR
+	FRedirectionSummary& GetRedirectionSummary();
+#endif
+
+private:
+	friend void FCoreRedirects::ValidateAllRedirects();
+	/** Sets that we have done our initial validation. After initial validation, each change to redirects will validate independently */
+	void SetHasValidated();
+
+	/** This lock allows exclusive locking (WriteLock) and shared locking (ReadLock)
+	 *  Additionally, it permits limited types of recursion. It is possible to ReadLock()
+	 *  or WriteLock() while locked for write. It is NOT possible to Read or WriteLock() while
+	 *  locked for read. I.e., if the lock is held exclusively, it re-acquiring it is always permitted.
+	 *  If the lock is held shared, re-acquiring it is never permitted.
+	 */
+	struct FRWLockWithExclusiveRecursion
+	{
+		void ReadLock();
+
+		void WriteLock();
+
+		void WriteUnlock();
+
+		void ReadUnlock();
+
+	private:
+		uint32 LoadWriteLockOwnerThreadIdRelaxed() const;
+		void StoreWriteLockOwnerThreadId(uint32 ThreadId, std::memory_order Order = std::memory_order_seq_cst);
+
+		FTransactionallySafeRWLock InternalLock;
+		std::atomic<uint32> WriteLockOwnerThreadId = 0;
+		int32 RecursionCount = 0;
+	};
+	friend struct FScopeCoreRedirectsReadLockedContext;
+	friend struct FScopeCoreRedirectsWriteLockedContext;
+
+private:
+	FString Name;
+	std::atomic<bool> bInitialized;
+	std::atomic<bool> bInDebugMode;
+	bool bValidatedOnce;
+	TMap<FName, ECoreRedirectFlags> ConfigKeyMap;
+	FCoreRedirects::FRedirectTypeMap RedirectTypeMap;
+	/** Lock to protect multithreaded access to the CoreRedirect system */
+	FRWLockWithExclusiveRecursion RWLock;
+
+#if WITH_EDITOR
+	FRedirectionSummary RedirectionSummary;
+#endif
+};
+
+/**
+* RAII type for swapping the current thread's FCoreRedirects context to a new context. Can optionally
+* copy the existing state from another context when created.
+*/
+struct COREUOBJECT_API FScopeCoreRedirectsContext
+{
+	FScopeCoreRedirectsContext(const FString& Name);
+	FScopeCoreRedirectsContext(const FString& Name, FCoreRedirectsContext& ContextToCopyFrom);
+	~FScopeCoreRedirectsContext();
+
+	FCoreRedirectsContext ScopeContext;
+	FCoreRedirectsContext& ContextToRestore;
+};
+
+
+/**
+* RAII type for locking a context for reading. Supports recursively entering the lock.
+*/
+struct FScopeCoreRedirectsReadLockedContext
+{
+public:
+	explicit FScopeCoreRedirectsReadLockedContext(FCoreRedirectsContext& InContext)
+		: Context(InContext)
+		, NeedsUnlock(true)
+	{
+		Context.RWLock.ReadLock();
+	}
+
+	~FScopeCoreRedirectsReadLockedContext()
+	{
+		if (NeedsUnlock)
+		{
+			Context.RWLock.ReadUnlock();
+		}
+	}
+
+	FCoreRedirectsContext& Get() { return Context; }
+
+protected:
+	enum class EInitFlag
+	{
+		InitFlag
+	};
+
+	FScopeCoreRedirectsReadLockedContext(FCoreRedirectsContext& InContext, EInitFlag Unused)
+		: Context(InContext)
+		, NeedsUnlock(false)
+	{}
+
+	FCoreRedirectsContext& Context;
+	bool NeedsUnlock;
+};
+
+/**
+* RAII type for locking a context for writing. Supports recursively entering the lock.
+*/
+struct FScopeCoreRedirectsWriteLockedContext : public FScopeCoreRedirectsReadLockedContext
+{
+public:
+	explicit FScopeCoreRedirectsWriteLockedContext(FCoreRedirectsContext& InContext)
+		: FScopeCoreRedirectsReadLockedContext(InContext, EInitFlag::InitFlag)
+	{
+		Context.RWLock.WriteLock();
+	}
+
+	~FScopeCoreRedirectsWriteLockedContext()
+	{
+		Context.RWLock.WriteUnlock();
+	}
+};
+
+///////////////////////////////////////////////////////
+// Inline implementations
+///////////////////////////////////////////////////////
+
+inline const FString& FCoreRedirectsContext::GetName() const 
+{
+	return Name;
+}
+
+/** Get whether this has been initialized */
+inline const bool FCoreRedirectsContext::IsInitialized() const
+{
+	// Use an atomic variable rather than take a read lock because we might already be under a read lock
+	return bInitialized.load(std::memory_order_acquire);
+}
+
+/** True if we are in debug mode that does extra validation */
+inline const bool FCoreRedirectsContext::IsInDebugMode() const
+{
+	// Use an atomic variable rather than take a read lock because we might already be under a read lock
+	return bInDebugMode.load(std::memory_order_relaxed);
+}
+
+/** True if we have done our initial validation. After initial validation, each change to redirects will validate independently */
+inline const bool FCoreRedirectsContext::HasValidated() const 
+{ 
+	return bValidatedOnce; 
+}
+
+/** Sets that we have done our initial validation. After initial validation, each change to redirects will validate independently */
+inline void FCoreRedirectsContext::SetHasValidated() 
+{ 
+	bValidatedOnce = true;
+}
+
+/** Map from config name to flag */
+inline TMap<FName, ECoreRedirectFlags>& FCoreRedirectsContext::GetConfigKeyMap()
+{ 
+	return ConfigKeyMap;
+}
+
+/** Map from name of thing being mapped to full list. List must be filtered further */
+inline FCoreRedirects::FRedirectTypeMap& FCoreRedirectsContext::GetRedirectTypeMap() 
+{ 
+	return RedirectTypeMap;
+}
+
+#if WITH_EDITOR
+inline FRedirectionSummary& FCoreRedirectsContext::GetRedirectionSummary() 
+{ 
+	return RedirectionSummary;
+}
+#endif
diff --git a/Engine/Source/Runtime/CoreUObject/Internal/UObject/CoreRedirects/RedirectionSummary.h b/Engine/Source/Runtime/CoreUObject/Internal/UObject/CoreRedirects/RedirectionSummary.h
new file mode 100644
--- /dev/null
+++ b/Engine/Source/Runtime/CoreUObject/Internal/UObject/CoreRedirects/RedirectionSummary.h
@@ -0,0 +1,70 @@
+// Copyright Epic Games, Inc. All Rights Reserved.
+
+#pragma once
+
+#if WITH_EDITOR
+
+#include "Containers/Array.h"
+#include "Containers/ArrayView.h"
+#include "Containers/Map.h"
+#include "Containers/SortedMap.h"
+#include "HAL/CriticalSection.h"
+#include "Hash/Blake3.h"
+#include "UObject/CoreRedirects.h"
+#include "UObject/NameTypes.h"
+
+/**
+ * Container for FCoreRedirects that can affect a package. Used by class FCoreRedirects to implement
+ * AppendHashOfRedirectsAffectingPackages.
+ */
+class FRedirectionSummary
+{
+public:
+	FRedirectionSummary() = default;
+	FRedirectionSummary(FRedirectionSummary&& Other);
+	FRedirectionSummary& operator=(FRedirectionSummary&& Other);
+
+	void Add(const FCoreRedirect& CoreRedirect, bool bIsWildcardMatch);
+	void Remove(const FCoreRedirect& CoreRedirect, bool bIsWildcardMatch);
+	void AppendHashAffectingPackages(FBlake3& Hasher, TConstArrayView<FName> PackageNames);
+	void AppendHashGlobal(FBlake3& Hasher);
+
+private:
+	struct FCompareRedirect
+	{
+		bool operator()(const FCoreRedirect& A, const FCoreRedirect& B);
+	};
+	struct FRedirectContainer
+	{
+	public:
+		void Add(FCoreRedirect&& Redirect);
+		void Remove(const FCoreRedirect& Redirect);
+		bool IsEmpty() const;
+		void Empty();
+		bool TryAppendHashInReadLock(FBlake3& Hasher) const;
+		void AppendHashInWriteLock(FBlake3& Hasher);
+	private:
+		void CalculateHash();
+		void AppendHashWithoutDirtyCheck(FBlake3& Hasher) const;
+	private:
+		TSortedMap<FCoreRedirect, bool, FDefaultAllocator, FCompareRedirect> Redirects;
+		FBlake3Hash Hash;
+		bool bHashDirty = false;
+	};
+
+private:
+	static TArray<FName, TInlineAllocator<2>> GetAffectedPackages(const FCoreRedirect& Redirect);
+
+private:
+	TMap<FName, FRedirectContainer> RedirectsForPackage;
+	FRedirectContainer GlobalRedirects;
+	/**
+	 * CoreRedirects are written when the engine is single threaded, but we do not write the hashes for the
+	 * global and per-package containers until they are requested on demand. The requests on demand can
+	 * occur on multiple threads, so we need a lock to ensure that the first thread that requests them and finds
+	 * them dirty does not cause a data race with the second thread requests them.
+	 */
+	FRWLock Lock;
+};
+
+#endif
\ No newline at end of file
diff --git a/Engine/Source/Runtime/CoreUObject/Private/UObject/CoreRedirects.cpp b/Engine/Source/Runtime/CoreUObject/Private/UObject/CoreRedirects.cpp
--- a/Engine/Source/Runtime/CoreUObject/Private/UObject/CoreRedirects.cpp
+++ b/Engine/Source/Runtime/CoreUObject/Private/UObject/CoreRedirects.cpp
@@ -2,6 +2,8 @@
 
 #include "UObject/CoreRedirects.h"
 
+#include "UObject/CoreRedirects/CoreRedirectsContext.h"
+#include "UObject/CoreRedirects/RedirectionSummary.h"
 #include "UObject/Linker.h"
 #include "UObject/Package.h"
 #include "UObject/PropertyHelper.h"
@@ -13,12 +15,9 @@
 #include "Hash/Blake3.h"
 #include "Logging/StructuredLog.h"
 #include "Misc/AutomationTest.h"
-#include "Misc/CommandLine.h"
 #include "Misc/ConfigCacheIni.h"
-#include "Misc/CoreDelegates.h"
 #include "Misc/PackageName.h"
 #include "Misc/ScopeRWLock.h"
-#include "RedirectionSummary.h"
 #include "Serialization/DeferredMessageLog.h"
 #include "Templates/Casts.h"
 
@@ -430,10 +429,6 @@ namespace UE::CoreRedirects::Private
 	}
 }
 
-#if WITH_EDITOR
-FRedirectionSummary GRedirectionSummary;
-#endif
-
 FCoreRedirectObjectName::FCoreRedirectObjectName(const FTopLevelAssetPath& TopLevelAssetPath)
 	: FCoreRedirectObjectName(TopLevelAssetPath.GetAssetName(), NAME_None, TopLevelAssetPath.GetPackageName())
 {
@@ -1263,66 +1258,6 @@ int FCoreRedirectObjectName::Compare(const FCoreRedirectObjectName& Other) const
 	return 0;
 }
 
-std::atomic<bool> FCoreRedirects::bInitialized = false;
-std::atomic<bool> FCoreRedirects::bInDebugMode = false;
-bool FCoreRedirects::bValidatedOnce = false;
-
-TMap<FName, ECoreRedirectFlags> FCoreRedirects::ConfigKeyMap;
-FCoreRedirects::FRedirectTypeMap FCoreRedirects::RedirectTypeMap;
-FCoreRedirects::FRWLockWithExclusiveRecursion FCoreRedirects::RWLock;
-
-namespace UE::CoreRedirects::Private
-{
-class FRWWithExclusiveRecursionScopeLockForRead
-{
-public:
-	explicit FRWWithExclusiveRecursionScopeLockForRead(FCoreRedirects::FRWLockWithExclusiveRecursion& InLock)
-	: InternalLock(InLock)
-	, NeedsUnlock(true)
-	{
-		InternalLock.ReadLock();
-	}
-
-	~FRWWithExclusiveRecursionScopeLockForRead()
-	{
-		if (NeedsUnlock)
-		{
-			InternalLock.ReadUnlock();
-		}
-	}
-
-protected:
-	enum class EInitFlag
-	{
-		InitFlag
-	};
-
-	FRWWithExclusiveRecursionScopeLockForRead(FCoreRedirects::FRWLockWithExclusiveRecursion& InLock, EInitFlag Unused)
-	: InternalLock(InLock)
-	, NeedsUnlock(false)
-	{}
-
-	FCoreRedirects::FRWLockWithExclusiveRecursion& InternalLock;
-	bool NeedsUnlock;
-};
-
-class FRWWithExclusiveRecursionScopeLockForWrite :  public FRWWithExclusiveRecursionScopeLockForRead
-{
-public:
-	explicit FRWWithExclusiveRecursionScopeLockForWrite(FCoreRedirects::FRWLockWithExclusiveRecursion& InLock)
-	: FRWWithExclusiveRecursionScopeLockForRead(InLock, EInitFlag::InitFlag)
-	{
-		InternalLock.WriteLock();
-	}
-
-	~FRWWithExclusiveRecursionScopeLockForWrite()
-	{
-		InternalLock.WriteUnlock();
-	}
-};
-
-}
-
 FCoreRedirects::FRedirectNameMap& FCoreRedirects::FRedirectTypeMap::FindOrAdd(ECoreRedirectFlags Key)
 {
 	FRedirectNameMap*& NameMap = Map.FindOrAdd(Key);
@@ -1350,7 +1285,7 @@ FCoreRedirects::FRedirectNameMap& FCoreRedirects::FRedirectTypeMap::FindOrAdd(EC
 
 	if (EnumHasAnyFlags(Key, ECoreRedirectFlags::Option_MatchWildcardMask))
 	{
-		NewPair.Value.Wildcards.Reset(new FWildcardData());
+		NewPair.Value.Wildcards = MakeShareable<FWildcardData>(new FWildcardData());
 	}
 	
 	return NewPair.Value;
@@ -1370,66 +1305,29 @@ void FCoreRedirects::FRedirectTypeMap::Empty()
 void FCoreRedirects::Initialize()
 {
 	ensureMsgf(IsInGameThread(), TEXT("FCoreRedirects can only be initialized on the game thread."));
-	if (bInitialized.load(std::memory_order_relaxed))
-	{
-		return;
-	}
-	
-	{
-		FCoreRedirectorScopeLockForWrite ScopeLock(RWLock);
-
-#if !UE_BUILD_SHIPPING
-		if (FParse::Param(FCommandLine::Get(), TEXT("FullDebugCoreRedirects")))
-		{
-			// Enable debug mode and set to maximum verbosity
-			bInDebugMode.store(true, std::memory_order_relaxed);
-			UE_SET_LOG_VERBOSITY(LogCoreRedirects, VeryVerbose);
-			FCoreDelegates::OnFEngineLoopInitComplete.AddStatic(ValidateAllRedirects);
-		}
-		else if (FParse::Param(FCommandLine::Get(), TEXT("DebugCoreRedirects")))
-		{
-			// Enable debug mode and increase log levels but don't show every message
-			bInDebugMode.store(true, std::memory_order_relaxed);
-			UE_SET_LOG_VERBOSITY(LogCoreRedirects, Verbose);
-			FCoreDelegates::OnFEngineLoopInitComplete.AddStatic(ValidateAllRedirects);
-		}
-#endif
-
-		// Setup map
-		ConfigKeyMap.Add(TEXT("ObjectRedirects"), ECoreRedirectFlags::Type_Object);
-		ConfigKeyMap.Add(TEXT("ClassRedirects"), ECoreRedirectFlags::Type_Class);
-		ConfigKeyMap.Add(TEXT("StructRedirects"), ECoreRedirectFlags::Type_Struct);
-		ConfigKeyMap.Add(TEXT("EnumRedirects"), ECoreRedirectFlags::Type_Enum);
-		ConfigKeyMap.Add(TEXT("FunctionRedirects"), ECoreRedirectFlags::Type_Function);
-		ConfigKeyMap.Add(TEXT("PropertyRedirects"), ECoreRedirectFlags::Type_Property);
-		ConfigKeyMap.Add(TEXT("PackageRedirects"), ECoreRedirectFlags::Type_Package);
-		ConfigKeyMap.Add(TEXT("AssetRedirects"), ECoreRedirectFlags::Type_Asset);
 
-		RegisterNativeRedirectsUnderWriteLock(ScopeLock);
+	FCoreRedirectsContext& Context = FCoreRedirectsContext::GetGlobalContext();
+	Context.InitializeContext();
 
-		// Prepopulate RedirectTypeMap entries that some threads write to after the engine goes multi-threaded.
-		// Most RedirectTypeMap entries are written to only from InitUObject's call to ReadRedirectsFromIni, and at that point the Engine is single-threaded.
-		// Known missing packages and plugin loads can add entries to existing lists but will not add brand new types.
-		// Taking advantage of this, we treat the list of Key/Value pairs of RedirectTypeMap as immutable and read from it without synchronization.
-		// Note that the values for those written-during-multithreading entries need to be synchronized; it is only the list of Key/Value pairs that is immutable.
-		RedirectTypeMap.FindOrAdd(ECoreRedirectFlags::Type_Package | ECoreRedirectFlags::Category_Removed | ECoreRedirectFlags::Option_MissingLoad);
-
-		bInitialized.store(true, std::memory_order_release);
-	}
 	// Enable to run startup tests
-	// RunTests();
+#if 0
+	{
+		FScopeCoreRedirectsContext ScopeContext(TEXT("FCoreRedirects::RunTests"));
+		ensure(RunTests());
+	}
+#endif
 }
 
 bool FCoreRedirects::RedirectNameAndValues(ECoreRedirectFlags Type, const FCoreRedirectObjectName& OldObjectName,
 	FCoreRedirectObjectName& NewObjectName, const FCoreRedirect** FoundValueRedirect, ECoreRedirectMatchFlags MatchFlags)
 {
-	FCoreRedirectorScopeLockForRead ScopeLock(RWLock);
-	return RedirectNameAndValuesUnderReadLock(Type, OldObjectName, NewObjectName, FoundValueRedirect, MatchFlags, ScopeLock);
+	FScopeCoreRedirectsReadLockedContext LockedContext(FCoreRedirectsContext::GetThreadContext());
+	return RedirectNameAndValuesUnderReadLock(Type, OldObjectName, NewObjectName, FoundValueRedirect, MatchFlags, LockedContext);
 }
 
 bool FCoreRedirects::RedirectNameAndValuesUnderReadLock(ECoreRedirectFlags Type, const FCoreRedirectObjectName& OldObjectName,
 	FCoreRedirectObjectName& NewObjectName, const FCoreRedirect** FoundValueRedirect, ECoreRedirectMatchFlags MatchFlags,
-	const FCoreRedirectorScopeLockForRead& HeldLock)
+	FScopeCoreRedirectsReadLockedContext& LockedContext)
 {
 	auto ProcessRedirect = [&FoundValueRedirect, OldObjectName](const FCoreRedirect* Redirect, const FCoreRedirectObjectName& NewObjectName)
 	{
@@ -1455,7 +1353,7 @@ bool FCoreRedirects::RedirectNameAndValuesUnderReadLock(ECoreRedirectFlags Type,
 
 	NewObjectName = OldObjectName;
 	TArray<const FCoreRedirect*> FoundRedirects;
-	if (GetMatchingRedirectsUnderReadLock(Type, OldObjectName, FoundRedirects, MatchFlags, HeldLock))
+	if (GetMatchingRedirectsUnderReadLock(Type, OldObjectName, FoundRedirects, MatchFlags, LockedContext))
 	{
 		if (FoundRedirects.Num() > 1)
 		{
@@ -1485,15 +1383,15 @@ bool FCoreRedirects::RedirectNameAndValuesUnderReadLock(ECoreRedirectFlags Type,
 	}
 
 	const bool bDidRedirect = NewObjectName != OldObjectName;
-	UE_CLOG(IsInDebugMode() && bDidRedirect, LogCoreRedirects, Verbose,
+	UE_CLOG(LockedContext.Get().IsInDebugMode() && bDidRedirect, LogCoreRedirects, Verbose,
 		TEXT("RedirectNameAndValues(%s) replaced by %s"), *OldObjectName.ToString(), *NewObjectName.ToString());
 	return bDidRedirect;
 }
 
-bool FCoreRedirects::ValidateAssetRedirectsUnderReadLock(const FCoreRedirectorScopeLockForRead& HeldLock)
+bool FCoreRedirects::ValidateAssetRedirectsUnderReadLock(FScopeCoreRedirectsReadLockedContext& LockedContext)
 {
 	bool bValidationSucceeded = true;
-	FRedirectNameMap* RedirectNameMap = RedirectTypeMap.Find(ECoreRedirectFlags::Type_Asset);
+	FRedirectNameMap* RedirectNameMap = LockedContext.Get().GetRedirectTypeMap().Find(ECoreRedirectFlags::Type_Asset);
 	if (RedirectNameMap != nullptr)
 	{
 		for (const TPair<FName, TArray<FCoreRedirect>>& Pair : RedirectNameMap->RedirectMap)
@@ -1521,10 +1419,10 @@ bool FCoreRedirects::ValidateAssetRedirectsUnderReadLock(const FCoreRedirectorSc
 FCoreRedirectObjectName FCoreRedirects::GetRedirectedName(ECoreRedirectFlags Type,
 	const FCoreRedirectObjectName& OldObjectName, ECoreRedirectMatchFlags MatchFlags)
 {
-	FCoreRedirectorScopeLockForRead ScopeLock(RWLock);
+	FScopeCoreRedirectsReadLockedContext LockedContext(FCoreRedirectsContext::GetThreadContext());
 	FCoreRedirectObjectName NewObjectName;
 
-	RedirectNameAndValuesUnderReadLock(Type, OldObjectName, NewObjectName, nullptr, MatchFlags, ScopeLock);
+	RedirectNameAndValuesUnderReadLock(Type, OldObjectName, NewObjectName, nullptr, MatchFlags, LockedContext);
 
 	return NewObjectName;
 }
@@ -1532,15 +1430,15 @@ FCoreRedirectObjectName FCoreRedirects::GetRedirectedName(ECoreRedirectFlags Typ
 const TMap<FString, FString>* FCoreRedirects::GetValueRedirects(ECoreRedirectFlags Type,
 	const FCoreRedirectObjectName& OldObjectName, ECoreRedirectMatchFlags MatchFlags)
 {
-	FCoreRedirectorScopeLockForRead ScopeLock(RWLock);
+	FScopeCoreRedirectsReadLockedContext LockedContext(FCoreRedirectsContext::GetThreadContext());
 	FCoreRedirectObjectName NewObjectName;
 	const FCoreRedirect* FoundRedirect = nullptr;
 
-	RedirectNameAndValuesUnderReadLock(Type, OldObjectName, NewObjectName, &FoundRedirect, MatchFlags, ScopeLock);
+	RedirectNameAndValuesUnderReadLock(Type, OldObjectName, NewObjectName, &FoundRedirect, MatchFlags, LockedContext);
 
 	if (FoundRedirect && FoundRedirect->ValueChanges.Num() > 0)
 	{
-		UE_CLOG(IsInDebugMode() , LogCoreRedirects, VeryVerbose, TEXT("GetValueRedirects found %d matches for %s"),
+		UE_CLOG(LockedContext.Get().IsInDebugMode(), LogCoreRedirects, VeryVerbose, TEXT("GetValueRedirects found %d matches for %s"),
 			FoundRedirect->ValueChanges.Num(), *OldObjectName.ToString());
 
 		return &FoundRedirect->ValueChanges;
@@ -1552,12 +1450,12 @@ const TMap<FString, FString>* FCoreRedirects::GetValueRedirects(ECoreRedirectFla
 bool FCoreRedirects::GetMatchingRedirects(ECoreRedirectFlags SearchFlags, const FCoreRedirectObjectName& OldObjectName,
 	TArray<const FCoreRedirect*>& FoundRedirects, ECoreRedirectMatchFlags MatchFlags)
 {
-	FCoreRedirectorScopeLockForRead ScopeLock(RWLock);
-	return GetMatchingRedirectsUnderReadLock(SearchFlags, OldObjectName, FoundRedirects, MatchFlags, ScopeLock);
+	FScopeCoreRedirectsReadLockedContext LockedContext(FCoreRedirectsContext::GetThreadContext());
+	return GetMatchingRedirectsUnderReadLock(SearchFlags, OldObjectName, FoundRedirects, MatchFlags, LockedContext);
 }
 
 bool FCoreRedirects::GetMatchingRedirectsUnderReadLock(ECoreRedirectFlags SearchFlags, const FCoreRedirectObjectName& OldObjectName,
-	TArray<const FCoreRedirect*>& FoundRedirects, ECoreRedirectMatchFlags MatchFlags, const FCoreRedirectorScopeLockForRead& HeldLock)
+	TArray<const FCoreRedirect*>& FoundRedirects, ECoreRedirectMatchFlags MatchFlags, FScopeCoreRedirectsReadLockedContext& LockedContext)
 {
 	// Look for all redirects that match the given names and flags
 	bool bFound = false;
@@ -1573,7 +1471,7 @@ bool FCoreRedirects::GetMatchingRedirectsUnderReadLock(ECoreRedirectFlags Search
 		(!(MatchFlags & ECoreRedirectMatchFlags::AllowPartialMatch) || !OldObjectName.PackageName.IsNone());
 
 	// Determine list of maps to look over, need to handle being passed multiple types in a bit mask
-	for (const TPair<ECoreRedirectFlags, FRedirectNameMap>& Pair : RedirectTypeMap)
+	for (const TPair<ECoreRedirectFlags, FRedirectNameMap>& Pair : LockedContext.Get().GetRedirectTypeMap())
 	{
 		ECoreRedirectFlags PairFlags = Pair.Key;
 
@@ -1610,7 +1508,7 @@ bool FCoreRedirects::GetMatchingRedirectsUnderReadLock(ECoreRedirectFlags Search
 
 bool FCoreRedirects::FindPreviousNames(ECoreRedirectFlags SearchFlags, const FCoreRedirectObjectName& NewObjectName, TArray<FCoreRedirectObjectName>& PreviousNames)
 {
-	FCoreRedirectorScopeLockForRead ScopeLock(RWLock);
+	FScopeCoreRedirectsReadLockedContext LockedContext(FCoreRedirectsContext::GetThreadContext());
 
 	// Look for reverse direction redirects
 	bool bFound = false;
@@ -1639,7 +1537,7 @@ bool FCoreRedirects::FindPreviousNames(ECoreRedirectFlags SearchFlags, const FCo
 		};
 
 	// Determine list of maps to look over, need to handle being passed multiple Flags in a bit mask
-	for (const TPair<ECoreRedirectFlags, FRedirectNameMap>& Pair : RedirectTypeMap)
+	for (const TPair<ECoreRedirectFlags, FRedirectNameMap>& Pair : LockedContext.Get().GetRedirectTypeMap())
 	{
 		ECoreRedirectFlags PairFlags = Pair.Key;
 
@@ -1698,7 +1596,7 @@ bool FCoreRedirects::FindPreviousNames(ECoreRedirectFlags SearchFlags, const FCo
 		}
 	}
 
-	UE_CLOG(bFound && IsInDebugMode(), LogCoreRedirects, VeryVerbose, TEXT("FindPreviousNames found %d previous names for %s"), PreviousNames.Num(), *NewObjectName.ToString());
+	UE_CLOG(bFound && LockedContext.Get().IsInDebugMode(), LogCoreRedirects, VeryVerbose, TEXT("FindPreviousNames found %d previous names for %s"), PreviousNames.Num(), *NewObjectName.ToString());
 
 	return bFound;
 }
@@ -1706,9 +1604,8 @@ bool FCoreRedirects::FindPreviousNames(ECoreRedirectFlags SearchFlags, const FCo
 bool FCoreRedirects::IsKnownMissing(ECoreRedirectFlags Type, const FCoreRedirectObjectName& ObjectName)
 {
 	TArray<const FCoreRedirect*> FoundRedirects;
-
-	FCoreRedirectorScopeLockForRead ScopeLock(RWLock);
-	return FCoreRedirects::GetMatchingRedirectsUnderReadLock(Type | ECoreRedirectFlags::Category_Removed, ObjectName, FoundRedirects, ECoreRedirectMatchFlags::None, ScopeLock);
+	FScopeCoreRedirectsReadLockedContext LockedContext(FCoreRedirectsContext::GetThreadContext());
+	return FCoreRedirects::GetMatchingRedirectsUnderReadLock(Type | ECoreRedirectFlags::Category_Removed, ObjectName, FoundRedirects, ECoreRedirectMatchFlags::None, LockedContext);
 }
 
 bool FCoreRedirects::AddKnownMissing(ECoreRedirectFlags Type, const FCoreRedirectObjectName& ObjectName, ECoreRedirectFlags Channel)
@@ -1734,8 +1631,8 @@ void FCoreRedirects::ClearKnownMissing(ECoreRedirectFlags Type, ECoreRedirectFla
 	check((Channel & ~ECoreRedirectFlags::Option_MissingLoad) == ECoreRedirectFlags::None);
 	ECoreRedirectFlags RedirectFlags = Type | ECoreRedirectFlags::Category_Removed | Channel;
 
-	FCoreRedirectorScopeLockForWrite ScopeLock(RWLock);
-	FRedirectNameMap* RedirectNameMap = RedirectTypeMap.Find(RedirectFlags);
+	FScopeCoreRedirectsReadLockedContext LockedContext(FCoreRedirectsContext::GetThreadContext());
+	FRedirectNameMap* RedirectNameMap = LockedContext.Get().GetRedirectTypeMap().Find(RedirectFlags);
 	if (RedirectNameMap)
 	{
 		RedirectNameMap->RedirectMap.Empty();
@@ -1745,26 +1642,26 @@ void FCoreRedirects::ClearKnownMissing(ECoreRedirectFlags Type, ECoreRedirectFla
 #if WITH_EDITOR
 void FCoreRedirects::AppendHashOfRedirectsAffectingPackages(FBlake3& Hasher, TConstArrayView<FName> PackageNames)
 {
-	GRedirectionSummary.AppendHashAffectingPackages(Hasher, PackageNames);
+	FCoreRedirectsContext::GetThreadContext().GetRedirectionSummary().AppendHashAffectingPackages(Hasher, PackageNames);
 }
 
 void FCoreRedirects::AppendHashOfGlobalRedirects(FBlake3& Hasher)
 {
-	GRedirectionSummary.AppendHashGlobal(Hasher);
+	FCoreRedirectsContext::GetThreadContext().GetRedirectionSummary().AppendHashGlobal(Hasher);
 }
 
 void FCoreRedirects::RecordAddedObjectRedirector(const FSoftObjectPath& Source, const FSoftObjectPath& Dest)
 {
 	FCoreRedirect ConvertedToCoreRedirect(ECoreRedirectFlags::Type_Object,
 		FCoreRedirectObjectName(Source), FCoreRedirectObjectName(Dest));
-	GRedirectionSummary.Add(ConvertedToCoreRedirect, false /* bIsWildcardMatch */);
+	FCoreRedirectsContext::GetThreadContext().GetRedirectionSummary().Add(ConvertedToCoreRedirect, false /* bIsWildcardMatch */);
 }
 
 void FCoreRedirects::RecordRemovedObjectRedirector(const FSoftObjectPath& Source, const FSoftObjectPath& Dest)
 {
 	FCoreRedirect ConvertedToCoreRedirect(ECoreRedirectFlags::Type_Object,
 		FCoreRedirectObjectName(Source), FCoreRedirectObjectName(Dest));
-	GRedirectionSummary.Remove(ConvertedToCoreRedirect, false /* bIsWildcardMatch */);
+	FCoreRedirectsContext::GetThreadContext().GetRedirectionSummary().Remove(ConvertedToCoreRedirect, false /* bIsWildcardMatch */);
 }
 
 #endif
@@ -1773,13 +1670,6 @@ bool FCoreRedirects::RunTests()
 {
 	bool bSuccess = true;
 
-	FRedirectTypeMap BackupMap = MoveTemp(RedirectTypeMap);
-	RedirectTypeMap.Empty();
-#if WITH_EDITOR
-	FRedirectionSummary BackupSummary = MoveTemp(GRedirectionSummary);
-	GRedirectionSummary = FRedirectionSummary();
-#endif
-
 	UE_LOG(LogCoreRedirects, Log, TEXT("Running FCoreRedirect Tests"));
 
 	TArray<FCoreRedirect> NewRedirects;
@@ -1971,91 +1861,87 @@ bool FCoreRedirects::RunTests()
 		UE_LOG(LogCoreRedirects, Error, TEXT("FCoreRedirect Test Failed: /Game/NotRemovedPackage should no longer be removed!"));
 	}
 
-	// Restore old state
-	RedirectTypeMap = MoveTemp(BackupMap);
-#if WITH_EDITOR
-	GRedirectionSummary = MoveTemp(BackupSummary);
-#endif
-
 	UE_LOG(LogCoreRedirects, Log, TEXT("FCoreRedirect Test %s!"), (bSuccess ? TEXT("Passed") : TEXT("Failed")));
 	return bSuccess;
 }
 
 void FCoreRedirects::AddAssetRedirects(const TMap<FSoftObjectPath, FSoftObjectPath>& InRedirects)
 {
-	if (InRedirects.Num() > 0)
+	if (InRedirects.IsEmpty())
+	{
+		return;
+	}
+
+	FScopeCoreRedirectsWriteLockedContext LockedContext(FCoreRedirectsContext::GetThreadContext());
+	FRedirectNameMap* ExistingMap = &LockedContext.Get().GetRedirectTypeMap().FindOrAdd(ECoreRedirectFlags::Type_Asset);
+	int32 NumAdded = 0;
+	int32 NumSkipped = 0;
+	for (const TPair<FSoftObjectPath, FSoftObjectPath>& Pair : InRedirects)
 	{
-		FCoreRedirectorScopeLockForWrite ScopeLock(RWLock);
-		FRedirectNameMap* ExistingMap = &RedirectTypeMap.FindOrAdd(ECoreRedirectFlags::Type_Asset);
-		int32 NumAdded = 0;
-		int32 NumSkipped = 0;
-		for (const TPair<FSoftObjectPath, FSoftObjectPath>& Pair : InRedirects)
+		// Asset redirects are, by definition, not package redirects
+		if (Pair.Key.GetLongPackageFName().IsNone() || Pair.Value.GetAssetFName().IsNone())
 		{
-			// Asset redirects are, by definition, not package redirects
-			if (Pair.Key.GetLongPackageFName().IsNone() || Pair.Value.GetAssetFName().IsNone())
-			{
-				UE_LOG(LogCoreRedirects, Warning, TEXT("Attempted to register asset redirector that was missing a package or object name. Redirector was from %s to %s"),
-					*Pair.Key.ToString(), *Pair.Value.ToString());
-				NumSkipped++;
-				continue;
-			}
+			UE_LOG(LogCoreRedirects, Warning, TEXT("Attempted to register asset redirector that was missing a package or object name. Redirector was from %s to %s"),
+				*Pair.Key.ToString(), *Pair.Value.ToString());
+			NumSkipped++;
+			continue;
+		}
 
-			// Asset redirects use the package as the lookup key but contain multiple redirects underneath
-			// Conceptually, for each object redirector we add two core redirects: one for the package and one for the object itself
-			// In practice, we can just add the package redirect the first time
+		// Asset redirects use the package as the lookup key but contain multiple redirects underneath
+		// Conceptually, for each object redirector we add two core redirects: one for the package and one for the object itself
+		// In practice, we can just add the package redirect the first time
 
-			FCoreRedirect ObjectRedirector(ECoreRedirectFlags::Type_Asset, Pair.Key.ToString(), Pair.Value.ToString());
-			TArray<FCoreRedirect>& ExistingRedirects = ExistingMap->RedirectMap.FindOrAdd(ObjectRedirector.GetSearchKey());
+		FCoreRedirect ObjectRedirector(ECoreRedirectFlags::Type_Asset, Pair.Key.ToString(), Pair.Value.ToString());
+		TArray<FCoreRedirect>& ExistingRedirects = ExistingMap->RedirectMap.FindOrAdd(ObjectRedirector.GetSearchKey());
 
-			if (ExistingRedirects.Num() == 0)
-			{
-				// This is a new redirector. Add two entries, one for the package first
-				FCoreRedirect PackageRedirect(ECoreRedirectFlags::Type_Asset, 
-					Pair.Key.GetLongPackageName(), Pair.Value.GetLongPackageName());
+		if (ExistingRedirects.Num() == 0)
+		{
+			// This is a new redirector. Add two entries, one for the package first
+			FCoreRedirect PackageRedirect(ECoreRedirectFlags::Type_Asset, 
+				Pair.Key.GetLongPackageName(), Pair.Value.GetLongPackageName());
 
-				ExistingRedirects.Add(PackageRedirect);
-			}
+			ExistingRedirects.Add(PackageRedirect);
+		}
 
-			// Check that, among the existing redirects, this one won't be a duplicate. Don't bother checking the package entry
-			bool bShouldAddRedirect = true;
-			for (int32 RedirectIndex = 1; RedirectIndex < ExistingRedirects.Num(); RedirectIndex++)
+		// Check that, among the existing redirects, this one won't be a duplicate. Don't bother checking the package entry
+		bool bShouldAddRedirect = true;
+		for (int32 RedirectIndex = 1; RedirectIndex < ExistingRedirects.Num(); RedirectIndex++)
+		{
+			if (Pair.Key.GetLongPackageFName() == ExistingRedirects[RedirectIndex].OldName.PackageName
+				&& Pair.Key.GetAssetName() == ExistingRedirects[RedirectIndex].OldName.ObjectName)
 			{
-				if (Pair.Key.GetLongPackageFName() == ExistingRedirects[RedirectIndex].OldName.PackageName
-					&& Pair.Key.GetAssetName() == ExistingRedirects[RedirectIndex].OldName.ObjectName)
-				{
-					const FCoreRedirectObjectName& ExistingTargetName = ExistingRedirects[0].NewName;
+				const FCoreRedirectObjectName& ExistingTargetName = ExistingRedirects[0].NewName;
 
-					UE_LOGFMT(LogCoreRedirects, Error, "Skipping new redirect target '{target}' due to existing map from '{source}' to '{dest}'",
-						Pair.Value.ToString(), Pair.Key.ToString(), ExistingTargetName.ToString());
-
-					bShouldAddRedirect = false;
-					NumSkipped++;
-					break;
-				}
-			}
+				UE_LOGFMT(LogCoreRedirects, Error, "Skipping new redirect target '{target}' due to existing map from '{source}' to '{dest}'",
+					Pair.Value.ToString(), Pair.Key.ToString(), ExistingTargetName.ToString());
 
-			if (bShouldAddRedirect)
-			{
-				ExistingRedirects.Add(ObjectRedirector);
-				NumAdded++;
+				bShouldAddRedirect = false;
+				NumSkipped++;
+				break;
 			}
 		}
 
-		UE_LOG(LogCoreRedirects, Display, 
-			TEXT("Object redirects provided to FCoreRedirects: %d. Redirects add: %d. Redirects skipped: %d"),
-			InRedirects.Num(), NumAdded, NumSkipped);
-
-		if (IsInDebugMode())
+		if (bShouldAddRedirect)
 		{
-			ValidateAssetRedirects();
+			ExistingRedirects.Add(ObjectRedirector);
+			NumAdded++;
 		}
 	}
+
+	UE_LOG(LogCoreRedirects, Display, 
+		TEXT("Object redirects provided to FCoreRedirects: %d. Redirects add: %d. Redirects skipped: %d"),
+		InRedirects.Num(), NumAdded, NumSkipped);
+
+	if (LockedContext.Get().IsInDebugMode())
+	{
+		ValidateAssetRedirects();
+	}
 }
 
 void FCoreRedirects::RemoveAllAssetRedirects()
 {
-	FCoreRedirectorScopeLockForWrite ScopeLock(RWLock);
-	FRedirectNameMap* ExistingMap = &RedirectTypeMap.FindOrAdd(ECoreRedirectFlags::Type_Asset);
+	FScopeCoreRedirectsWriteLockedContext LockedContext(FCoreRedirectsContext::GetThreadContext());
+	FRedirectNameMap* ExistingMap = &LockedContext.Get().GetRedirectTypeMap().FindOrAdd(ECoreRedirectFlags::Type_Asset);
 	ExistingMap->RedirectMap.Reset();
 }
 
@@ -2067,7 +1953,7 @@ bool FCoreRedirectTest::RunTest(const FString& Parameters)
 
 bool FCoreRedirects::ReadRedirectsFromIni(const FString& IniName)
 {
-	FCoreRedirectorScopeLockForWrite ScopeLock(RWLock);
+	FScopeCoreRedirectsWriteLockedContext LockedContext(FCoreRedirectsContext::GetThreadContext());
 	ensureMsgf(IsInitialized(), TEXT("FCoreRedirects must be initialized on the game thread before use."));
 
 	if (GConfig)
@@ -2099,7 +1985,7 @@ bool FCoreRedirects::ReadRedirectsFromIni(const FString& IniName)
 
 				FParse::Value(*ValueString, TEXT("OverrideClassName="), OverrideClassName);
 
-				ECoreRedirectFlags* FlagPtr = ConfigKeyMap.Find(It.Key());
+				ECoreRedirectFlags* FlagPtr = LockedContext.Get().GetConfigKeyMap().Find(It.Key());
 
 				// If valid type
 				if (FlagPtr)
@@ -2218,17 +2104,17 @@ bool FCoreRedirects::ReadRedirectsFromIni(const FString& IniName)
 
 bool FCoreRedirects::AddRedirectList(TArrayView<const FCoreRedirect> Redirects, const FString& SourceString)
 {
-	FCoreRedirectorScopeLockForWrite ScopeLock(RWLock);
+	FScopeCoreRedirectsWriteLockedContext ScopeLock(FCoreRedirectsContext::GetThreadContext());
 	ensureMsgf(IsInitialized(), TEXT("FCoreRedirects must be initialized on the game thread before use."));
 	return AddRedirectListUnderWriteLock(Redirects, SourceString, ScopeLock);
 }
 
 bool FCoreRedirects::AddRedirectListUnderWriteLock(TArrayView<const FCoreRedirect> Redirects, const FString & SourceString, 
-	const FCoreRedirectorScopeLockForWrite& HeldLock)
+	FScopeCoreRedirectsWriteLockedContext& LockedContext)
 {
 	UE_LOG(LogCoreRedirects, Verbose, TEXT("AddRedirect(%s) adding %d redirects"), *SourceString, Redirects.Num());
 
-	if (IsInDebugMode() && bValidatedOnce)
+	if (LockedContext.Get().IsInDebugMode() && LockedContext.Get().HasValidated())
 	{
 		// Validate on apply because we finished our initial validation pass
 		ValidateRedirectList(Redirects, SourceString);
@@ -2265,7 +2151,7 @@ bool FCoreRedirects::AddRedirectListUnderWriteLock(TArrayView<const FCoreRedirec
 			UE_LOG(LogCoreRedirects, Log, TEXT("AddRedirect(%s) has wildcard redirect %s, these are very slow and should be resolved as soon as possible! Please refer to the documentation in Engine/Config/BaseEngine.ini."), *SourceString, *NewRedirect.OldName.ToString());
 		}
 		
-		if (AddSingleRedirectUnderWriteLock(NewRedirect, SourceString, HeldLock))
+		if (AddSingleRedirectUnderWriteLock(NewRedirect, SourceString, LockedContext))
 		{
 			bAddedAny = true;
 
@@ -2275,7 +2161,7 @@ bool FCoreRedirects::AddRedirectListUnderWriteLock(TArrayView<const FCoreRedirec
 				FCoreRedirect ValueRedirect = NewRedirect;
 				ValueRedirect.OldName = ValueRedirect.NewName;
 
-				AddSingleRedirectUnderWriteLock(ValueRedirect, SourceString, HeldLock);
+				AddSingleRedirectUnderWriteLock(ValueRedirect, SourceString, LockedContext);
 			}
 		}
 	}
@@ -2284,12 +2170,12 @@ bool FCoreRedirects::AddRedirectListUnderWriteLock(TArrayView<const FCoreRedirec
 }
 
 bool FCoreRedirects::AddSingleRedirectUnderWriteLock(const FCoreRedirect& NewRedirect, const FString& SourceString, 
-	const FCoreRedirectorScopeLockForWrite& HeldLock)
+	FScopeCoreRedirectsWriteLockedContext& LockedContext)
 {
 	const bool bIsWildcardMatch = NewRedirect.IsWildcardMatch();
 	FRedirectNameMap* ExistingNameMap;
 
-	ExistingNameMap = &RedirectTypeMap.FindOrAdd(NewRedirect.RedirectFlags);
+	ExistingNameMap = &LockedContext.Get().GetRedirectTypeMap().FindOrAdd(NewRedirect.RedirectFlags);
 
 	TArray<FCoreRedirect>* ExistingRedirects = nullptr;
 	if (bIsWildcardMatch)
@@ -2374,8 +2260,9 @@ bool FCoreRedirects::AddSingleRedirectUnderWriteLock(const FCoreRedirect& NewRed
 	{
 		ExistingRedirects->Add(NewRedirect);
 	}
+
 #if WITH_EDITOR
-	GRedirectionSummary.Add(NewRedirect, bIsWildcardMatch);
+	LockedContext.Get().GetRedirectionSummary().Add(NewRedirect, bIsWildcardMatch);
 #endif
 
 	return true;
@@ -2383,7 +2270,7 @@ bool FCoreRedirects::AddSingleRedirectUnderWriteLock(const FCoreRedirect& NewRed
 
 bool FCoreRedirects::RemoveRedirectList(TArrayView<const FCoreRedirect> Redirects, const FString& SourceString)
 {
-	FCoreRedirectorScopeLockForWrite ScopeLock(RWLock);
+	FScopeCoreRedirectsWriteLockedContext LockedContext(FCoreRedirectsContext::GetThreadContext());
 	UE_LOG(LogCoreRedirects, Verbose, TEXT("RemoveRedirect(%s) Removing %d redirects"), *SourceString, Redirects.Num());
 
 	bool bRemovedAny = false;
@@ -2418,7 +2305,7 @@ bool FCoreRedirects::RemoveRedirectList(TArrayView<const FCoreRedirect> Redirect
 			UE_LOG(LogCoreRedirects, Log, TEXT("RemoveRedirect(%s) has wildcard redirect %s, these are very slow and should be resolved as soon as possible! Please refer to the documentation in Engine/Config/BaseEngine.ini."), *SourceString, *RedirectToRemove.OldName.ToString());
 		}
 
-		bRemovedAny |= RemoveSingleRedirectUnderWriteLock(RedirectToRemove, SourceString, ScopeLock);
+		bRemovedAny |= RemoveSingleRedirectUnderWriteLock(RedirectToRemove, SourceString, LockedContext);
 	}
 
 	return bRemovedAny;
@@ -2426,21 +2313,19 @@ bool FCoreRedirects::RemoveRedirectList(TArrayView<const FCoreRedirect> Redirect
 
 bool FCoreRedirects::IsInitialized()
 {
-	// Use an atomic variable rather than take a read lock because we might already be under a read lock
-	return bInitialized.load(std::memory_order_acquire);
+	return FCoreRedirectsContext::GetThreadContext().IsInitialized();
 }
 
 bool FCoreRedirects::IsInDebugMode()
 {
-	// Use an atomic variable rather than take a read lock because we might already be under a read lock
-	return bInDebugMode.load(std::memory_order_relaxed);
+	return FCoreRedirectsContext::GetThreadContext().IsInDebugMode();
 }
 
 bool FCoreRedirects::RemoveSingleRedirectUnderWriteLock(const FCoreRedirect& RedirectToRemove, const FString& SourceString,
-	const FCoreRedirectorScopeLockForWrite& HeldLock)
+	FScopeCoreRedirectsWriteLockedContext& LockedContext)
 {
 	const bool bIsWildcardMatch = RedirectToRemove.IsWildcardMatch();
-	FRedirectNameMap* ExistingNameMap = RedirectTypeMap.Find(RedirectToRemove.RedirectFlags);
+	FRedirectNameMap* ExistingNameMap = LockedContext.Get().GetRedirectTypeMap().Find(RedirectToRemove.RedirectFlags);
 
 	if (!ExistingNameMap)
 	{
@@ -2500,8 +2385,9 @@ bool FCoreRedirects::RemoveSingleRedirectUnderWriteLock(const FCoreRedirect& Red
 			// tables to avoid unnecessary false positives
 			ExistingNameMap->Wildcards->Rebuild();
 		}
+
 #if WITH_EDITOR
-		GRedirectionSummary.Remove(RedirectToRemove, bIsWildcardMatch);
+		LockedContext.Get().GetRedirectionSummary().Remove(RedirectToRemove, bIsWildcardMatch);
 #endif
 	}
 
@@ -2598,12 +2484,12 @@ void FCoreRedirects::ValidateRedirectList(TArrayView<const FCoreRedirect> Redire
 
 void FCoreRedirects::ValidateAllRedirects()
 {
-	FCoreRedirectorScopeLockForRead ReadLock(RWLock);
-	bValidatedOnce = true;
+	FScopeCoreRedirectsReadLockedContext LockedContext(FCoreRedirectsContext::GetThreadContext());
+	LockedContext.Get().SetHasValidated();
 
 	// Validate all existing redirects
 
-	for (const TPair<ECoreRedirectFlags, FRedirectNameMap>& Pair : RedirectTypeMap)
+	for (const TPair<ECoreRedirectFlags, FRedirectNameMap>& Pair : LockedContext.Get().GetRedirectTypeMap())
 	{
 		ECoreRedirectFlags PairFlags = Pair.Key;
 		FString ListName = FString::Printf(TEXT("Type %d"), (int32)PairFlags);
@@ -2614,20 +2500,20 @@ void FCoreRedirects::ValidateAllRedirects()
 		}
 	}
 
-	ValidateAssetRedirectsUnderReadLock(ReadLock);
+	ValidateAssetRedirectsUnderReadLock(LockedContext);
 }
 
 bool FCoreRedirects::ValidateAssetRedirects()
 {
-	FCoreRedirectorScopeLockForRead ReadLock(RWLock);
-	return ValidateAssetRedirectsUnderReadLock(ReadLock);
+	FScopeCoreRedirectsReadLockedContext LockedContext(FCoreRedirectsContext::GetThreadContext());
+	return ValidateAssetRedirectsUnderReadLock(LockedContext);
 }
 
 const TMap<FName, ECoreRedirectFlags>& FCoreRedirects::GetConfigKeyMap()
 {
 	// The config key map is only written to during initialization. That allows us to provide unguarded access after initialization
 	ensureMsgf(IsInitialized(), TEXT("It is not legal to read the config key map until after FCoreRedirects has been initialized."));
-	return ConfigKeyMap;
+	return FCoreRedirectsContext::GetThreadContext().GetConfigKeyMap();
 }
 
 ECoreRedirectFlags FCoreRedirects::GetFlagsForTypeName(FName PackageName, FName TypeName)
@@ -3427,7 +3313,7 @@ static void RegisterNativeRedirects49(TArray<FCoreRedirect>& Redirects)
 
 UE_ENABLE_OPTIMIZATION_SHIP
 
-void FCoreRedirects::RegisterNativeRedirectsUnderWriteLock(const FCoreRedirectorScopeLockForWrite& HeldLock)
+void FCoreRedirects::RegisterNativeRedirectsUnderWriteLock(FScopeCoreRedirectsWriteLockedContext& LockedContext)
 {
 	// Registering redirects here instead of in baseengine.ini is faster to parse and can clean up the ini, but is not required
 	TArray<FCoreRedirect> Redirects;
@@ -3438,111 +3324,10 @@ void FCoreRedirects::RegisterNativeRedirectsUnderWriteLock(const FCoreRedirector
 
 	// 4.10 and later are in baseengine.ini
 
-	AddRedirectListUnderWriteLock(Redirects, TEXT("RegisterNativeRedirects"), HeldLock);
+	AddRedirectListUnderWriteLock(Redirects, TEXT("RegisterNativeRedirects"), LockedContext);
 }
 #else
-void FCoreRedirects::RegisterNativeRedirectsUnderWriteLock(const FCoreRedirectorScopeLock& HeldLock)
+void FCoreRedirects::RegisterNativeRedirectsUnderWriteLock(const FCoreRedirectorScopeLock& HeldLock, FCoreRedirectsContext& Context)
 {
 }
 #endif // UE_WITH_CORE_REDIRECTS
-
-uint32 FCoreRedirects::FRWLockWithExclusiveRecursion::LoadWriteLockOwnerThreadIdRelaxed() const
-{
-	uint32 ThreadId;
-	// Memory validation is disabled, as this is a hot code path.
-	UE_AUTORTFM_OPEN_NO_VALIDATION
-	{
-		ThreadId = WriteLockOwnerThreadId.load(std::memory_order_relaxed);
-	};
-	return ThreadId;
-}
-
-void FCoreRedirects::FRWLockWithExclusiveRecursion::StoreWriteLockOwnerThreadId(uint32 ThreadId, std::memory_order Order)
-{
-	if (AutoRTFM::IsClosed())
-	{
-		// Memory validation is disabled, as this is a hot code path.
-		// If an abort occurs, we want to restore the variable's previous value, so we use
-		// `RecordOpenWrite` here so that AutoRTFM can undo our work. (We aren't overly
-		// concerned with atomicity as all loads are done in `memory_order_relaxed`.)
-		UE_AUTORTFM_OPEN_NO_VALIDATION
-		{
-			AutoRTFM::RecordOpenWrite(&ThreadId);
-			WriteLockOwnerThreadId.store(ThreadId, Order);
-		};
-	}
-	else
-	{
-		WriteLockOwnerThreadId.store(ThreadId, Order);
-	}
-}
-
-void FCoreRedirects::FRWLockWithExclusiveRecursion::ReadLock()
-{
-	uint32 InitialLockOwner = LoadWriteLockOwnerThreadIdRelaxed();
-	
-	// Avoid calling GetCurrentThreadId in the common case where we start out unlocked
-	// If we are unowned at the start, then there are three cases to consider:
-	// 1. We are unlocked. In that case we'll pass the InitialLockOwner check and call ReadLock()
-	// 2. We are already in a ReadLock() state. Taking the read lock again will cause us to deadlock. 
-	// That is part of the contract for this lock so it's okay.
-	// 3. We are write locked, in which case we'll enter the first block 
-	if (InitialLockOwner != 0)
-	{
-		// If we are initially write locked, there are two possibilities:
-		// 1. This thread owns the write lock. In that case, it's safe to just bump the recursion count
-		// 2. Another thread owns the write lock, in which case we want to block with a ReadLock()
-		if (InitialLockOwner == FPlatformTLS::GetCurrentThreadId())
-		{
-			// Allow recursive read locking by bumping the recursion count
-			RecursionCount++;
-		}
-		else
-		{
-			InternalLock.ReadLock();
-		}
-	}
-	else
-	{
-		// We may or may not still be locked at this point if another thread took a WriteLock
-		// but it doesn't matter because in either case we want to take a ReadLock and potentially wait
-		InternalLock.ReadLock();
-	}
-}
-
-void FCoreRedirects::FRWLockWithExclusiveRecursion::WriteLock()
-{
-	uint32 CurrentThreadId = FPlatformTLS::GetCurrentThreadId();
-	if (LoadWriteLockOwnerThreadIdRelaxed() != CurrentThreadId)
-	{
-		InternalLock.WriteLock();
-		StoreWriteLockOwnerThreadId(CurrentThreadId);
-	}
-	RecursionCount++;
-}
-
-void FCoreRedirects::FRWLockWithExclusiveRecursion::WriteUnlock()
-{
-	ensure(LoadWriteLockOwnerThreadIdRelaxed() == FPlatformTLS::GetCurrentThreadId());
-	ensure(RecursionCount > 0);
-	RecursionCount--;
-	if (RecursionCount == 0)
-	{
-		StoreWriteLockOwnerThreadId(0, std::memory_order_relaxed);
-		InternalLock.WriteUnlock();
-	}
-}
-
-void FCoreRedirects::FRWLockWithExclusiveRecursion::ReadUnlock()
-{
-	if (RecursionCount > 0)
-	{
-		ensureMsgf(LoadWriteLockOwnerThreadIdRelaxed() == FPlatformTLS::GetCurrentThreadId(),
-			TEXT("Called ReadUnlock() on a lock held exclusively by another thread."));
-		RecursionCount--;
-	}
-	else
-	{
-		InternalLock.ReadUnlock();
-	}
-}
diff --git a/Engine/Source/Runtime/CoreUObject/Private/UObject/CoreRedirectsContext.cpp b/Engine/Source/Runtime/CoreUObject/Private/UObject/CoreRedirectsContext.cpp
new file mode 100644
--- /dev/null
+++ b/Engine/Source/Runtime/CoreUObject/Private/UObject/CoreRedirectsContext.cpp
@@ -0,0 +1,221 @@
+// Copyright Epic Games, Inc. All Rights Reserved.
+
+#include "UObject/CoreRedirects/CoreRedirectsContext.h"
+#include "UObject/CoreRedirects.h"
+
+#include "Misc/CommandLine.h"
+#include "Misc/CoreDelegates.h"
+
+namespace UE::CoreRedirects::Private
+{
+	FCoreRedirectsContext GlobalCoreRedirectContext(TEXT("Global CoreRedirect Context"));
+	thread_local FCoreRedirectsContext* ThreadCoreRedirectContext = nullptr;
+}
+
+FCoreRedirectsContext::FCoreRedirectsContext(const FString& ContextName)
+	: Name(ContextName)
+	, bInitialized(false)
+	, bInDebugMode(false)
+	, bValidatedOnce(false)
+{
+}
+
+FCoreRedirectsContext::FCoreRedirectsContext(const FString& ContextName, FCoreRedirectsContext& OtherContext)
+	: Name(ContextName)
+	, bInitialized(OtherContext.bInitialized.load())
+	, bInDebugMode(OtherContext.bInDebugMode.load())
+	, bValidatedOnce(OtherContext.bValidatedOnce)
+	, ConfigKeyMap(OtherContext.ConfigKeyMap)
+	, RedirectTypeMap(OtherContext.RedirectTypeMap)
+{
+}
+
+void FCoreRedirectsContext::InitializeContext()
+{
+	if (bInitialized.load(std::memory_order_relaxed))
+	{
+		return;
+	}
+
+	{
+		FScopeCoreRedirectsWriteLockedContext WriteLockedContext(*this);
+
+#if !UE_BUILD_SHIPPING
+		if (FParse::Param(FCommandLine::Get(), TEXT("FullDebugCoreRedirects")))
+		{
+			// Enable debug mode and set to maximum verbosity
+			bInDebugMode.store(true, std::memory_order_relaxed);
+			UE_SET_LOG_VERBOSITY(LogCoreRedirects, VeryVerbose);
+			FCoreDelegates::OnFEngineLoopInitComplete.AddStatic(FCoreRedirects::ValidateAllRedirects);
+		}
+		else if (FParse::Param(FCommandLine::Get(), TEXT("DebugCoreRedirects")))
+		{
+			// Enable debug mode and increase log levels but don't show every message
+			bInDebugMode.store(true, std::memory_order_relaxed);
+			UE_SET_LOG_VERBOSITY(LogCoreRedirects, Verbose);
+			FCoreDelegates::OnFEngineLoopInitComplete.AddStatic(FCoreRedirects::ValidateAllRedirects);
+		}
+#endif
+
+		// Setup map
+		ConfigKeyMap.Add(TEXT("ObjectRedirects"), ECoreRedirectFlags::Type_Object);
+		ConfigKeyMap.Add(TEXT("ClassRedirects"), ECoreRedirectFlags::Type_Class);
+		ConfigKeyMap.Add(TEXT("StructRedirects"), ECoreRedirectFlags::Type_Struct);
+		ConfigKeyMap.Add(TEXT("EnumRedirects"), ECoreRedirectFlags::Type_Enum);
+		ConfigKeyMap.Add(TEXT("FunctionRedirects"), ECoreRedirectFlags::Type_Function);
+		ConfigKeyMap.Add(TEXT("PropertyRedirects"), ECoreRedirectFlags::Type_Property);
+		ConfigKeyMap.Add(TEXT("PackageRedirects"), ECoreRedirectFlags::Type_Package);
+		ConfigKeyMap.Add(TEXT("AssetRedirects"), ECoreRedirectFlags::Type_Asset);
+
+		FCoreRedirects::RegisterNativeRedirectsUnderWriteLock(WriteLockedContext);
+
+		// Prepopulate RedirectTypeMap entries that some threads write to after the engine goes multi-threaded.
+		// Most RedirectTypeMap entries are written to only from InitUObject's call to ReadRedirectsFromIni, and at that point the Engine is single-threaded.
+		// Known missing packages and plugin loads can add entries to existing lists but will not add brand new types.
+		// Taking advantage of this, we treat the list of Key/Value pairs of RedirectTypeMap as immutable and read from it without synchronization.
+		// Note that the values for those written-during-multithreading entries need to be synchronized; it is only the list of Key/Value pairs that is immutable.
+		RedirectTypeMap.FindOrAdd(ECoreRedirectFlags::Type_Package | ECoreRedirectFlags::Category_Removed | ECoreRedirectFlags::Option_MissingLoad);
+
+		bInitialized.store(true, std::memory_order_release);
+	}
+}
+
+FCoreRedirectsContext& FCoreRedirectsContext::GetGlobalContext()
+{
+	using namespace UE::CoreRedirects::Private;
+	return GlobalCoreRedirectContext;
+}
+
+FCoreRedirectsContext& FCoreRedirectsContext::GetThreadContext()
+{
+	using namespace UE::CoreRedirects::Private;
+	return ThreadCoreRedirectContext ? *ThreadCoreRedirectContext : GetGlobalContext();
+}
+
+void FCoreRedirectsContext::SetThreadContext(FCoreRedirectsContext& NewContext)
+{
+	using namespace UE::CoreRedirects::Private;
+	ThreadCoreRedirectContext = &NewContext;
+}
+
+FScopeCoreRedirectsContext::FScopeCoreRedirectsContext(const FString& Name)
+	: ScopeContext(Name)
+	, ContextToRestore(FCoreRedirectsContext::GetThreadContext())
+{
+	FCoreRedirectsContext::SetThreadContext(ScopeContext);
+	ScopeContext.InitializeContext();
+}
+
+FScopeCoreRedirectsContext::FScopeCoreRedirectsContext(const FString& Name, FCoreRedirectsContext& ContextToCopyFrom)
+	: ScopeContext(Name, ContextToCopyFrom)
+	, ContextToRestore(FCoreRedirectsContext::GetThreadContext())
+{
+	FCoreRedirectsContext::SetThreadContext(ScopeContext);
+	ScopeContext.InitializeContext();
+}
+
+FScopeCoreRedirectsContext::~FScopeCoreRedirectsContext()
+{
+	FCoreRedirectsContext::SetThreadContext(ContextToRestore);
+}
+
+uint32 FCoreRedirectsContext::FRWLockWithExclusiveRecursion::LoadWriteLockOwnerThreadIdRelaxed() const
+{
+	uint32 ThreadId;
+	// Memory validation is disabled, as this is a hot code path.
+	UE_AUTORTFM_OPEN_NO_VALIDATION
+	{
+		ThreadId = WriteLockOwnerThreadId.load(std::memory_order_relaxed);
+	};
+	return ThreadId;
+}
+
+void FCoreRedirectsContext::FRWLockWithExclusiveRecursion::StoreWriteLockOwnerThreadId(uint32 ThreadId, std::memory_order Order)
+{
+	if (AutoRTFM::IsClosed())
+	{
+		// Memory validation is disabled, as this is a hot code path.
+		// If an abort occurs, we want to restore the variable's previous value, so we use
+		// `RecordOpenWrite` here so that AutoRTFM can undo our work. (We aren't overly
+		// concerned with atomicity as all loads are done in `memory_order_relaxed`.)
+		UE_AUTORTFM_OPEN_NO_VALIDATION
+		{
+			AutoRTFM::RecordOpenWrite(&ThreadId);
+			WriteLockOwnerThreadId.store(ThreadId, Order);
+		};
+	}
+	else
+	{
+		WriteLockOwnerThreadId.store(ThreadId, Order);
+	}
+}
+
+void FCoreRedirectsContext::FRWLockWithExclusiveRecursion::ReadLock()
+{
+	uint32 InitialLockOwner = LoadWriteLockOwnerThreadIdRelaxed();
+
+	// Avoid calling GetCurrentThreadId in the common case where we start out unlocked
+	// If we are unowned at the start, then there are three cases to consider:
+	// 1. We are unlocked. In that case we'll pass the InitialLockOwner check and call ReadLock()
+	// 2. We are already in a ReadLock() state. Taking the read lock again will cause us to deadlock. 
+	// That is part of the contract for this lock so it's okay.
+	// 3. We are write locked, in which case we'll enter the first block 
+	if (InitialLockOwner != 0)
+	{
+		// If we are initially write locked, there are two possibilities:
+		// 1. This thread owns the write lock. In that case, it's safe to just bump the recursion count
+		// 2. Another thread owns the write lock, in which case we want to block with a ReadLock()
+		if (InitialLockOwner == FPlatformTLS::GetCurrentThreadId())
+		{
+			// Allow recursive read locking by bumping the recursion count
+			RecursionCount++;
+		}
+		else
+		{
+			InternalLock.ReadLock();
+		}
+	}
+	else
+	{
+		// We may or may not still be locked at this point if another thread took a WriteLock
+		// but it doesn't matter because in either case we want to take a ReadLock and potentially wait
+		InternalLock.ReadLock();
+	}
+}
+
+void FCoreRedirectsContext::FRWLockWithExclusiveRecursion::WriteLock()
+{
+	uint32 CurrentThreadId = FPlatformTLS::GetCurrentThreadId();
+	if (LoadWriteLockOwnerThreadIdRelaxed() != CurrentThreadId)
+	{
+		InternalLock.WriteLock();
+		StoreWriteLockOwnerThreadId(CurrentThreadId);
+	}
+	RecursionCount++;
+}
+
+void FCoreRedirectsContext::FRWLockWithExclusiveRecursion::WriteUnlock()
+{
+	ensure(LoadWriteLockOwnerThreadIdRelaxed() == FPlatformTLS::GetCurrentThreadId());
+	ensure(RecursionCount > 0);
+	RecursionCount--;
+	if (RecursionCount == 0)
+	{
+		StoreWriteLockOwnerThreadId(0, std::memory_order_relaxed);
+		InternalLock.WriteUnlock();
+	}
+}
+
+void FCoreRedirectsContext::FRWLockWithExclusiveRecursion::ReadUnlock()
+{
+	if (RecursionCount > 0)
+	{
+		ensureMsgf(LoadWriteLockOwnerThreadIdRelaxed() == FPlatformTLS::GetCurrentThreadId(),
+			TEXT("Called ReadUnlock() on a lock held exclusively by another thread."));
+		RecursionCount--;
+	}
+	else
+	{
+		InternalLock.ReadUnlock();
+	}
+}
\ No newline at end of file
diff --git a/Engine/Source/Runtime/CoreUObject/Private/UObject/RedirectionSummary.cpp b/Engine/Source/Runtime/CoreUObject/Private/UObject/RedirectionSummary.cpp
--- a/Engine/Source/Runtime/CoreUObject/Private/UObject/RedirectionSummary.cpp
+++ b/Engine/Source/Runtime/CoreUObject/Private/UObject/RedirectionSummary.cpp
@@ -1,6 +1,6 @@
 // Copyright Epic Games, Inc. All Rights Reserved.
 
-#include "RedirectionSummary.h"
+#include "UObject/CoreRedirects/RedirectionSummary.h"
 
 #if WITH_EDITOR
 
diff --git a/Engine/Source/Runtime/CoreUObject/Private/UObject/RedirectionSummary.h b/Engine/Source/Runtime/CoreUObject/Private/UObject/RedirectionSummary.h
deleted file mode 100644
diff --git a/Engine/Source/Runtime/CoreUObject/Public/UObject/CoreRedirects.h b/Engine/Source/Runtime/CoreUObject/Public/UObject/CoreRedirects.h
--- a/Engine/Source/Runtime/CoreUObject/Public/UObject/CoreRedirects.h
+++ b/Engine/Source/Runtime/CoreUObject/Public/UObject/CoreRedirects.h
@@ -15,6 +15,7 @@
 #include "CoreTypes.h"
 #include "Misc/TransactionallySafeRWLock.h"
 #include "Misc/EnumClassFlags.h"
+#include "Templates/SharedPointer.h"
 #include "UObject/NameTypes.h"
 #include "UObject/UnrealNames.h"
 #include "UObject/CoreRedirects/PM-k.h"
@@ -24,9 +25,10 @@ class IPakFile;
 class UClass;
 struct FSoftObjectPath;
 struct FTopLevelAssetPath;
+struct FCoreRedirectsContext;
+struct FScopeCoreRedirectsReadLockedContext;
+struct FScopeCoreRedirectsWriteLockedContext;
 namespace UE::CoreRedirects::Private { struct FCoreRedirectObjectUtf8Name; }
-namespace UE::CoreRedirects::Private { class FRWWithExclusiveRecursionScopeLockForRead; }
-namespace UE::CoreRedirects::Private { class FRWWithExclusiveRecursionScopeLockForWrite; }
 
 DECLARE_LOG_CATEGORY_EXTERN(LogCoreRedirects, Log, All);
 
@@ -408,39 +410,37 @@ struct FCoreRedirects
 	static COREUOBJECT_API void RemoveAllAssetRedirects();
 	 
 private:
-	typedef UE::CoreRedirects::Private::FRWWithExclusiveRecursionScopeLockForRead FCoreRedirectorScopeLockForRead;
-	typedef UE::CoreRedirects::Private::FRWWithExclusiveRecursionScopeLockForWrite FCoreRedirectorScopeLockForWrite;
 
 	/** Static only class, never constructed */
 	COREUOBJECT_API FCoreRedirects();
 
 	/** Internal implementation for AddRedirectList that requires a write lock to already have been acquired */
 	static COREUOBJECT_API bool AddRedirectListUnderWriteLock(TArrayView<const FCoreRedirect> Redirects, 
-		const FString& SourceString, const FCoreRedirectorScopeLockForWrite& HeldLock);
+		const FString& SourceString, FScopeCoreRedirectsWriteLockedContext& LockedContext);
 
 	/** Internal implementation for AddSingleRedirect that requires a write lock to already have been acquired */
 	static COREUOBJECT_API bool AddSingleRedirectUnderWriteLock(const FCoreRedirect& NewRedirect, 
-		const FString& SourceString, const FCoreRedirectorScopeLockForWrite& HeldLock);
+		const FString& SourceString, FScopeCoreRedirectsWriteLockedContext& LockedContext);
 
 	/** Internal implementation for RemoveSingleRedirect that requires a write lock to already have been acquired */
 	static COREUOBJECT_API bool RemoveSingleRedirectUnderWriteLock(const FCoreRedirect& OldRedirect, 
-		const FString& SourceString, const FCoreRedirectorScopeLockForWrite& HeldLock);
+		const FString& SourceString, FScopeCoreRedirectsWriteLockedContext& LockedContext);
 
 	/** Add native redirects, called before ini is parsed for the first time */
-	static COREUOBJECT_API void RegisterNativeRedirectsUnderWriteLock(const FCoreRedirectorScopeLockForWrite& HeldLock);
+	static COREUOBJECT_API void RegisterNativeRedirectsUnderWriteLock(FScopeCoreRedirectsWriteLockedContext& LockedContext);
 
 	/** Internal implementation for GetMatchingRedirects that requires a read lock to already have been acquired */
 	static COREUOBJECT_API bool GetMatchingRedirectsUnderReadLock(ECoreRedirectFlags Type, 
 		const FCoreRedirectObjectName& OldObjectName, TArray<const FCoreRedirect*>& FoundRedirects, ECoreRedirectMatchFlags MatchFlags, 
-		const FCoreRedirectorScopeLockForRead& HeldLock);
+		FScopeCoreRedirectsReadLockedContext& LockedContext);
 
 	/** Internal implementation for RedirectNameAndValues that requires a read lock to already have been acquired */
 	static COREUOBJECT_API bool RedirectNameAndValuesUnderReadLock(ECoreRedirectFlags Type, const FCoreRedirectObjectName& OldObjectName,
 		FCoreRedirectObjectName& NewObjectName, const FCoreRedirect** FoundValueRedirect,
-		ECoreRedirectMatchFlags MatchFlags, const FCoreRedirectorScopeLockForRead& HeldLock);
+		ECoreRedirectMatchFlags MatchFlags, FScopeCoreRedirectsReadLockedContext& LockedContext);
 
 	/** Internal implementation for ValidateAssetRedirects that requires a read lock to already have been acquired */
-	static COREUOBJECT_API bool ValidateAssetRedirectsUnderReadLock(const FCoreRedirectorScopeLockForRead& HeldLock);
+	static COREUOBJECT_API bool ValidateAssetRedirectsUnderReadLock(FScopeCoreRedirectsReadLockedContext& LockedContext);
 
 	/** Container for managing Wildcard redirects (substrings, prefixes, suffixes) */
 	struct FWildcardData
@@ -467,21 +467,9 @@ private:
 		/** Map from name of thing being mapped to full list. List must be filtered further */
 		TMap<FName, TArray<FCoreRedirect>> RedirectMap;
 		/** Used to manage wildcard data and accelerate wildcard queries */
-		TUniquePtr<FWildcardData> Wildcards;
+		TSharedPtr<FWildcardData> Wildcards;
 	};
 
-	/** Whether this has been initialized at least once */
-	static COREUOBJECT_API std::atomic<bool> bInitialized;
-
-	/** True if we are in debug mode that does extra validation */
-	static COREUOBJECT_API std::atomic<bool> bInDebugMode;
-
-	/** True if we have done our initial validation. After initial validation, each change to redirects will validate independently */
-	static COREUOBJECT_API bool bValidatedOnce;
-
-	/** Map from config name to flag */
-	static COREUOBJECT_API TMap<FName, ECoreRedirectFlags> ConfigKeyMap;
-
 	/** Map from name of thing being mapped to full list. List must be filtered further */
 	struct FRedirectTypeMap
 	{
@@ -496,38 +484,8 @@ private:
 		TMap<ECoreRedirectFlags, FRedirectNameMap*> Map;
 		TArray<TPair<ECoreRedirectFlags, FRedirectNameMap>> FastIterable;
 	};
-	static COREUOBJECT_API FRedirectTypeMap RedirectTypeMap;
-
-	/** This lock allows exclusive locking (WriteLock) and shared locking (ReadLock)
-	 *  Additionally, it permits limited types of recursion. It is possible to ReadLock() 
-	 *  or WriteLock() while locked for write. It is NOT possible to Read or WriteLock() while
-	 *  locked for read. I.e., if the lock is held exclusively, it re-acquiring it is always permitted.
-	 *  If the lock is held shared, re-acquiring it is never permitted.
-	 */
-	struct FRWLockWithExclusiveRecursion
-	{
-		void ReadLock();
-
-		void WriteLock();
-
-		void WriteUnlock();
 
-		void ReadUnlock();
-
-	private:
-		uint32 LoadWriteLockOwnerThreadIdRelaxed() const;
-		void StoreWriteLockOwnerThreadId(uint32 ThreadId, std::memory_order Order = std::memory_order_seq_cst);
-
-		FTransactionallySafeRWLock InternalLock;
-		std::atomic<uint32> WriteLockOwnerThreadId = 0;
-		int32 RecursionCount = 0;
-	};
-	friend class UE::CoreRedirects::Private::FRWWithExclusiveRecursionScopeLockForRead;
-	friend class UE::CoreRedirects::Private::FRWWithExclusiveRecursionScopeLockForWrite;
-
-	/**
-	 * Lock to protect multithreaded access to the CoreRedirect system
-	 */
-	static COREUOBJECT_API FRWLockWithExclusiveRecursion RWLock;
+	friend struct FCoreRedirectsContext;
 };
 
+
