[CoreRedirects]

Suppress warning about returning local reference. We are returning a global reference. Fixes issue 689807

#rnx

diff --git a/Engine/Source/Runtime/CoreUObject/Private/UObject/CoreRedirectsContext.cpp b/Engine/Source/Runtime/CoreUObject/Private/UObject/CoreRedirectsContext.cpp
--- a/Engine/Source/Runtime/CoreUObject/Private/UObject/CoreRedirectsContext.cpp
+++ b/Engine/Source/Runtime/CoreUObject/Private/UObject/CoreRedirectsContext.cpp
@@ -83,7 +83,7 @@ void FCoreRedirectsContext::InitializeContext()
 FCoreRedirectsContext& FCoreRedirectsContext::GetGlobalContext()
 {
 	using namespace UE::CoreRedirects::Private;
-	return GlobalCoreRedirectContext;
+	return GlobalCoreRedirectContext; // -V558
 }
 
 FCoreRedirectsContext& FCoreRedirectsContext::GetThreadContext()
