Adding missing file from UHT change

#rnx

diff --git a/Engine/Source/Runtime/CoreUObject/Public/UObject/UObjectGlobals.h b/Engine/Source/Runtime/CoreUObject/Public/UObject/UObjectGlobals.h
--- a/Engine/Source/Runtime/CoreUObject/Public/UObject/UObjectGlobals.h
+++ b/Engine/Source/Runtime/CoreUObject/Public/UObject/UObjectGlobals.h
@@ -3509,7 +3509,7 @@ namespace UECodeGen_Private
 		FieldPath         = 0x1F,
 		LargeWorldCoordinatesReal = 0x20,
 		Optional          = 0x21,
-		VValue            = 0x22,
+		VerseValue        = 0x22,
 		Utf8Str           = 0x23,
 		AnsiStr           = 0x24,
 		VerseString       = 0x25,
