[AutoRTFM] Fix closed-transaction reconstruction of named UObjects

Partial reland of 38332615 (backed out with CL38342442) - this does not include the mutex changes to properties which negatively affected performance.

The memory validator flagged that performing the Memzero and UObjectBase in-place reconstruction was clobbering tracked memory.
It was right to flag this - aborting the transaction after reconstruction would leave the object in a zero'd state, not the state before reconstruction.

To fix this record the writes in the open block only if the object is being reconstructed:
* If a this is a new object, we *do not* want to record the write as reverting the memory would bring the object to an uninitialized memory state, which will cause the GC to access uninitialized memory.
* If this is a recycled object, then we *do* want to record the write so an abort will return the object to the previous object's state.

We disable validation on this write because this spans the entire object's memory including fields that may be accessed only in the open (e.g. object flags which are set atomically), and mixing closed writes with open writes will trigger the memory validator.

diff --git a/Engine/Source/Runtime/CoreUObject/Private/UObject/UObjectGlobals.cpp b/Engine/Source/Runtime/CoreUObject/Private/UObject/UObjectGlobals.cpp
--- a/Engine/Source/Runtime/CoreUObject/Private/UObject/UObjectGlobals.cpp
+++ b/Engine/Source/Runtime/CoreUObject/Private/UObject/UObjectGlobals.cpp
@@ -3503,7 +3503,8 @@ UObject* StaticAllocateObject
 	UObject* Obj = NULL;
 	if(InName == NAME_None)
 	{
-		AutoRTFM::Open([&]{
+		UE_AUTORTFM_OPEN
+		{
 #if WITH_EDITOR
 			if ( GOutputCookingWarnings && GetTransientPackage() != InOuter->GetOutermost() )
 			{
@@ -3514,7 +3515,7 @@ UObject* StaticAllocateObject
 			{
 				InName = MakeUniqueObjectName(InOuter, InClass);
 			}
-		});
+		};
 	}
 	else
 	{
@@ -3705,11 +3706,21 @@ UObject* StaticAllocateObject
 			*InClass->GetPathName(), *GetPathNameSafe(InOuter), *InName.ToString(), *OldRemoteId.ToString(), *RemoteId.ToString());
 #endif // UE_WITH_REMOTE_OBJECT_HANDLE
 
-		// perform the UObjectBase construction in the open - we expect the GC to invoke
-		// the destructor in the case of transaction abort
-		// Validation temporarily disabled until SOL-7253 is investigated
 		UE_AUTORTFM_OPEN_NO_VALIDATION
 		{
+			// If a this is a new object, we *do not* want to record the write as reverting the
+			// memory would bring the object to an uninitialized memory state, which will
+			// cause the GC to access uninitialized memory.
+			// If this is a recycled object, then we *do* want to record the write so an abort
+			// will return the object to the previous object's state.
+			if (bWasConstructedOnOldObject)
+			{
+				// We disable validation on this write because this spans the entire object's memory
+				// including fields that may be accessed only in the open (e.g. object flags which
+				// are set atomically), and mixing closed writes with open writes will trigger the
+				// memory validator.
+				AutoRTFM::RecordOpenWriteNoMemoryValidation((void *)Obj, TotalSize);
+			}
 			FMemory::Memzero((void *)Obj, TotalSize);
 			new ((void *)Obj) UObjectBase(const_cast<UClass*>(InClass), InFlags|RF_NeedInitialization, InternalSetFlags, InOuter, InName, OldIndex, OldSerialNumber == 0 ? SerialNumber : OldSerialNumber, 
 #if UE_WITH_REMOTE_OBJECT_HANDLE
@@ -4269,7 +4280,7 @@ void FObjectInitializer::PostConstructInit()
 	// need a fence to guarantee that the cleared flag is only visible to other threads
 	// after the other initialization-related writes
 	std::atomic_thread_fence(std::memory_order_release);
-	UE_AUTORTFM_OPEN{ Obj->ClearFlags(RF_NeedInitialization); };
+	Obj->ClearFlags(RF_NeedInitialization);
 
 	// clear the object pointer so we can guard against running this function again
 	Obj = nullptr;
diff --git a/Engine/Source/Runtime/Engine/Private/Tests/AutoRTFM/AutoRTFMActorComponentTests.cpp b/Engine/Source/Runtime/Engine/Private/Tests/AutoRTFM/AutoRTFMActorComponentTests.cpp
--- a/Engine/Source/Runtime/Engine/Private/Tests/AutoRTFM/AutoRTFMActorComponentTests.cpp
+++ b/Engine/Source/Runtime/Engine/Private/Tests/AutoRTFM/AutoRTFMActorComponentTests.cpp
@@ -3,7 +3,7 @@
 #include "Misc/AutomationTest.h"
 #include "Animation/AnimInstance.h"
 #include "Animation/Skeleton.h"
-#include "AutoRTFM/AutoRTFM.h"
+#include "AutoRTFM/AutoRTFMTesting.h"
 #include "AutoRTFMTestActor.h"
 #include "AutoRTFMTestAnotherActor.h"
 #include "AutoRTFMTestBodySetup.h"
@@ -20,6 +20,7 @@
 #include "Physics/Experimental/PhysScene_Chaos.h"
 #include "PhysicsProxy/SingleParticlePhysicsProxy.h"
 #include "Rendering/SkeletalMeshRenderData.h"
+#include "UObject/Package.h"
 #include "UObject/UObjectArray.h"
 #include "UObject/UObjectGlobals.h"
 
@@ -492,6 +493,27 @@ AUTORTFM_ACTOR_COMPONENT_TEST(WorldBlockTillLevelStreamingCompleted)
 	TestTrueExpr(AutoRTFM::ETransactionResult::Committed == Result);
 }
 
+AUTORTFM_ACTOR_COMPONENT_TEST(ReconstructObject)
+{
+	FName Name = TEXT("MyObjectToBeReplaced");
+	UObject* Outer = static_cast<UObject*>(GetTransientPackage());
+	AActor* Old = NewObject<AActor>(Outer, Name);
+	TestTrueExpr(0 == Old->ActorCategory);
+	Old->ActorCategory = 123;
+
+	AActor* New = nullptr;
+
+	AutoRTFM::Testing::Abort([&]
+		{
+			// Reconstruct the object by using the same name, then immediately abort.
+			New = NewObject<AActor>(Outer, Name);
+			AutoRTFM::AbortTransaction();
+		});
+
+	TestTrueExpr(nullptr == New);
+	TestTrueExpr(123 == Old->ActorCategory);
+}
+
 }  // anonymous namespace
 
 #undef AUTORTFM_ACTOR_COMPONENT_TEST
