MakeUniversallyUniqueObjectName now generate names with the following logic:
- Generate an immutable unique session id (64 bits) that won't leak any user data, in base-36 to support case-insensitive file systems.
- Generate a random id for every generated unique name, based on a simple linear congruential generator seeded with the session id.
- Generated names are ClassName_SessionId_RandomId, which results in a single name added to the name table for each user session.

Added NibbleToTCharLower along with existing NibbleToTChar.

diff --git a/Engine/Source/Runtime/Core/Public/Containers/UnrealString.h b/Engine/Source/Runtime/Core/Public/Containers/UnrealString.h
--- a/Engine/Source/Runtime/Core/Public/Containers/UnrealString.h
+++ b/Engine/Source/Runtime/Core/Public/Containers/UnrealString.h
@@ -66,7 +66,7 @@ inline int32 StringToBytes( const FString& String, uint8* OutBytes, int32 MaxBuf
 	return NumBytes;
 }
 
-/** Returns Char value of Nibble */
+/** Returns uppercase Char value of Nibble */
 [[nodiscard]] inline TCHAR NibbleToTChar(uint8 Num)
 {
 	if (Num > 9)
@@ -76,6 +76,16 @@ inline int32 StringToBytes( const FString& String, uint8* OutBytes, int32 MaxBuf
 	return (TCHAR)('0' + Num);
 }
 
+/** Returns lowercase Char value of Nibble */
+[[nodiscard]] inline TCHAR NibbleToTCharLower(uint8 Num)
+{
+	if (Num > 9)
+	{
+		return (TCHAR)('a' + (Num - 10));
+	}
+	return (TCHAR)('0' + Num);
+}
+
 /** 
  * Convert a byte to hex
  * @param In byte value to convert
diff --git a/Engine/Source/Runtime/CoreUObject/Private/UObject/UObjectGlobals.cpp b/Engine/Source/Runtime/CoreUObject/Private/UObject/UObjectGlobals.cpp
--- a/Engine/Source/Runtime/CoreUObject/Private/UObject/UObjectGlobals.cpp
+++ b/Engine/Source/Runtime/CoreUObject/Private/UObject/UObjectGlobals.cpp
@@ -9,7 +9,6 @@
 #include "HAL/PlatformFileManager.h"
 #include "HAL/FileManager.h"
 #include "Logging/StructuredLog.h"
-#include "Math/Int128.h"
 #include "Misc/AsciiSet.h"
 #include "Misc/EnumRange.h"
 #include "Misc/PackageAccessTrackingOps.h"
@@ -2912,51 +2911,68 @@ FName MakeUniqueObjectName(UObject* Parent, const UClass* Class, FName InBaseNam
 	return TestName;
 }
 
-FName MakeUniversallyUniqueObjectName(UObject* Outer, FName BaseName)
+FName MakeUniversallyUniqueObjectName(UObject* InOuter, FName InBaseName)
 {
 	// Generate a unique generation context for this session.
 	static struct FContext
 	{
-		FGuid Guid = FGuid::NewGuid();
-		FDateTime Origin = FDateTime(2024, 10, 1);
-		std::atomic<uint32> Counter = 0;
-	} Context;
+		FContext()
+		{
+			const FGuid Guid = FGuid::NewGuid();
+			
+			uint64 Value = ((uint64)Guid.A << 32) | Guid.C;
+			uint32 Index = 0;
 
-	FName OutObjectName;
+			while (Value)
+			{
+				const uint8 Remainder = uint8(Value % 36);
+				SessionIdBase36[Index++] = NibbleToTCharLower(Remainder);
+				check(Index < (sizeof(SessionIdBase36) / sizeof(SessionIdBase36[0])));
+				Value /= 36;
+			}
 
-	do
-	{
-		TStringBuilderWithBuffer<TCHAR, NAME_SIZE> StringBuilder;
+			SessionIdBase36[Index] = 0;
+
+			Seed = Guid.D;
+		}
 
-		// Append the provided base name if present.
-		if (!BaseName.IsNone())
+		uint32 GetNewValue() const
 		{
-			StringBuilder += BaseName.ToString();
-			StringBuilder += TEXT("_");
+			for (uint32 CurValue = Seed.load(std::memory_order_relaxed);;)
+			{
+				const uint32 NewValue = CurValue * 1664525 + 1013904223;
+				if (Seed.compare_exchange_weak(CurValue, NewValue))
+				{
+					return NewValue;
+				}
+			}
 		}
 
-		// Compute the amount of microseconds since a specific date to construct our GUID.
-		const FTimespan Elapsed = FDateTime::Now() - Context.Origin;
-		const uint64 ElapsedUs = (uint64)Elapsed.GetTotalMilliseconds() * 1000 + (Context.Counter++ % 1000);
+		TCHAR SessionIdBase36[14];
+		mutable std::atomic<uint32> Seed;
+	} const Context;
+	
+	TStringBuilder<NAME_SIZE> StringBuilder;
 
-		// The lowest 30 bits of the timestamp are stored in the name number below. This is to minimize the total names 
-		// generated for globally unique names (the string part will change every ~17 minutes).
-		FUInt128 Value(0, uint32(ElapsedUs >> 30), Context.Guid.A ^ Context.Guid.B, Context.Guid.C ^ Context.Guid.D);
+	// Append the provided base name if present.
+	if (!InBaseName.IsNone())
+	{
+		StringBuilder += InBaseName.ToString();
+		StringBuilder.AppendChar(TEXT('_'));
+	}
 
-		// Convert the constructed value in base 36
-		while (!Value.IsZero())
-		{
-			uint32 Remainder;
-			Value = Value.Divide(36, Remainder);
-			StringBuilder.AppendChar(NibbleToTChar(uint8(Remainder)));
-		}
+	StringBuilder += Context.SessionIdBase36;
 
-		// The name number bit 30 is reserved to know if this is a globally unique name. This is not 100% safe, but
-		// should cover most cases. The name number bit 31 is reserved to preserve the fast path name generation
-		// (see GFastPathUniqueNameGeneration).
-		OutObjectName = FName(*StringBuilder, (ElapsedUs & 0x3fffffff) | (1 << 30));
+	FName OutObjectName;
+	do
+	{
+		// Store the random part in the lowest 31 bits of the name number below. This is to minimize the total names 
+		// generated for globally unique names. The name number bit 31 is reserved to preserve the fast path names
+		// generation (see GFastPathUniqueNameGeneration).
+		int32 OutObjectIndex = Context.GetNewValue() & 0x7fffffff;
+		OutObjectName = FName(*StringBuilder, OutObjectIndex);		
 	}
-	while (StaticFindObjectFast(nullptr, Outer, OutObjectName));
+	while (StaticFindObjectFast(nullptr, InOuter, OutObjectName));
 
 	return OutObjectName;
 }
