Add UpdatePriority to Batch Read Bulkdata Streaming API

diff --git a/Engine/Source/Runtime/CoreUObject/Private/Serialization/BulkDataStreaming.cpp b/Engine/Source/Runtime/CoreUObject/Private/Serialization/BulkDataStreaming.cpp
--- a/Engine/Source/Runtime/CoreUObject/Private/Serialization/BulkDataStreaming.cpp
+++ b/Engine/Source/Runtime/CoreUObject/Private/Serialization/BulkDataStreaming.cpp
@@ -61,6 +61,7 @@ protected:
 	bool WaitForChunkRequest(float TimeLimitSeconds = 0.0f);
 	void CancelChunkRequest();
 	int64 GetSizeResult() const { return SizeResult; }
+	void UpdatePriority(const uint32 Pri) { Request.UpdatePriority(Pri); }
 	
 	FIoBuffer			Buffer;
 
@@ -480,6 +481,7 @@ public:
 	virtual FBulkDataRequest::EStatus GetStatus() const = 0;
 	virtual bool Cancel() = 0;
 	virtual bool Wait(uint32 Milliseconds) = 0;
+	virtual void UpdatePriority(const EAsyncIOPriorityAndFlags Priority) = 0;
 };
 
 //////////////////////////////////////////////////////////////////////////////
@@ -532,6 +534,10 @@ public:
 		return false;
 	}
 
+	virtual void UpdatePriority(const EAsyncIOPriorityAndFlags Priority) override
+	{
+	}
+
 private:
 	std::atomic<int32>	RefCount{0};
 	std::atomic<uint32> Status{0};
@@ -603,6 +609,15 @@ public:
 		return true;
 	}
 
+	virtual void UpdatePriority(const EAsyncIOPriorityAndFlags Priority) override
+	{
+		if (IoHandle.Status().GetErrorCode() <= EIoErrorCode::Unknown )
+		{
+			const int32 IoPriority = ConvertToIoDispatcherPriority(Priority);
+			IoHandle.UpdatePriority(IoPriority);
+		}
+	}
+
 	FBulkDataRequest::IHandle* Batch = nullptr;
 	FIoRequest IoHandle;
 };
@@ -655,6 +670,14 @@ public:
 		return DoneEvent.WaitFor(UE::FMonotonicTimeSpan::FromMilliseconds(Milliseconds));
 	}
 
+	virtual void UpdatePriority(const EAsyncIOPriorityAndFlags Priority) override final
+	{
+		for (FChunkBatchReadRequest& ReadRequest : Requests)
+		{
+			ReadRequest.UpdatePriority(Priority);
+		}
+	}
+
 	void Read(
 		const FIoChunkId& BulkChunkId,
 		const FIoReadOptions& Options,
@@ -829,6 +852,14 @@ bool FBulkDataBatchRequest::WaitFor(const FTimespan& WaitTime)
 	return WaitFor((uint32)FMath::Clamp<int64>(WaitTime.GetTicks() / ETimespan::TicksPerMillisecond, 0, MAX_uint32));
 }
 
+void FBulkDataBatchRequest::UpdatePriority(const EAsyncIOPriorityAndFlags Priority)
+{
+	if (Handle.IsValid() && Handle->GetStatus() <= EStatus::Pending) 
+	{
+		Handle->UpdatePriority(Priority);
+	}
+}
+
 //////////////////////////////////////////////////////////////////////////////
 
 FBulkDataBatchRequest::FBuilder::FBuilder(int32 MaxCount)
diff --git a/Engine/Source/Runtime/CoreUObject/Public/Serialization/BulkData.h b/Engine/Source/Runtime/CoreUObject/Public/Serialization/BulkData.h
--- a/Engine/Source/Runtime/CoreUObject/Public/Serialization/BulkData.h
+++ b/Engine/Source/Runtime/CoreUObject/Public/Serialization/BulkData.h
@@ -1318,6 +1318,12 @@ public:
 	 */
 	COREUOBJECT_API bool WaitFor(const FTimespan& WaitTime);
 
+	/**
+	 * Adjusts the priority of the request(s) (if possible).
+	 * @param Priority The new priority to use
+	 */
+	COREUOBJECT_API void UpdatePriority(const EAsyncIOPriorityAndFlags Priority);
+
 	/** Issue one or more I/O request in a single batch. */
 	class FBatchBuilder : public FBuilder
 	{
