Enforce that structs with AddStructReferencedObjects() are pointer-aligned

GC assumes they are, but TWeakObjectPtr is only 4B-aligned

#udn 14044102
#rnx

diff --git a/Engine/Source/Runtime/CoreUObject/Private/UObject/GarbageCollection.cpp b/Engine/Source/Runtime/CoreUObject/Private/UObject/GarbageCollection.cpp
--- a/Engine/Source/Runtime/CoreUObject/Private/UObject/GarbageCollection.cpp
+++ b/Engine/Source/Runtime/CoreUObject/Private/UObject/GarbageCollection.cpp
@@ -6540,6 +6540,7 @@ void FStructProperty::EmitReferenceInfo(UE::GC::FSchemaBuilder& Schema, int32 Ba
 	const int32 Offset = BaseOffset + GetOffset_ForGC();
 	if (Struct->StructFlags & STRUCT_AddStructReferencedObjects)
 	{
+		checkf(Struct->GetCppStructOps()->GetAlignment() >= 8, TEXT("'%s' has AddStructReferencedObjects(), which requires at least 8B alignment "), *Struct->GetPathName());
 		StructAROFn StructARO = Struct->GetCppStructOps()->AddStructReferencedObjects();	
 		for (int32 Idx = 0, Num = ArrayDim; Idx < Num; ++Idx)
 		{	
diff --git a/Engine/Source/Runtime/CoreUObject/Public/UObject/Class.h b/Engine/Source/Runtime/CoreUObject/Public/UObject/Class.h
--- a/Engine/Source/Runtime/CoreUObject/Public/UObject/Class.h
+++ b/Engine/Source/Runtime/CoreUObject/Public/UObject/Class.h
@@ -1317,6 +1317,7 @@ public:
 		TCppStructOps()
 			: ICppStructOps(sizeof(CPPSTRUCT), alignof(CPPSTRUCT))
 		{
+			static_assert(!(TTraits::WithAddStructReferencedObjects && alignof(CPPSTRUCT) < alignof(UObject*)), "AddStructReferencedObjects() requires pointer alignment");
 		}
 
 		virtual FCapabilities GetCapabilities() const override
