Initial implementation of FRemoteObjectReference - a structure that wraps the process of sharing references to remote objects with different servers via RPC

diff --git a/Engine/Source/Runtime/CoreUObject/Private/UObject/Class.cpp b/Engine/Source/Runtime/CoreUObject/Private/UObject/Class.cpp
--- a/Engine/Source/Runtime/CoreUObject/Private/UObject/Class.cpp
+++ b/Engine/Source/Runtime/CoreUObject/Private/UObject/Class.cpp
@@ -45,6 +45,7 @@
 #include "UObject/PropertyHelper.h"
 #include "UObject/CoreRedirects.h"
 #include "UObject/ObjectMacros.h"
+#include "UObject/RemoteObjectTransfer.h"
 #include "Internationalization/PolyglotTextData.h"
 #include "Serialization/ArchiveScriptReferenceCollector.h"
 #include "Serialization/ArchiveUObjectFromStructuredArchive.h"
@@ -7922,6 +7923,13 @@ UScriptStruct* TBaseStructure<FRemoteServerId>::Get()
 	return ScriptStruct;
 }
 
+UScriptStruct* TBaseStructure<FRemoteObjectReference>::Get()
+{
+	static UScriptStruct* ScriptStruct = StaticGetBaseStructureInternal(TEXT("RemoteObjecReference"));
+	return ScriptStruct;
+}
+UE_IMPLEMENT_STRUCT("/Script/CoreUObject", RemoteObjectReference);
+
 #define UE_DEFINE_CORE_VARIANT_TYPE(VARIANT, CORE)												\
 UScriptStruct* TBaseStructure<F##CORE>::Get()												\
 {																								\
diff --git a/Engine/Source/Runtime/CoreUObject/Private/UObject/RemoteObject.cpp b/Engine/Source/Runtime/CoreUObject/Private/UObject/RemoteObject.cpp
--- a/Engine/Source/Runtime/CoreUObject/Private/UObject/RemoteObject.cpp
+++ b/Engine/Source/Runtime/CoreUObject/Private/UObject/RemoteObject.cpp
@@ -124,7 +124,8 @@ void RegisterRemoteObjectId(FRemoteObjectId ObjectId, FRemoteServerId OwnerServe
 
 void RegisterSharedObject(UObject* Object)
 {
-	MarkAsLocal(Object);
+	FUObjectItem* ObjectItem = GUObjectArray.ObjectToObjectItem(Object);
+	ObjectItem->SetFlags(EInternalObjectFlags::RemoteReference);
 }
 
 void MarkAsRemote(UObject* Object, FRemoteServerId DestinationServerId)
@@ -393,3 +394,9 @@ FArchive& operator<<(FArchive& Ar, FRemoteObjectId& Id)
 	Ar << Id.Id;
 	return Ar;
 }
+
+FArchive& operator<<(FArchive& Ar, FRemoteServerId& Id)
+{
+	Ar << Id.Id;
+	return Ar;
+}
diff --git a/Engine/Source/Runtime/CoreUObject/Private/UObject/RemoteObjectTransfer.cpp b/Engine/Source/Runtime/CoreUObject/Private/UObject/RemoteObjectTransfer.cpp
--- a/Engine/Source/Runtime/CoreUObject/Private/UObject/RemoteObjectTransfer.cpp
+++ b/Engine/Source/Runtime/CoreUObject/Private/UObject/RemoteObjectTransfer.cpp
@@ -643,4 +643,61 @@ void StoreObjectToDatabase(UObject* Object)
 	} while (Object->HasAnyInternalFlags(EInternalObjectFlags::RemoteReference));
 }
 
-} // namespace UE::RemoteObject::Transfer::Private
\ No newline at end of file
+} // namespace UE::RemoteObject::Transfer::Private
+
+
+FRemoteObjectReference::FRemoteObjectReference(const FObjectPtr& Ptr)
+{
+	ObjectId = Ptr.GetRemoteId();
+	if (ObjectId.IsValid())
+	{
+		if (Ptr.IsResolved())
+		{
+			ServerId = UE::RemoteObject::GetGlobalServerId();
+			UE::RemoteObject::Private::RegisterSharedObject(Ptr.Get());
+		}
+		else if (UE::RemoteObject::Handle::FRemoteObjectStub* Stub = UE::RemoteObject::Private::FindRemoteObjectStub(ObjectId))
+		{
+			ServerId = Stub->RemoteServerId;
+		}
+	}
+}
+
+FRemoteObjectReference::FRemoteObjectReference(const FWeakObjectPtr& WeakPtr)
+#if UE_WITH_REMOTE_OBJECT_HANDLE
+	: FRemoteObjectReference(FObjectPtr(WeakPtr.GetRemoteId()))
+#endif
+{
+}
+
+FObjectPtr FRemoteObjectReference::ToObjectPtr() const
+{
+	return FObjectPtr(ObjectId);
+}
+
+FWeakObjectPtr FRemoteObjectReference::ToWeakPtr() const
+{
+#if UE_WITH_REMOTE_OBJECT_HANDLE
+	return FWeakObjectPtr(ObjectId);
+#else
+	return FWeakObjectPtr();
+#endif
+}
+
+UObject* FRemoteObjectReference::Resolve() const
+{
+	return ToObjectPtr().Get();
+}
+
+FArchive& operator<<(FArchive& Ar, FRemoteObjectReference& Ref)
+{
+	Ar << Ref.ObjectId;
+	Ar << Ref.ServerId;
+
+	if (Ar.IsLoading())
+	{
+		UE::RemoteObject::Private::RegisterRemoteObjectId(Ref.ObjectId, Ref.ServerId);
+	}
+
+	return Ar;
+}
\ No newline at end of file
diff --git a/Engine/Source/Runtime/CoreUObject/Public/UObject/Class.h b/Engine/Source/Runtime/CoreUObject/Public/UObject/Class.h
--- a/Engine/Source/Runtime/CoreUObject/Public/UObject/Class.h
+++ b/Engine/Source/Runtime/CoreUObject/Public/UObject/Class.h
@@ -4530,6 +4530,12 @@ template<> struct TBaseStructure<FRemoteServerId>
 	static COREUOBJECT_API UScriptStruct* Get();
 };
 
+struct FRemoteObjectReference;
+template<> struct TBaseStructure<FRemoteObjectReference>
+{
+	static COREUOBJECT_API UScriptStruct* Get();
+};
+
 // TBaseStructure for explicit core variant types only. e.g. FVector3d returns "Vector3d" struct. 
 template< class T > struct TVariantStructure
 {
diff --git a/Engine/Source/Runtime/CoreUObject/Public/UObject/NoExportTypes.h b/Engine/Source/Runtime/CoreUObject/Public/UObject/NoExportTypes.h
--- a/Engine/Source/Runtime/CoreUObject/Public/UObject/NoExportTypes.h
+++ b/Engine/Source/Runtime/CoreUObject/Public/UObject/NoExportTypes.h
@@ -3441,6 +3441,18 @@ struct FRemoteServerId
 	uint32 Id;
 };
 
+USTRUCT(noexport, BlueprintType)
+struct FRemoteObjectReference
+{
+	GENERATED_BODY()
+
+	UPROPERTY()
+	FRemoteObjectId ObjectId;
+
+	UPROPERTY()
+	FRemoteServerId ServerId;
+};
+
 /// @endcond
 
 #endif
diff --git a/Engine/Source/Runtime/CoreUObject/Public/UObject/RemoteObjectTransfer.h b/Engine/Source/Runtime/CoreUObject/Public/UObject/RemoteObjectTransfer.h
--- a/Engine/Source/Runtime/CoreUObject/Public/UObject/RemoteObjectTransfer.h
+++ b/Engine/Source/Runtime/CoreUObject/Public/UObject/RemoteObjectTransfer.h
@@ -5,7 +5,7 @@
 #include "CoreTypes.h"
 #include "Misc/AssertionMacros.h"
 #include "UObject/RemoteObjectTypes.h"
-#include "UObject/Object.h"
+#include "UObject/Class.h"
 #include "Async/Async.h"
 
 namespace UE::RemoteObject::Transfer
@@ -156,3 +156,69 @@ namespace UE::RemoteObject::Transfer
 
 } // namespace UE::RemoveObject::Transfer
 
+
+struct FRemoteObjectReference
+{
+private:
+	/** Object id being shared with another server */
+	FRemoteObjectId ObjectId;
+	/** Id of a server that shared the object (last owner of the object) */
+	FRemoteServerId ServerId;
+
+public:
+	FRemoteObjectReference() = default;
+
+	COREUOBJECT_API explicit FRemoteObjectReference(const FObjectPtr& Ptr);
+
+	template <typename T>
+	explicit FRemoteObjectReference(const TObjectPtr<T>& Ptr)
+		: FRemoteObjectReference(FObjectPtr(Ptr.GetHandle()))
+	{
+	}
+
+	COREUOBJECT_API explicit FRemoteObjectReference(const FWeakObjectPtr& WeakPtr);
+
+	template <typename T>
+	explicit FRemoteObjectReference(const TWeakObjectPtr<T>& WeakPtr)
+		: FRemoteObjectReference((FWeakObjectPtr)WeakPtr)
+	{
+	}
+
+	COREUOBJECT_API FObjectPtr ToObjectPtr() const;
+	COREUOBJECT_API FWeakObjectPtr ToWeakPtr() const;
+	COREUOBJECT_API UObject* Resolve() const;
+
+	FORCEINLINE FRemoteObjectId GetRemoteId() const
+	{
+		return ObjectId;
+	}
+	FORCEINLINE FRemoteServerId GetSharingServerId() const
+	{
+		return ServerId;
+	}
+
+	bool Serialize(FArchive& Ar)
+	{
+		Ar << *this;
+		return true;
+	}
+
+	bool NetSerialize(FArchive& Ar, class UPackageMap* Map, bool& bOutSuccess)
+	{
+		Ar << *this;
+		bOutSuccess = true;
+		return bOutSuccess;
+	}
+
+	friend COREUOBJECT_API FArchive& operator<<(FArchive& Ar, FRemoteObjectReference& Ref);
+};
+
+template<>
+struct TStructOpsTypeTraits<FRemoteObjectReference> : public TStructOpsTypeTraitsBase2<FRemoteObjectReference>
+{
+	enum
+	{
+		WithNetSerializer = true,
+		WithSerializer = true,
+	};
+};
\ No newline at end of file
diff --git a/Engine/Source/Runtime/CoreUObject/Public/UObject/RemoteObjectTypes.h b/Engine/Source/Runtime/CoreUObject/Public/UObject/RemoteObjectTypes.h
--- a/Engine/Source/Runtime/CoreUObject/Public/UObject/RemoteObjectTypes.h
+++ b/Engine/Source/Runtime/CoreUObject/Public/UObject/RemoteObjectTypes.h
@@ -41,6 +41,8 @@ struct FRemoteServerId
 		return Id < Other.Id;
 	}
 
+	friend COREUOBJECT_API FArchive& operator<<(FArchive& Ar, FRemoteServerId& Id);
+
 private:
 	uint32 Id = 0;
 };
@@ -141,4 +143,4 @@ struct FRemoteObjectData
 	TArray<FName> Names;
 	TArray<FRemoteObjectPathName> PathNames;
 	TArray<uint8> Bytes;
-};
\ No newline at end of file
+};
