Disable split pins on instanced structs because it results in an incorrect display of an exec pin which would always cause the node to fail to compile.

diff --git a/Engine/Source/Runtime/CoreUObject/Public/StructUtils/InstancedStruct.h b/Engine/Source/Runtime/CoreUObject/Public/StructUtils/InstancedStruct.h
--- a/Engine/Source/Runtime/CoreUObject/Public/StructUtils/InstancedStruct.h
+++ b/Engine/Source/Runtime/CoreUObject/Public/StructUtils/InstancedStruct.h
@@ -23,7 +23,7 @@ class UUserDefinedStruct;
  *	UPROPERTY(EditAnywhere, Category = Foo, meta = (BaseStruct = "/Script/ModuleName.TestStructBase"))
  *	TArray<FInstancedStruct> TestArray;
  */
-USTRUCT(BlueprintType, meta = (HasNativeMake = "/Script/Engine.BlueprintInstancedStructLibrary.MakeInstancedStruct"))
+USTRUCT(BlueprintType, meta = (DisableSplitPin, HasNativeMake = "/Script/Engine.BlueprintInstancedStructLibrary.MakeInstancedStruct"))
 struct COREUOBJECT_API FInstancedStruct
 {
 	GENERATED_BODY()
