When migrating objects, serialize all FObjectProperty values as FObjectPtr to avoid resolving remote objects

diff --git a/Engine/Source/Runtime/CoreUObject/Private/UObject/PropertyObject.cpp b/Engine/Source/Runtime/CoreUObject/Private/UObject/PropertyObject.cpp
--- a/Engine/Source/Runtime/CoreUObject/Private/UObject/PropertyObject.cpp
+++ b/Engine/Source/Runtime/CoreUObject/Private/UObject/PropertyObject.cpp
@@ -175,7 +175,7 @@ void FObjectProperty::SerializeItem(FStructuredArchive::FSlot Slot, void* Value,
 			Slot << (*ObjectPtr);
 		}
 #if !(UE_BUILD_TEST || UE_BUILD_SHIPPING) 
-		if (!UnderlyingArchive.IsSaving())
+		if (!UnderlyingArchive.IsSaving() && !UnderlyingArchive.IsMigratingRemoteObjects())
 		{
 			CheckValidObject(Value, CurrentValue);
 		}
@@ -184,7 +184,7 @@ void FObjectProperty::SerializeItem(FStructuredArchive::FSlot Slot, void* Value,
 	else
 	{
 #if UE_WITH_OBJECT_HANDLE_LATE_RESOLVE || UE_WITH_REMOTE_OBJECT_HANDLE
-		if (HasAnyPropertyFlags(CPF_TObjectPtr))
+		if (HasAnyPropertyFlags(CPF_TObjectPtr) || UnderlyingArchive.IsMigratingRemoteObjects())
 		{
 			FObjectHandle OriginalHandle = ObjectPtr->GetHandle();
 			Slot << *ObjectPtr;
