[UObject]

- Updates Rename API comment to remove warning about ForceNoResetLoaders which no longer appies
- Updates Rename API comment to clarify that a unique name will be generated for the object when the new outer is equal to the current outer, not just when using nullptr.

#rnx

diff --git a/Engine/Source/Runtime/CoreUObject/Public/UObject/Object.h b/Engine/Source/Runtime/CoreUObject/Public/UObject/Object.h
--- a/Engine/Source/Runtime/CoreUObject/Public/UObject/Object.h
+++ b/Engine/Source/Runtime/CoreUObject/Public/UObject/Object.h
@@ -682,7 +682,7 @@ public:
 
 	/** 
 	 * Rename this object to a unique name, or change its outer.
-	 * @warning Unless ForceNoResetLoaders is passed in, this will cause a flush of all level streaming.
+     * @remark A unique name will be generated if NewName is null and if either is true: NewOuter is null or NewOuter equals the current outer.
 	 * 
 	 * @param	NewName		The new name of the object, if null then NewOuter should be set
 	 * @param	NewOuter	New Outer this object will be placed within, if null it will use the current outer
