Add FTaskScope to AsyncLoading thread so it is correctly tagged.

#rnx

diff --git a/Engine/Source/Runtime/CoreUObject/Private/Serialization/AsyncLoading2.cpp b/Engine/Source/Runtime/CoreUObject/Private/Serialization/AsyncLoading2.cpp
--- a/Engine/Source/Runtime/CoreUObject/Private/Serialization/AsyncLoading2.cpp
+++ b/Engine/Source/Runtime/CoreUObject/Private/Serialization/AsyncLoading2.cpp
@@ -8623,6 +8623,8 @@ bool FAsyncLoadingThread2::Init()
 
 uint32 FAsyncLoadingThread2::Run()
 {
+	const FTaskTagScope AsyncLoadingScope(ETaskTag::EAsyncLoadingThread);
+
 	LLM_SCOPE(ELLMTag::AsyncLoading);
 
 	AsyncLoadingThreadID = FPlatformTLS::GetCurrentThreadId();
