When AppendLocalFullPath attempts to lookup package extension look in for packages in IOStore containers and pak chunks.
This prevents failure to find issues when packages are in UToc containers rather than a pak file.

diff --git a/Engine/Source/Runtime/CoreUObject/Private/Misc/PackagePath.cpp b/Engine/Source/Runtime/CoreUObject/Private/Misc/PackagePath.cpp
--- a/Engine/Source/Runtime/CoreUObject/Private/Misc/PackagePath.cpp
+++ b/Engine/Source/Runtime/CoreUObject/Private/Misc/PackagePath.cpp
@@ -795,10 +795,10 @@ void FPackagePath::AppendLocalFullPath(FStringBuilderBase& Builder, EPackageSegm
 		if (Extension == EPackageExtension::Unspecified)
 		{
 			check(PackageSegment == EPackageSegment::Header);
-			FPackagePath PathWithExtension;
-			if (IPackageResourceManager::Get().DoesPackageExist(*this, FBulkDataCookedIndex::Default, EPackageSegment::Header, &PathWithExtension))
+			FPackagePath FoundPackage;
+			if (FPackageName::DoesPackageExist(*this, &FoundPackage))
 			{
-				Extension = PathWithExtension.GetHeaderExtension();
+				Extension = FoundPackage.GetHeaderExtension();
 				// DoesPackageExist should not search for files with non-standard extensions, so we don't handle CustomExtensions here
 				check(Extension != EPackageExtension::Custom);
 
