UObject construction will no longer be recycling remote subobjects

diff --git a/Engine/Source/Runtime/CoreUObject/Private/UObject/UObjectGlobals.cpp b/Engine/Source/Runtime/CoreUObject/Private/UObject/UObjectGlobals.cpp
--- a/Engine/Source/Runtime/CoreUObject/Private/UObject/UObjectGlobals.cpp
+++ b/Engine/Source/Runtime/CoreUObject/Private/UObject/UObjectGlobals.cpp
@@ -3630,8 +3630,12 @@ UObject* StaticAllocateObject
 			);
 #endif//UE_BUILD_SHIPPING
 		}
-		// Subobjects are always created in the constructor, no need to re-create them here unless their archetype != CDO or they're blueprint generated.	
-		if (!bCreatingCDO && (!bCanRecycleSubobjects || !Obj->IsDefaultSubobject()))
+		// Subobjects are always created in the constructor, no need to re-create them here unless their archetype != CDO or they're blueprint generated or they're remote
+		if (!bCreatingCDO && (!bCanRecycleSubobjects || !Obj->IsDefaultSubobject()
+#if UE_WITH_REMOTE_OBJECT_HANDLE
+			|| UE::RemoteObject::Handle::IsRemote(Obj)
+#endif // UE_WITH_REMOTE_OBJECT_HANDLE
+			))
 		{
 			OldIndex = GUObjectArray.ObjectToIndex(Obj);
 			OldSerialNumber = GUObjectArray.GetSerialNumber(OldIndex);
