Preventing unnecessary migration of remote objects when performing operations related to object migration

diff --git a/Engine/Source/Runtime/CoreUObject/Private/UObject/RemoteObject.cpp b/Engine/Source/Runtime/CoreUObject/Private/UObject/RemoteObject.cpp
--- a/Engine/Source/Runtime/CoreUObject/Private/UObject/RemoteObject.cpp
+++ b/Engine/Source/Runtime/CoreUObject/Private/UObject/RemoteObject.cpp
@@ -25,6 +25,7 @@ namespace UE::RemoteObject::Private
 
 FRemoteServerId RemoteServerId;
 std::atomic<uint64> RemoteObjectSerialNumber(1);
+int32 UnsafeToMigrateObjects = 0; // This should go into TLS
 
 class FRemoteObjectMaps
 {
@@ -149,6 +150,22 @@ UE::RemoteObject::Handle::FRemoteObjectStub* FindRemoteObjectStub(FRemoteObjectI
 	return ObjectMaps->FindRemoteObjectStub(ObjectId);
 }
 
+FUnsafeToMigrateScope::FUnsafeToMigrateScope()
+{
+	UnsafeToMigrateObjects++;
+}
+FUnsafeToMigrateScope::~FUnsafeToMigrateScope()
+{
+	UnsafeToMigrateObjects--;
+	check(UnsafeToMigrateObjects >= 0);
+}
+
+bool IsSafeToMigrateObjects()
+{
+	// Not a thread safe test but atm we assume we're running single-threaded
+	return !(GIsGarbageCollecting || UnsafeToMigrateObjects);
+}
+
 } // namespace UE::RemoteObject::Private
 
 namespace UE::RemoteObject
@@ -208,7 +225,7 @@ UObject* ResolveObject(const FRemoteObjectStub* Stub, ERemoteReferenceType RefTy
 
 	// This is a slightly faster version of IsRemote(FRemoteObjectId) because we already know a Stub exists and we are going to re-use the Object pointer
 	UObject* Object = StaticFindObjectFastInternal(Stub->Id);
-	if (GIsGarbageCollecting && (Object || RefType == ERemoteReferenceType::Weak)) // Not a thread-safe test
+	if (!IsSafeToMigrateObjects() && (Object || RefType == ERemoteReferenceType::Weak)) // Not a thread-safe test
 	{
 		// Begin/FinishDestroy overrides may attempt to access subobjects of objects that have been migrated in which case we
 		// don't want to accidentally migrate them back mid-purge and if the Object memory is still valid (but has EInternalObjectFlags::Remote flag)
@@ -268,7 +285,7 @@ UObject* ResolveObject(UObject* Object, ERemoteReferenceType RefType /*= ERemote
 	// Begin/FinishDestroy overrides may attempt to access subobjects of objects that have been migrated in which case we
 	// don't want to accidentally migrate them back mid-purge and if the Object memory is still valid (but has EInternalObjectFlags::Remote flag)
 	// we can just return it and let the owner finish its cleanup.
-	if (!GIsGarbageCollecting) // Note: this is not a thread-safe check
+	if (IsSafeToMigrateObjects()) // Note: this is not a thread-safe check
 	{
 		FRemoteObjectStub* Stub = ObjectMaps->FindRemoteObjectStub(UE::CoreUObject::Private::FObjectHandleUtils::GetRemoteId(Object));
 		checkf(Stub, TEXT("Failed to find remote object stub for %s"), *GetPathNameSafe(Object));
@@ -309,7 +326,7 @@ bool CanResolveObject(FRemoteObjectId ObjectId)
 	{
 		// A stub exists so the object memory is not local but we can (attempt to) migrate it back if we're not garbage collecting
 		// Note: GIsGarbageCollecting checks are not thread safe
-		return !GIsGarbageCollecting;
+		return IsSafeToMigrateObjects();
 	}
 
 	// ObjectId is local or represents an object that has never been migrated
diff --git a/Engine/Source/Runtime/CoreUObject/Private/UObject/RemoteObjectPrivate.h b/Engine/Source/Runtime/CoreUObject/Private/UObject/RemoteObjectPrivate.h
--- a/Engine/Source/Runtime/CoreUObject/Private/UObject/RemoteObjectPrivate.h
+++ b/Engine/Source/Runtime/CoreUObject/Private/UObject/RemoteObjectPrivate.h
@@ -42,6 +42,12 @@ namespace UE::RemoteObject::Private
 	* Finds remote object stub
 	*/
 	UE::RemoteObject::Handle::FRemoteObjectStub* FindRemoteObjectStub(FRemoteObjectId ObjectId);
+
+	struct FUnsafeToMigrateScope
+	{
+		FUnsafeToMigrateScope();
+		~FUnsafeToMigrateScope();
+	};
 }
 
 namespace UE::RemoteObject::Transfer::Private
diff --git a/Engine/Source/Runtime/CoreUObject/Private/UObject/RemoteObjectSerialization.cpp b/Engine/Source/Runtime/CoreUObject/Private/UObject/RemoteObjectSerialization.cpp
--- a/Engine/Source/Runtime/CoreUObject/Private/UObject/RemoteObjectSerialization.cpp
+++ b/Engine/Source/Runtime/CoreUObject/Private/UObject/RemoteObjectSerialization.cpp
@@ -7,6 +7,7 @@
 #include "UObject/ObjectPtr.h"
 #include "UObject/Class.h"
 #include "UObject/PropertyPortFlags.h"
+#include "UObject/UObjectHash.h"
 #include "Serialization/ArchiveUObject.h"
 #include "Serialization/MemoryWriter.h"
 #include "Serialization/MemoryReader.h"
@@ -608,16 +609,39 @@ UObject* ConstructRemoteObject(const FRemoteObjectHeader& Header, FRemoteObjectC
 	using namespace UE::RemoteObject::Private;
 
 	UClass* Class = CastChecked<UClass>(Header.Class.Get());
-	UObject* Outer = Header.Outer.Get();	
+	UObject* Outer = nullptr;
+#if UE_WITH_REMOTE_OBJECT_HANDLE
+	if (Header.Outer.IsResolved())
+	{
+		Outer = Header.Outer.Get();
+	}
+	else
+	{
+		Outer = StaticFindObjectFastInternal(Header.Outer.GetRemoteId());
+		UE_CLOG(!Outer, LogRemoteObject, Fatal, TEXT("Failed to resolve an Outer when constructing remote object"))
+	}
+#endif
 	FName Name = Header.Name;
 
 	// The object may already exist in memory (it could be a default subobject of an object we've just created)
 	UObject* Object = StaticFindObjectFast(Class, Outer, Name);
-	if (Object && FRemoteObjectId(Object) != Header.RemoteId)
+	if (Object)
 	{
-		UE_LOG(LogRemoteObject, Warning, TEXT("Received remote object %s with identical pathname (%s) as a local object %s. Remote object will be renamed."), *Header.RemoteId.ToString(), *Object->GetPathName(), *FRemoteObjectId(Object).ToString());
-		Name = FName();
-		Object = nullptr;
+		if (FRemoteObjectId(Object) != Header.RemoteId)
+		{
+			UE_LOG(LogRemoteObject, Warning, TEXT("Received remote object %s with identical pathname (%s) as a local object %s. Remote object will be renamed."), *Header.RemoteId.ToString(), *Object->GetPathName(), *FRemoteObjectId(Object).ToString());
+			Name = FName();
+			Object = nullptr;
+		}
+	}
+	else
+	{
+		Object = StaticFindObjectFastInternal(Header.RemoteId);
+		if (Object)
+		{
+			// The object already exists on this server but has been renamed
+			Name = Object->GetFName();
+		}
 	}
 
 	// If not or the object is marked as remote (which means we brought it back before it was GC'd) (re)construct it
@@ -633,8 +657,13 @@ UObject* ConstructRemoteObject(const FRemoteObjectHeader& Header, FRemoteObjectC
 		Params.RemoteSubObjectOverrides = ConstructionOverrides;
 #endif
 
-		// Using StaticConstructObject_Internal to pass the extra parameters (RemoteId and SubobjectOverrides) which are not exposed to normal APIs
-		Object = StaticConstructObject_Internal(Params);
+		{
+			// In case we're allocating on top of existing object that's marked as remote don't try to resolve any of its references since they all are going to be destroyed anyway
+			FUnsafeToMigrateScope UnsafeToMigrateScope;
+
+			// Using StaticConstructObject_Internal to pass the extra parameters (RemoteId and SubobjectOverrides) which are not exposed to normal APIs
+			Object = StaticConstructObject_Internal(Params);
+		}
 	}
 
 	checkf(FRemoteObjectId(Object) == Header.RemoteId, TEXT("Created an object with a different ID:%s than requested:%s"), *FRemoteObjectId(Object).ToString(), *Header.RemoteId.ToString());
diff --git a/Engine/Source/Runtime/CoreUObject/Private/UObject/RemoteObjectTransfer.cpp b/Engine/Source/Runtime/CoreUObject/Private/UObject/RemoteObjectTransfer.cpp
--- a/Engine/Source/Runtime/CoreUObject/Private/UObject/RemoteObjectTransfer.cpp
+++ b/Engine/Source/Runtime/CoreUObject/Private/UObject/RemoteObjectTransfer.cpp
@@ -334,15 +334,21 @@ void MigrateObjectToRemoteServer(uint32 RequestToken, UObject* Object, FRemoteSe
 
 		checkf(SentObjects.Num() > 0, TEXT("PeerId:%s unable to migrate object %s to PeerId:%s"), *GetGlobalServerId().ToString(), *UE::CoreUObject::Private::FObjectHandleUtils::GetRemoteId(Object).ToString(), *DestinationServerId.ToString());
 
-		// Notify that objects have been migrated to a remote server
-		OnObjectsSentDelegate.Broadcast(SentObjects, DestinationServerId);
-
-		for (UObject* SentObject : SentObjects)
 		{
+			// PostMigrate callbacks may be manipulating migrated objects and may try to migrate them back so prevent that from happening since
+			// they all will be destroyed anyway
+			FUnsafeToMigrateScope UnsafeToMigrateScope;
+
+			// Notify that objects have been migrated to a remote server
+			OnObjectsSentDelegate.Broadcast(SentObjects, DestinationServerId);
+
+			for (UObject* SentObject : SentObjects)
+			{
 #if UE_WITH_REMOTE_OBJECT_HANDLE
-			SentObject->PostMigrate(DestinationServerId);
+				SentObject->PostMigrate(DestinationServerId);
 #endif
-			MarkAsRemote(SentObject, DestinationServerId);
+				MarkAsRemote(SentObject, DestinationServerId);
+			}
 		}
 	}
 }
