[BugFix] Added ArDontResetArrayData flag that makes array serialization idempotent. This fixes an issue with the CopyTaggedProperties method for IDOs
[BugFix] Set ArSerializingDefaults in CopyTaggedProperties when copying CDOs

diff --git a/Engine/Source/Runtime/Core/Private/Serialization/Archive.cpp b/Engine/Source/Runtime/Core/Private/Serialization/Archive.cpp
--- a/Engine/Source/Runtime/Core/Private/Serialization/Archive.cpp
+++ b/Engine/Source/Runtime/Core/Private/Serialization/Archive.cpp
@@ -194,6 +194,7 @@ PRAGMA_ENABLE_DEPRECATION_WARNINGS
 	ArCustomPropertyList				= nullptr;
 	ArUseCustomPropertyList				= false;
 	ArMergeOverrides					= false;
+	ArPreserveArrayElements				= false;
 	ArShouldSkipUpdateCustomVersion		= false;
 	ArIsMigratingRemoteObjects			= false;
 	SavePackageData						= nullptr;
@@ -257,6 +258,7 @@ PRAGMA_ENABLE_DEPRECATION_WARNINGS
 	ArCustomPropertyList                 = ArchiveToCopy.ArCustomPropertyList;
 	ArUseCustomPropertyList              = ArchiveToCopy.ArUseCustomPropertyList;
 	ArMergeOverrides					 = ArchiveToCopy.ArMergeOverrides;
+	ArPreserveArrayElements					 = ArchiveToCopy.ArPreserveArrayElements;
 	ArShouldSkipUpdateCustomVersion		 = ArchiveToCopy.ArShouldSkipUpdateCustomVersion;
 	ArIsMigratingRemoteObjects			 = ArchiveToCopy.ArIsMigratingRemoteObjects;
 	SavePackageData						 = ArchiveToCopy.SavePackageData;
diff --git a/Engine/Source/Runtime/Core/Public/Serialization/Archive.h b/Engine/Source/Runtime/Core/Public/Serialization/Archive.h
--- a/Engine/Source/Runtime/Core/Public/Serialization/Archive.h
+++ b/Engine/Source/Runtime/Core/Public/Serialization/Archive.h
@@ -895,6 +895,9 @@ public:
 	/** Set to false if OverriddenProperties should be cleared at the beginning of SerializeVersionedTaggedProperties. */
 	uint8 ArMergeOverrides : 1;
 
+	/** Set TRUE to stop arrays from being cleared  before they're loaded. Instead, array data will be merged */
+	uint8 ArPreserveArrayElements : 1;
+
 	/** Whether we are currently serializing defaults. > 0 means yes, <= 0 means no. */
 	int32 ArSerializingDefaults;
 
@@ -2134,6 +2137,7 @@ PRAGMA_ENABLE_DEPRECATION_WARNINGS
 	using FArchiveState::ArPortFlags;
 	using FArchiveState::ArMaxSerializeSize;
 	using FArchiveState::ArMergeOverrides;
+	using FArchiveState::ArPreserveArrayElements;
 
 public:
 	using FArchiveState::SetIsLoading;
diff --git a/Engine/Source/Runtime/CoreUObject/Private/UObject/InstanceDataObjectUtils.cpp b/Engine/Source/Runtime/CoreUObject/Private/UObject/InstanceDataObjectUtils.cpp
--- a/Engine/Source/Runtime/CoreUObject/Private/UObject/InstanceDataObjectUtils.cpp
+++ b/Engine/Source/Runtime/CoreUObject/Private/UObject/InstanceDataObjectUtils.cpp
@@ -834,9 +834,11 @@ namespace UE
 		FUObjectSerializeContext* SerializeContext = FUObjectThreadContext::Get().GetSerializeContext();
 		TGuardValue<bool> ImpersonatePropertiesScope(SerializeContext->bImpersonateProperties, true);
 
+		const bool bIsCDO = Source->HasAnyFlags(RF_ClassDefaultObject);
+		
 		// Track only initialized properties when copying. This is required to skip uninitialized properties
 		// during saving and to mark initialized properties during loading.
-		TGuardValue<bool> ScopedTrackInitializedProperties(SerializeContext->bTrackInitializedProperties, true);
+		TGuardValue<bool> ScopedTrackInitializedProperties(SerializeContext->bTrackInitializedProperties, !bIsCDO);
 		TGuardValue<bool> ScopedTrackSerializedProperties(SerializeContext->bTrackSerializedProperties, false);
 		TGuardValue<bool> ScopedTrackUnknownProperties(SerializeContext->bTrackUnknownProperties, false);
 		TGuardValue<bool> ScopedTrackUnknownEnumNames(SerializeContext->bTrackUnknownEnumNames, false);
@@ -852,6 +854,7 @@ namespace UE
 		FObjectReader Reader(Buffer);
 		FSerializingDefaultsScope ReaderDefaultsScope(Reader, Dest);
 		Reader.ArMergeOverrides = true;
+		Reader.ArPreserveArrayElements = true;
 		Dest->GetClass()->SerializeTaggedProperties(Reader, (uint8*)Dest, Dest->GetClass(), nullptr);
 	}
 
diff --git a/Engine/Source/Runtime/CoreUObject/Private/UObject/PropertyArray.cpp b/Engine/Source/Runtime/CoreUObject/Private/UObject/PropertyArray.cpp
--- a/Engine/Source/Runtime/CoreUObject/Private/UObject/PropertyArray.cpp
+++ b/Engine/Source/Runtime/CoreUObject/Private/UObject/PropertyArray.cpp
@@ -603,7 +603,7 @@ void FArrayProperty::SerializeItem(FStructuredArchive::FSlot Slot, void* Value,
 	if (UnderlyingArchive.IsLoading())
 	{
 		// If using a custom property list, don't empty the array on load. Not all indices may have been serialized, so we need to preserve existing values at those slots.
-		if (UnderlyingArchive.ArUseCustomPropertyList)
+		if (UnderlyingArchive.ArUseCustomPropertyList || UnderlyingArchive.ArPreserveArrayElements)
 		{
 			const int32 OldNum = ArrayHelper.Num();
 			if (ElementCount > OldNum)
