[BugFix] more bug fixes for property change event resolutions in IDO code.

- Object properties need to be redirected from instance pointers to IDO pointers where applicable. Essentially because we use impersonation during serialization rather than pointing object pointers in IDOs to other IDOs, we need to do the same sort of impersonation when resolving change events for IDOs

- An overlooked feature of FPropertyChangedChainEvent was that the chain may have more properties at the head or tail than the notification is referring to and that resolution callbacks (like this one) should only iterate from PropertyChain.GetActiveMemberNode() to PropertyChain.GetActiveNode()

- Added FOptionalProperty support to resolution

diff --git a/Engine/Source/Runtime/CoreUObject/Private/UObject/PropertyBagRepository.cpp b/Engine/Source/Runtime/CoreUObject/Private/UObject/PropertyBagRepository.cpp
--- a/Engine/Source/Runtime/CoreUObject/Private/UObject/PropertyBagRepository.cpp
+++ b/Engine/Source/Runtime/CoreUObject/Private/UObject/PropertyBagRepository.cpp
@@ -2,6 +2,8 @@
 
 #include "UObject/PropertyBagRepository.h"
 
+#include "UObject/PropertyPathFunctions.h"
+
 #if WITH_EDITORONLY_DATA
 
 #include "Containers/Queue.h"
@@ -181,300 +183,19 @@ void FPropertyBagRepository::ReassociateObjects(const TMap<UObject*, UObject*>&
 	}
 }
 
-static FProperty* FindPropertyByNameAndType(const UStruct* Struct, FName InName, FName Type)
+static FProperty* FindPropertyByNameAndType(const UStruct* Struct, const FProperty* SourceProperty)
 {
-	for (FProperty* Property = Struct->PropertyLink; Property != nullptr; Property = Property->PropertyLinkNext)
-	{
-		if (Property->GetFName() == InName && Property->GetID() == Type)
-		{
-			return Property;
-		}
-	}
-
-	return nullptr;
-}
-
-static bool ConstructRemappedPropertyChain(const FEditPropertyChain& Chain, FEditPropertyChain& NewChain, const UObject* Destination)
-{
-	UStruct* Struct = Destination->GetClass();
-	for (FEditPropertyChain::TDoubleLinkedListNode* Itr = Chain.GetHead(); Itr; Itr = Itr->GetNextNode())
-	{
-		FProperty* Property = Itr->GetValue();
-		Property = FindPropertyByNameAndType(Struct, Property->GetFName(), Property->GetID());
-		if (!Property)
-		{
-			NewChain.Empty();
-			return false;
-		}
-		NewChain.AddTail(Property);
-
-		// iterate the struct to look in
-		if (FOptionalProperty* AsOptionalProperty = CastField<FOptionalProperty>(Property))
-		{
-			Property = AsOptionalProperty->GetValueProperty();
-		}
-		else if (FArrayProperty* AsArrayProperty = CastField<FArrayProperty>(Property))
-		{
-			Property = AsArrayProperty->Inner;
-		}
-		else if (FSetProperty* AsSetProperty = CastField<FSetProperty>(Property))
-		{
-			Property = AsSetProperty->ElementProp;
-		}
-		else if (FMapProperty* AsMapProperty = CastField<FMapProperty>(Property))
-		{
-			Property = AsMapProperty->ValueProp;
-		}
-		
-		if (FStructProperty* AsStructProperty = CastField<FStructProperty>(Property))
-		{
-			Struct = AsStructProperty->Struct;
-		}
-		else if (FObjectProperty* AsObjectProperty = CastField<FObjectProperty>(Property))
-		{
-			Struct = AsObjectProperty->PropertyClass;
-		}
-		else
-		{
-			check(Itr->GetNextNode() == nullptr);
-		}
-
-		// remap active and active member nodes
-		if (Chain.GetActiveNode() == Itr)
-		{
-			NewChain.SetActivePropertyNode(NewChain.GetTail()->GetValue());
-		}
-		if (Chain.GetActiveMemberNode() == Itr)
-		{
-			NewChain.SetActiveMemberPropertyNode(NewChain.GetTail()->GetValue());
-		}
-	}
-	return true;
-}
-
-// note: this method is very similar to RemapChangeEvent so there's a good chance bugs found here will also need to be fixed there
-static void* ResolveChangePath(const void* StructData, FProperty*& ResolvedProperty, FPropertyChangedChainEvent& ChangeEvent)
-{
-	if (ChangeEvent.PropertyChain.GetHead() == nullptr)
-	{
-		return nullptr;
-	}
-		
-	FEditPropertyChain::TDoubleLinkedListNode* PropertyNode = ChangeEvent.PropertyChain.GetHead();
-	void* MemoryPtr = const_cast<void*>(StructData);
-	do
-	{
-		FProperty* Property = PropertyNode->GetValue();
-		MemoryPtr = Property->ContainerPtrToValuePtr<uint8>(MemoryPtr);
-		ResolvedProperty = Property;
-		
-		const int32 ArrayIndex = ChangeEvent.GetArrayIndex(Property->GetName());
-		if (ArrayIndex != INDEX_NONE && !(ChangeEvent.ChangeType & EPropertyChangeType::ArrayRemove))
-		{
-			PropertyNode = PropertyNode->GetNextNode();
-			FProperty* InnerProperty = PropertyNode ? PropertyNode->GetValue() : ChangeEvent.Property;
-			if (const FArrayProperty* ArrayProperty = CastField<FArrayProperty>(Property))
-			{
-				FScriptArrayHelper ArrayHelper(ArrayProperty, MemoryPtr);
-				if(!ArrayHelper.IsValidIndex(ArrayIndex))
-				{
-					return nullptr;
-				}
-				MemoryPtr = ArrayHelper.GetRawPtr(ArrayIndex);
-				ResolvedProperty = ArrayProperty->Inner;
-			}
-			else if (const FSetProperty* SetProperty = CastField<FSetProperty>(Property))
-			{
-				FScriptSetHelper SetHelper(SetProperty, MemoryPtr);
-				if(!SetHelper.IsValidIndex(ArrayIndex))
-				{
-					return nullptr;
-				}
-				MemoryPtr = SetHelper.GetElementPtr(ArrayIndex);
-				ResolvedProperty = SetProperty->ElementProp;
-			}
-			else if (const FMapProperty* MapProperty = CastField<FMapProperty>(Property))
-			{
-				FScriptMapHelper MapHelper(MapProperty, MemoryPtr);
-				if(!MapHelper.IsValidIndex(ArrayIndex))
-				{
-					return nullptr;
-				}
-				if (InnerProperty == MapProperty->KeyProp)
-				{
-					MemoryPtr = MapHelper.GetKeyPtr(ArrayIndex);
-					ResolvedProperty = MapProperty->KeyProp;
-				}
-				else
-				{
-					MemoryPtr = MapHelper.GetValuePtr(ArrayIndex);
-					ResolvedProperty = MapProperty->ValueProp;
-				}
-			}
-		}
-		else if (PropertyNode->GetNextNode())
-		{
-			if (const FObjectProperty* AsObjectProperty = CastField<FObjectProperty>(Property))
-			{
-				MemoryPtr = AsObjectProperty->GetObjectPropertyValue(MemoryPtr);
-			}
-		}
-
-		
-		if (PropertyNode)
-		{
-			PropertyNode = PropertyNode->GetNextNode();
-		}
-	}
-	while (PropertyNode);
-
-	return MemoryPtr;
-}
-
-
-// note: this method is very similar to ResolveChangePath so there's a good chance bugs found here will also need to be fixed there
-static bool RemapChangeEvent(const FPropertyChangedChainEvent& InChangeEvent, FPropertyChangedChainEvent& OutRemappedChangeEvent, void*& MemoryPtr, FProperty*& ResolvedProperty, TArray<TMap<FString, int32>>& ArrayIndices)
-{
-	
-	if (!ConstructRemappedPropertyChain(InChangeEvent.PropertyChain, OutRemappedChangeEvent.PropertyChain, static_cast<UObject*>(MemoryPtr)))
-	{
-		return false;
-	}
+	// use property impersonation for SaveTypeName so that keys of IDOs and non-IDOs match
+	FUObjectSerializeContext* SerializeContext = FUObjectThreadContext::Get().GetSerializeContext();
+	TGuardValue<bool> ScopedImpersonateProperties(SerializeContext->bImpersonateProperties, true);
 	
-	if (OutRemappedChangeEvent.PropertyChain.GetHead() == nullptr)
-	{
-		return false;
-	}
-		
-	FEditPropertyChain::TDoubleLinkedListNode* RemappedPropertyNode = OutRemappedChangeEvent.PropertyChain.GetHead();
-	FEditPropertyChain::TDoubleLinkedListNode* InPropertyNode = InChangeEvent.PropertyChain.GetHead();
-	do
-	{
-		FProperty* RemappedProperty = RemappedPropertyNode->GetValue();
-		FProperty* InProperty = InPropertyNode->GetValue();
-		MemoryPtr = RemappedProperty->ContainerPtrToValuePtr<uint8>(MemoryPtr);
-		ResolvedProperty = RemappedProperty;
-
-		if (InChangeEvent.Property == InProperty)
-		{
-			OutRemappedChangeEvent.Property = RemappedProperty;
-		}
-		if (InChangeEvent.MemberProperty == InProperty)
-		{
-			OutRemappedChangeEvent.MemberProperty = RemappedProperty;
-		}
-		
-		int32 ArrayIndex = InChangeEvent.GetArrayIndex(RemappedProperty->GetName());
-		if (ArrayIndex != INDEX_NONE && !(InChangeEvent.ChangeType & EPropertyChangeType::ArrayRemove))
-		{
-			RemappedPropertyNode = RemappedPropertyNode->GetNextNode();
-			
-			InPropertyNode = InPropertyNode->GetNextNode();
-			FProperty* InnerProperty = InPropertyNode ? InPropertyNode->GetValue() : InChangeEvent.Property;
-			if (const FArrayProperty* ArrayProperty = CastField<FArrayProperty>(RemappedProperty))
-			{
-				ResolvedProperty = ArrayProperty->Inner;
-				if (InChangeEvent.Property->Owner == InProperty)
-				{
-					OutRemappedChangeEvent.Property = ArrayProperty->Inner;
-				}
-				if (InChangeEvent.MemberProperty->Owner == InProperty)
-				{
-					OutRemappedChangeEvent.MemberProperty = ArrayProperty->Inner;
-				}
-				
-				FScriptArrayHelper ArrayHelper(ArrayProperty, MemoryPtr);
-				if(!ArrayHelper.IsValidIndex(ArrayIndex))
-				{
-					check(OutRemappedChangeEvent.ChangeType & EPropertyChangeType::ArrayAdd);
-					check(ArrayHelper.Num() == ArrayIndex);
-					ArrayHelper.Resize(ArrayIndex + 1);
-				}
-				ArrayIndices[InChangeEvent.ObjectIteratorIndex].Add(RemappedProperty->GetName(), ArrayIndex);
-				MemoryPtr = ArrayHelper.GetRawPtr(ArrayIndex);
-			}
-			if (const FSetProperty* SetProperty = CastField<FSetProperty>(RemappedProperty))
-			{
-				ResolvedProperty = SetProperty->ElementProp;
-				if (InChangeEvent.Property->Owner == InProperty)
-				{
-					OutRemappedChangeEvent.Property = SetProperty->ElementProp;
-				}
-				if (InChangeEvent.MemberProperty->Owner == InProperty)
-				{
-					OutRemappedChangeEvent.MemberProperty = SetProperty->ElementProp;
-				}
-				
-				FScriptSetHelper SetHelper(SetProperty, MemoryPtr);
-				if(!SetHelper.IsValidIndex(ArrayIndex))
-				{
-					check(OutRemappedChangeEvent.ChangeType & EPropertyChangeType::ArrayAdd);
-					ArrayIndex = SetHelper.AddUninitializedValue();
-				}
-				ArrayIndices[InChangeEvent.ObjectIteratorIndex].Add(RemappedProperty->GetName(), ArrayIndex);
-				MemoryPtr = SetHelper.GetElementPtr(ArrayIndex); // TODO: do more testing with sets
-			}
-			if (const FMapProperty* MapProperty = CastField<FMapProperty>(RemappedProperty))
-			{
-				if (InnerProperty == CastField<FMapProperty>(InProperty)->KeyProp) // TODO: do more testing with maps
-				{
-					ResolvedProperty = MapProperty->KeyProp;
-					if (InChangeEvent.Property->Owner == InProperty)
-					{
-						OutRemappedChangeEvent.Property = MapProperty->KeyProp;
-					}
-					if (InChangeEvent.MemberProperty->Owner == InProperty)
-					{
-						OutRemappedChangeEvent.MemberProperty = MapProperty->KeyProp;
-					}
-				}
-				else
-				{
-					ResolvedProperty = MapProperty->ValueProp;
-					if (InChangeEvent.Property->Owner == InProperty)
-					{
-						OutRemappedChangeEvent.Property = MapProperty->ValueProp;
-					}
-					if (InChangeEvent.MemberProperty->Owner == InProperty)
-					{
-						OutRemappedChangeEvent.MemberProperty = MapProperty->ValueProp;
-					}
-				}
-				
-				FScriptMapHelper MapHelper(MapProperty, MemoryPtr);
-				if(!MapHelper.IsValidIndex(ArrayIndex))
-				{
-					check(OutRemappedChangeEvent.ChangeType & EPropertyChangeType::ArrayAdd);
-					ArrayIndex = MapHelper.AddUninitializedValue();
-				}
-				ArrayIndices[InChangeEvent.ObjectIteratorIndex].Add(RemappedProperty->GetName(), ArrayIndex);
-				MemoryPtr = InnerProperty == CastField<FMapProperty>(InProperty)->KeyProp ? MapHelper.GetKeyPtr(ArrayIndex) : MapHelper.GetValuePtr(ArrayIndex);
-			}
-		}
-		else if (RemappedPropertyNode->GetNextNode())
-		{
-			if (const FObjectProperty* AsObjectProperty = CastField<FObjectProperty>(RemappedProperty))
-         	{
-         		MemoryPtr = AsObjectProperty->GetObjectPropertyValue(MemoryPtr);
-         	}
-		}
-		
-		if (RemappedPropertyNode)
-		{
-			RemappedPropertyNode = RemappedPropertyNode->GetNextNode();
-			InPropertyNode = InPropertyNode->GetNextNode();
-		}
-	}
-	while (RemappedPropertyNode);
-
-	OutRemappedChangeEvent.SetArrayIndexPerObject(ArrayIndices);
-
-	check(InChangeEvent.Property != OutRemappedChangeEvent.Property);
-	check(InChangeEvent.MemberProperty != OutRemappedChangeEvent.MemberProperty);
-
-	return true;
+	FPropertyTypeNameBuilder SourceTypeNameBuilder;
+	SourceProperty->SaveTypeName(SourceTypeNameBuilder);
+	FPropertyTypeName SourceTypeName = SourceTypeNameBuilder.Build();
+	return FindPropertyByNameAndTypeName(Struct, SourceProperty->GetFName(), SourceTypeName);
 }
 
+	
 static void CopyProperty(const FProperty* SourceProperty, const void* SourceValue, const FProperty* DestProperty, void* DestValue)
 {
 	check(SourceProperty->GetID() == DestProperty->GetID());
@@ -488,7 +209,7 @@ static void CopyProperty(const FProperty* SourceProperty, const void* SourceValu
 		UStruct* DestStruct = CastFieldChecked<FStructProperty>(DestProperty)->Struct;
 		for (FProperty* SourceChild : TFieldRange<FProperty>(SourceStruct))
 		{
-			if (FProperty* DestChild = FindPropertyByNameAndType(DestStruct, SourceChild->GetFName(), SourceChild->GetID()))
+			if (FProperty* DestChild = FindPropertyByNameAndType(DestStruct, SourceChild))
 			{
 				CopyProperty(SourceChild, SourceChild->ContainerPtrToValuePtr<void>(SourceValue),
 					DestChild, DestChild->ContainerPtrToValuePtr<void>(DestValue));
@@ -565,6 +286,285 @@ static void CopyProperty(const FProperty* SourceProperty, const void* SourceValu
 	}
 }
 
+
+// This is a utility struct that's used to take an existing FPropertyChangedChainEvent, and remap its property chain to instead have the properties
+// of a different class. In the process of remapping, it also resolves the data pointer of both the source property, and the resolved property.
+struct FRemappedChainEvent : public FPropertyChangedChainEvent
+{
+public:
+	FRemappedChainEvent(FPropertyChangedChainEvent& InEvent, const UObject* InSourceObject, const UObject* InObjectToResolve)
+		: FPropertyChangedChainEvent(this->FRemappedChainEvent::PropertyChain, InEvent)
+		, SourceEvent(InEvent)
+		, SourceChainNode(nullptr)
+		, ResolvedChainNode(nullptr)
+		, SourceField(InSourceObject->GetClass())
+		, ResolvedField(InObjectToResolve->GetClass())
+		, SourceMemoryPtr(const_cast<UObject*>(InSourceObject))
+		, ResolvedMemoryPtr(const_cast<UObject*>(InObjectToResolve))
+		, SourceObject(InSourceObject)
+		, ObjectToResolve(InObjectToResolve)
+	{
+		ResolvedArrayIndices.AddDefaulted(InEvent.ObjectIteratorIndex + 1);
+		SetArrayIndexPerObject(ResolvedArrayIndices);
+
+		// Iterate property chain, Remap fields, and resolve data pointers
+		while(!ConstructionComplete())
+		{
+			if (UStruct* ResolvedStruct = ResolvedField.Get<UStruct>())
+			{
+				SourceChainNode = (SourceChainNode) ? SourceChainNode->GetNextNode() : SourceEvent.PropertyChain.GetActiveMemberNode();
+				check(SourceChainNode);
+				
+				SourceField = SourceChainNode->GetValue();
+				SourceMemoryPtr = SourceField.Get<FProperty>()->ContainerPtrToValuePtr<void>(SourceMemoryPtr);
+				
+				ResolvedChainNode = ResolveCurChainNode(ResolvedStruct, SourceField.Get<FProperty>());
+				ResolvedField = ResolvedChainNode->GetValue();
+				ResolvedMemoryPtr = ResolvedField.Get<FProperty>()->ContainerPtrToValuePtr<void>(ResolvedMemoryPtr);
+			}
+			else if (UObject* ResolvedObject = ResolvedField.Get<UObject>())
+			{
+				check(false);
+			}
+			else if (FArrayProperty* ResolvedArrayProperty = ResolvedField.Get<FArrayProperty>())
+			{
+				FArrayProperty* SourceArrayProperty = SourceField.Get<FArrayProperty>();
+				int32 ArrayIndex = SourceEvent.GetArrayIndex(SourceArrayProperty->GetName());
+				if (ArrayIndex == INDEX_NONE)
+				{
+					checkf(SourceChainNode->GetNextNode() == nullptr, TEXT("Expected this to be the last property because there's no index"));
+					SourceField = nullptr;
+					ResolvedChainNode = nullptr;
+					ResolvedField = nullptr;
+				}
+				else
+				{
+					FScriptArrayHelper SourceArrayHelper(SourceArrayProperty, SourceMemoryPtr);
+	                FScriptArrayHelper ResolvedArrayHelper(ResolvedArrayProperty, ResolvedMemoryPtr);
+	                if(!ResolvedArrayHelper.IsValidIndex(ArrayIndex))
+	                {
+                		check(ChangeType & EPropertyChangeType::ArrayAdd);
+                		check(ResolvedArrayHelper.Num() == ArrayIndex);
+                		ResolvedArrayHelper.Resize(ArrayIndex + 1);
+	                }
+	                ResolvedArrayIndices[SourceEvent.ObjectIteratorIndex].Add(ResolvedArrayProperty->GetName(), ArrayIndex);
+	                SourceMemoryPtr = SourceArrayHelper.GetRawPtr(ArrayIndex);
+	                ResolvedMemoryPtr = ResolvedArrayHelper.GetRawPtr(ArrayIndex);
+	                SourceField = SourceArrayProperty->Inner;
+	                ResolvedField = ResolvedArrayProperty->Inner;
+				}
+			}
+			else if (FSetProperty* ResolvedSetProperty = ResolvedField.Get<FSetProperty>())
+			{
+				FSetProperty* SourceSetProperty = SourceField.Get<FSetProperty>();
+                int32 SetIndex = SourceEvent.GetArrayIndex(SourceSetProperty->GetName());
+                if (SetIndex == INDEX_NONE)
+                {
+                	checkf(SourceChainNode->GetNextNode() == nullptr, TEXT("Expected this to be the last property because there's no index"));
+                	SourceField = nullptr;
+                	ResolvedChainNode = nullptr;
+                	ResolvedField = nullptr;
+                }
+                else
+                {
+                	FScriptSetHelper SourceSetHelper(SourceSetProperty, SourceMemoryPtr);
+                	FScriptSetHelper ResolvedSetHelper(ResolvedSetProperty, ResolvedMemoryPtr);
+                	if(!ResolvedSetHelper.IsValidIndex(SetIndex))
+                	{
+                		check(ChangeType & EPropertyChangeType::ArrayAdd);
+                		check(ResolvedSetHelper.Num() == SetIndex);
+                		int32 AddedIndex = ResolvedSetHelper.AddUninitializedValue();
+                		check(AddedIndex == SetIndex);
+                	}
+                    ResolvedArrayIndices[SourceEvent.ObjectIteratorIndex].Add(ResolvedSetProperty->GetName(), SetIndex);
+                	SourceMemoryPtr = SourceSetHelper.GetElementPtr(SetIndex);
+                	ResolvedMemoryPtr = ResolvedSetHelper.GetElementPtr(SetIndex);
+                	SourceField = SourceSetProperty->ElementProp;
+                	ResolvedField = ResolvedSetProperty->ElementProp;
+                }
+			}
+			else if (FMapProperty* ResolvedMapProperty = ResolvedField.Get<FMapProperty>())
+			{
+				FMapProperty* SourceMapProperty = SourceField.Get<FMapProperty>();
+				int32 MapIndex = SourceEvent.GetArrayIndex(SourceMapProperty->GetName());
+				if (MapIndex == INDEX_NONE)
+				{
+					checkf(SourceChainNode->GetNextNode() == nullptr, TEXT("Expected this to be the last property because there's no index"));
+					SourceField = nullptr;
+					ResolvedChainNode = nullptr;
+					ResolvedField = nullptr;
+				}
+				else
+				{
+					FScriptMapHelper SourceMapHelper(SourceMapProperty, SourceMemoryPtr);
+					FScriptMapHelper ResolvedMapHelper(ResolvedMapProperty, ResolvedMemoryPtr);
+					if(!ResolvedMapHelper.IsValidIndex(MapIndex))
+					{
+						check(ChangeType & EPropertyChangeType::ArrayAdd);
+						check(ResolvedMapHelper.Num() == MapIndex);
+						int32 AddedIndex = ResolvedMapHelper.AddUninitializedValue();
+						check(AddedIndex == MapIndex);
+					}
+	                ResolvedArrayIndices[SourceEvent.ObjectIteratorIndex].Add(ResolvedMapProperty->GetName(), MapIndex);
+
+					// try to figure out whether we should be traversing into the key or the value
+					auto TryEnterMapWithHint = [&](const FProperty* Hint)
+					{
+						if (Hint == SourceMapProperty->KeyProp)
+						{
+							SourceMemoryPtr = SourceMapHelper.GetKeyPtr(MapIndex);
+							ResolvedMemoryPtr = ResolvedMapHelper.GetKeyPtr(MapIndex);
+							SourceField = SourceMapProperty->KeyProp;
+							ResolvedField = ResolvedMapProperty->KeyProp;
+							return true;
+						}
+						else if (Hint == SourceMapProperty->ValueProp)
+						{
+							SourceMemoryPtr = SourceMapHelper.GetValuePtr(MapIndex);
+							ResolvedMemoryPtr = ResolvedMapHelper.GetValuePtr(MapIndex);
+							SourceField = SourceMapProperty->ValueProp;
+							ResolvedField = ResolvedMapProperty->ValueProp;
+							return true;
+						}
+						return false;
+					};
+
+					// this process isn't well-defined we need to check several places for the property to see if it's a key or a value prop
+					if (FEditPropertyChain::TDoubleLinkedListNode* NextSourceChainNode = SourceChainNode->GetNextNode())
+					{
+						bool bSuccess = TryEnterMapWithHint(NextSourceChainNode->GetValue());
+						check(bSuccess);
+					}
+					else if (TryEnterMapWithHint(SourceEvent.Property)) {}
+					else if (TryEnterMapWithHint(SourceEvent.MemberProperty)) {}
+					else
+					{
+						checkf(false, TEXT("An index was provided for the map but it's ambiguous whether the key or value was modified"));
+					}
+				}
+			}
+			else if (FStructProperty* ResolvedStructProperty = ResolvedField.Get<FStructProperty>())
+			{
+				FStructProperty* SourceStructProperty = SourceField.Get<FStructProperty>();
+				SourceField = SourceStructProperty->Struct;
+				ResolvedField = ResolvedStructProperty->Struct;
+			}
+			else if (FObjectProperty* ResolvedObjectProperty = ResolvedField.Get<FObjectProperty>())
+			{
+				UObject* ResolvedSubObject = GetObjectRefFromProperty(ResolvedObjectProperty, ResolvedMemoryPtr, ResolvedObject);
+				ResolvedField = ResolvedSubObject->GetClass();
+				ResolvedMemoryPtr = ResolvedSubObject;
+				
+				FObjectProperty* SourceObjectProperty = SourceField.Get<FObjectProperty>();
+				UObject* SourceSubObject = GetObjectRefFromProperty(SourceObjectProperty, SourceMemoryPtr, SourceObject);
+				SourceField = SourceSubObject->GetClass();
+				SourceMemoryPtr = SourceSubObject;
+			}
+			else if (FOptionalProperty* ResolvedOptionalProperty = ResolvedField.Get<FOptionalProperty>())
+			{
+				FOptionalProperty* SourceOptionalProperty = SourceField.Get<FOptionalProperty>();
+				
+				check(SourceOptionalProperty->IsSet(SourceMemoryPtr))
+				SourceField = SourceOptionalProperty->GetValueProperty();
+				SourceMemoryPtr = SourceOptionalProperty->GetValuePointerForReadOrReplace(SourceMemoryPtr);
+				
+	            if (!ResolvedOptionalProperty->IsSet(ResolvedMemoryPtr))
+	            {
+            		ResolvedField = ResolvedOptionalProperty->GetValueProperty();
+            		ResolvedMemoryPtr = ResolvedOptionalProperty->MarkSetAndGetInitializedValuePointerToReplace(ResolvedMemoryPtr);
+	            }
+	            else
+	            {
+            		ResolvedField = ResolvedOptionalProperty->GetValueProperty();
+            		ResolvedMemoryPtr = ResolvedOptionalProperty->GetValuePointerForReadOrReplace(ResolvedMemoryPtr);
+	            }
+			}
+			else
+			{
+				checkf(SourceChainNode->GetNextNode() == nullptr, TEXT("Expected this to be the last property in the chain because it's not a type known to contain properties"));
+				SourceChainNode = nullptr;
+				SourceField = nullptr;
+				
+				ResolvedChainNode = nullptr;
+				ResolvedField = nullptr;
+			}
+			
+			if (SourceEvent.Property == SourceField.Get<FProperty>())
+			{
+				Property = ResolvedField.Get<FProperty>();
+			}
+			if (SourceEvent.MemberProperty == SourceField.Get<FProperty>())
+			{
+				MemberProperty = ResolvedField.Get<FProperty>();
+			}
+			if (SourceEvent.PropertyChain.GetActiveNode() == SourceChainNode)
+			{
+				PropertyChain.SetActivePropertyNode(ResolvedChainNode->GetValue());
+			}
+			if (SourceEvent.PropertyChain.GetActiveMemberNode() == SourceChainNode)
+			{
+				PropertyChain.SetActiveMemberPropertyNode(ResolvedChainNode->GetValue());
+			}
+		}
+		
+		// make sure the Property and MemberProperty members were remapped as well
+		check(SourceEvent.Property != Property);
+		check(SourceEvent.MemberProperty != MemberProperty);
+	}
+
+	bool ConstructionComplete() const
+	{
+		if (FProperty* SourceProperty = SourceField.Get<FProperty>())
+		{
+			if (SourceEvent.GetArrayIndex(SourceProperty->GetName()) != INDEX_NONE && SourceProperty->Owner.Get<FProperty>() == SourceEvent.Property)
+			{
+				return false;
+			}
+			return SourceProperty == SourceEvent.Property;
+		}
+		check(SourceField.IsValid());
+		check(ResolvedField.IsValid());
+		return false;
+	}
+
+	FEditPropertyChain::TDoubleLinkedListNode* ResolveCurChainNode(const UStruct* ResolvedStruct, const FProperty* SourceProperty)
+	{
+		FProperty* FoundProperty = FindPropertyByNameAndType(ResolvedStruct, SourceProperty);
+		check(FoundProperty);
+		
+		PropertyChain.AddTail(FoundProperty);
+		return PropertyChain.GetTail();
+	}
+
+	static UObject* GetObjectRefFromProperty(const FObjectProperty* Property, const void* Memory, const UObject* OwningObject)
+	{
+		UObject* SubObject = Property->GetObjectPropertyValue(Memory);
+		if (SubObject && UE::IsClassOfInstanceDataObjectClass(OwningObject->GetClass()))
+		{
+			if (UObject* FoundIdo = FPropertyBagRepository::Get().FindInstanceDataObject(SubObject))
+			{
+				SubObject = FoundIdo;
+			}
+		}
+		return SubObject;
+	}
+
+	FEditPropertyChain PropertyChain;
+	const FPropertyChangedChainEvent& SourceEvent;
+	FEditPropertyChain::TDoubleLinkedListNode* SourceChainNode;
+	
+	FEditPropertyChain::TDoubleLinkedListNode* ResolvedChainNode;
+	TArray<TMap<FString, int32>> ResolvedArrayIndices;
+	
+	FFieldVariant SourceField;
+	FFieldVariant ResolvedField;
+	
+	void* SourceMemoryPtr;
+	void* ResolvedMemoryPtr;
+	const UObject* SourceObject;
+	const UObject* ObjectToResolve;
+};
+
 void FPropertyBagRepository::PostEditChangeChainProperty(const UObject* Object, FPropertyChangedChainEvent& PropertyChangedEvent)
 {
 #if WITH_EDITOR
@@ -577,29 +577,17 @@ void FPropertyBagRepository::PostEditChangeChainProperty(const UObject* Object,
 	
 	auto CopyChanges = [&PropertyChangedEvent](const UObject* Source, UObject* Dest)
 	{
-		FPropertyChangedEvent BaseRemappedEvent(PropertyChangedEvent);
-		FEditPropertyChain RemappedChain;
-		FPropertyChangedChainEvent RemappedChangeEvent(RemappedChain, BaseRemappedEvent);
-		void* DestData = Dest;
-		TArray<TMap<FString, int32>> RemappedArrayIndices;
-		RemappedArrayIndices.AddDefaulted(PropertyChangedEvent.ObjectIteratorIndex + 1);
-		FProperty* DestProperty = nullptr;
-		if (RemapChangeEvent(PropertyChangedEvent, RemappedChangeEvent, DestData, DestProperty, RemappedArrayIndices))
-		{
-			Dest->PreEditChange(RemappedChangeEvent.PropertyChain);
-			
-            FProperty* SourceProperty = nullptr;
-            const void* SourceData = ResolveChangePath(Source, SourceProperty, PropertyChangedEvent);
-			
-			CopyProperty(SourceProperty, SourceData, DestProperty, DestData);
-            
-            Dest->PostEditChangeChainProperty(RemappedChangeEvent);
-		}
-		else
-		{
-			ensureMsgf(false, TEXT("A const loose property was modified on an instance data object"));
-		}
+		FRemappedChainEvent RemappedChainEvent(PropertyChangedEvent, Source, Dest);
+		
+		Dest->PreEditChange(RemappedChainEvent.PropertyChain);
 		
+		CopyProperty(
+			RemappedChainEvent.SourceField.Get<FProperty>(),
+			RemappedChainEvent.SourceMemoryPtr,
+			RemappedChainEvent.ResolvedField.Get<FProperty>(),
+			RemappedChainEvent.ResolvedMemoryPtr);
+            
+		Dest->PostEditChangeChainProperty(RemappedChainEvent);
 	};
 	
 	if (UObject* Ido = Get().FindInstanceDataObject(Object))
