[StructUtils] minor massaging of TStructTypeBitSet functionality, and added  a bunch of comments

#rnx

diff --git a/Engine/Source/Developer/MassEntityTestSuite/Private/MassSerializationTest.cpp b/Engine/Source/Developer/MassEntityTestSuite/Private/MassSerializationTest.cpp
--- a/Engine/Source/Developer/MassEntityTestSuite/Private/MassSerializationTest.cpp
+++ b/Engine/Source/Developer/MassEntityTestSuite/Private/MassSerializationTest.cpp
@@ -15,7 +15,7 @@ UE_DISABLE_OPTIMIZATION_SHIP
 struct FTestStructTracker : FStructTracker
 {
 	FTestStructTracker()
-		: FStructTracker([](){ return FMassFragment::StaticStruct(); })
+		: FStructTracker(FMassFragment::StaticStruct())
 	{}
 
 	template<typename T>
diff --git a/Engine/Source/Runtime/CoreUObject/Private/StructUtils/StructTypeBitSet.cpp b/Engine/Source/Runtime/CoreUObject/Private/StructUtils/StructTypeBitSet.cpp
--- a/Engine/Source/Runtime/CoreUObject/Private/StructUtils/StructTypeBitSet.cpp
+++ b/Engine/Source/Runtime/CoreUObject/Private/StructUtils/StructTypeBitSet.cpp
@@ -57,6 +57,21 @@ namespace FBitSetSerializationDictionary
 //-----------------------------------------------------------------------------
 FStructTracker::FStructTracker(const FBaseStructGetter& InBaseStructGetter)
 	: BaseStructGetter(InBaseStructGetter)
+	, TypeVerification([this](const UStruct* InStruct)
+		{
+			return InStruct != nullptr && InStruct->IsChildOf(GetBaseType());
+		})
+{
+}
+
+FStructTracker::FStructTracker(const UStruct* InBaseType, const FTypeVerification& InTypeVerification)
+	: TypeVerification(InTypeVerification 
+		? InTypeVerification
+		: [this](const UStruct* Struct)
+			{
+				return Struct && this->GetBaseType() && Struct->IsChildOf(this->BaseType);
+			})
+	, BaseType(InBaseType)
 {
 }
 
@@ -65,6 +80,11 @@ FStructTracker::~FStructTracker()
 	FBitSetSerializationDictionary::UnregisterTracker(*this);
 }
 
+FORCEINLINE const UStruct* FStructTracker::GetBaseType() const
+{
+	return BaseType ? BaseType : (BaseStructGetter ? (BaseType = BaseStructGetter()) : nullptr);
+}
+
 int32 FStructTracker::FindOrAddStructTypeIndex(const UStruct& InStructType)
 {
 	// Get existing index...
@@ -75,21 +95,34 @@ int32 FStructTracker::FindOrAddStructTypeIndex(const UStruct& InStructType)
 	{
 		return ElementId.AsInteger();
 	}
-	return Register(InStructType);
+	return RegisterImplementation(InStructType, /*bCheckExisting=*/false);
 }
 
-int32 FStructTracker::Register(const UStruct& InStructType)
+int32 FStructTracker::RegisterImplementation(const UStruct& InStructType, const bool bCheckExisting)
 {
+	if (!ensureMsgf(TypeVerification(&InStructType)
+		, TEXT("Unable to register %s singe it fails the validation test.")
+		, *InStructType.GetName()))
+	{
+		return INDEX_NONE;
+	}
+
 	const uint32 Hash = PointerHash(&InStructType);
-	checkSlow(InStructType.IsChildOf(BaseStructGetter()));
+	if (bCheckExisting)
+	{
+		const FSetElementId ElementId = StructTypeToIndexSet.FindIdByHash(Hash, Hash);
+		if (ElementId.IsValidId())
+		{
+			return ElementId.AsInteger();
+		}
+	}
 
 	// .. or create a new one
 	const FSetElementId ElementId = StructTypeToIndexSet.AddByHash(Hash, Hash);
-	checkSlow(ElementId.IsValidId());
+	check(ElementId.IsValidId());
 
 	const int32 NewIndex = ElementId.AsInteger();
-
-	checkSlow(StructTypesList.Num() == NewIndex);
+	check(StructTypesList.Num() == NewIndex);
 	StructTypesList.Add(&InStructType);
 
 	if (bIsSerializable)
@@ -98,9 +131,9 @@ int32 FStructTracker::Register(const UStruct& InStructType)
 		if (SerializationHash == 0)
 		{
 			ensure(StructTypesList.Num() == 1);
-			UStruct* BaseType = BaseStructGetter();
-			CA_ASSUME(BaseType);
-			SerializationHash = GetTypeHash(BaseType->GetFullName());
+			const UStruct* LocalBaseType = GetBaseType();
+			CA_ASSUME(LocalBaseType);
+			SerializationHash = GetTypeHash(LocalBaseType->GetFullName());
 		}
 		SerializationHash = HashCombine(SerializationHash, GetTypeHash(InStructType.GetFullName()));
 
@@ -156,7 +189,7 @@ void FStructTracker::Serialize(FArchive& Ar, FStructTypeBitSet::FBitSetContainer
 		const int64 InitialDataOffset = Ar.Tell();
 
 		// store information on the base type, so that we can verify we're trying to load the right data later on
-		FTopLevelAssetPath BaseStructPath(BaseStructGetter());
+		FTopLevelAssetPath BaseStructPath(GetBaseType());
 		Ar << BaseStructPath;
 
 		int32 StructTypesListNum = StructTypesList.Num();
@@ -204,7 +237,7 @@ void FStructTracker::Serialize(FArchive& Ar, FStructTypeBitSet::FBitSetContainer
 			// the wrong types creating the Mapping (via FindOrAddStructTypeIndex below)
 			UStruct* SerializedBaseStructType = FindObject<UStruct>(BaseStructPath);
 
-			if (SerializedBaseStructType == BaseStructGetter())
+			if (SerializedBaseStructType == GetBaseType())
 			{
 				// create a translator 
 				TArray<int32>& BitMapping = FBitSetSerializationDictionary::GetOrCreateBitMapping(LoadedSerializationHash);
@@ -222,7 +255,7 @@ void FStructTracker::Serialize(FArchive& Ar, FStructTypeBitSet::FBitSetContainer
 
 						if (UStruct* StructType = FindObject<UStruct>(TypePath))
 						{
-							checkSlow(StructType->IsChildOf(BaseStructGetter()));
+							checkSlow(StructType->IsChildOf(GetBaseType()));
 							BitMapping[Index] = FindOrAddStructTypeIndex(*StructType);
 						}
 						else
@@ -253,7 +286,7 @@ void FStructTracker::Serialize(FArchive& Ar, FStructTypeBitSet::FBitSetContainer
 			else
 			{
 				UE_LOG(LogStructUtils, Error, TEXT("Trying to load mismatching BitSet data. Current base class: %s, read base class: %s")
-					, *GetNameSafe(BaseStructGetter()), *BaseStructPath.ToString());
+					, *GetNameSafe(GetBaseType()), *BaseStructPath.ToString());
 			}
 		}
 	}
diff --git a/Engine/Source/Runtime/CoreUObject/Public/StructUtils/StructTypeBitSet.h b/Engine/Source/Runtime/CoreUObject/Public/StructUtils/StructTypeBitSet.h
--- a/Engine/Source/Runtime/CoreUObject/Public/StructUtils/StructTypeBitSet.h
+++ b/Engine/Source/Runtime/CoreUObject/Public/StructUtils/StructTypeBitSet.h
@@ -7,18 +7,25 @@
 #include "StructUtils/InstancedStruct.h"
 #include "Templates/FunctionFwd.h"
 
-
 class FArchive;
 
 namespace FStructTypeBitSet
 {
 
+/**
+ * A constant bitset container that extends TBitArray<> with additional utilities.
+ * Used for operations such as bitwise checks and hashing
+ */
 struct FConstBitSetContainer : TBitArray<>
 {
 	using Super = TBitArray<>;
 
 	FConstBitSetContainer() = default;
 
+	/**
+	 * Checks if all bits set in the Other bit array are also set in this bitset
+	 * @return True if all bits in Other are set in this bitset
+	 */
 	FORCEINLINE bool HasAll(const TBitArray<>& Other) const
 	{
 		FConstWordIterator ThisIterator(*this);
@@ -40,6 +47,10 @@ struct FConstBitSetContainer : TBitArray<>
 		return true;
 	}
 
+	/**
+	 * Checks if any bits in the Other bit array are also set in this bitset
+	 * @return True if any bits in Other are set in this bitset
+	 */
 	FORCEINLINE bool HasAny(const TBitArray<>& Other) const
 	{
 		FConstWordIterator ThisIterator(*this);
@@ -61,6 +72,9 @@ struct FConstBitSetContainer : TBitArray<>
 		return false;
 	}
 
+	/**
+	 * @return whether the bitset is empty (no bits are set)
+	 */
 	FORCEINLINE bool IsEmpty() const
 	{
 		FConstWordIterator Iterator(*this);
@@ -73,11 +87,19 @@ struct FConstBitSetContainer : TBitArray<>
 		return !Iterator;
 	}
 
+	/**
+	 * Compares two bitsets for equality
+	 * @return True if the bitsets are equal
+	 */
 	bool operator==(const FConstBitSetContainer& Other) const
 	{
 		return (IsEmpty() && Other.IsEmpty()) || Super::operator==(Other);
 	}
 
+	/**
+	 * Computes the hash value for the given bitset
+	 * @return The computed hash value
+	 */
 	FORCEINLINE friend uint32 GetTypeHash(const FConstBitSetContainer& Instance)
 	{
 		FConstWordIterator Iterator(Instance);
@@ -100,34 +122,55 @@ struct FConstBitSetContainer : TBitArray<>
 		return Hash;
 	}
 
+	/**
+	 * Checks if a specific bit is set in the bitset
+	 * @param Index - The index of the bit to check
+	 * @return True if the bit is set, false otherwise
+	 */
 	bool Contains(const int32 Index) const
 	{
 		check(Index >= 0);
 		return (Index < Num()) && (*this)[Index];
 	}
 
+	/**
+	 * Counts the number of set bits in the bitset
+	 * @return The number of set bits
+	 */
 	int32 CountStoredTypes() const
 	{
 		return CountSetBits();
 	}
 
 protected:
-	// only for use by the FBitSetContainer
+	/**
+	 * Copy-constructor for creating a bitset from another bit array
+	 * @param Source - The source bit array
+	 */
 	explicit FConstBitSetContainer(const TBitArray<>& Source)
 		: Super(Source)
 	{
 	}
 
+	/**
+	 * Move-constructor for creating a bitset from another bit array
+	 * @param Source - The source bit array to move from
+	 */
 	explicit FConstBitSetContainer(TBitArray<>&& Source)
 		: Super(Forward<TBitArray<>>(Source))
 	{
 	}
 };
 
+/**
+ * A mutable extension of FConstBitSetContainer, adding methods for modifying bits
+ */
 struct FBitSetContainer : FConstBitSetContainer
 {
 	using FConstBitSetContainer::FConstBitSetContainer;
 
+	FBitSetContainer() = default;
+
 	explicit FBitSetContainer(const TBitArray<>& Source)
 		: FConstBitSetContainer(Source)
 	{
@@ -145,11 +188,21 @@ struct FBitSetContainer : FConstBitSetContainer
 		return *this;
 	}
 
+
+	/**
+	 * Initializes the bitset to a specified size and value
+	 * @param bValue - The value to set all bits to
+	 * @param Count - The number of bits to initialize
+	 */
 	void SetAll(const bool bValue, const int32 Count)
 	{
 		Init(bValue, Count);
 	}
-	
+
+	/**
+	 * Removes bits set in the Other bit array from this bitset
+	 * @param Other - The bit array to remove
+	 */
 	FORCEINLINE void operator-=(const TBitArray<>& Other)
 	{
 		FWordIterator ThisIterator(*this);
@@ -164,17 +217,21 @@ struct FBitSetContainer : FConstBitSetContainer
 		}
 	}
 
-	FORCEINLINE friend uint32 GetTypeHash(const FBitSetContainer& Instance)
-	{
-		return GetTypeHash(*static_cast<const FConstBitSetContainer*>(&Instance));
-	}
-
+	/**
+	 * Adds a bit at the specified index
+	 * @param Index - The index of the bit to add
+	 */
 	void AddAtIndex(const int32 Index)
 	{
+		check(Index >= 0);
 		PadToNum(Index + 1, false);
 		SetBitNoCheck(Index, true);
 	}
 
+	/**
+	 * Removes a bit at the specified index
+	 * @param Index - The index of the bit to remove
+	 */
 	void RemoveAtIndex(const int32 Index)
 	{
 		check(Index >= 0);
@@ -187,9 +244,9 @@ struct FBitSetContainer : FConstBitSetContainer
 
 protected:
 	/**
-	 * duplication of TBitArray::SetBitNoCheck needed since it's private but it's the performant way of setting bits
-	 * when we know the index is valid.
-	 * @todo ask Core team about exposing that
+	 * Directly sets a bit without boundary checks for improved performance.
+	 * @param Index - The index of the bit to set.
+	 * @param Value - The value to set the bit to.
 	 */
 	void SetBitNoCheck(const int32 Index, const bool Value)
 	{
@@ -198,72 +255,89 @@ protected:
 		Word = (Word & ~(1 << BitOffset)) | (((uint32)Value) << BitOffset);
 	}
 };
-} // namespace FStructTypeBitSet
 
+// Ensure that FBitSetContainer does not add any new member variables compared to FConstBitSetContainer.
+static_assert(sizeof(FBitSetContainer) == sizeof(FConstBitSetContainer), "FBitSetContainer as a functional extension of FConstBitSetContainer is not allowed add new member variables.");
+
+} // namespace FStructTypeBitSet
 
 /**
- * The TStructTypeBitSet holds information on "existence" of subtypes of a given UStruct. The information on 
- * available child-structs is gathered lazily - the internal FStructTracker is assigning a given type a new index the
- * very first time the type is encountered. 
- * To create a specific instantiation of the type you also need to provide a type that will instantiate a static 
- * FStructTracker instance. The supplied macros hide this details away however.
- * 
- * To do declare a bitset type for an arbitrary struct type FFooBar add the following in your header or cpp file:
- * 
- *	DECLARE_STRUCTTYPEBITSET(FMyFooBarBitSet, FFooBar);
- * 
- * where FMyFooBarBitSet is the alias of the type you can use in your code. To have your type exposed to other modules 
- * use DECLARE_STRUCTTYPEBITSET_EXPORTED, like so:
- * 
- *	DECLARE_STRUCTTYPEBITSET_EXPORTED(MYMODULE_API, FMyFooBarBitSet, FFooBar);
+ * FStructTracker is a utility class used to track and map UStruct types to indices, which are used in bitsets.
+ * It manages a mapping between UStruct instances and integer indices, allowing for efficient storage and querying
+ * of types in a bitset
  *
- * You also need to instantiate the static FStructTracker added by the DECLARE* macro. You can easily do it by placing 
- * the following in your cpp file (continuing the FFooBar example):
- *
- *	DEFINE_TYPEBITSET(FMyFooBarBitSet);
- * 
+ * The FStructTracker assigns an index to a given type the first time it encounters it
  */
 struct COREUOBJECT_API FStructTracker
 {
 	using FBaseStructGetter = TFunction<UStruct*()>;
+	using FTypeVerification = TFunction<bool(const UStruct*)>;
 
 	/** 
 	 * The input parameter is a function that fetches the UStruct representing the base class for all the stored structs.
 	 * We're unable to get the UStruct parameter directly since FStructTracker instances are being created during
-	 * module loading (via DEFINE_TYPEBITSET macro) and the ::StaticStruct call fail at that point for types defined
+	 * module loading (via DEFINE_TYPEBITSET macro), and the ::StaticStruct call fails at that point for types defined
 	 * in the same module.
 	 */
 	explicit FStructTracker(const FBaseStructGetter& InBaseStructGetter);
+	explicit FStructTracker(const UStruct* InBaseType, const FTypeVerification& InTypeVerification = FTypeVerification());
 	~FStructTracker();
 
-	/** @return the index for the given type, or INDEX_NONE if it hasn't been registered with this struct tracker */
+	UE_NONCOPYABLE(FStructTracker);
+
+	/** 
+	 * @return the index for the given type, or INDEX_NONE if it hasn't been registered with this struct tracker
+	 */
 	int32 FindStructTypeIndex(const UStruct& InStructType) const;
 
 	/**
-	 * Fetches the internal index representing given UStruct instance. If it hasn't been registered yet, it will be automatically added
-	 * @return the index for the given type
+	 * Fetches the internal index representing the given UStruct instance. If it hasn't been registered yet, it will be automatically added.
+	 * @return the index for the given type.
 	 */
 	int32 FindOrAddStructTypeIndex(const UStruct& InStructType);
 
 	/**
-	 * Registers given UStruct with the struct tracker.
-	 * @return the index for the registered type
+	 * Registers the given UStruct with the struct tracker.
+	 * @return the index for the registered type.
 	 */
-	int32 Register(const UStruct& InStructType);
+	int32 Register(const UStruct& InStructType)
+	{
+		return RegisterImplementation(InStructType, /*bCheckPrevious=*/true);
+	}
 
+	/**
+	 * Retrieves the UStruct type associated with a given index
+	 * @param StructTypeIndex - The index of the struct type
+	 * @return Pointer to the UStruct if found, nullptr otherwise
+	 */
 	const UStruct* GetStructType(const int32 StructTypeIndex) const
 	{
 		return StructTypesList.IsValidIndex(StructTypeIndex) ? StructTypesList[StructTypeIndex].Get() : nullptr;
 	}
 
+	/**
+	 * Retrieves the base UStruct type used for validation.
+	 * @return Pointer to the base UStruct.
+	 */
+	const UStruct* GetBaseType() const;
+
+	/**
+	 * Gets the number of registered struct types.
+	 * @return Number of struct types.
+	 */
 	int32 Num() const {	return StructTypeToIndexSet.Num(); }
 
+	/**
+	 * Serializes the struct types bit array using the archive provided.
+	 * @param Ar - The archive to serialize with.
+	 * @param StructTypesBitArray - The bit array representing struct types.
+	 */
 	void Serialize(FArchive& Ar, FStructTypeBitSet::FBitSetContainer& StructTypesBitArray);
 
 #if WITH_STRUCTUTILS_DEBUG
 	/**
-	* @return index identifying given tag or INDEX_NONE if it has never been used/seen before.
-	*/
+	 * @return Name identifying the given struct type or NAME_None if it has never been used/seen before.
+	 */
 	FName DebugGetStructTypeName(const int32 StructTypeIndex) const
 	{
 		return DebugStructTypeNamesList.IsValidIndex(StructTypeIndex) ? DebugStructTypeNamesList[StructTypeIndex] : FName();
@@ -275,45 +349,78 @@ struct COREUOBJECT_API FStructTracker
 		return MakeArrayView(reinterpret_cast<const TWeakObjectPtr<const T>*>(StructTypesList.GetData()), StructTypesList.Num());
 	}
 
+	/**
+	 * Resets all struct type mapping information. Used for debugging and testing purposes.
+	 */
 	void DebugResetStructTypeMappingInfo()
 	{
 		StructTypeToIndexSet.Reset();
 		StructTypesList.Reset();
 		DebugStructTypeNamesList.Reset();
 	}
-	TArray<FName, TInlineAllocator<64>> DebugStructTypeNamesList;
+
+protected:
+	mutable TArray<FName, TInlineAllocator<64>> DebugStructTypeNamesList;
 #endif // WITH_STRUCTUTILS_DEBUG
 
+protected:
+	int32 RegisterImplementation(const UStruct& InStructType, const bool bCheckPrevious);
+
+	/** Set mapping struct type hashes to their indices */
 	TSet<uint32> StructTypeToIndexSet;
+	/** List of weak pointers to struct types, indexed by their assigned indices */
 	TArray<TWeakObjectPtr<const UStruct>, TInlineAllocator<64>> StructTypesList;
+	/** Hash used during serialization to detect mismatches */
 	uint32 SerializationHash = 0;
+	/** Flag indicating whether this tracker is serializable */
 	uint32 bIsSerializable : 1 = true;
+	/** Function to retrieve the base UStruct */
 	const FBaseStructGetter BaseStructGetter;
+	/** Function to perform additional type verification */
+	const FTypeVerification TypeVerification;
+
+private:
+	mutable const UStruct* BaseType;
 };
 
-template<typename TImplementation, typename TBaseStruct, typename TStructType, bool bTestInheritanceAtRuntime=WITH_STRUCTUTILS_DEBUG>
+/**
+ * Base class for managing bitsets associated with specific struct types.
+ * Provides a foundation for derived classes to handle operations like adding, removing, and checking for struct types.
+ *
+ * @param TImplementation - The derived type implementing this base class.
+ * @param TBaseStruct - The base struct type for type validation.
+ * @param TStructType - Unreal's struct type, typically UScriptStruct or UClass.
+ * @param TBitSetContainer - The container type for storing bitsets.
+ * @param bTestInheritanceAtRuntime - Flag to enable runtime inheritance checks.
+ */
+template<typename TImplementation, typename TBaseStruct, typename TStructType, typename TBitSetContainer, bool bTestInheritanceAtRuntime=WITH_STRUCTUTILS_DEBUG>
 struct TTypeBitSetBase
 {
 	using FUStructType = TStructType;
 	using FBaseStruct = TBaseStruct;
+	using FContainer = TBitSetContainer;
 
-	/** Retrieves indices of bits of given value. */
+	/**
+	 * Iterator for traversing indices of bits with a specified value (true or false).
+	 */
 	struct FIndexIterator
 	{
-	public:
 		explicit FIndexIterator(const FStructTypeBitSet::FBitSetContainer& BitArray, const bool bInValueToCheck = true)
 			: It(BitArray), bValueToCheck(bInValueToCheck)
 		{
 			if (It && It.GetValue() != bInValueToCheck)
 			{
-				// will result in either setting IT to the first bInValue, or making bool(It) == false
+				// Will result in either setting It to the first bit with bInValueToCheck, or making bool(It) == false
 				++(*this);
 			}
 		}
 
-		operator bool() const { return bool(It); }
+		operator bool() const
+		{
+			return bool(It);
+		}
 
-		/** @todo can be optimized to skip whole empty words, similar to TBitArray<>::Find */
+		/** Advances the iterator to the next bit with the specified value */
 		FIndexIterator& operator++() 
 		{ 
 			while (++It)
@@ -326,6 +433,7 @@ struct TTypeBitSetBase
 			return *this; 
 		}
 
+		/** Dereferences the iterator to get the current index */
 		int32 operator*() const
 		{
 			return It.GetIndex();
@@ -336,21 +444,38 @@ struct TTypeBitSetBase
 		const bool bValueToCheck;
 	};
 
+	/**
+	 * Creates an index iterator for bits with the specified value.
+	 * @param bValueToCheck - The bit value to search for (true or false).
+	 * @return An iterator to traverse indices of matching bits.
+	 */
 	FIndexIterator GetIndexIterator(const bool bValueToCheck = true) const
 	{
 		return FIndexIterator(StructTypesBitArray, bValueToCheck);
 	}
 
+	/**
+	 * Retrieves the derived implementation of this base class.
+	 * @return Reference to the derived implementation.
+	 */
 	TImplementation& GetImplementation()
 	{
 		return *reinterpret_cast<TImplementation*>(this);
 	}
 
+	/**
+	 * Retrieves the derived implementation of this base class (const version).
+	 * @return Const reference to the derived implementation.
+	 */
 	const TImplementation& GetImplementation() const 
 	{
 		return *reinterpret_cast<const TImplementation*>(this);
 	}
 
+	/**
+	 * Retrieves the base UStruct type used for validation.
+	 * @return Pointer to the base UStruct.
+	 */
 	FORCEINLINE_DEBUGGABLE static const UStruct* GetBaseUStruct()
 	{
 		if constexpr (bTestInheritanceAtRuntime)
@@ -363,67 +488,100 @@ struct TTypeBitSetBase
 			return nullptr;
 		}
 	}
-	
+
+	/**
+	 * Sets all bits in the bitset to the specified value.
+	 * @param bValue - The value to set for all bits (default is true).
+	 */
 	void SetAll(const bool bValue = true)
 	{
 		StructTypesBitArray.SetAll(bValue, GetImplementation().GetStructTracker().Num());
 	}
 
+	/**
+	 * Adds a struct type to the bitset.
+	 * @param InStructType - The struct type to add.
+	 */
 	void Add(const FUStructType& InStructType)
 	{
-#if WITH_STRUCTUTILS_DEBUG
 		if constexpr (bTestInheritanceAtRuntime)
 		{
-			ensureMsgf(InStructType.IsChildOf(GetBaseUStruct())
-				, TEXT("Registering '%s' with FStructTracker while it doesn't derive from the expected struct type %s")
-				, *InStructType.GetPathName(), *GetBaseUStruct()->GetName());
+			if (UNLIKELY(!ensureMsgf(InStructType.IsChildOf(GetBaseUStruct()),
+				TEXT("Registering '%s' with FStructTracker while it doesn't derive from the expected struct type %s"),
+				*InStructType.GetPathName(), *GetBaseUStruct()->GetName())))
+			{
+				return;
+			}
 		}
-#endif // WITH_STRUCTUTILS_DEBUG
 
 		const int32 StructTypeIndex = GetImplementation().GetStructTracker().FindOrAddStructTypeIndex(InStructType);
 		StructTypesBitArray.AddAtIndex(StructTypeIndex);
 	}
 
+	/**
+	 * Removes a struct type from the bitset.
+	 * @param InStructType - The struct type to remove.
+	 */
 	void Remove(const FUStructType& InStructType)
 	{
-#if WITH_STRUCTUTILS_DEBUG
 		if constexpr (bTestInheritanceAtRuntime)
 		{
-			ensureMsgf(InStructType.IsChildOf(GetBaseUStruct())
-				, TEXT("Registering '%s' with FStructTracker while it doesn't derive from the expected struct type %s")
-				, *InStructType.GetPathName(), *GetBaseUStruct()->GetName());
+			if (UNLIKELY(!ensureMsgf(InStructType.IsChildOf(GetBaseUStruct()),
+				TEXT("Registering '%s' with FStructTracker while it doesn't derive from the expected struct type %s"),
+				*InStructType.GetPathName(), *GetBaseUStruct()->GetName())))
+			{
+				return;
+			}
 		}
-#endif // WITH_STRUCTUTILS_DEBUG
 
 		const int32 StructTypeIndex = GetImplementation().GetStructTracker().FindOrAddStructTypeIndex(InStructType);
 		StructTypesBitArray.RemoveAtIndex(StructTypeIndex);
 	}
 
+	/**
+	 * Resets all bits in the bitset
+	 */
 	void Reset()
 	{
 		StructTypesBitArray.Reset();
 	}
 
+	/**
+	 * Checks if the bitset contains a specific struct type.
+	 * @param InStructType - The struct type to check for.
+	 * @return True if the struct type is in the bitset; false otherwise.
+	 */
 	bool Contains(const FUStructType& InStructType) const
 	{
-#if WITH_STRUCTUTILS_DEBUG
 		if constexpr (bTestInheritanceAtRuntime)
 		{
-			ensureMsgf(InStructType.IsChildOf(GetBaseUStruct())
-				, TEXT("Registering '%s' with FStructTracker while it doesn't derive from the expected struct type %s")
-				, *InStructType.GetPathName(), *GetBaseUStruct()->GetName());
+			if (UNLIKELY(!ensureMsgf(InStructType.IsChildOf(GetBaseUStruct()),
+				TEXT("Registering '%s' with FStructTracker while it doesn't derive from the expected struct type %s"),
+				*InStructType.GetPathName(), *GetBaseUStruct()->GetName())))
+			{
+				return false;
+			}
 		}
-#endif // WITH_STRUCTUTILS_DEBUG
 
 		const int32 StructTypeIndex = GetImplementation().GetStructTracker().FindStructTypeIndex(InStructType);
 		return StructTypeIndex != INDEX_NONE && StructTypesBitArray.Contains(StructTypeIndex);
 	}
 
+	/**
+	 * Performs a bitwise AND operation with another bitset.
+	 * @param Other - The other bitset to intersect with.
+	 * @return A new bitset representing the intersection.
+	 */
 	FORCEINLINE TImplementation operator&(const TImplementation& Other) const
 	{
 		return TImplementation(TBitArray<>::BitwiseAND(StructTypesBitArray, Other.StructTypesBitArray, EBitwiseOperatorFlags::MinSize));
 	}
 
+	/**
+	 * Performs a bitwise OR operation with another bitset.
+	 * @param Other - The other bitset to union with.
+	 * @return A new bitset representing the union.
+	 */
 	FORCEINLINE TImplementation operator|(const TImplementation& Other) const
 	{
 		return TImplementation(TBitArray<>::BitwiseOR(StructTypesBitArray, Other.StructTypesBitArray, EBitwiseOperatorFlags::MaxSize));
@@ -434,6 +592,12 @@ struct TTypeBitSetBase
 		return GetImplementation() & Other;
 	}
 
+	/**
+	 * Checks if the current bitset is equivalent to another, i.e., whether both contain same "true" bits
+	 * and ignoring the trailing "false" bits.
+	 * @param Other - The other bitset to compare.
+	 * @return True if both bitsets are equivalent; false otherwise.
+	 */
 	FORCEINLINE bool IsEquivalent(const TImplementation& Other) const
 	{
 		return StructTypesBitArray.CompareSetBits(Other.StructTypesBitArray, /*bMissingBitValue=*/false);
@@ -454,26 +618,48 @@ struct TTypeBitSetBase
 		return !StructTypesBitArray.HasAny(Other.StructTypesBitArray);
 	}
 
+	/**
+	 * Checks if the bitset is empty (no bits are set).
+	 * @return True if the bitset is empty.
+	 */
 	bool IsEmpty() const 
 	{ 
 		return StructTypesBitArray.IsEmpty();
 	}
 
+	/**
+	 * Checks if a specific bit is set in the bitset.
+	 * @param BitIndex - The index of the bit to check.
+	 * @return True if the bit is set; false otherwise.
+	 */
 	FORCEINLINE bool IsBitSet(const int32 BitIndex) const
 	{
 		return StructTypesBitArray.Contains(BitIndex);
 	}
 
+	/**
+	 * Adds the bits from another bitset to this one (union).
+	 * @param Other - The other bitset to add.
+	 */
 	FORCEINLINE void operator+=(const TImplementation& Other)
 	{
 		StructTypesBitArray = TBitArray<>::BitwiseOR(StructTypesBitArray, Other.StructTypesBitArray, EBitwiseOperatorFlags::MaxSize);
 	}
 
+	/**
+	 * Removes the bits from another bitset from this one (difference).
+	 * @param Other - The other bitset to subtract.
+	 */
 	FORCEINLINE void operator-=(const TImplementation& Other)
 	{
 		StructTypesBitArray -= Other.StructTypesBitArray;
 	}
 
+	/**
+	 * Adds a struct type to the bitset, returning a new bitset.
+	 * @param NewElement - The struct type to add.
+	 * @return A new bitset with the added struct type.
+	 */
 	FORCEINLINE TImplementation operator+(const FUStructType& NewElement) const
 	{
 		TImplementation Result = GetImplementation();
@@ -481,6 +667,11 @@ struct TTypeBitSetBase
 		return MoveTemp(Result);
 	}
 
+	/**
+	 * Removes a struct type from the bitset, returning a new bitset.
+	 * @param NewElement - The struct type to remove.
+	 * @return A new bitset with the struct type removed.
+	 */
 	FORCEINLINE TImplementation operator-(const FUStructType& NewElement) const
 	{
 		TImplementation Result = GetImplementation();
@@ -488,14 +679,20 @@ struct TTypeBitSetBase
 		return MoveTemp(Result);
 	}
 
+	/**
+	 * Counts the number of set bits in the bitset.
+	 * @return The number of set bits.
+	 */
 	int32 CountStoredTypes() const
 	{
 		return StructTypesBitArray.CountSetBits();
 	}
 
 	/**
-	 * note that this function is slow(ish) due to the FStructTracker utilizing WeakObjectPtrs to store types. 
+	 * Exports types stored in the bitset to an output array.
+	 * Note: This method can be slow due to the use of weak pointers in the struct tracker.
 	 * @todo To be improved.
+	 * @param OutTypes - The array to populate with struct types.
 	 */
 	template<typename TOutStructType, typename Allocator>
 	void ExportTypes(TArray<const TOutStructType*, Allocator>& OutTypes) const
@@ -512,6 +709,9 @@ struct TTypeBitSetBase
 	}
 
 #if WITH_STRUCTUTILS_DEBUG
+	/**
+	 * Provides a debug string description of the bitset contents via the provided FOutputDevice.
+	 */
 	void DebugGetStringDesc(FOutputDevice& Ar) const
 	{
 		for (int32 Index = 0; Index < StructTypesBitArray.Num(); ++Index)
@@ -523,6 +723,10 @@ struct TTypeBitSetBase
 		}
 	}
 
+	/**
+	 * Retrieves the names of individual types stored in the bitset.
+	 * @param OutFNames - The array to populate with type names.
+	 */
 	void DebugGetIndividualNames(TArray<FName>& OutFNames) const
 	{
 		for (int32 Index = 0; Index < StructTypesBitArray.Num(); ++Index)
@@ -535,39 +739,73 @@ struct TTypeBitSetBase
 	}
 #endif // WITH_STRUCTUTILS_DEBUG
 
+	/**
+	 * Gets the allocated size of the bitset.
+	 * @return The allocated memory size in bytes.
+	 */
 	SIZE_T GetAllocatedSize() const
 	{
 		return StructTypesBitArray.GetAllocatedSize();
 	}
 
 protected:
-	/** 
-	 * A private constructor for a creating an instance straight from TBitArrays. 
-	 * @Note that this constructor needs to remain private to ensure consistency of stored values with data tracked 
-	 * by the TStructTrackerWrapper::StructTracker
+	/**
+	 * Default constructor.
 	 */
-	TTypeBitSetBase(TBitArray<>&& Source)
-		: StructTypesBitArray(Forward<TBitArray<>>(Source))
+	TTypeBitSetBase() = default;
+
+	/**
+	 * Constructor accepting a container.
+	 * @param InContainer - The bitset container to use.
+	 */
+	TTypeBitSetBase(FContainer& InContainer)
+		: StructTypesBitArray(InContainer)
 	{
 	}
 
-	TTypeBitSetBase() = default;
-
 #if WITH_STRUCTUTILS_DEBUG
-	// unittesting purposes only
+	// For unit testing purposes only
 	const TBitArray<>& DebugGetStructTypesBitArray() const { return StructTypesBitArray; }
 	TBitArray<>& DebugGetMutableStructTypesBitArray() { return StructTypesBitArray; }
 #endif // WITH_STRUCTUTILS_DEBUG
 
-	FStructTypeBitSet::FBitSetContainer StructTypesBitArray;
+	/** The bitset container storing the bits representing struct types */
+	FContainer StructTypesBitArray;
 };
 
+
+/**
+ * The TStructTypeBitSet holds information on "existence" of subtypes of a given UStruct. The information on 
+ * available child-structs is gathered lazily - the internal FStructTracker assigns a given type a new index the
+ * very first time the type is encountered. 
+ * To create a specific instantiation of the type, you also need to provide a type that will instantiate a static 
+ * FStructTracker instance. The supplied macros hide this detail.
+ * 
+ * To declare a bitset type for an arbitrary struct type FFooBar, add the following in your header or cpp file:
+ * 
+ *	DECLARE_STRUCTTYPEBITSET(FMyFooBarBitSet, FFooBar);
+ * 
+ * where FMyFooBarBitSet is the alias of the type you can use in your code. To have your type exposed to other modules 
+ * use DECLARE_STRUCTTYPEBITSET_EXPORTED, like so:
+ * 
+ *	DECLARE_STRUCTTYPEBITSET_EXPORTED(MYMODULE_API, FMyFooBarBitSet, FFooBar);
+ *
+ * You also need to instantiate the static FStructTracker added by the DECLARE* macro. You can easily do it by placing 
+ * the following in your cpp file (continuing the FFooBar example):
+ *
+ *	DEFINE_TYPEBITSET(FMyFooBarBitSet);
+ *
+ *
+ * @param TBaseStruct - The base struct type for type validation.
+ * @param TStructTrackerWrapper - A wrapper providing a FStructTracker instance.
+ * @param TUStructType - Unreal's struct type, typically UScriptStruct or UClass (defaults to UScriptStruct).
+ */
 template<typename TBaseStruct, typename TStructTrackerWrapper, typename TUStructType = UScriptStruct>
-struct TStructTypeBitSet : TTypeBitSetBase<TStructTypeBitSet<TBaseStruct, TStructTrackerWrapper, TUStructType>, TBaseStruct, TUStructType>
+struct TStructTypeBitSet : TTypeBitSetBase<TStructTypeBitSet<TBaseStruct, TStructTrackerWrapper, TUStructType>, TBaseStruct, TUStructType, FStructTypeBitSet::FBitSetContainer>
 {
-	using Super = TTypeBitSetBase<TStructTypeBitSet, TBaseStruct, TUStructType>;
+	using Super = TTypeBitSetBase<TStructTypeBitSet, TBaseStruct, TUStructType, FStructTypeBitSet::FBitSetContainer>;
 	using FStructTrackerWrapper = TStructTrackerWrapper;
-	using FUStructType =  TUStructType;
+	using FUStructType = TUStructType;
 	using FBaseStruct = TBaseStruct;
 
 	friend Super;
@@ -584,11 +822,19 @@ struct TStructTypeBitSet : TTypeBitSetBase<TStructTypeBitSet<TBaseStruct, TStruc
 
 	TStructTypeBitSet() = default;
 
+	/**
+	 * Constructor that initializes the bitset with a single struct type.
+	 * @param StructType - The struct type to add.
+	 */
 	explicit TStructTypeBitSet(const FUStructType& StructType)
 	{
 		Add(StructType);
 	}
 
+	/**
+	 * Constructor that initializes the bitset with an initializer list of struct types.
+	 * @param InitList - The initializer list of struct type pointers.
+	 */
 	explicit TStructTypeBitSet(std::initializer_list<const FUStructType*> InitList)
 	{
 		for (const FUStructType* StructType : InitList)
@@ -600,6 +846,10 @@ struct TStructTypeBitSet : TTypeBitSetBase<TStructTypeBitSet<TBaseStruct, TStruc
 		}
 	}
 
+	/**
+	 * Constructor that initializes the bitset with an array view of struct types.
+	 * @param InitList - The array view of struct type pointers.
+	 */
 	explicit TStructTypeBitSet(TConstArrayView<const FUStructType*> InitList)
 	{
 		for (const FUStructType* StructType : InitList)
@@ -611,7 +861,11 @@ struct TStructTypeBitSet : TTypeBitSetBase<TStructTypeBitSet<TBaseStruct, TStruc
 		}
 	}
 
-	/** This flavor of constructor is only available for UScriptStructs */
+	/**
+	 * This constructor is only available for UScriptStructs.
+	 * Initializes the bitset with an array view of FInstancedStruct.
+	 * @param InitList - The array view of FInstancedStruct.
+	 */
 	template<typename T = TBaseStruct, typename = typename TEnableIf<!TIsDerivedFrom<T, UObject>::IsDerived>::Type>
 	explicit TStructTypeBitSet(TConstArrayView<FInstancedStruct> InitList)
 	{
@@ -625,15 +879,19 @@ struct TStructTypeBitSet : TTypeBitSetBase<TStructTypeBitSet<TBaseStruct, TStruc
 	}
 
 private:
-
 	/** 
-	 * A private constructor for a creating an instance straight from TBitArrays. 
-	 * @Note that this constructor needs to remain private to ensure consistency of stored values with data tracked 
-	 * by the TStructTrackerWrapper::StructTracker
+	 * A private constructor for creating an instance straight from TBitArrays. 
+	 * Note that this constructor needs to remain private to ensure consistency of stored values with data tracked 
+	 * by the TStructTrackerWrapper::StructTracker.
 	 */
 	TStructTypeBitSet(TBitArray<>&& Source)
-		: Super(Forward<TBitArray<>>(Source))
 	{
+		StructTypesBitArray = Forward<TBitArray<>>(Source);
+	}
+
+	TStructTypeBitSet(const TBitArray<>& Source)
+	{
+		StructTypesBitArray = Source;
 	}
 
 	TStructTypeBitSet(const int32 BitToSet)
@@ -643,16 +901,29 @@ private:
 
 public:
 
+	/**
+	 * Retrieves the FStructTracker associated with this bitset.
+	 * @return Reference to the FStructTracker.
+	 */
 	FStructTracker& GetStructTracker()
 	{
 		return TStructTrackerWrapper::StructTracker;
 	}
 
+	/**
+	 * Retrieves the FStructTracker associated with this bitset (const version).
+	 * @return Const reference to the FStructTracker.
+	 */
 	const FStructTracker& GetStructTracker() const
 	{
 		return TStructTrackerWrapper::StructTracker;
 	}
 
+	/**
+	 * Gets the index of a struct type within the tracker, adding it if not already present.
+	 * @param InStructType - The struct type to get the index for.
+	 * @return The index of the struct type.
+	 */
 	static int32 GetTypeIndex(const FUStructType& InStructType)
 	{
 #if WITH_STRUCTUTILS_DEBUG
@@ -664,6 +935,11 @@ public:
 		return TStructTrackerWrapper::StructTracker.FindOrAddStructTypeIndex(InStructType);
 	}
 
+	/**
+	 * Gets the index of a struct type within the tracker, adding it if not already present.
+	 * @param T - The struct type to get the index for.
+	 * @return The index of the struct type.
+	 */
 	template<typename T>
 	static int32 GetTypeIndex()
 	{
@@ -672,6 +948,11 @@ public:
 		return TypeIndex;
 	}
 	
+	/**
+	 * Gets a bitset representing a single struct type.
+	 * @param T - The struct type to get the bitset for.
+	 * @return Reference to a static bitset representing the struct type.
+	 */
 	template<typename T>
 	static const TStructTypeBitSet& GetTypeBitSet()
 	{
@@ -680,11 +961,20 @@ public:
 		return TypeBitSet;
 	}
 
+	/**
+	 * Gets the struct type associated with a given index.
+	 * @param Index - The index of the struct type.
+	 * @return Pointer to the struct type if found; nullptr otherwise.
+	 */
 	static const FUStructType* GetTypeAtIndex(const int32 Index)
 	{
 		return Cast<const FUStructType>(TStructTrackerWrapper::StructTracker.GetStructType(Index));
 	}
 
+	/**
+	 * Adds a struct type to the bitset.
+	 * @param T - The struct type to add.
+	 */
 	template<typename T>
 	FORCEINLINE void Add()
 	{
@@ -693,6 +983,10 @@ public:
 		StructTypesBitArray.AddAtIndex(StructTypeIndex);
 	}
 
+	/**
+	 * Removes a struct type from the bitset.
+	 * @param T - The struct type to remove.
+	 */
 	template<typename T>
 	FORCEINLINE void Remove()
 	{
@@ -701,11 +995,20 @@ public:
 		StructTypesBitArray.RemoveAtIndex(StructTypeIndex);
 	}
 
+	/**
+	 * Removes the bits from another bitset from this one (difference).
+	 * @param Other - The other bitset to subtract.
+	 */
 	FORCEINLINE void Remove(const TStructTypeBitSet& Other)
 	{
 		StructTypesBitArray -= Other.StructTypesBitArray;
 	}
 
+	/**
+	 * Checks if the bitset contains a specific struct type.
+	 * @param T - The struct type to check for.
+	 * @return True if the struct type is in the bitset; false otherwise.
+	 */
 	template<typename T>
 	FORCEINLINE bool Contains() const
 	{
@@ -714,6 +1017,11 @@ public:
 		return StructTypesBitArray.Contains(StructTypeIndex);
 	}
 
+	/**
+	 * Performs a union operation with another bitset.
+	 * @param Other - The other bitset to combine with.
+	 * @return A new bitset representing the union.
+	 */
 	FORCEINLINE TStructTypeBitSet operator+(const TStructTypeBitSet& Other) const
 	{
 		TStructTypeBitSet Result;
@@ -721,14 +1029,40 @@ public:
 		return MoveTemp(Result);
 	}
 
+	/**
+	 * Gets the maximum number of struct types registered in the tracker.
+	 * @return The number of struct types.
+	 */
 	static int32 GetMaxNum() 
 	{		
 		return TStructTrackerWrapper::StructTracker.Num();
 	}
 
-	FORCEINLINE bool operator==(const TStructTypeBitSet& Other) const { return Super::StructTypesBitArray == Other.StructTypesBitArray; }
-	FORCEINLINE bool operator!=(const TStructTypeBitSet& Other) const { return !(*this == Other); }
+	/**
+	 * Equality operator.
+	 * @param Other - The other bitset to compare
+	 * @return True if the bitsets are equal; false otherwise.
+	 */
+	FORCEINLINE bool operator==(const TStructTypeBitSet& Other) const
+	{
+		return Super::StructTypesBitArray == Other.StructTypesBitArray;
+	}
 
+	/**
+	 * Inequality operator.
+	 * @param Other - The other bitset to compare.
+	 * @return True if the bitsets are not equal; false otherwise.
+	 */
+	FORCEINLINE bool operator!=(const TStructTypeBitSet& Other) const
+	{
+		return !(*this == Other);
+	}
+
+	/**
+	 * Performs a difference operation with another bitset.
+	 * @param Other - The other bitset to subtract.
+	 * @return A new bitset representing the difference.
+	 */
 	FORCEINLINE TStructTypeBitSet operator-(const TStructTypeBitSet& Other) const
 	{
 		TStructTypeBitSet Result = *this;
@@ -737,11 +1071,12 @@ public:
 	}
 
 	/**
-	 * Lists all types used by this bit set, calling the provided callback for each one. Returning false from
-	 * the callback will early-out of iterating over the types.
-	 * 
-	 * note that this function is slow(ish) due to the FStructTracker utilizing WeakObjectPtrs to store types.
+	 * Lists all types used by this bitset, calling the provided callback for each one.
+	 * Returning false from the callback will early-out of iterating over the types.
+	 *
+	 * Note that this function is slow due to the FStructTracker utilizing WeakObjectPtrs to store types.
 	 * @todo To be improved.
+	 * @param Callback - The callback function to invoke for each type.
 	 */
 	void ExportTypes(TFunctionRef<bool(const FUStructType*)> Callback) const
 	{
@@ -759,6 +1094,10 @@ public:
 
 #if WITH_STRUCTUTILS_DEBUG
 	using Super::DebugGetStringDesc;
+	/**
+	 * Provides a debug string description of the bitset contents.
+	 * @return The debug string.
+	 */
 	FString DebugGetStringDesc() const
 	{
 		FStringOutputDevice Ar;
@@ -766,6 +1105,10 @@ public:
 		return static_cast<FString>(Ar);
 	}
 #else
+	/**
+	 * Provides a placeholder debug string when debug info is compiled out.
+	 * @return The placeholder string.
+	 */
 	FString DebugGetStringDesc() const
 	{
 		return TEXT("DEBUG INFO COMPILED OUT");
@@ -773,19 +1116,29 @@ public:
 #endif //WITH_STRUCTUTILS_DEBUG
 
 #if WITH_STRUCTUTILS_DEBUG
+	/**
+	 * Retrieves the name of a struct type by index (for debugging).
+	 * @param StructTypeIndex - The index of the struct type.
+	 * @return The name of the struct type.
+	 */
 	static FName DebugGetStructTypeName(const int32 StructTypeIndex)
 	{
 		return TStructTrackerWrapper::StructTracker.DebugGetStructTypeName(StructTypeIndex);
 	}
 
+	/**
+	 * Retrieves all struct types (for debugging).
+	 * @return An array view of weak pointers to the struct types.
+	 */
 	static TConstArrayView<TWeakObjectPtr<const TUStructType>> DebugGetAllStructTypes()
 	{
 		return TStructTrackerWrapper::StructTracker.template DebugGetAllStructTypes<TUStructType>();
 	}
 
 	/**
-	 * Resets all the information gathered on the tags. Calling this results in invalidating all previously created
-	 * FStructTypeCollection instances. Used only for debugging and unit/functional testing.
+	 * Resets all the information gathered on the types.
+	 * Calling this results in invalidating all previously created TStructTypeBitSet instances.
+	 * Used only for debugging and unit/functional testing.
 	 */
 	static void DebugResetStructTypeMappingInfo()
 	{
@@ -794,6 +1147,11 @@ public:
 #endif // WITH_STRUCTUTILS_DEBUG
 
 public:
+	/**
+	 * Hash function for the bitset.
+	 * @param Instance - The bitset instance to hash.
+	 * @return The hash value.
+	 */
 	FORCEINLINE friend uint32 GetTypeHash(const TStructTypeBitSet& Instance)
 	{
 		const uint32 StoredTypeHash = PointerHash(GetBaseUStruct());
@@ -801,11 +1159,21 @@ public:
 		return HashCombine(StoredTypeHash, BitArrayHash);
 	}
 
+	/**
+	 * Serializes the bitset using the provided archive.
+	 * @param Ar - The archive to serialize with.
+	 */
 	void Serialize(FArchive& Ar)
 	{
 		TStructTrackerWrapper::StructTracker.Serialize(Ar, Super::StructTypesBitArray);
 	}
 
+	/**
+	 * Stream insertion operator for serialization.
+	 * @param Ar - The archive to serialize with.
+	 * @param Instance - The bitset instance to serialize.
+	 * @return The archive after serialization.
+	 */
 	FORCEINLINE friend FArchive& operator<<(FArchive& Ar, TStructTypeBitSet& Instance)
 	{
 		Instance.Serialize(Ar);
@@ -815,7 +1183,7 @@ public:
 
 /** 
  * We're declaring StructTracker this way rather than being a static member variable of TStructTypeBitSet to avoid linking
- * issues. We're run into all sorts of issues depending on compiler and whether strict mode was on, and this is the best
+ * issues. We've run into all sorts of issues depending on compiler and whether strict mode was on, and this is the best
  * way we could come up with to solve it. Thankfully the user doesn't need to even know about this class's existence
  * as long as they are using the macros below.
  */
