Removed redundant Verse property initialization tracking for default-initialized objects.
We now track using the list of property initializers stored in the Verse class (done in FUObjectGenerator::Assemble).

#rnx

diff --git a/Engine/Source/Runtime/CoreUObject/Private/Blueprint/BlueprintSupport.cpp b/Engine/Source/Runtime/CoreUObject/Private/Blueprint/BlueprintSupport.cpp
--- a/Engine/Source/Runtime/CoreUObject/Private/Blueprint/BlueprintSupport.cpp
+++ b/Engine/Source/Runtime/CoreUObject/Private/Blueprint/BlueprintSupport.cpp
@@ -2655,13 +2655,6 @@ bool FObjectInitializer::InitNonNativeProperty(FProperty* Property, UObject* Dat
 		if (!Property->HasAnyPropertyFlags(CPF_ZeroConstructor)) // this stuff is already zero
 		{
 			Property->InitializeValue_InContainer(Data);
-
-			#if WITH_EDITORONLY_DATA
-			if (Property->HasAnyPropertyFlags(CPF_RequiredParm))
-			{
-				UE::FInitializedPropertyValueState(Data->GetClass(), Data).Set(Property);
-			}
-			#endif
 		}
 		return true;
 	}
