[CoreRedirects]

- Moves FCoreRedirects global state into a FCoreRedirectsContext.
- Threads may now push/pop a thread local context for core redirects using FScopedCoreRedirectsContext. If no context is pushed a shared global context is used (as was the case previously).
- When creating a context to push you may inherit the state of the existing context. This copies the existing state protecting the original state from non-threadsafe mutation should it be shared.
-RedirectionSummary was moved into internal so it may be used by the context since it's part of the CoreRedirects global state.
- Private locking structures have been moved into LockedContext types for passing around rather than passing around a lock and a context.

[CoreRedirects]

Suppress warning about returning local reference.

[AssetHeaderPatcher]

- Removes global CoreRedirect use from the patcher and instead each task thread will push patcher redirects for the duration of the task. This is to prevent the patcher from affecting global redirects happening on other threads (adding in irrelevant redirects but also prevents removing important redirects unintentionally)
- Adds copy constructor for RedirectTypeMap to properly copy FRedirectNameMap and FWildcardData values to avoid contaminating existing scopes

#rnx

diff --git a/Engine/Source/Runtime/CoreUObject/Internal/UObject/CoreRedirects/CoreRedirectsContext.h b/Engine/Source/Runtime/CoreUObject/Internal/UObject/CoreRedirects/CoreRedirectsContext.h
--- a/Engine/Source/Runtime/CoreUObject/Internal/UObject/CoreRedirects/CoreRedirectsContext.h
+++ b/Engine/Source/Runtime/CoreUObject/Internal/UObject/CoreRedirects/CoreRedirectsContext.h
@@ -3,6 +3,7 @@
 #pragma once
 #include "Containers/Map.h"
 #include "Containers/UnrealString.h"
+#include "Misc/TransactionallySafeRWLock.h"
 #include "RedirectionSummary.h"
 #include "UObject/CoreRedirects.h"
 
