[StructUtils] Make the FInstancedPropertyBag write into the archive in Serialize when reinstanciating objects.

* Otherwise object reference to be replaced are not caught and reset to None
* Impact PCG with any reference to an object instance of a BP class, which is very disruptive.

#rnx

diff --git a/Engine/Source/Runtime/CoreUObject/Private/StructUtils/PropertyBag.cpp b/Engine/Source/Runtime/CoreUObject/Private/StructUtils/PropertyBag.cpp
--- a/Engine/Source/Runtime/CoreUObject/Private/StructUtils/PropertyBag.cpp
+++ b/Engine/Source/Runtime/CoreUObject/Private/StructUtils/PropertyBag.cpp
@@ -2739,7 +2739,7 @@ bool FInstancedPropertyBag::Serialize(FArchive& Ar)
 				Ar.Seek(Ar.Tell() + SerialSize);
 			}
 		}
-		else if (Ar.IsSaving())
+		else if (Ar.IsSaving() || Ar.IsObjectReferenceCollector() || Ar.IsCountingMemory())
 		{
 			check(BagStruct);
 			
