FPackagePath: Fix warnings in CookedEditor: "GetLocalFullPath called on FPackagePath %s which has an unspecified header extension". We had some callsites that were not handling missing packages and packages in IoStore correctly.
GetExplanationForUnavailablePackage: suppress the warning from FPackagePath; we expect it to be unknown, and calling code handles the missing extension.
FPackageName::DoesPackageExistEx: When getting information about the package being in IoStore from the AssetRegistry, we need to set the extension to the special empty extension that indicates the package is in IoStore rather than leaving it unspecified.
#rnx

diff --git a/Engine/Source/Runtime/CoreUObject/Private/Misc/PackageName.cpp b/Engine/Source/Runtime/CoreUObject/Private/Misc/PackageName.cpp
--- a/Engine/Source/Runtime/CoreUObject/Private/Misc/PackageName.cpp
+++ b/Engine/Source/Runtime/CoreUObject/Private/Misc/PackageName.cpp
@@ -1781,6 +1781,7 @@ void FPackageName::GetExplanationForUnavailablePackage(const FName& UnavailableP
 			{
 				InOutExplanation.Appendf(TEXT("FPackageName: Skipped package %s has a valid, mounted, mount point but does not exist either on disk or in iostore."),
 					*UnavailablePackageNameString);
+				SkippedPackagePath.SetHeaderExtension(EPackageExtension::EmptyString); // Suppress the warning from FPackagePath about unknown asset
 				const FString& LocalFullPath = SkippedPackagePath.GetLocalFullPath();
 				if (LocalFullPath.Len())
 				{
@@ -2007,6 +2008,15 @@ FPackageName::EPackageLocationFilter FPackageName::DoesPackageExistEx(const FPac
 				}
 				else
 				{
+					// IODispatcher packagepaths use the EmptyString HeaderExtension to indicate they are an IODispatcher path.
+					// We do need to provide the EmptyString extension, rather than leaving it unspecified, to suppress wasted effort and
+					// warnings from FPackagePath about missing extensions.
+					// For IODispatcher paths we respect an existing header extension and don't overwrite it. Unknown whether that's
+					// important; it has been that way for a while so some systems may be relying on it.
+					if (OutPackagePath->GetHeaderExtension() == EPackageExtension::Unspecified)
+					{
+						OutPackagePath->SetHeaderExtension(EPackageExtension::EmptyString);
+					}
 					return PackageLocation;
 				}
 			}
@@ -2117,6 +2127,11 @@ FPackageName::EPackageLocationFilter FPackageName::InternalDoesPackageExistEx(co
 			if (OutPackagePath)
 			{
 				*OutPackagePath = PackagePath;
+				// IODispatcher packagepaths use the EmptyString HeaderExtension to indicate they are an IODispatcher path.
+				// We do need to provide the EmptyString extension, rather than leaving it unspecified, to suppress wasted effort and
+				// warnings from FPackagePath about missing extensions.
+				// For IODispatcher paths we respect an existing header extension and don't overwrite it. Unknown whether that's
+				// important; it has been that way for a while so some systems may be relying on it.
 				if (OutPackagePath->GetHeaderExtension() == EPackageExtension::Unspecified)
 				{
 					OutPackagePath->SetHeaderExtension(EPackageExtension::EmptyString);
