[AsyncPackageLoader]

When using the AsyncLoadingThread in Editor, force the GC to flush loading before collection. This is to prevent issues where loading in editor can result in extra processes (e.g. reinstancing) that can leave object hierarchies of now invalid/transient objects which might be found in memory during a load preventing loading from disk.

#rnx

diff --git a/Engine/Source/Runtime/CoreUObject/Private/Serialization/AsyncPackageLoader.cpp b/Engine/Source/Runtime/CoreUObject/Private/Serialization/AsyncPackageLoader.cpp
--- a/Engine/Source/Runtime/CoreUObject/Private/Serialization/AsyncPackageLoader.cpp
+++ b/Engine/Source/Runtime/CoreUObject/Private/Serialization/AsyncPackageLoader.cpp
@@ -146,6 +146,19 @@ FAsyncLoadingThreadSettings::FAsyncLoadingThreadSettings()
 	bool bCommandLineEnable = FParse::Param(FCommandLine::Get(), TEXT("AsyncLoadingThread"));
 	bAsyncLoadingThreadEnabled = bCommandLineEnable || (bConfigValue && FApp::ShouldUseThreadingForPerformance() && !bCommandLineDisable);
 
+#if WITH_EDITOR
+	if (bAsyncLoadingThreadEnabled)
+	{
+		const TCHAR* FlushStreamingOnGCStr = TEXT("gc.FlushStreamingOnGC");
+		IConsoleVariable* FlushStreamingOnGCVar = IConsoleManager::Get().FindConsoleVariable(FlushStreamingOnGCStr);
+		check(FlushStreamingOnGCVar);
+		const int CurrentFlushStreamingOnGCValue = FlushStreamingOnGCVar->GetInt();
+
+		UE_CLOG(CurrentFlushStreamingOnGCValue == 0, LogStreaming, Display, TEXT("Using AsyncLoadingThread in Editor; forcing %s=1"), FlushStreamingOnGCStr);
+		FlushStreamingOnGCVar->Set(1);
+	}
+#endif
+
 	// if we are on an uncooked platform the default value for async loading is disabled while it is enabled for cooked platform
 	bConfigValue = FPlatformProperties::RequiresCookedData() ? true : false;
 
