[AsyncLoading2]

- NotifyConstructedDuringAsyncLoading will no longer append objects to the ConstructedObjects list for packages unless they are outered to the package owning the ConstructedObjects list. Previously objects that have (or whose parent has) been invalidated/reinstanced may be included in this list. This can result in unnecessary iteration and unnecessary postloading.
- Brings back change in 30713584 to discard constructed objects for packages that are done post loading. There is no need to add them since we won't postload them anymore, and proceeding can potentially lead to accessing the FAsyncPackage2 past its lifetime.
- Adds additional MultiThreadDetector read/write scopes for ConstructedObjects as there were some codepaths missing such guards

#rnx

diff --git a/Engine/Source/Runtime/CoreUObject/Private/Serialization/AsyncLoading2.cpp b/Engine/Source/Runtime/CoreUObject/Private/Serialization/AsyncLoading2.cpp
--- a/Engine/Source/Runtime/CoreUObject/Private/Serialization/AsyncLoading2.cpp
+++ b/Engine/Source/Runtime/CoreUObject/Private/Serialization/AsyncLoading2.cpp
@@ -8325,6 +8325,7 @@ EAsyncPackageState::Type FAsyncLoadingThread2::ProcessLoadedPackagesFromGameThre
 				{
 					// In the editor we need to find any assets and packages and add them to list for later callback
 					EditorCompletedUPackages.Add(Package->LinkerRoot);
+					UE_MT_SCOPED_READ_ACCESS(Package->ConstructedObjectsAccessDetector)
 					for (UObject* Object : Package->ConstructedObjects)
 					{
 						if (Object->IsAsset())
@@ -9119,35 +9120,44 @@ void FAsyncLoadingThread2::NotifyUnreachableObjects(const TArrayView<FUObjectIte
  */
 void FAsyncLoadingThread2::NotifyConstructedDuringAsyncLoading(UObject* Object, bool bSubObjectThatAlreadyExists)
 {
+	check(Object);
 	FUObjectThreadContext& ThreadContext = FUObjectThreadContext::Get();
-	if (!ThreadContext.AsyncPackage)
+	FAsyncPackage2* ContextAsyncPackage = (FAsyncPackage2*)ThreadContext.AsyncPackage;
+	UPackage* ObjectPackage = Object->GetPackage();
+
+	if (!ContextAsyncPackage || !ObjectPackage || ObjectPackage == GetTransientPackage())
 	{
 		// Something is creating objects on the async loading thread outside of the actual async loading code
 		// e.g. ShaderCodeLibrary::OnExternalReadCallback doing FTaskGraphInterface::Get().WaitUntilTaskCompletes(Event);
+		// Or the object doesn't have a package, in which case we didn't load it from disk
+		// Or the object is now in the transient package, and is no longer needed
+		UE_LOG(LogStreaming, VeryVerbose, TEXT("Constructed object '%s' during loading will not be post-loaded due to its package '%s' (null and transient package sub-objects are ignored)."), *Object->GetFullName(), (ObjectPackage ? *ObjectPackage->GetFullName() : TEXT("null")));
 		return;
 	}
 
-	FAsyncPackage2* AsyncPackage2 = (FAsyncPackage2*)ThreadContext.AsyncPackage;
-
 #if WITH_EDITOR
 	// In editor, objects from other packages might be constructed with the wrong package currently on stack (i.e. reinstantiation).
 	// It's a lot cleaner and easier to properly dispatch them to their respective package than adding scopes everywhere.
-	if (UPackage* ObjectPackage = Object->GetPackage())
+	if (FLinkerLoad* LinkerLoad = ObjectPackage->GetLinker())
 	{
-		if (FLinkerLoad* LinkerLoad = ObjectPackage->GetLinker())
+		if (FAsyncPackage2* AsyncPackage = (FAsyncPackage2*)LinkerLoad->AsyncRoot)
 		{
-			if (FAsyncPackage2* AsyncPackage = (FAsyncPackage2*)LinkerLoad->AsyncRoot)
+			if (AsyncPackage->AsyncPackageLoadingState < EAsyncPackageLoadingState2::DeferredPostLoad)
 			{
+				check(AsyncPackage->Desc.UPackageId == ObjectPackage->GetPackageId());
 				AsyncPackage->AddConstructedObject(Object, bSubObjectThatAlreadyExists);
 				return;
 			}
 		}
+	}
 
-		FPackageId PackageId = ObjectPackage->GetPackageId();
-		if (PackageId.IsValid() && PackageId != AsyncPackage2->Desc.UPackageId)
+	const FPackageId ObjectPackageId = ObjectPackage->GetPackageId();
+	if (ObjectPackageId.IsValid())
+	{
+		FScopeLock LockAsyncPackages(&AsyncPackagesCritical);
+		if (FAsyncPackage2* AsyncPackage = AsyncPackageLookup.FindRef(ObjectPackageId))
 		{
-			FScopeLock LockAsyncPackages(&AsyncPackagesCritical);
-			if (FAsyncPackage2* AsyncPackage = AsyncPackageLookup.FindRef(PackageId))
+			if (AsyncPackage->AsyncPackageLoadingState < EAsyncPackageLoadingState2::DeferredPostLoad)
 			{
 				AsyncPackage->AddConstructedObject(Object, bSubObjectThatAlreadyExists);
 				return;
@@ -9155,8 +9165,23 @@ void FAsyncLoadingThread2::NotifyConstructedDuringAsyncLoading(UObject* Object,
 		}
 	}
 #endif
-	
-	AsyncPackage2->AddConstructedObject(Object, bSubObjectThatAlreadyExists);
+
+	if (ContextAsyncPackage->AsyncPackageLoadingState <= EAsyncPackageLoadingState2::DeferredPostLoad)
+	{
+		UE_CLOG(ContextAsyncPackage->Desc.UPackageId != ObjectPackage->GetPackageId(), LogStreaming, VeryVerbose,
+			TEXT("Constructed object '%s' is part of package '%s' which is no longer postloading objects. This object will be postloaded with the FUObjectThreadContext AsyncPackage '%s'(%s) instead."),
+			*Object->GetFullName(),
+			*ObjectPackage->GetFullName(),
+			*ContextAsyncPackage->Desc.UPackageName.ToString(),
+			LexToString(ContextAsyncPackage->AsyncPackageLoadingState));
+
+		ContextAsyncPackage->AddConstructedObject(Object, bSubObjectThatAlreadyExists);
+	}
+	else
+	{
+		UE_LOG(LogStreaming, Warning, TEXT("Object '%s' was created while the FUObjectThreadContext's async package '%s(0x%s)' is done post-loading objects. This object may be left in a partially loaded state."),
+			*Object->GetFullName(), *ContextAsyncPackage->Desc.UPackageName.ToString(), *LexToString(ContextAsyncPackage->Desc.UPackageId));
+	}
 }
 
 void FAsyncLoadingThread2::NotifyRegistrationEvent(
@@ -9295,12 +9320,14 @@ void FAsyncPackage2::ClearImportedPackages()
 void FAsyncPackage2::MoveConstructedObjectsToPhase2()
 {
 	TRACE_CPUPROFILER_EVENT_SCOPE(MoveConstructedObjectsToPhase2);
-
-	for (UObject* Object : ConstructedObjects)
 	{
-		check(!Object->HasAnyFlags(RF_NeedLoad | RF_NeedInitialization));
-		Object->SetInternalFlags(EInternalObjectFlags::AsyncLoadingPhase2);
-		Object->AtomicallyClearInternalFlags(EInternalObjectFlags::AsyncLoadingPhase1);
+		UE_MT_SCOPED_READ_ACCESS(ConstructedObjectsAccessDetector)
+		for (UObject* Object : ConstructedObjects)
+		{
+			check(!Object->HasAnyFlags(RF_NeedLoad | RF_NeedInitialization));
+			Object->SetInternalFlags(EInternalObjectFlags::AsyncLoadingPhase2);
+			Object->AtomicallyClearInternalFlags(EInternalObjectFlags::AsyncLoadingPhase1);
+		}
 	}
 
 	TArray<UObject*> CDODefaultSubobjects;
@@ -9350,6 +9377,7 @@ void FAsyncPackage2::MoveConstructedObjectsToPhase2()
 void FAsyncPackage2::ClearConstructedObjects()
 {
 	TRACE_CPUPROFILER_EVENT_SCOPE(ClearConstructedObjects);
+	UE_MT_SCOPED_WRITE_ACCESS(ConstructedObjectsAccessDetector)
 
 	for (UObject* Object : ConstructedObjects)
 	{
