Struct Utils: Add fast-path for getting data with base type in templated structs

diff --git a/Engine/Source/Runtime/CoreUObject/Public/StructUtils/InstancedStruct.h b/Engine/Source/Runtime/CoreUObject/Public/StructUtils/InstancedStruct.h
--- a/Engine/Source/Runtime/CoreUObject/Public/StructUtils/InstancedStruct.h
+++ b/Engine/Source/Runtime/CoreUObject/Public/StructUtils/InstancedStruct.h
@@ -424,7 +424,14 @@ public:
 	template<typename T = BaseStructT, typename = std::enable_if_t<std::is_base_of_v<BaseStructT, std::decay_t<T>>>>
 	const T* GetPtr() const
 	{
-		return InstancedStruct.GetPtr<T>();
+		if constexpr (std::is_same_v<BaseStructT, std::decay_t<T>>)
+		{
+			return ((T*)InstancedStruct.GetMemory());
+		}
+		else
+		{
+			return InstancedStruct.GetPtr<T>();
+		}
 	}
 
 	/** Returns a mutable pointer to raw struct memory. */
@@ -444,7 +451,14 @@ public:
 	template<typename T = BaseStructT, typename = std::enable_if_t<std::is_base_of_v<BaseStructT, std::decay_t<T>>>>
 	T* GetMutablePtr()
 	{
-		return InstancedStruct.GetMutablePtr<T>();
+		if constexpr (std::is_same_v<BaseStructT, std::decay_t<T>>)
+		{
+			return ((T*)InstancedStruct.GetMutableMemory());
+		}
+		else
+		{
+			return InstancedStruct.GetMutablePtr<T>();
+		}
 	}
 
 	/** Returns True if the struct is valid.*/
diff --git a/Engine/Source/Runtime/CoreUObject/Public/StructUtils/StructView.h b/Engine/Source/Runtime/CoreUObject/Public/StructUtils/StructView.h
--- a/Engine/Source/Runtime/CoreUObject/Public/StructUtils/StructView.h
+++ b/Engine/Source/Runtime/CoreUObject/Public/StructUtils/StructView.h
@@ -173,7 +173,14 @@ public:
 	template<typename T = BaseStructT, typename = std::enable_if_t<std::is_base_of_v<BaseStructT, std::decay_t<T>>>>
 	T* GetPtr() const
 	{
-		return UE::StructUtils::GetStructPtr<T>(ScriptStruct, StructMemory);
+		if constexpr (std::is_same_v<BaseStructT, std::decay_t<T>>)
+		{
+			return ((T*)StructMemory);
+		}
+		else
+		{
+			return UE::StructUtils::GetStructPtr<T>(ScriptStruct, StructMemory);
+		}
 	}
 
 	/** Comparison operators. Note: it does not compare the internal structure itself */
@@ -378,7 +385,14 @@ struct TConstStructView : FConstStructView
 	template<typename T = BaseStructT, typename = std::enable_if_t<std::is_base_of_v<BaseStructT, std::decay_t<T>>>>
 	constexpr const T* GetPtr() const
 	{
-		return UE::StructUtils::GetStructPtr<T>(ScriptStruct, StructMemory);
+		if constexpr (std::is_same_v<BaseStructT, std::decay_t<T>>)
+		{
+			return ((T*)StructMemory);
+		}
+		else
+		{
+			return UE::StructUtils::GetStructPtr<T>(ScriptStruct, StructMemory);
+		}
 	}
 
 	/** Comparison operators. Note: it does not compare the internal structure itself */
