Add IS_ANY_PACKAGE_DEPRECATED that checks whether a given UObject* is equal to ANY_PACKAGE_DEPRECATED without the misalignment UB of casting -1 to a UObject*.

diff --git a/Engine/Source/Runtime/CoreUObject/Private/UObject/AnyPackagePrivate.h b/Engine/Source/Runtime/CoreUObject/Private/UObject/AnyPackagePrivate.h
--- a/Engine/Source/Runtime/CoreUObject/Private/UObject/AnyPackagePrivate.h
+++ b/Engine/Source/Runtime/CoreUObject/Private/UObject/AnyPackagePrivate.h
@@ -6,4 +6,6 @@
 
 #pragma once
 
-#define ANY_PACKAGE_DEPRECATED ((UPackage*)-1)
\ No newline at end of file
+#define ANY_PACKAGE_DEPRECATED ((UPackage*)-1)
+
+#define IS_ANY_PACKAGE_DEPRECATED(Package) (reinterpret_cast<PTRINT>(Package) == -1)
\ No newline at end of file
diff --git a/Engine/Source/Runtime/CoreUObject/Private/UObject/UObjectGlobals.cpp b/Engine/Source/Runtime/CoreUObject/Private/UObject/UObjectGlobals.cpp
--- a/Engine/Source/Runtime/CoreUObject/Private/UObject/UObjectGlobals.cpp
+++ b/Engine/Source/Runtime/CoreUObject/Private/UObject/UObjectGlobals.cpp
@@ -555,7 +555,7 @@ UObject* StaticFindObject( UClass* ObjectClass, UObject* InObjectPackage, const
 	INC_DWORD_STAT(STAT_FindObject);
 
 	// Resolve the object and package name.
-	const bool bAnyPackage = InObjectPackage == ANY_PACKAGE_DEPRECATED;
+	const bool bAnyPackage = IS_ANY_PACKAGE_DEPRECATED(InObjectPackage);
 	UObject* ObjectPackage = bAnyPackage ? nullptr : InObjectPackage;
 
 	UObject* MatchingObject = nullptr;
@@ -659,7 +659,7 @@ UObject* StaticFindObjectChecked( UClass* ObjectClass, UObject* ObjectParent, co
 #if !(UE_BUILD_SHIPPING || UE_BUILD_TEST)
 	if( !Result )
 	{
-		UE_LOG(LogUObjectGlobals, Fatal, TEXT("Failed to find object '%s %s.%s'"), *ObjectClass->GetName(), ObjectParent == ANY_PACKAGE_DEPRECATED ? TEXT("Any") : ObjectParent ? *ObjectParent->GetName() : TEXT("None"), InName);
+		UE_LOG(LogUObjectGlobals, Fatal, TEXT("Failed to find object '%s %s.%s'"), *ObjectClass->GetName(), IS_ANY_PACKAGE_DEPRECATED(ObjectParent) ? TEXT("Any") : ObjectParent ? *ObjectParent->GetName() : TEXT("None"), InName);
 	}
 #endif
 	return Result;
@@ -1349,7 +1349,7 @@ bool ParseObject( const TCHAR* Stream, const TCHAR* Match, UClass* Class, UObjec
 					? StaticLoadObject(Class, nullptr, ObjNameOrPathName)
 					: StaticFindObject(Class, nullptr, ObjNameOrPathName);
 			}
-			else if (InParent && InParent != ANY_PACKAGE_DEPRECATED)
+			else if (InParent && !IS_ANY_PACKAGE_DEPRECATED(InParent))
 			{
 				// Try to find the object within its parent
 				return StaticFindObject(Class, InParent, ObjNameOrPathName);
@@ -2784,7 +2784,7 @@ namespace NameReuse
 		}
 
 		static const FName NamePackage(NAME_Package);
-		if (!Parent || Parent == ANY_PACKAGE_DEPRECATED || !!(Options & EUniqueObjectNameOptions::GloballyUnique) || BaseName == NamePackage || FPlatformProperties::HasEditorOnlyData() || !GFastPathUniqueNameGeneration)
+		if (!Parent || IS_ANY_PACKAGE_DEPRECATED(Parent) || !!(Options & EUniqueObjectNameOptions::GloballyUnique) || BaseName == NamePackage || FPlatformProperties::HasEditorOnlyData() || !GFastPathUniqueNameGeneration)
 		{
 			return FName();
 		}
@@ -2863,7 +2863,7 @@ FName MakeUniqueObjectName(UObject* Parent, const UClass* Class, FName InBaseNam
 				int32 NameNumber = 0;
 				UE_AUTORTFM_OPEN
 				{
-					if (Parent && (Parent != ANY_PACKAGE_DEPRECATED) && !(Options & EUniqueObjectNameOptions::GloballyUnique))
+					if (Parent && !IS_ANY_PACKAGE_DEPRECATED(Parent) && !(Options & EUniqueObjectNameOptions::GloballyUnique))
 					{
 						if (!FPlatformProperties::HasEditorOnlyData() && GFastPathUniqueNameGeneration)
 						{
@@ -2898,7 +2898,7 @@ FName MakeUniqueObjectName(UObject* Parent, const UClass* Class, FName InBaseNam
 				TestName = FName(BaseName, NameNumber);
 			}
 
-			if (Parent == ANY_PACKAGE_DEPRECATED || !!(Options & EUniqueObjectNameOptions::GloballyUnique))
+			if (IS_ANY_PACKAGE_DEPRECATED(Parent) || !!(Options & EUniqueObjectNameOptions::GloballyUnique))
 			{
 				ExistingObject = StaticFindFirstObject(nullptr, *TestName.ToString());
 			}
diff --git a/Engine/Source/Runtime/CoreUObject/Private/UObject/UObjectHash.cpp b/Engine/Source/Runtime/CoreUObject/Private/UObject/UObjectHash.cpp
--- a/Engine/Source/Runtime/CoreUObject/Private/UObject/UObjectHash.cpp
+++ b/Engine/Source/Runtime/CoreUObject/Private/UObject/UObjectHash.cpp
@@ -869,7 +869,7 @@ UObject* StaticFindObjectFastInternal(const UClass* ObjectClass, const UObject*
 	{
 		INC_DWORD_STAT(STAT_FindObjectFast);
 
-		check(ObjectPackage != ANY_PACKAGE_DEPRECATED); // this could never have returned anything but nullptr
+		check(!IS_ANY_PACKAGE_DEPRECATED(ObjectPackage)); // this could never have returned anything but nullptr
 
 		// If they specified an outer use that during the hashing
 		FUObjectHashTables& ThreadHash = FUObjectHashTables::Get();
@@ -889,7 +889,7 @@ UObject* StaticFindObjectFastInternal(const UClass* ObjectClass, const UObject*
 	{
 		INC_DWORD_STAT(STAT_FindObjectFast);
 
-		check(ObjectPackage != ANY_PACKAGE_DEPRECATED); // this could never have returned anything but nullptr
+		check(!IS_ANY_PACKAGE_DEPRECATED(ObjectPackage)); // this could never have returned anything but nullptr
 
 		// If they specified an outer use that during the hashing
 		FUObjectHashTables& ThreadHash = FUObjectHashTables::Get();
