PR #12461: FAsyncPackage::CreateLinker Added referencer to error

#rnx

diff --git a/Engine/Source/Runtime/CoreUObject/Private/Serialization/AsyncLoading.cpp b/Engine/Source/Runtime/CoreUObject/Private/Serialization/AsyncLoading.cpp
--- a/Engine/Source/Runtime/CoreUObject/Private/Serialization/AsyncLoading.cpp
+++ b/Engine/Source/Runtime/CoreUObject/Private/Serialization/AsyncLoading.cpp
@@ -6257,7 +6257,13 @@ EAsyncPackageState::Type FAsyncPackage::CreateLinker()
 
 					if (!FLinkerLoad::IsKnownMissingPackage(NameToLoadFName))
 					{
-						UE_LOGFMT(LogStreaming, Error, "Couldn't find file for package {Package} requested by async loading code. NameToLoad: {NameToLoad}", UE::FAssetLog(*Desc.Name.ToString()), UE::FAssetLog(Desc.PackagePath));
+#if UE_WITH_PACKAGE_ACCESS_TRACKING
+						UE_LOGFMT(LogStreaming, Error, "Couldn't find file for package {Package} requested by async loading code. NameToLoad: {NameToLoad}. Referencer package: {Referencer}",
+							UE::FAssetLog(*Desc.Name.ToString()), UE::FAssetLog(Desc.PackagePath), UE::FAssetLog(*Desc.ReferencerPackageName.ToString()));
+#else
+						UE_LOGFMT(LogStreaming, Error, "Couldn't find file for package {Package} requested by async loading code. NameToLoad: {NameToLoad}",
+							UE::FAssetLog(*Desc.Name.ToString()), UE::FAssetLog(Desc.PackagePath));
+#endif
 
 #if !WITH_EDITORONLY_DATA
 						UE_CLOG(bUseTimeLimit, LogStreaming, Error, TEXT("This will hitch streaming because it ends up searching the disk instead of finding the file in the pak file."));
