Use UE::FMutex in certain systems instead of system objects synchronization. This is enabled on certain platforms that have a memory limitation for system objects

diff --git a/Engine/Source/Runtime/Core/Public/HAL/Platform.h b/Engine/Source/Runtime/Core/Public/HAL/Platform.h
--- a/Engine/Source/Runtime/Core/Public/HAL/Platform.h
+++ b/Engine/Source/Runtime/Core/Public/HAL/Platform.h
@@ -626,6 +626,11 @@
 	#define PLATFORM_WRITES_ARE_SLOW 0
 #endif
 
+#ifndef PLATFORM_FORCE_UE_LOCK_USAGE
+	#define PLATFORM_FORCE_UE_LOCK_USAGE 0
+#endif
+
+
 // deprecated, do not use
 #define PLATFORM_HAS_THREADSAFE_RHIGetRenderQueryResult	#
 #define PLATFORM_SUPPORTS_RHI_THREAD #
diff --git a/Engine/Source/Runtime/CoreUObject/Private/UObject/LinkerInstancingContext.cpp b/Engine/Source/Runtime/CoreUObject/Private/UObject/LinkerInstancingContext.cpp
--- a/Engine/Source/Runtime/CoreUObject/Private/UObject/LinkerInstancingContext.cpp
+++ b/Engine/Source/Runtime/CoreUObject/Private/UObject/LinkerInstancingContext.cpp
@@ -13,10 +13,20 @@
 
 LLM_DEFINE_TAG(Loading_LinkerInstancingContext);
 
+#if PLATFORM_FORCE_UE_LOCK_USAGE && !UE_AUTORTFM
+	typedef UE::TUniqueLock<UE::FMutex> FLinkerInstancingContextReadLock;
+	typedef UE::TUniqueLock<UE::FMutex> FLinkerInstancingContextWriteLock;
+	typedef UE::FMutex FLinkerInstancingLock;
+#else
+	typedef UE::TReadScopeLock<FTransactionallySafeRWLock> FLinkerInstancingContextReadLock;
+	typedef UE::TWriteScopeLock<FTransactionallySafeRWLock> FLinkerInstancingContextWriteLock;
+	typedef FTransactionallySafeRWLock FLinkerInstancingLock;
+#endif
+
 class FLinkerInstancingContext::FSharedLinkerInstancingContextData
 {
 	/** The shared content needs to be protected as it could be modified by the caller as well as the loading thread */
-	mutable FTransactionallySafeRWLock Lock;
+	mutable FLinkerInstancingLock Lock;
 	/** Map of original package name to their instance counterpart. */
 	FLinkerInstancedPackageMap InstancedPackageMap;
 	/** Optional function to map original package name to their instance counterpart. The result of this function should be immutable, as it will be cached. */
@@ -51,14 +61,14 @@ public:
 		
 	void AddPackageMapping(FName Original, FName Instanced)
 	{
-		UE::TWriteScopeLock ScopeLock(Lock);
+		FLinkerInstancingContextWriteLock ScopeLock(Lock);
 		LLM_SCOPE_BYTAG(Loading_LinkerInstancingContext);
 		InstancedPackageMap.AddPackageMapping(Original, Instanced);
 	}
 
 	bool FindPackageMapping(FName Original, FName& Instanced) const
 	{
-		UE::TReadScopeLock ScopeLock(Lock);
+		FLinkerInstancingContextReadLock ScopeLock(Lock);
 		LLM_SCOPE_BYTAG(Loading_LinkerInstancingContext);
 		if (const FName* InstancedPtr = InstancedPackageMap.InstancedPackageMapping.Find(Original))
 		{
@@ -70,14 +80,14 @@ public:
 
 	void FixupSoftObjectPath(FSoftObjectPath& InOutSoftObjectPath) const
 	{
-		UE::TReadScopeLock ScopeLock(Lock);
+		FLinkerInstancingContextReadLock ScopeLock(Lock);
 		LLM_SCOPE_BYTAG(Loading_LinkerInstancingContext);
 		InstancedPackageMap.FixupSoftObjectPath(InOutSoftObjectPath);
 	}
 
 	FSoftObjectPath RemapPath(const FSoftObjectPath& Path) const
 	{
-		UE::TReadScopeLock ScopeLock(Lock);
+		FLinkerInstancingContextReadLock ScopeLock(Lock);
 		if (const FTopLevelAssetPath* Remapped = PathMapping.Find(Path.GetAssetPath()))
 		{
 			return FSoftObjectPath(*Remapped, Path.GetSubPathString());
@@ -87,40 +97,40 @@ public:
 
 	bool IsInstanced() const
 	{
-		UE::TReadScopeLock ScopeLock(Lock);
+		FLinkerInstancingContextReadLock ScopeLock(Lock);
 		return InstancedPackageMap.IsInstanced() || InstancedPackageMapFunc || PathMapping.Num() > 0;
 	}
 
 	void EnableAutomationTest()
 	{
-		UE::TWriteScopeLock ScopeLock(Lock);
+		FLinkerInstancingContextWriteLock ScopeLock(Lock);
 		InstancedPackageMap.EnableAutomationTest();
 	}
 
 	void BuildPackageMapping(FName Original, FName Instanced, bool bInSoftObjectPathRemappingEnabled)
 	{
-		UE::TWriteScopeLock ScopeLock(Lock);
+		FLinkerInstancingContextWriteLock ScopeLock(Lock);
 		LLM_SCOPE_BYTAG(Loading_LinkerInstancingContext);
 		InstancedPackageMap.BuildPackageMapping(Original, Instanced, bInSoftObjectPathRemappingEnabled);
 	}
 
 	void AddTag(FName NewTag)
 	{
-		UE::TWriteScopeLock ScopeLock(Lock);
+		FLinkerInstancingContextWriteLock ScopeLock(Lock);
 		LLM_SCOPE_BYTAG(Loading_LinkerInstancingContext);
 		Tags.Add(NewTag);
 	}
 
 	void AppendTags(const TSet<FName>& NewTags)
 	{
-		UE::TWriteScopeLock ScopeLock(Lock);
+		FLinkerInstancingContextWriteLock ScopeLock(Lock);
 		LLM_SCOPE_BYTAG(Loading_LinkerInstancingContext);
 		Tags.Append(NewTags);
 	}
 
 	bool HasTag(FName Tag) const
 	{
-		UE::TReadScopeLock ScopeLock(Lock);
+		FLinkerInstancingContextReadLock ScopeLock(Lock);
 		return Tags.Contains(Tag);
 	}
 
@@ -136,14 +146,14 @@ public:
 
 	void AddPathMapping(FSoftObjectPath Original, FSoftObjectPath Instanced)
 	{
-		UE::TWriteScopeLock ScopeLock(Lock);
+		FLinkerInstancingContextWriteLock ScopeLock(Lock);
 		LLM_SCOPE_BYTAG(Loading_LinkerInstancingContext);
 		PathMapping.Emplace(Original.GetAssetPath(), Instanced.GetAssetPath());
 	}
 
 	void SetPackageMappingFunc(TFunction<FName(FName)> InInstancedPackageMapFunc)
 	{
-		UE::TWriteScopeLock ScopeLock(Lock);
+		FLinkerInstancingContextWriteLock ScopeLock(Lock);
 		InstancedPackageMapFunc = MoveTemp(InInstancedPackageMapFunc);
 	}
 
@@ -153,7 +163,7 @@ public:
 
 		FName RemappedPackageName;
 		{
-			UE::TReadScopeLock ScopeLock(Lock);
+			FLinkerInstancingContextReadLock ScopeLock(Lock);
 			RemappedPackageName = InstancedPackageMap.RemapPackage(PackageName);
 
 			if ((RemappedPackageName == PackageName) && InstancedPackageMapFunc)
@@ -169,7 +179,7 @@ public:
 
 		if (bAddPackageMapping)
 		{
-			UE::TWriteScopeLock ScopeLock(Lock);
+			FLinkerInstancingContextWriteLock ScopeLock(Lock);
 			LLM_SCOPE_BYTAG(Loading_LinkerInstancingContext);
 			InstancedPackageMap.AddPackageMapping(PackageName, RemappedPackageName);
 		}
diff --git a/Engine/Source/Runtime/CoreUObject/Public/UObject/Class.h b/Engine/Source/Runtime/CoreUObject/Public/UObject/Class.h
--- a/Engine/Source/Runtime/CoreUObject/Public/UObject/Class.h
+++ b/Engine/Source/Runtime/CoreUObject/Public/UObject/Class.h
@@ -113,11 +113,7 @@ enum class EPropertyObjectReferenceType : uint32;
 COREUOBJECT_API DECLARE_LOG_CATEGORY_EXTERN(LogClass, Log, All);
 COREUOBJECT_API DECLARE_LOG_CATEGORY_EXTERN(LogScriptSerialization, Log, All);
 
-#ifndef USE_UE_LOCK_FOR_UCLASS_FUNCTION_HANDLING
-	#define USE_UE_LOCK_FOR_UCLASS_FUNCTION_HANDLING 0
-#endif
-
-#if USE_UE_LOCK_FOR_UCLASS_FUNCTION_HANDLING
+#if PLATFORM_FORCE_UE_LOCK_USAGE
 	typedef UE::TUniqueLock<UE::FMutex> FUClassFuncScopeReadLock;
 	typedef UE::TUniqueLock<UE::FMutex> FUClassFuncScopeWriteLock;
 	typedef UE::FMutex FUClassFuncLock;
