Cleanup spammy Core stat events.

+ Remove FBulkData scopes from basic operators as there can be tens of thousands of these in a frame and they're basically 0 time so they just bloat parent scopes.
+ Mark CreateStatId as verbose as it's not very useful.
+ Mark STAT_PostTickComponentLW and STAT_Canvas_GetBatchElementsTime verbose as they're very numerous and pretty much instant.
+ Mark STAT_ComponentCreateRenderState/STAT_ComponentDestroyRenderState and STAT_ComponentCreatePhysicsState/STAT_ComponentDestroyPhysicsState verbose.
+ Mark STAT_InitProperties as it's very frequent during object construction but typically very fast.

#rnx

diff --git a/Engine/Source/Runtime/CoreUObject/Private/Serialization/BulkData.cpp b/Engine/Source/Runtime/CoreUObject/Private/Serialization/BulkData.cpp
--- a/Engine/Source/Runtime/CoreUObject/Private/Serialization/BulkData.cpp
+++ b/Engine/Source/Runtime/CoreUObject/Private/Serialization/BulkData.cpp
@@ -604,8 +604,6 @@ const void* FOwnedBulkDataPtr::GetPointer()
  */
 FBulkData::FBulkData( const FBulkData& Other )
 {
-	DECLARE_SCOPE_CYCLE_COUNTER(TEXT("FBulkData::FBulkData"), STAT_UBD_Constructor, STATGROUP_Memory);
-
 	*this = Other;
 
 #if TRACK_BULKDATA_USE
@@ -618,8 +616,6 @@ FBulkData::FBulkData( const FBulkData& Other )
  */
 FBulkData::~FBulkData()
 {
-	DECLARE_SCOPE_CYCLE_COUNTER(TEXT("FBulkData::~FBulkData"), STAT_UBD_Destructor, STATGROUP_Memory);
-
 	check(IsUnlocked());
 
 	// Free memory.
@@ -646,8 +642,6 @@ FBulkData::~FBulkData()
  */
 FBulkData& FBulkData::operator=( const FBulkData& Other )
 {
-	DECLARE_SCOPE_CYCLE_COUNTER(TEXT("FBulkData::operator="), STAT_UBD_Constructor, STATGROUP_Memory);
-
 	checkf(IsUnlocked(), TEXT("Attempting to modify a BulkData object that is locked"));
 	checkf(Other.BulkMeta.GetLockStatus() != LOCKSTATUS_ReadWriteLock, TEXT("Attempting to read from a BulkData object that is locked for write"));
 
@@ -877,8 +871,6 @@ PRAGMA_ENABLE_DEPRECATION_WARNINGS
  */
 void FBulkData::GetCopy( void** Dest, bool bDiscardInternalCopy )
 {
-	DECLARE_SCOPE_CYCLE_COUNTER(TEXT("FBulkData::GetCopy"), STAT_UBD_GetCopy, STATGROUP_Memory);
-
 	check( IsUnlocked() );
 	check( Dest );
 
@@ -1017,8 +1009,6 @@ const void* FBulkData::LockReadOnly() const
 
 void* FBulkData::Realloc(int64 ElementCount, int64 ElementSize)
 {
-	DECLARE_SCOPE_CYCLE_COUNTER(TEXT("FBulkData::Realloc"), STAT_UBD_Realloc, STATGROUP_Memory);
-
 	checkf(IsLocked(), TEXT("BulkData must be locked for 'write' before reallocating!"));
 #if !WITH_EDITOR
 	checkf(!CanLoadFromDisk(), TEXT("Cannot re-allocate a FBulkDataBase object that represents a file on disk!"));
@@ -1040,8 +1030,6 @@ void* FBulkData::Realloc(int64 ElementCount, int64 ElementSize)
  */
 void FBulkData::Unlock() const
 {
-	DECLARE_SCOPE_CYCLE_COUNTER(TEXT("FBulkData::Unlock"), STAT_UBD_Unlock, STATGROUP_Memory);
-
 	check(IsUnlocked() == false);
 
 	FBulkData* mutable_this = const_cast<FBulkData*>(this);
@@ -1060,8 +1048,6 @@ void FBulkData::Unlock() const
  */
 void FBulkData::RemoveBulkData()
 {
-	DECLARE_SCOPE_CYCLE_COUNTER(TEXT("FBulkData::RemoveBulkData"), STAT_UBD_RemoveBulkData, STATGROUP_Memory);
-
 	check(IsUnlocked());
 
 #if WITH_EDITOR
@@ -1253,7 +1239,8 @@ bool FBulkData::HasFlags(EBulkDataFlags Flags, EBulkDataFlags Contains)
 
 void FBulkData::Serialize(FArchive& Ar, UObject* Owner, bool bAttemptFileMapping, int32 ElementSize, EFileRegionType FileRegionType)
 {
-	DECLARE_SCOPE_CYCLE_COUNTER(TEXT("FBulkData::Serialize"), STAT_UBD_Serialize, STATGROUP_Memory);
+	DECLARE_CYCLE_STAT_WITH_FLAGS(TEXT("FBulkData::Serialize"), STAT_UBD_Serialize, STATGROUP_Memory, EStatFlags::Verbose);
+	SCOPE_CYCLE_COUNTER(STAT_UBD_Serialize);
 
 	SCOPED_LOADTIMER(BulkData_Serialize);
 
@@ -1648,8 +1635,6 @@ void FBulkData::MakeSureBulkDataIsLoaded()
  */
 bool FBulkData::TryLoadDataIntoMemory(FIoBuffer Dest)
 {
-	DECLARE_SCOPE_CYCLE_COUNTER(TEXT("FBulkData::TryLoadDataIntoMemory"), STAT_UBD_LoadDataDirectly, STATGROUP_Memory);
-
 	const int64 BulkDataOffset = GetBulkDataOffsetInFile();
 	const int64 BulkDataSize = GetBulkDataSize();
 	const int64 BulkDataSizeOnDisk = GetBulkDataSizeOnDisk();
diff --git a/Engine/Source/Runtime/CoreUObject/Private/UObject/UObjectArray.cpp b/Engine/Source/Runtime/CoreUObject/Private/UObject/UObjectArray.cpp
--- a/Engine/Source/Runtime/CoreUObject/Private/UObject/UObjectArray.cpp
+++ b/Engine/Source/Runtime/CoreUObject/Private/UObject/UObjectArray.cpp
@@ -10,6 +10,7 @@
 #include "HAL/LowLevelMemStats.h"
 #include "Misc/ScopeLock.h"
 #include "ProfilingDebugging/MetadataTrace.h"
+#include "Stats/Stats.h"
 #include "UObject/UObjectAllocator.h"
 #include "UObject/Class.h"
 #include "UObject/GarbageCollectionInternalFlags.h"
@@ -19,13 +20,15 @@
 
 DEFINE_LOG_CATEGORY_STATIC(LogUObjectArray, Log, All);
 
+DECLARE_CYCLE_STAT_WITH_FLAGS(TEXT("CreateStatId"), CreateStatId, STATGROUP_Object, EStatFlags::Verbose);
+
 FUObjectClusterContainer GUObjectClusters;
 
 #if STATS || ENABLE_STATNAMEDEVENTS_UOBJECT
 void FUObjectItem::CreateStatID() const
 {
 	LLM_SCOPE_BYNAME(TEXT("Debug/CreateStatID"));
-	QUICK_SCOPE_CYCLE_COUNTER(CreateStatId);
+	SCOPE_CYCLE_COUNTER(CreateStatId);
 
 	FString LongName;
 	LongName.Reserve(255);
diff --git a/Engine/Source/Runtime/CoreUObject/Public/UObject/UObjectGlobals.h b/Engine/Source/Runtime/CoreUObject/Public/UObject/UObjectGlobals.h
--- a/Engine/Source/Runtime/CoreUObject/Public/UObject/UObjectGlobals.h
+++ b/Engine/Source/Runtime/CoreUObject/Public/UObject/UObjectGlobals.h
@@ -83,7 +83,7 @@ DECLARE_CYCLE_STAT_EXTERN(TEXT("AllocateObject"),STAT_AllocateObject,STATGROUP_O
 DECLARE_CYCLE_STAT_EXTERN(TEXT("PostConstructInitializeProperties"),STAT_PostConstructInitializeProperties,STATGROUP_ObjectVerbose, );
 DECLARE_CYCLE_STAT_EXTERN(TEXT("LoadConfig"),STAT_LoadConfig,STATGROUP_Object, );
 DECLARE_CYCLE_STAT_EXTERN(TEXT("LoadObject"),STAT_LoadObject,STATGROUP_Object, );
-DECLARE_CYCLE_STAT_EXTERN(TEXT("InitProperties"),STAT_InitProperties,STATGROUP_Object, );
+DECLARE_CYCLE_STAT_WITH_FLAGS_EXTERN(TEXT("InitProperties"),STAT_InitProperties,STATGROUP_Object, EStatFlags::Verbose, );
 DECLARE_DWORD_ACCUMULATOR_STAT_EXTERN(TEXT("NameTable Entries"),STAT_NameTableEntries,STATGROUP_Object, );
 DECLARE_DWORD_ACCUMULATOR_STAT_EXTERN(TEXT("NameTable ANSI Entries"),STAT_NameTableAnsiEntries,STATGROUP_Object, );
 DECLARE_DWORD_ACCUMULATOR_STAT_EXTERN(TEXT("NameTable Wide Entries"),STAT_NameTableWideEntries,STATGROUP_Object, );
diff --git a/Engine/Source/Runtime/Engine/Private/Components/ActorComponent.cpp b/Engine/Source/Runtime/Engine/Private/Components/ActorComponent.cpp
--- a/Engine/Source/Runtime/Engine/Private/Components/ActorComponent.cpp
+++ b/Engine/Source/Runtime/Engine/Private/Components/ActorComponent.cpp
@@ -70,11 +70,11 @@ DECLARE_CYCLE_STAT(TEXT("UnregisterComponent"), STAT_UnregisterComponent, STATGR
 DECLARE_CYCLE_STAT(TEXT("Component OnRegister"), STAT_ComponentOnRegister, STATGROUP_Component);
 DECLARE_CYCLE_STAT(TEXT("Component OnUnregister"), STAT_ComponentOnUnregister, STATGROUP_Component);
 
-DECLARE_CYCLE_STAT(TEXT("Component CreateRenderState"), STAT_ComponentCreateRenderState, STATGROUP_Component);
-DECLARE_CYCLE_STAT(TEXT("Component DestroyRenderState"), STAT_ComponentDestroyRenderState, STATGROUP_Component);
+DECLARE_CYCLE_STAT_WITH_FLAGS(TEXT("Component CreateRenderState"), STAT_ComponentCreateRenderState, STATGROUP_Component, EStatFlags::Verbose);
+DECLARE_CYCLE_STAT_WITH_FLAGS(TEXT("Component DestroyRenderState"), STAT_ComponentDestroyRenderState, STATGROUP_Component, EStatFlags::Verbose);
 
-DECLARE_CYCLE_STAT(TEXT("Component CreatePhysicsState"), STAT_ComponentCreatePhysicsState, STATGROUP_Component);
-DECLARE_CYCLE_STAT(TEXT("Component DestroyPhysicsState"), STAT_ComponentDestroyPhysicsState, STATGROUP_Component);
+DECLARE_CYCLE_STAT_WITH_FLAGS(TEXT("Component CreatePhysicsState"), STAT_ComponentCreatePhysicsState, STATGROUP_Component, EStatFlags::Verbose);
+DECLARE_CYCLE_STAT_WITH_FLAGS(TEXT("Component DestroyPhysicsState"), STAT_ComponentDestroyPhysicsState, STATGROUP_Component, EStatFlags::Verbose);
 
 // Should we tick latent actions fired for a component at the same time as the component?
 // - Non-zero values behave the same way as actors do, ticking pending latent action when the component ticks, instead of later on in the frame
diff --git a/Engine/Source/Runtime/Engine/Public/EngineStats.h b/Engine/Source/Runtime/Engine/Public/EngineStats.h
--- a/Engine/Source/Runtime/Engine/Public/EngineStats.h
+++ b/Engine/Source/Runtime/Engine/Public/EngineStats.h
@@ -60,7 +60,7 @@ DECLARE_CYCLE_STAT_EXTERN(TEXT("PlayerController Tick"), STAT_PlayerControllerTi
 DECLARE_CYCLE_STAT_EXTERN(TEXT("Post Tick Component Update"),STAT_PostTickComponentUpdate,STATGROUP_Game, );
 DECLARE_CYCLE_STAT_EXTERN(TEXT("Post Tick Component Wait"),STAT_PostTickComponentUpdateWait,STATGROUP_Game, );
 DECLARE_CYCLE_STAT_EXTERN(TEXT("Recreate"),STAT_PostTickComponentRecreate,STATGROUP_Game, );
-DECLARE_CYCLE_STAT_EXTERN(TEXT("Transform or RenderData"),STAT_PostTickComponentLW,STATGROUP_Game, );
+DECLARE_CYCLE_STAT_WITH_FLAGS_EXTERN(TEXT("Transform or RenderData"),STAT_PostTickComponentLW,STATGROUP_Game,EStatFlags::Verbose, );
 DECLARE_CYCLE_STAT_EXTERN(TEXT("Update Particle Data"),STAT_ParticleManagerUpdateData,STATGROUP_Game, );
 DECLARE_CYCLE_STAT_EXTERN(TEXT("Async Work Wait"),STAT_AsyncWorkWaitTime,STATGROUP_Game, );
 DECLARE_CYCLE_STAT_EXTERN(TEXT("Visual Logging time"), STAT_VisualLog, STATGROUP_Game, ENGINE_API);
@@ -197,7 +197,7 @@ DECLARE_CYCLE_STAT_EXTERN(TEXT("Draw Texture Tile Time"),STAT_Canvas_DrawTexture
 DECLARE_CYCLE_STAT_EXTERN(TEXT("Draw Material Tile Time"),STAT_Canvas_DrawMaterialTileTime,STATGROUP_Canvas, );
 DECLARE_CYCLE_STAT_EXTERN(TEXT("Draw String Time"),STAT_Canvas_DrawStringTime,STATGROUP_Canvas, );
 DECLARE_CYCLE_STAT_EXTERN(TEXT("Word Wrapping Time"),STAT_Canvas_WordWrappingTime,STATGROUP_Canvas, );
-DECLARE_CYCLE_STAT_EXTERN(TEXT("Get Batched Element Time"),STAT_Canvas_GetBatchElementsTime,STATGROUP_Canvas, );
+DECLARE_CYCLE_STAT_WITH_FLAGS_EXTERN(TEXT("Get Batched Element Time"),STAT_Canvas_GetBatchElementsTime,STATGROUP_Canvas,EStatFlags::Verbose, );
 DECLARE_CYCLE_STAT_EXTERN(TEXT("Add Material Tile Time"),STAT_Canvas_AddTileRenderTime,STATGROUP_Canvas, );
 DECLARE_CYCLE_STAT_EXTERN(TEXT("Add Material Triangle Time"),STAT_Canvas_AddTriangleRenderTime,STATGROUP_Canvas, );
 DECLARE_DWORD_COUNTER_STAT_EXTERN(TEXT("Num Batches Created"),STAT_Canvas_NumBatchesCreated,STATGROUP_Canvas, );	
