Improves performance of AssetHeaderPatcher

- Adds EFlags to FCoreRedirectsContext to specify different switches the FCoreRedirects internals have for performing debug checks/logging, validating input, etc.. (e.g. Adding a new redirect calls HasValidCharacters which is very slow in bulk due to the excessive string copies from FNames and isn't necessary when we know the input is already well-formed).
- Building off the above, there is a flag to make the RedirectionsSummary optionally enabled in FCoreRedirectsContext (on by default), as the RedirectionSummary creates many allocations and is only used in narrow circumstances (incremental cooking)
- We no longer push a new FCoreRedirectsContext for each task thread used when AssetHeaderPatching and instead uses a single context owned by the patcher. This allows us to copy the global context once per patcher rather than once per file being patched which was done mistakenly and causes tons of allocator contention as you'd expect.
- Removes the debug name from the FCoreRedirectsContext and scope as it's ultimately not helpful and such information about where scopes come from are easily determined in callstacks or very light debugging

#rnx

diff --git a/Engine/Source/Developer/AssetTools/Internal/AssetHeaderPatcher.h b/Engine/Source/Developer/AssetTools/Internal/AssetHeaderPatcher.h
--- a/Engine/Source/Developer/AssetTools/Internal/AssetHeaderPatcher.h
+++ b/Engine/Source/Developer/AssetTools/Internal/AssetHeaderPatcher.h
@@ -5,6 +5,7 @@
 #include "Tasks/Task.h"
 #include "Containers/ContainersFwd.h"
 #include "UObject/CoreRedirects.h"
+#include "UObject/CoreRedirects/CoreRedirectsContext.h"
 
 /**
  * Delegate called on when patch operation completes
@@ -94,6 +95,7 @@ UE_INTERNAL struct ASSETTOOLS_API FAssetHeaderPatcher
 
 		// Todo: Make TSet once FCoreRedirect GetTypeHash is implemented
 		TArray<FCoreRedirect> Redirects;
+		mutable FCoreRedirectsContext RedirectsContext;
 
 		// String mappings are only used for best-effort replacements. These will be error-prone 
 		// and we should strive for more structured data formats to guard against errors here
diff --git a/Engine/Source/Developer/AssetTools/Private/AssetHeaderPatcher.cpp b/Engine/Source/Developer/AssetTools/Private/AssetHeaderPatcher.cpp
--- a/Engine/Source/Developer/AssetTools/Private/AssetHeaderPatcher.cpp
+++ b/Engine/Source/Developer/AssetTools/Private/AssetHeaderPatcher.cpp
@@ -888,11 +888,7 @@ public:
 
 FAssetHeaderPatcher::EResult FAssetHeaderPatcher::DoPatch(const FString& InSrcAsset, const FString& InDstAsset, const FContext& InContext)
 {
-	FScopeCoreRedirectsContext PatchContext(TEXT("Asset Header Patcher Context"), FCoreRedirectsContext::GetThreadContext());
-	FCoreRedirects::AddRedirectList(InContext.Redirects, TEXT("Asset Header Patcher"));
-
 	FAssetHeaderPatcherInner Inner(InSrcAsset, InDstAsset, InContext.StringReplacements, InContext.StringMountReplacements);
-
 	if (!FFileHelper::LoadFileToArray(Inner.SrcBuffer, *Inner.SrcAsset))
 	{
 		UE_LOG(LogAssetHeaderPatcher, Error, TEXT("Failed to load %s"), *Inner.SrcAsset);
@@ -900,6 +896,12 @@ FAssetHeaderPatcher::EResult FAssetHeaderPatcher::DoPatch(const FString& InSrcAs
 	}
 	else
 	{
+		// Swap in the CoreRedirect context for the patcher since we might be running on a different thread with a separate context
+		// We do not use a FScopeCoreRedirectsContext as that will copy into a new context but we want to re-use the patcher's context
+		FCoreRedirectsContext& OriginalContext = FCoreRedirectsContext::GetThreadContext();
+		FCoreRedirectsContext::SetThreadContext(InContext.RedirectsContext);
+		ON_SCOPE_EXIT{ FCoreRedirectsContext::SetThreadContext(OriginalContext); };
+
 		return Inner.PatchHeader();
 	}
 }
@@ -917,6 +919,28 @@ void FAssetHeaderPatcher::SetContext(FContext InContext)
 {
 	checkf(!IsPatching(), TEXT("Cannot set the patcher context while patching"));
 	Context = InContext;
+
+	// Copy the global context into our own to inherit any global redirects already loaded
+	Context.RedirectsContext = FCoreRedirectsContext(FCoreRedirectsContext::GetGlobalContext());
+	Context.RedirectsContext.InitializeContext();
+
+	// Disable validation of the CoreRedirects as these redirects are handmade and should not need extra validation.
+	// RedirectionSummary is unimportant for patching and will result in unnecessary allocations when we add our redirects.
+	// Leave DebugMode on since if someone wants to debug, we should not prevent them.
+	FCoreRedirectsContext::EFlags NewFlags = Context.RedirectsContext.GetFlags();
+	NewFlags = NewFlags & ~(FCoreRedirectsContext::EFlags::ValidateAddedRedirects | FCoreRedirectsContext::EFlags::UseRedirectionSummary);
+	Context.RedirectsContext.SetFlags(NewFlags);
+
+	{
+		// Swap the thread context to the patcher's FCoreRedirectsContext so we may populate it once and share it with the task threads
+		// We do not use a FScopeCoreRedirectsContext as that will copy into a new context but we want to re-use the patcher's context
+		FCoreRedirectsContext& OriginalContext = FCoreRedirectsContext::GetThreadContext();
+		FCoreRedirectsContext::SetThreadContext(Context.RedirectsContext);
+		ON_SCOPE_EXIT{ FCoreRedirectsContext::SetThreadContext(OriginalContext); };
+
+		FCoreRedirects::AddRedirectList(Context.Redirects, TEXT("Asset Header Patcher"));
+	}
+
 	Reset();
 }
 
diff --git a/Engine/Source/Runtime/CoreUObject/Internal/UObject/CoreRedirects/CoreRedirectsContext.h b/Engine/Source/Runtime/CoreUObject/Internal/UObject/CoreRedirects/CoreRedirectsContext.h
--- a/Engine/Source/Runtime/CoreUObject/Internal/UObject/CoreRedirects/CoreRedirectsContext.h
+++ b/Engine/Source/Runtime/CoreUObject/Internal/UObject/CoreRedirects/CoreRedirectsContext.h
@@ -3,6 +3,7 @@
 #pragma once
 #include "Containers/Map.h"
 #include "Containers/UnrealString.h"
+#include "Misc/EnumClassFlags.h"
 #include "Misc/TransactionallySafeRWLock.h"
 #include "RedirectionSummary.h"
 #include "UObject/CoreRedirects.h"
@@ -16,6 +17,16 @@
 */
 struct COREUOBJECT_API FCoreRedirectsContext
 {
+	enum class EFlags : uint32
+	{
+		None						= 0,
+		DebugMode					= 1 << 0,	// Enables extra validation and logging
+		ValidateAddedRedirects		= 1 << 1,	// Validates CoreRedirects are well-formed before addition
+		UseRedirectionSummary		= 1 << 2,   // New redirects will be appended to the RedirectionSummary
+
+		Default = ValidateAddedRedirects | UseRedirectionSummary
+	};
+
 	/*
 	* Returns the global context used by FCoreRedirects when no new context has been provided as an override.
 	*/
@@ -39,14 +50,13 @@ struct COREUOBJECT_API FCoreRedirectsContext
 	void InitializeContext();
 
 	/** Creates a context with a debug context name */
-	FCoreRedirectsContext(const FString& ContextName);
-
-	/** Creates a  context with a debug context name, and inherits the state from OtherContext.
-	The state from OtherContext is copied into the new context*/
-	FCoreRedirectsContext(const FString& ContextName, FCoreRedirectsContext& OtherContext);
+	FCoreRedirectsContext();
 
-	/** Describes the context's intended purpose */
-	const FString& GetName() const;
+	/**
+	* Creates a  context with a debug context name, and inherits the state from OtherContext.
+	* The state from OtherContext is copied into the new context
+	*/
+	FCoreRedirectsContext(const FCoreRedirectsContext& OtherContext);
 
 	/** Get whether this has been initialized */
 	const bool IsInitialized() const;
@@ -54,6 +64,12 @@ struct COREUOBJECT_API FCoreRedirectsContext
 	/** True if we are in debug mode that does extra validation */
 	const bool IsInDebugMode() const;
 
+	/** Gets the flags for the context */
+	EFlags GetFlags() const;
+
+	/** Sets the flags for the context */
+	void SetFlags(EFlags Flags);
+
 	/** True if we have done our initial validation. After initial validation, each change to redirects will validate independently */
 	const bool HasValidated() const;
 
@@ -67,6 +83,8 @@ struct COREUOBJECT_API FCoreRedirectsContext
 	FRedirectionSummary& GetRedirectionSummary();
 #endif
 
+	FCoreRedirectsContext& operator=(const FCoreRedirectsContext& Other);
+
 private:
 	friend void FCoreRedirects::ValidateAllRedirects();
 	/** Sets that we have done our initial validation. After initial validation, each change to redirects will validate independently */
@@ -100,9 +118,8 @@ private:
 	friend struct FScopeCoreRedirectsWriteLockedContext;
 
 private:
-	FString Name;
 	std::atomic<bool> bInitialized;
-	std::atomic<bool> bInDebugMode;
+	std::atomic<EFlags> Flags;
 	bool bValidatedOnce;
 	TMap<FName, ECoreRedirectFlags> ConfigKeyMap;
 	FCoreRedirects::FRedirectTypeMap RedirectTypeMap;
@@ -113,6 +130,7 @@ private:
 	FRedirectionSummary RedirectionSummary;
 #endif
 };
+ENUM_CLASS_FLAGS(FCoreRedirectsContext::EFlags);
 
 /**
 * RAII type for swapping the current thread's FCoreRedirects context to a new context. Can optionally
@@ -120,8 +138,8 @@ private:
 */
 struct COREUOBJECT_API FScopeCoreRedirectsContext
 {
-	FScopeCoreRedirectsContext(const FString& Name);
-	FScopeCoreRedirectsContext(const FString& Name, FCoreRedirectsContext& ContextToCopyFrom);
+	FScopeCoreRedirectsContext();
+	FScopeCoreRedirectsContext(FCoreRedirectsContext& ContextToCopyFrom);
 	~FScopeCoreRedirectsContext();
 
 	FCoreRedirectsContext ScopeContext;
@@ -189,11 +207,6 @@ public:
 // Inline implementations
 ///////////////////////////////////////////////////////
 
-inline const FString& FCoreRedirectsContext::GetName() const 
-{
-	return Name;
-}
-
 /** Get whether this has been initialized */
 inline const bool FCoreRedirectsContext::IsInitialized() const
 {
@@ -203,9 +216,21 @@ inline const bool FCoreRedirectsContext::IsInitialized() const
 
 /** True if we are in debug mode that does extra validation */
 inline const bool FCoreRedirectsContext::IsInDebugMode() const
+{
+	return Flags.load(std::memory_order_relaxed) == EFlags::DebugMode;
+}
+
+/** Gets the flags for the context */
+inline FCoreRedirectsContext::EFlags FCoreRedirectsContext::GetFlags() const
 {
 	// Use an atomic variable rather than take a read lock because we might already be under a read lock
-	return bInDebugMode.load(std::memory_order_relaxed);
+	return Flags.load(std::memory_order_relaxed);
+}
+
+/** Sets the flags for the context */
+inline void FCoreRedirectsContext::SetFlags(FCoreRedirectsContext::EFlags NewFlags)
+{
+	return Flags.store(NewFlags, std::memory_order_relaxed);
 }
 
 /** True if we have done our initial validation. After initial validation, each change to redirects will validate independently */
diff --git a/Engine/Source/Runtime/CoreUObject/Private/UObject/CoreRedirects.cpp b/Engine/Source/Runtime/CoreUObject/Private/UObject/CoreRedirects.cpp
--- a/Engine/Source/Runtime/CoreUObject/Private/UObject/CoreRedirects.cpp
+++ b/Engine/Source/Runtime/CoreUObject/Private/UObject/CoreRedirects.cpp
@@ -1277,6 +1277,27 @@ FCoreRedirects::FRedirectTypeMap::FRedirectTypeMap(const FRedirectTypeMap& Other
 	}
 }
 
+FCoreRedirects::FRedirectTypeMap& FCoreRedirects::FRedirectTypeMap::operator=(const FCoreRedirects::FRedirectTypeMap& Other)
+{
+	FastIterable.Empty(Other.FastIterable.Num());
+	FastIterable.Append(Other.FastIterable);
+
+	Map.Empty(FastIterable.Num());
+	for (TPair<ECoreRedirectFlags, FRedirectNameMap>& Pair : FastIterable)
+	{
+		FRedirectNameMap& RedirectNameMap = Pair.Value;
+		Map.Add(Pair.Key, &RedirectNameMap);
+
+		if (RedirectNameMap.Wildcards)
+		{
+			FWildcardData* WildCardCopy = new FWildcardData();
+			*WildCardCopy = *RedirectNameMap.Wildcards;
+			RedirectNameMap.Wildcards = MakeShareable(WildCardCopy);
+		}
+	}
+	return *this;
+}
+
 FCoreRedirects::FRedirectNameMap& FCoreRedirects::FRedirectTypeMap::FindOrAdd(ECoreRedirectFlags Key)
 {
 	FRedirectNameMap*& NameMap = Map.FindOrAdd(Key);
@@ -1661,26 +1682,42 @@ void FCoreRedirects::ClearKnownMissing(ECoreRedirectFlags Type, ECoreRedirectFla
 #if WITH_EDITOR
 void FCoreRedirects::AppendHashOfRedirectsAffectingPackages(FBlake3& Hasher, TConstArrayView<FName> PackageNames)
 {
-	FCoreRedirectsContext::GetThreadContext().GetRedirectionSummary().AppendHashAffectingPackages(Hasher, PackageNames);
+	FCoreRedirectsContext& Context = FCoreRedirectsContext::GetThreadContext();
+	if (!!(Context.GetFlags() & FCoreRedirectsContext::EFlags::UseRedirectionSummary))
+	{
+		Context.GetRedirectionSummary().AppendHashAffectingPackages(Hasher, PackageNames);
+	}
 }
 
 void FCoreRedirects::AppendHashOfGlobalRedirects(FBlake3& Hasher)
 {
-	FCoreRedirectsContext::GetThreadContext().GetRedirectionSummary().AppendHashGlobal(Hasher);
+	FCoreRedirectsContext& Context = FCoreRedirectsContext::GetThreadContext();
+	if (!!(Context.GetFlags() & FCoreRedirectsContext::EFlags::UseRedirectionSummary))
+	{
+		Context.GetRedirectionSummary().AppendHashGlobal(Hasher);
+	}
 }
 
 void FCoreRedirects::RecordAddedObjectRedirector(const FSoftObjectPath& Source, const FSoftObjectPath& Dest)
 {
-	FCoreRedirect ConvertedToCoreRedirect(ECoreRedirectFlags::Type_Object,
-		FCoreRedirectObjectName(Source), FCoreRedirectObjectName(Dest));
-	FCoreRedirectsContext::GetThreadContext().GetRedirectionSummary().Add(ConvertedToCoreRedirect, false /* bIsWildcardMatch */);
+	FCoreRedirectsContext& Context = FCoreRedirectsContext::GetThreadContext();
+	if (!!(Context.GetFlags() & FCoreRedirectsContext::EFlags::UseRedirectionSummary))
+	{
+		FCoreRedirect ConvertedToCoreRedirect(ECoreRedirectFlags::Type_Object,
+			FCoreRedirectObjectName(Source), FCoreRedirectObjectName(Dest));
+		Context.GetRedirectionSummary().Add(ConvertedToCoreRedirect, false /* bIsWildcardMatch */);
+	}
 }
 
 void FCoreRedirects::RecordRemovedObjectRedirector(const FSoftObjectPath& Source, const FSoftObjectPath& Dest)
 {
-	FCoreRedirect ConvertedToCoreRedirect(ECoreRedirectFlags::Type_Object,
-		FCoreRedirectObjectName(Source), FCoreRedirectObjectName(Dest));
-	FCoreRedirectsContext::GetThreadContext().GetRedirectionSummary().Remove(ConvertedToCoreRedirect, false /* bIsWildcardMatch */);
+	FCoreRedirectsContext& Context = FCoreRedirectsContext::GetThreadContext();
+	if (!!(Context.GetFlags() & FCoreRedirectsContext::EFlags::UseRedirectionSummary))
+	{
+		FCoreRedirect ConvertedToCoreRedirect(ECoreRedirectFlags::Type_Object,
+			FCoreRedirectObjectName(Source), FCoreRedirectObjectName(Dest));
+		Context.GetRedirectionSummary().Remove(ConvertedToCoreRedirect, false /* bIsWildcardMatch */);
+	}
 }
 
 #endif
@@ -2155,7 +2192,8 @@ bool FCoreRedirects::AddRedirectListUnderWriteLock(TArrayView<const FCoreRedirec
 		// Because the validation in HasValidCharacters does not perfectly align with other systems,
 		// there can exist assets which are valid, are referenced by UObjectRedirector assets, but which
 		// fail the check.
-		if (!EnumHasAnyFlags(NewRedirect.RedirectFlags, ECoreRedirectFlags::Type_Asset))
+		if (!!(LockedContext.Get().GetFlags() & FCoreRedirectsContext::EFlags::ValidateAddedRedirects)
+			&& !EnumHasAnyFlags(NewRedirect.RedirectFlags, ECoreRedirectFlags::Type_Asset))
 		{
 			if ((!NewRedirect.OldName.HasValidCharacters(NewRedirect.RedirectFlags) && !FPackageName::IsVersePackage(NewRedirect.OldName.PackageName.ToString()))
 				|| (!NewRedirect.NewName.HasValidCharacters(NewRedirect.RedirectFlags) && !FPackageName::IsVersePackage(NewRedirect.NewName.PackageName.ToString())))
@@ -2167,7 +2205,7 @@ bool FCoreRedirects::AddRedirectListUnderWriteLock(TArrayView<const FCoreRedirec
 
 		if (NewRedirect.IsWildcardMatch())
 		{
-			UE_LOG(LogCoreRedirects, Log, TEXT("AddRedirect(%s) has wildcard redirect %s, these are very slow and should be resolved as soon as possible! Please refer to the documentation in Engine/Config/BaseEngine.ini."), *SourceString, *NewRedirect.OldName.ToString());
+			UE_LOG(LogCoreRedirects, Verbose, TEXT("AddRedirect(%s) has wildcard redirect %s, these are very slow and should be resolved as soon as possible! Please refer to the documentation in Engine/Config/BaseEngine.ini."), *SourceString, *NewRedirect.OldName.ToString());
 		}
 		
 		if (AddSingleRedirectUnderWriteLock(NewRedirect, SourceString, LockedContext))
@@ -2281,7 +2319,10 @@ bool FCoreRedirects::AddSingleRedirectUnderWriteLock(const FCoreRedirect& NewRed
 	}
 
 #if WITH_EDITOR
-	LockedContext.Get().GetRedirectionSummary().Add(NewRedirect, bIsWildcardMatch);
+	if (!!(LockedContext.Get().GetFlags() & FCoreRedirectsContext::EFlags::UseRedirectionSummary))
+	{
+		LockedContext.Get().GetRedirectionSummary().Add(NewRedirect, bIsWildcardMatch);
+	}
 #endif
 
 	return true;
@@ -2321,7 +2362,7 @@ bool FCoreRedirects::RemoveRedirectList(TArrayView<const FCoreRedirect> Redirect
 
 		if (RedirectToRemove.IsWildcardMatch())
 		{
-			UE_LOG(LogCoreRedirects, Log, TEXT("RemoveRedirect(%s) has wildcard redirect %s, these are very slow and should be resolved as soon as possible! Please refer to the documentation in Engine/Config/BaseEngine.ini."), *SourceString, *RedirectToRemove.OldName.ToString());
+			UE_LOG(LogCoreRedirects, Verbose, TEXT("RemoveRedirect(%s) has wildcard redirect %s, these are very slow and should be resolved as soon as possible! Please refer to the documentation in Engine/Config/BaseEngine.ini."), *SourceString, *RedirectToRemove.OldName.ToString());
 		}
 
 		bRemovedAny |= RemoveSingleRedirectUnderWriteLock(RedirectToRemove, SourceString, LockedContext);
@@ -2406,7 +2447,10 @@ bool FCoreRedirects::RemoveSingleRedirectUnderWriteLock(const FCoreRedirect& Red
 		}
 
 #if WITH_EDITOR
-		LockedContext.Get().GetRedirectionSummary().Remove(RedirectToRemove, bIsWildcardMatch);
+		if (!!(LockedContext.Get().GetFlags() & FCoreRedirectsContext::EFlags::UseRedirectionSummary))
+		{
+			LockedContext.Get().GetRedirectionSummary().Remove(RedirectToRemove, bIsWildcardMatch);
+		}
 #endif
 	}
 
diff --git a/Engine/Source/Runtime/CoreUObject/Private/UObject/CoreRedirectsContext.cpp b/Engine/Source/Runtime/CoreUObject/Private/UObject/CoreRedirectsContext.cpp
--- a/Engine/Source/Runtime/CoreUObject/Private/UObject/CoreRedirectsContext.cpp
+++ b/Engine/Source/Runtime/CoreUObject/Private/UObject/CoreRedirectsContext.cpp
@@ -8,28 +8,32 @@
 
 namespace UE::CoreRedirects::Private
 {
-	FCoreRedirectsContext GlobalCoreRedirectContext(TEXT("Global CoreRedirect Context"));
+	FCoreRedirectsContext GlobalCoreRedirectContext;
 	thread_local FCoreRedirectsContext* ThreadCoreRedirectContext = nullptr;
 }
 
-FCoreRedirectsContext::FCoreRedirectsContext(const FString& ContextName)
-	: Name(ContextName)
-	, bInitialized(false)
-	, bInDebugMode(false)
+FCoreRedirectsContext::FCoreRedirectsContext()
+	: bInitialized(false)
+	, Flags(EFlags::Default)
 	, bValidatedOnce(false)
 {
 }
 
-FCoreRedirectsContext::FCoreRedirectsContext(const FString& ContextName, FCoreRedirectsContext& OtherContext)
-	: Name(ContextName)
-	, bInitialized(OtherContext.bInitialized.load())
-	, bInDebugMode(OtherContext.bInDebugMode.load())
+FCoreRedirectsContext::FCoreRedirectsContext(const FCoreRedirectsContext& OtherContext)
+	: bInitialized(OtherContext.bInitialized.load())
+	, Flags(OtherContext.Flags.load())
 	, bValidatedOnce(OtherContext.bValidatedOnce)
 	, ConfigKeyMap(OtherContext.ConfigKeyMap)
 	, RedirectTypeMap(OtherContext.RedirectTypeMap)
 {
 }
 
+FCoreRedirectsContext& FCoreRedirectsContext::operator=(const FCoreRedirectsContext& Other)
+{
+	new(this) FCoreRedirectsContext(Other);
+	return *this;
+}
+
 void FCoreRedirectsContext::InitializeContext()
 {
 	if (bInitialized.load(std::memory_order_relaxed))
@@ -44,14 +48,14 @@ void FCoreRedirectsContext::InitializeContext()
 		if (FParse::Param(FCommandLine::Get(), TEXT("FullDebugCoreRedirects")))
 		{
 			// Enable debug mode and set to maximum verbosity
-			bInDebugMode.store(true, std::memory_order_relaxed);
+			Flags.store(Flags.load(std::memory_order_relaxed) | EFlags::DebugMode, std::memory_order_relaxed);
 			UE_SET_LOG_VERBOSITY(LogCoreRedirects, VeryVerbose);
 			FCoreDelegates::OnFEngineLoopInitComplete.AddStatic(FCoreRedirects::ValidateAllRedirects);
 		}
 		else if (FParse::Param(FCommandLine::Get(), TEXT("DebugCoreRedirects")))
 		{
 			// Enable debug mode and increase log levels but don't show every message
-			bInDebugMode.store(true, std::memory_order_relaxed);
+			Flags.store(Flags.load(std::memory_order_relaxed) | EFlags::DebugMode, std::memory_order_relaxed);
 			UE_SET_LOG_VERBOSITY(LogCoreRedirects, Verbose);
 			FCoreDelegates::OnFEngineLoopInitComplete.AddStatic(FCoreRedirects::ValidateAllRedirects);
 		}
@@ -98,16 +102,15 @@ void FCoreRedirectsContext::SetThreadContext(FCoreRedirectsContext& NewContext)
 	ThreadCoreRedirectContext = &NewContext;
 }
 
-FScopeCoreRedirectsContext::FScopeCoreRedirectsContext(const FString& Name)
-	: ScopeContext(Name)
-	, ContextToRestore(FCoreRedirectsContext::GetThreadContext())
+FScopeCoreRedirectsContext::FScopeCoreRedirectsContext()
+	: ContextToRestore(FCoreRedirectsContext::GetThreadContext())
 {
 	FCoreRedirectsContext::SetThreadContext(ScopeContext);
 	ScopeContext.InitializeContext();
 }
 
-FScopeCoreRedirectsContext::FScopeCoreRedirectsContext(const FString& Name, FCoreRedirectsContext& ContextToCopyFrom)
-	: ScopeContext(Name, ContextToCopyFrom)
+FScopeCoreRedirectsContext::FScopeCoreRedirectsContext(FCoreRedirectsContext& ContextToCopyFrom)
+	: ScopeContext(ContextToCopyFrom)
 	, ContextToRestore(FCoreRedirectsContext::GetThreadContext())
 {
 	FCoreRedirectsContext::SetThreadContext(ScopeContext);
diff --git a/Engine/Source/Runtime/CoreUObject/Public/UObject/CoreRedirects.h b/Engine/Source/Runtime/CoreUObject/Public/UObject/CoreRedirects.h
--- a/Engine/Source/Runtime/CoreUObject/Public/UObject/CoreRedirects.h
+++ b/Engine/Source/Runtime/CoreUObject/Public/UObject/CoreRedirects.h
@@ -476,6 +476,7 @@ private:
 	public:
 		FRedirectTypeMap() = default;
 		FRedirectTypeMap(const FRedirectTypeMap& Other);
+		FRedirectTypeMap& operator=(const FRedirectTypeMap& Other);
 		FRedirectNameMap& FindOrAdd(ECoreRedirectFlags Key);
 		FRedirectNameMap* Find(ECoreRedirectFlags Key);
 		void Empty();
@@ -484,7 +485,6 @@ private:
 		TArray<TPair<ECoreRedirectFlags, FRedirectNameMap>>::RangedForIteratorType end() { return FastIterable.end(); }
 
 		FRedirectTypeMap(const FRedirectTypeMap&& Other) = delete;
-		FRedirectTypeMap& operator=(const FRedirectTypeMap& Other) = delete;
 		FRedirectTypeMap& operator=(const FRedirectTypeMap&& Other) = delete;
 	private:
 		TMap<ECoreRedirectFlags, FRedirectNameMap*> Map;
