[StructUtils] InstanceStruct should not do delta serialization in ExportText

* The memory is reset in ImportText

#rnx

diff --git a/Engine/Source/Runtime/CoreUObject/Private/StructUtils/InstancedStruct.cpp b/Engine/Source/Runtime/CoreUObject/Private/StructUtils/InstancedStruct.cpp
--- a/Engine/Source/Runtime/CoreUObject/Private/StructUtils/InstancedStruct.cpp
+++ b/Engine/Source/Runtime/CoreUObject/Private/StructUtils/InstancedStruct.cpp
@@ -240,7 +240,8 @@ bool FInstancedStruct::ExportTextItem(FString& ValueStr, FInstancedStruct const&
 	if (const UScriptStruct* StructTypePtr = GetScriptStruct())
 	{
 		ValueStr += StructTypePtr->GetPathName();
-		StructTypePtr->ExportText(ValueStr, GetMemory(), StructTypePtr == DefaultValue.GetScriptStruct() ? DefaultValue.GetMemory() : nullptr, Parent, PortFlags, ExportRootScope);
+		// Force the default to nullptr to disable delta serialization because we reset the memory in import text.
+		StructTypePtr->ExportText(ValueStr, GetMemory(), nullptr, Parent, PortFlags, ExportRootScope);
 	}
 	else
 	{
