MakeUniqueObjectName: update per-class suffix number annotation to support custom names.
This avoids incrementing the per-class suffix for calls to MakeUniqueObjectName using a base name other than the class name.

diff --git a/Engine/Source/Runtime/CoreUObject/Private/UObject/UObjectGlobals.cpp b/Engine/Source/Runtime/CoreUObject/Private/UObject/UObjectGlobals.cpp
--- a/Engine/Source/Runtime/CoreUObject/Private/UObject/UObjectGlobals.cpp
+++ b/Engine/Source/Runtime/CoreUObject/Private/UObject/UObjectGlobals.cpp
@@ -422,7 +422,7 @@ struct FPerClassNumberSuffixAnnotation
 		return false;
 	}
 
-	TMap<const UClass*, int32> Suffixes;
+	TMap<FName, int32> Suffixes;
 };
 
 /**
@@ -430,18 +430,23 @@ struct FPerClassNumberSuffixAnnotation
  *
  * Updating is done via a callback because a lock needs to be maintained while this happens.
  */
-int32 UpdateSuffixForNextNewObject(UObject* Parent, const UClass* Class, TFunctionRef<void(int32&)> IndexMutator)
+int32 UpdateSuffixForNextNewObject(UObject* Parent, FName BaseName, TFunctionRef<void(int32&)> IndexMutator)
 {
 	static FCriticalSection PerClassNumberSuffixAnnotationMutex;
 	static FUObjectAnnotationDense<FPerClassNumberSuffixAnnotation, true> PerClassNumberSuffixAnnotation;
 
 	FPerClassNumberSuffixAnnotation& Annotation = PerClassNumberSuffixAnnotation.GetAnnotationRef(Parent);
 	FScopeLock Lock(&PerClassNumberSuffixAnnotationMutex);
-	int32& Result = Annotation.Suffixes.FindOrAdd(Class);
+	int32& Result = Annotation.Suffixes.FindOrAdd(BaseName);
 	IndexMutator(Result);
 	return Result;
 }
 
+int32 UpdateSuffixForNextNewObject(UObject* Parent, const UClass* Class, TFunctionRef<void(int32&)> IndexMutator)
+{
+	return UpdateSuffixForNextNewObject(Parent, *Class->GetPathName(), IndexMutator);
+}
+
 //
 // Find an object, path must unqualified
 //
@@ -2822,6 +2827,7 @@ FName MakeUniqueObjectName(UObject* Parent, const UClass* Class, FName InBaseNam
 	CSV_SCOPED_TIMING_STAT(UObject, MakeUniqueObjectName);
 	check(Class);
 	const FName BaseName = InBaseName.IsNone() ? Class->GetFName() : InBaseName;
+	const FName ClassPathName = InBaseName.IsNone() ? *Class->GetPathName() : InBaseName;
 
 	if (!!(Options & EUniqueObjectNameOptions::UniversallyUnique))
 	{
@@ -2887,7 +2893,7 @@ FName MakeUniqueObjectName(UObject* Parent, const UClass* Class, FName InBaseNam
 						}
 						else
 						{
-							NameNumber = UpdateSuffixForNextNewObject(Parent, Class, [](int32& Index) { ++Index; });
+							NameNumber = UpdateSuffixForNextNewObject(Parent, ClassPathName, [](int32& Index) { ++Index; });
 						}
 					}
 					else
