Compute redirect flags from redirected class names

This allows uncooked assets that refer to Verse classes via an older name redirected to /Script/CoreUObject.VerseClass to still pick up Type_Class redirects.

diff --git a/Engine/Source/Runtime/CoreUObject/Private/UObject/LinkerLoad.cpp b/Engine/Source/Runtime/CoreUObject/Private/UObject/LinkerLoad.cpp
--- a/Engine/Source/Runtime/CoreUObject/Private/UObject/LinkerLoad.cpp
+++ b/Engine/Source/Runtime/CoreUObject/Private/UObject/LinkerLoad.cpp
@@ -2146,7 +2146,7 @@ FLinkerLoad::ELinkerStatus FLinkerLoad::FixupImportMap()
 				}
 
 				FCoreRedirectObjectName OldObjectName(GetImportPathName(i)), NewObjectName;
-				ECoreRedirectFlags ObjectRedirectFlags = FCoreRedirects::GetFlagsForTypeName(Import->ClassPackage, Import->ClassName);
+				ECoreRedirectFlags ObjectRedirectFlags = FCoreRedirects::GetFlagsForTypeName(NewClassName.PackageName, NewClassName.ObjectName);
 				const FCoreRedirect* ValueRedirect = nullptr;
 					
 				FCoreRedirects::RedirectNameAndValues(ObjectRedirectFlags, OldObjectName, NewObjectName, &ValueRedirect);
