Fix for issue where a Verse object would not be created with the correct name during re-compilation.
This could happen for example during re-compilation when the type of a Verse member would change (and then the object saved on disk). The end result was that we would fail to find the correct CDO for that object during deserialization from disk.

#rnx

diff --git a/Engine/Source/Runtime/CoreUObject/Private/UObject/PropertyBaseObject.cpp b/Engine/Source/Runtime/CoreUObject/Private/UObject/PropertyBaseObject.cpp
--- a/Engine/Source/Runtime/CoreUObject/Private/UObject/PropertyBaseObject.cpp
+++ b/Engine/Source/Runtime/CoreUObject/Private/UObject/PropertyBaseObject.cpp
@@ -759,7 +759,7 @@ static bool CanPropertyHoldObject(const FObjectPropertyBase* Prop, TObjectPtr<co
 	return false;
 }
 
-UObject* FObjectPropertyBase::ConstructDefaultObjectValueIfNecessary(UObject* ExistingValue) const
+UObject* FObjectPropertyBase::ConstructDefaultObjectValueIfNecessary(UObject* ExistingValue, const void* Defaults) const
 {
 	UObject* NewDefaultObjectValue = nullptr;
 
@@ -780,7 +780,7 @@ UObject* FObjectPropertyBase::ConstructDefaultObjectValueIfNecessary(UObject* Ex
 				ExistingValue->GetOuter() != Outer) // Unless the template's Outer is the same as the new Outer in which case the template (ExistingValue) IS the object we can reuse
 			{
 				// We probably got here because an object value failed to load (missing import class) and the property is left with a template of default subobject
-				NewDefaultObjectValue = NewObject<UObject>(Outer, ExistingValue->GetClass(), NAME_None, RF_NoFlags, ExistingValue);
+				NewDefaultObjectValue = NewObject<UObject>(Outer, ExistingValue->GetClass(), ExistingValue->GetFName(), ExistingValue->GetFlags(), ExistingValue);
 			}
 			else
 			{
@@ -798,16 +798,20 @@ UObject* FObjectPropertyBase::ConstructDefaultObjectValueIfNecessary(UObject* Ex
 #endif // WITH_METADATA
 	)
 	{
+		FName NewName = ExistingValue ? ExistingValue->GetFName() : NAME_None;
+		EObjectFlags NewFlags = ExistingValue ? ExistingValue->GetFlags() : RF_NoFlags;
+		UObject* Template = Defaults ? GetObjectPropertyValue(Defaults) : nullptr;
+
 		// Existing value did not exist or it could not be used as a template
 		// Existing value may be null in case we were serializing an array of UObjects that failed to load (missing import class). Since the array is first pre-allocated with null values
 		// it will not have any existing objects to instantiate
-		NewDefaultObjectValue = NewObject<UObject>(Outer, PropertyClass);
+		NewDefaultObjectValue = NewObject<UObject>(Outer, PropertyClass, NewName, NewFlags, Template);
 	}
 
 	return NewDefaultObjectValue;
 }
 
-void FObjectPropertyBase::CheckValidObject(void* ValueAddress, TObjectPtr<UObject> OldValue) const
+void FObjectPropertyBase::CheckValidObject(void* ValueAddress, TObjectPtr<UObject> OldValue, const void* Defaults) const
 {
 	const TObjectPtr<UObject> Object = GetObjectPtrPropertyValue(ValueAddress);
 	if (!Object)
@@ -885,7 +889,7 @@ void FObjectPropertyBase::CheckValidObject(void* ValueAddress, TObjectPtr<UObjec
 				ENonNullableBehavior NonNullableBehavior = GetNonNullableBehavior();
 				if (NonNullableBehavior == ENonNullableBehavior::CreateDefaultObjectIfPossible)
 				{
-					DefaultValue = ConstructDefaultObjectValueIfNecessary(OldValue);
+					DefaultValue = ConstructDefaultObjectValueIfNecessary(OldValue, Defaults);
 				}
 
 				if (DefaultValue)
diff --git a/Engine/Source/Runtime/CoreUObject/Private/UObject/PropertyObject.cpp b/Engine/Source/Runtime/CoreUObject/Private/UObject/PropertyObject.cpp
--- a/Engine/Source/Runtime/CoreUObject/Private/UObject/PropertyObject.cpp
+++ b/Engine/Source/Runtime/CoreUObject/Private/UObject/PropertyObject.cpp
@@ -86,7 +86,7 @@ EConvertFromTypeResult FObjectProperty::ConvertFromType(const FPropertyTag& Tag,
 		SetPropertyValue_InContainer(Data, PreviousValueObj, Tag.ArrayIndex);
 
 		// Validate the type is proper
-		CheckValidObject(GetPropertyValuePtr_InContainer(Data, Tag.ArrayIndex), PreviousValueObj);
+		CheckValidObject(GetPropertyValuePtr_InContainer(Data, Tag.ArrayIndex), PreviousValueObj, Defaults);
 
 		return EConvertFromTypeResult::Converted;
 	}
@@ -102,7 +102,7 @@ EConvertFromTypeResult FObjectProperty::ConvertFromType(const FPropertyTag& Tag,
 		}
 
 		SetPropertyValue_InContainer(Data, ObjectValue, Tag.ArrayIndex);
-		CheckValidObject(GetPropertyValuePtr_InContainer(Data, Tag.ArrayIndex), ObjectValue);
+		CheckValidObject(GetPropertyValuePtr_InContainer(Data, Tag.ArrayIndex), ObjectValue, Defaults);
 		return EConvertFromTypeResult::Converted;
 	}
 
@@ -177,7 +177,7 @@ void FObjectProperty::SerializeItem(FStructuredArchive::FSlot Slot, void* Value,
 #if !(UE_BUILD_TEST || UE_BUILD_SHIPPING) 
 		if (!UnderlyingArchive.IsSaving() && !UnderlyingArchive.IsMigratingRemoteObjects())
 		{
-			CheckValidObject(Value, CurrentValue);
+			CheckValidObject(Value, CurrentValue, Defaults);
 		}
 #endif
 	}
@@ -221,7 +221,7 @@ void FObjectProperty::SerializeItem(FStructuredArchive::FSlot Slot, void* Value,
 				if (!UnderlyingArchive.IsSaving())
 				{
 					TObjectPtr<UObject> CurrentValue(*ObjectPtr);
-					CheckValidObject(Value, CurrentValue);
+					CheckValidObject(Value, CurrentValue, Defaults);
 				}
 			}
 		}
@@ -237,12 +237,12 @@ void FObjectProperty::SerializeItem(FStructuredArchive::FSlot Slot, void* Value,
 			TObjectPtr<UObject> CurrentValuePtr = GetObjectPtrPropertyValue(Value);
 			check(CurrentValuePtr.IsResolved());
 			UObject* CurrentValue = UE::CoreUObject::Private::ReadObjectHandlePointerNoCheck(CurrentValuePtr.GetHandle());
-			PostSerializeObjectItem(UnderlyingArchive, Value, CurrentValue, ObjectValue, EObjectPropertyOptions::None);
+			PostSerializeObjectItem(UnderlyingArchive, Value, CurrentValue, ObjectValue, EObjectPropertyOptions::None, Defaults);
 		}
 	}
 }
 
-void FObjectProperty::PostSerializeObjectItem(FArchive& SerializingArchive, void* Value, UObject* CurrentValue, UObject* ObjectValue, EObjectPropertyOptions Options /*= EObjectPropertyOptions::None*/) const
+void FObjectProperty::PostSerializeObjectItem(FArchive& SerializingArchive, void* Value, UObject* CurrentValue, UObject* ObjectValue, EObjectPropertyOptions Options /*= EObjectPropertyOptions::None*/, const void* Defaults /*= nullptr*/) const
 {
 	// Make sure non-nullable properties don't end up with null values.
 	if (!(Options & EObjectPropertyOptions::AllowNullValuesOnNonNullableProperty) &&
@@ -260,7 +260,7 @@ void FObjectProperty::PostSerializeObjectItem(FArchive& SerializingArchive, void
 		ENonNullableBehavior NonNullableBehavior = GetNonNullableBehavior();
 		if (NonNullableBehavior == ENonNullableBehavior::CreateDefaultObjectIfPossible)
 		{
-			DefaultValue = ConstructDefaultObjectValueIfNecessary(CurrentValue);
+			DefaultValue = ConstructDefaultObjectValueIfNecessary(CurrentValue, Defaults);
 		}
 
 		if (DefaultValue)
@@ -321,7 +321,7 @@ void FObjectProperty::PostSerializeObjectItem(FArchive& SerializingArchive, void
 
 		if (!SerializingArchive.IsSaving())
 		{
-			CheckValidObject(Value, ObjectValue);
+			CheckValidObject(Value, ObjectValue, Defaults);
 		}
 	}
 }
diff --git a/Engine/Source/Runtime/CoreUObject/Private/UObject/PropertySoftObjectPtr.cpp b/Engine/Source/Runtime/CoreUObject/Private/UObject/PropertySoftObjectPtr.cpp
--- a/Engine/Source/Runtime/CoreUObject/Private/UObject/PropertySoftObjectPtr.cpp
+++ b/Engine/Source/Runtime/CoreUObject/Private/UObject/PropertySoftObjectPtr.cpp
@@ -98,7 +98,7 @@ void FSoftObjectProperty::SerializeItem( FStructuredArchive::FSlot Slot, void* V
 		{
 			if (OldValue.GetUniqueID() != ((FSoftObjectPtr*)Value)->GetUniqueID())
 			{
-				CheckValidObject(Value, nullptr); // FSoftObjectProperty is never non-nullable at this point so it's ok to pass null as the current value
+				CheckValidObject(Value, nullptr, Defaults); // FSoftObjectProperty is never non-nullable at this point so it's ok to pass null as the current value
 			}
 		}
 #endif
diff --git a/Engine/Source/Runtime/CoreUObject/Public/UObject/UnrealType.h b/Engine/Source/Runtime/CoreUObject/Public/UObject/UnrealType.h
--- a/Engine/Source/Runtime/CoreUObject/Public/UObject/UnrealType.h
+++ b/Engine/Source/Runtime/CoreUObject/Public/UObject/UnrealType.h
@@ -2647,7 +2647,7 @@ public:
 	FORCEINLINE void SetPropertyClass(UClass* NewPropertyClass) { PropertyClass = NewPropertyClass; }
 #endif // USE_CIRCULAR_DEPENDENCY_LOAD_DEFERRING
 
-	virtual void CheckValidObject(void* ValueAddress, TObjectPtr<UObject> OldValue) const;
+	virtual void CheckValidObject(void* ValueAddress, TObjectPtr<UObject> OldValue, const void* Defaults = nullptr) const;
 	virtual bool AllowObjectTypeReinterpretationTo(const FObjectPropertyBase* Other) const;
 
 protected:
@@ -2658,9 +2658,10 @@ protected:
 	 * Constructs a new object if the existing one is missing or is not compatible with the property class
 	 * Used for making sure non-nullable properties have valid values.
 	 * @param ExistingValue Previous object value (can be null)
+	 * @param Defaults Defaults object to use when constructing the object value (can be null).
 	 * @return Non-null object that was assigned to this property value address
 	 */
-	UObject* ConstructDefaultObjectValueIfNecessary(UObject* ExistingValue) const;
+	UObject* ConstructDefaultObjectValueIfNecessary(UObject* ExistingValue, const void* Defaults = nullptr) const;
 
 	// Disable false positive buffer overrun warning during pgoprofile linking step
 	PRAGMA_DISABLE_BUFFER_OVERRUN_WARNING
@@ -2919,8 +2920,9 @@ public:
 	 * @param Value Property address
 	 * @param CurrentValue Current Object value
 	 * @param ObjectValue Deserialized Object value
+	 * @param Defaults Defaults used during serialization
 	 */
-	void PostSerializeObjectItem(FArchive& SerializingArchive, void* Value, UObject* CurrentValue, UObject* ObjectValue, EObjectPropertyOptions Options = EObjectPropertyOptions::None) const;
+	void PostSerializeObjectItem(FArchive& SerializingArchive, void* Value, UObject* CurrentValue, UObject* ObjectValue, EObjectPropertyOptions Options = EObjectPropertyOptions::None, const void* Defaults = nullptr) const;
 
 	inline TObjectPtr<UObject>* GetObjectPtrPropertyValuePtr(const void* PropertyValueAddress) const
 	{
