Re-enable IDOs by default. Placeholders remain disabled for now.

diff --git a/Engine/Source/Runtime/CoreUObject/Private/UObject/InstanceDataObjectUtils.cpp b/Engine/Source/Runtime/CoreUObject/Private/UObject/InstanceDataObjectUtils.cpp
--- a/Engine/Source/Runtime/CoreUObject/Private/UObject/InstanceDataObjectUtils.cpp
+++ b/Engine/Source/Runtime/CoreUObject/Private/UObject/InstanceDataObjectUtils.cpp
@@ -249,7 +249,7 @@ namespace UE
 	static const FName NAME_IDOMapKey(ANSITEXTVIEW("Key"));
 	static const FName NAME_IDOMapValue(ANSITEXTVIEW("Value"));
 
-	bool bEnableIDOSupport = false;
+	bool bEnableIDOSupport = true;
 	FAutoConsoleVariableRef EnableIDOSupportCVar(
 		TEXT("IDO.Enable"),
 		bEnableIDOSupport,
diff --git a/Engine/Source/Runtime/CoreUObject/Private/UObject/PropertyBagRepository.cpp b/Engine/Source/Runtime/CoreUObject/Private/UObject/PropertyBagRepository.cpp
--- a/Engine/Source/Runtime/CoreUObject/Private/UObject/PropertyBagRepository.cpp
+++ b/Engine/Source/Runtime/CoreUObject/Private/UObject/PropertyBagRepository.cpp
@@ -951,7 +951,7 @@ bool FPropertyBagRepository::IsPropertyBagPlaceholderObjectSupportEnabled()
 {
 #if WITH_EDITOR && UE_WITH_OBJECT_HANDLE_TYPE_SAFETY
 	static bool bIsInitialized = false;
-	static bool bForceDisabled = false;
+	static bool bForceDisabled = true;
 	if (!bIsInitialized)
 	{
 		Private::bEnablePropertyBagPlaceholderObjectSupport |= FParse::Param(FCommandLine::Get(), TEXT("WithPropertyBagPlaceholderObjects"));
