[AutoRTFM] Clear FUObjectItem::[Flags|RefCount] in the open

This must be cleared in the open as all other writes are explicitly done in the open.
Mixing open and closed write can cause memory corruption (SOL-6743).
Caught by the memory validator when scaling back the no-validation flags on opens.

diff --git a/Engine/Source/Runtime/CoreUObject/Private/UObject/UObjectArray.cpp b/Engine/Source/Runtime/CoreUObject/Private/UObject/UObjectArray.cpp
--- a/Engine/Source/Runtime/CoreUObject/Private/UObject/UObjectArray.cpp
+++ b/Engine/Source/Runtime/CoreUObject/Private/UObject/UObjectArray.cpp
@@ -383,9 +383,16 @@ void FUObjectArray::FreeUObjectIndex(UObjectBase* Object)
 		ObjectItem->ThisThreadAtomicallyClearedFlag(EInternalObjectFlags_RootFlags);
 	}
 
+	// Due to atomic operations, these fields are only modified in the open.
+	// Mixing open and closed writes to the same memory location can cause memory corruption (SOL-6743)
+	// so, reset these fields in the open.
+	UE_AUTORTFM_OPEN
+	{
+		ObjectItem->Flags = 0;
+		ObjectItem->RefCount = 0;
+	};
+
 	ObjectItem->Object = nullptr;
-	ObjectItem->Flags = 0;
-	ObjectItem->RefCount = 0;
 	ObjectItem->ClusterRootIndex = 0;
 	ObjectItem->SerialNumber = 0;
 #if UE_WITH_REMOTE_OBJECT_HANDLE
