Folded const handling into the return types of Cast functions.

diff --git a/Engine/Source/Runtime/Core/Public/Delegates/DelegateSignatureImpl.inl b/Engine/Source/Runtime/Core/Public/Delegates/DelegateSignatureImpl.inl
--- a/Engine/Source/Runtime/Core/Public/Delegates/DelegateSignatureImpl.inl
+++ b/Engine/Source/Runtime/Core/Public/Delegates/DelegateSignatureImpl.inl
@@ -15,6 +15,7 @@
 #include "UObject/ScriptDelegates.h"
 #include "UObject/WeakObjectPtrTemplates.h"
 #include "Templates/IsConst.h"
+#include "Templates/CopyQualifiersFromTo.h"
 #include "Templates/RemoveReference.h"
 #include "Templates/SharedPointer.h"
 #include "Templates/UnrealTemplate.h"
@@ -30,7 +31,7 @@ template <typename T>
 T* ToRawPtr(const TObjectPtr<T>& Ptr);
 
 template <typename To, typename From>
-To* Cast(From* Src);
+TCopyQualifiersFromTo_T<From, To>* Cast(From* Src);
 
 template<typename UserPolicy> class TMulticastDelegateBase;
 
diff --git a/Engine/Source/Runtime/CoreUObject/Public/Templates/Casts.h b/Engine/Source/Runtime/CoreUObject/Public/Templates/Casts.h
--- a/Engine/Source/Runtime/CoreUObject/Public/Templates/Casts.h
+++ b/Engine/Source/Runtime/CoreUObject/Public/Templates/Casts.h
@@ -91,7 +91,7 @@ struct UE_DEPRECATED(5.5, "TCastFlags has been deprecated - use Cast instead.")
 
 // Dynamically cast an object type-safely.
 template <typename To, typename From>
-FORCEINLINE To* Cast(From* Src)
+FORCEINLINE TCopyQualifiersFromTo_T<From, To>* Cast(From* Src)
 {
 	static_assert(sizeof(From) > 0 && sizeof(To) > 0, "Attempting to cast between incomplete types");
 
@@ -169,10 +169,10 @@ FORCEINLINE To* Cast(From* Src)
 	return nullptr;
 }
 
-template< class T >
-FORCEINLINE T* ExactCast( UObject* Src )
+template <typename To, typename From>
+FORCEINLINE TCopyQualifiersFromTo_T<From, To>* ExactCast(From* Src)
 {
-	return Src && (Src->GetClass() == T::StaticClass()) ? (T*)Src : nullptr;
+	return Src && (Src->GetClass() == To::StaticClass()) ? (To*)Src : nullptr;
 }
 
 #if DO_CHECK
@@ -187,9 +187,9 @@ FORCEINLINE T* ExactCast( UObject* Src )
 			return GetFullNameSafe(InObjectOrInterface);
 		}
 		else if constexpr (std::is_base_of_v<UObject, T>)
-    	{
-    		return InObjectOrInterface->GetFullName();;
-    	}
+		{
+			return InObjectOrInterface->GetFullName();;
+		}
 		else
 		{
 			return Cast<UObject>(InObjectOrInterface)->GetFullName();
@@ -198,7 +198,7 @@ FORCEINLINE T* ExactCast( UObject* Src )
 
 	template <typename To, typename From>
 	FUNCTION_NON_NULL_RETURN_START
-		To* CastChecked(From* Src)
+		TCopyQualifiersFromTo_T<From, To>* CastChecked(From* Src)
 	FUNCTION_NON_NULL_RETURN_END
 	{
 		static_assert(sizeof(From) > 0 && sizeof(To) > 0, "Attempting to cast between incomplete types");
@@ -208,7 +208,7 @@ FORCEINLINE T* ExactCast( UObject* Src )
 			CastLogError(TEXT("nullptr"), *GetTypeName<To>());
 		}
 
-		To* Result = Cast<To>(Src);
+		TCopyQualifiersFromTo_T<From, To>* Result = Cast<To>(Src);
 		if (!Result)
 		{
 			CastLogError(*GetFullNameForCastLogError(Src), *GetTypeName<To>());
@@ -218,13 +218,13 @@ FORCEINLINE T* ExactCast( UObject* Src )
 	}
 
 	template <typename To, typename From>
-	To* CastChecked(From* Src, ECastCheckedType::Type CheckType)
+	TCopyQualifiersFromTo_T<From, To>* CastChecked(From* Src, ECastCheckedType::Type CheckType)
 	{
 		static_assert(sizeof(From) > 0 && sizeof(To) > 0, "Attempting to cast between incomplete types");
 
 		if (Src)
 		{
-			To* Result = Cast<To>(Src);
+			TCopyQualifiersFromTo_T<From, To>* Result = Cast<To>(Src);
 			if (!Result)
 			{
 				CastLogError(*GetFullNameForCastLogError(Src), *GetTypeName<To>());
@@ -245,7 +245,7 @@ FORCEINLINE T* ExactCast( UObject* Src )
 
 	template <typename To, typename From>
 	FUNCTION_NON_NULL_RETURN_START
-		FORCEINLINE To* CastChecked(From* Src)
+		FORCEINLINE TCopyQualifiersFromTo_T<From, To>* CastChecked(From* Src)
 	FUNCTION_NON_NULL_RETURN_END
 	{
 		static_assert(sizeof(From) > 0 && sizeof(To) > 0, "Attempting to cast between incomplete types");
@@ -278,7 +278,7 @@ FORCEINLINE T* ExactCast( UObject* Src )
 	}
 
 	template <typename To, typename From>
-	FORCEINLINE To* CastChecked(From* Src, ECastCheckedType::Type CheckType)
+	FORCEINLINE TCopyQualifiersFromTo_T<From, To>* CastChecked(From* Src, ECastCheckedType::Type CheckType)
 	{
 		return CastChecked<To>(Src);
 	}
@@ -286,9 +286,26 @@ FORCEINLINE T* ExactCast( UObject* Src )
 #endif
 
 // auto weak versions
-template< class T, class U > FORCEINLINE T* Cast       ( const TWeakObjectPtr<U>& Src                                                                   ) { return Cast       <T>(Src.Get()); }
-template< class T, class U > FORCEINLINE T* ExactCast  ( const TWeakObjectPtr<U>& Src                                                                   ) { return ExactCast  <T>(Src.Get()); }
-template< class T, class U > FORCEINLINE T* CastChecked( const TWeakObjectPtr<U>& Src, ECastCheckedType::Type CheckType = ECastCheckedType::NullChecked ) { return CastChecked<T>(Src.Get(), CheckType); }
+template <typename To, typename From>
+FORCEINLINE TCopyQualifiersFromTo_T<From, To>* Cast(const TWeakObjectPtr<From>& Src)
+{
+	return Cast<To>(Src.Get());
+}
+template <typename To, typename From>
+FORCEINLINE TCopyQualifiersFromTo_T<From, To>* ExactCast(const TWeakObjectPtr<From>& Src)
+{
+	return ExactCast  <To>(Src.Get());
+}
+template <typename To, typename From>
+FORCEINLINE TCopyQualifiersFromTo_T<From, To>* CastChecked(const TWeakObjectPtr<From>& Src)
+{
+	return CastChecked<To>(Src.Get());
+}
+template <typename To, typename From>
+FORCEINLINE TCopyQualifiersFromTo_T<From, To>* CastChecked(const TWeakObjectPtr<From>& Src, ECastCheckedType::Type CheckType)
+{
+	return CastChecked<To>(Src.Get(), CheckType);
+}
 
 // object ptr versions
 template <typename To, typename From>
@@ -407,11 +424,6 @@ FORCEINLINE TCopyQualifiersFromTo_T<From, To>* CastChecked(const TObjectPtr<From
 template< class T, class U > FORCEINLINE T* Cast       ( const TSubclassOf<U>& Src                                                                   ) { return Cast       <T>(*Src); }
 template< class T, class U > FORCEINLINE T* CastChecked( const TSubclassOf<U>& Src, ECastCheckedType::Type CheckType = ECastCheckedType::NullChecked ) { return CastChecked<T>(*Src, CheckType); }
 
-// Const versions of the casts
-template< class T, class U > FORCEINLINE const T* Cast       ( const U      * Src                                                                   ) { return Cast       <T>(const_cast<U      *>(Src)); }
-template< class T          > FORCEINLINE const T* ExactCast  ( const UObject* Src                                                                   ) { return ExactCast  <T>(const_cast<UObject*>(Src)); }
-template< class T, class U > FORCEINLINE const T* CastChecked( const U      * Src, ECastCheckedType::Type CheckType = ECastCheckedType::NullChecked ) { return CastChecked<T>(const_cast<U      *>(Src), CheckType); }
-
 #define DECLARE_CAST_BY_FLAG(ClassName) \
 	class ClassName; \
 	template <> \
