Remove unnecessary block scope.

The braces here don't do anything meaningful and are confusing for a reader trying to understand the code at a glance, since they differ from the neighboring cases.

diff --git a/Engine/Source/Runtime/CoreUObject/Public/UObject/ObjectHandle.h b/Engine/Source/Runtime/CoreUObject/Public/UObject/ObjectHandle.h
--- a/Engine/Source/Runtime/CoreUObject/Public/UObject/ObjectHandle.h
+++ b/Engine/Source/Runtime/CoreUObject/Public/UObject/ObjectHandle.h
@@ -356,7 +356,7 @@ namespace UE::CoreUObject::Private
 	inline FObjectHandle MakeObjectHandle(UObject* Object)
 	{
 #if UE_WITH_OBJECT_HANDLE_LATE_RESOLVE
-		{ return { UPTRINT(Object) }; }
+		return { UPTRINT(Object) };
 #elif UE_WITH_REMOTE_OBJECT_HANDLE
 		return FObjectHandle(Object);
 #else
