Fix typo in code documentation

#rnx

diff --git a/Engine/Source/Runtime/CoreUObject/Private/Serialization/EditorBulkData.cpp b/Engine/Source/Runtime/CoreUObject/Private/Serialization/EditorBulkData.cpp
--- a/Engine/Source/Runtime/CoreUObject/Private/Serialization/EditorBulkData.cpp
+++ b/Engine/Source/Runtime/CoreUObject/Private/Serialization/EditorBulkData.cpp
@@ -95,7 +95,7 @@ static bool ShouldSaveToPackageSidecar()
 }
 
 /** 
- * Returns true if the project wants to log if an editor bulkdata is cooked or uncooked annd returns
+ * Returns true if the project wants to log if an editor bulkdata is cooked or uncooked and returns
  * false if the project does not care. The engine ini file defaults to false 
  */
 static bool ShouldLogCookedStatus()
