Fixing garbage reference reporting

diff --git a/Engine/Source/Runtime/CoreUObject/Private/UObject/GarbageCollection.cpp b/Engine/Source/Runtime/CoreUObject/Private/UObject/GarbageCollection.cpp
--- a/Engine/Source/Runtime/CoreUObject/Private/UObject/GarbageCollection.cpp
+++ b/Engine/Source/Runtime/CoreUObject/Private/UObject/GarbageCollection.cpp
@@ -4230,6 +4230,36 @@ public:
 		GGCStats.NumRoots = OutRootObjects.Num();
 	}
 
+	void ResetReachabilityFlags(const EGatherOptions Options)
+	{
+		using namespace UE::GC;
+		using namespace UE::GC::Private;
+
+		using FMarkObjectsState = TThreadedGather<TArray<UObject*>>;
+		FMarkObjectsState MarkObjectsState;
+
+		MarkObjectsState.Start(Options, GUObjectArray.GetObjectArrayNum(), GUObjectArray.GetFirstGCIndex());
+
+		// Reset all objects to 'MaybeUnreachable' state
+		FMarkObjectsState::FThreadIterators& ThreadIterators = MarkObjectsState.GetThreadIterators();
+		ParallelFor(TEXT("GC.ResetReachabilityFlags"), MarkObjectsState.NumWorkerThreads(), 1, [&ThreadIterators](int32 ThreadIndex)
+			{
+				TRACE_CPUPROFILER_EVENT_SCOPE(ResetReachabilityFlagsTask);
+				FMarkObjectsState::FIterator& ThreadState = ThreadIterators[ThreadIndex];
+				for (; ThreadState.Index <= ThreadState.LastIndex; ++ThreadState.Index)
+				{
+					FUObjectItem* ObjectItem = &GUObjectArray.GetObjectItemArrayUnsafe()[ThreadState.Index];
+					if (ObjectItem->Object)
+					{
+						FGCFlags::SetMaybeUnreachable_ForGC(ObjectItem);
+					}
+				}
+			}, (MarkObjectsState.NumWorkerThreads() == 1) ? EParallelForFlags::ForceSingleThread : EParallelForFlags::None);
+
+		TArray<UObject*> Unused; // We didn't collect any objects when resetting so this is just a dummy to satisfy the API
+		MarkObjectsState.Finish(Unused);
+	}
+
 	/**
 	 * Marks all objects that don't have KeepFlags and EInternalObjectFlags_GarbageCollectionKeepFlags as MaybeUnreachable
 	 */
@@ -4238,18 +4268,24 @@ public:
 		using namespace UE::GC;
 		using namespace UE::GC::Private;
 
+		EGatherOptions GatherOptions = GetObjectGatherOptions();
+
 		// Don't swap the flags if we're re-entering this function to track garbage references
 		if (const bool bInitialMark = !Stats.bFoundGarbageRef)
 		{
 			// This marks all UObjects as MaybeUnreachable
 			FGCFlags::SwapReachableAndMaybeUnreachable();
 		}
+		else
+		{
+			// Swapping flags would inverse reachability results from the initial (normal) pass but what we want
+			// is to reset reachability state of all objects to 'MaybeUnreachable'
+			ResetReachabilityFlags(GatherOptions);
+		}
 
 		// Not counting the disregard for GC set to preserve legacy behavior
 		GObjectCountDuringLastMarkPhase.Set(GUObjectArray.GetObjectArrayNumMinusAvailable() - GUObjectArray.GetFirstGCIndex());
 
-		EGatherOptions GatherOptions = GetObjectGatherOptions();
-
 		// Now make sure all clustered objects and root objects are marked as Reachable. 
 		// This could be considered as initial part of reachability analysis and could be made incremental.
 		MarkClusteredObjectsAsReachable(GatherOptions, InitialObjects);
@@ -5247,6 +5283,7 @@ static void DumpGarbageReferencers(TConstArrayView<TUniquePtr<FWorkerContext>> C
 		{
 			for (FGarbageReferenceInfo& GarbageReference : Context->GarbageReferences)
 			{
+				// Note: EngineTest tests rely on the format of this message, see TestGarbageElimination
 				UE_LOG(LogGarbage, Warning, TEXT("Reachable garbage object: %s"), *GarbageReference.GarbageObject->GetFullName());
 				UE_LOG(LogGarbage, Warning, TEXT("Referenced by:            %s"), *GarbageReference.GetReferencingObjectInfo());
 				UE_LOG(LogGarbage, Warning, TEXT(""));
@@ -5270,6 +5307,7 @@ static void DumpGarbageReferencers(TConstArrayView<TUniquePtr<FWorkerContext>> C
 				FKey Key(GarbageReference.bReferencerUObject ? (void*)GarbageReference.Referencer.Object->GetClass() : (void*)  GarbageReference.Referencer.GCObject, GarbageReference.PropertyName);
 				if (Seen.Contains(Key) == false)
 				{
+					// Note: EngineTest tests rely on the format of this message, see TestGarbageElimination
 					UE_LOG(LogGarbage, Warning, TEXT("Reachable garbage object: %s"), *GarbageReference.GarbageObject->GetFullName());
 					UE_LOG(LogGarbage, Warning, TEXT("Referenced by:            %s"), *GarbageReference.GetReferencingObjectInfo());
 					UE_LOG(LogGarbage, Warning, TEXT(""));
