Removed unnecessary const_casts.

#rnx

diff --git a/Engine/Source/Runtime/CoreUObject/Private/UObject/ObjectBaseUtility.cpp b/Engine/Source/Runtime/CoreUObject/Private/UObject/ObjectBaseUtility.cpp
--- a/Engine/Source/Runtime/CoreUObject/Private/UObject/ObjectBaseUtility.cpp
+++ b/Engine/Source/Runtime/CoreUObject/Private/UObject/ObjectBaseUtility.cpp
@@ -31,7 +31,7 @@ FPackageFileVersion UObjectBaseUtility::GetLinkerUEVersion() const
 		// is the linker root, find our loader in the global list.
 		if (GetOutermost() == this)
 		{
-			Loader = FLinkerLoad::FindExistingLinkerForPackage(const_cast<UPackage*>(CastChecked<UPackage>((const UObject*)this)));
+			Loader = FLinkerLoad::FindExistingLinkerForPackage(CastChecked<UPackage>((const UObject*)this));
 		}
 	}
 
@@ -64,7 +64,7 @@ int32 UObjectBaseUtility::GetLinkerCustomVersion(FGuid CustomVersionKey) const
 		// is the linker root, find our loader in the global list.
 		if( GetOutermost() == this )
 		{
-			Loader = FLinkerLoad::FindExistingLinkerForPackage(const_cast<UPackage*>(CastChecked<UPackage>((const UObject*)this)));
+			Loader = FLinkerLoad::FindExistingLinkerForPackage(CastChecked<UPackage>((const UObject*)this));
 		}
 	}
 
@@ -105,7 +105,7 @@ int32 UObjectBaseUtility::GetLinkerLicenseeUEVersion() const
 		// is the linker root, find our loader in the global list.
 		if( GetOutermost() == this )
 		{
-			Loader = FLinkerLoad::FindExistingLinkerForPackage(const_cast<UPackage*>(CastChecked<UPackage>((const UObject*)this)));
+			Loader = FLinkerLoad::FindExistingLinkerForPackage(CastChecked<UPackage>((const UObject*)this));
 		}
 	}
 
