Fixed the TSubclassOf overloads of Casts to be const correct.
Expanded the CastChecked overload into two to take advantage of the overload pairs it forwards to.

diff --git a/Engine/Source/Runtime/CoreUObject/Public/Templates/Casts.h b/Engine/Source/Runtime/CoreUObject/Public/Templates/Casts.h
--- a/Engine/Source/Runtime/CoreUObject/Public/Templates/Casts.h
+++ b/Engine/Source/Runtime/CoreUObject/Public/Templates/Casts.h
@@ -387,15 +387,15 @@ FORCEINLINE TCopyQualifiersFromTo_T<From, To>* CastChecked(const TObjectPtr<From
 #if DO_CHECK
 	if (Src)
 	{
-		auto* Result = Cast<To>(Src);
-		if (!Result)
-		{
+	auto* Result = Cast<To>(Src);
+	if (!Result)
+	{
 			CastLogError(*GetFullNameForCastLogError(Src.Get()), *GetTypeName<To>());
-		}
-
-		return Result;
 	}
 
+	return Result;
+}
+
 	if (CheckType == ECastCheckedType::NullChecked)
 	{
 		CastLogError(TEXT("nullptr"), *GetTypeName<To>());
@@ -421,8 +421,21 @@ FORCEINLINE TCopyQualifiersFromTo_T<From, To>* CastChecked(const TObjectPtr<From
 }
 
 // TSubclassOf versions
-template< class T, class U > FORCEINLINE T* Cast       ( const TSubclassOf<U>& Src                                                                   ) { return Cast       <T>(*Src); }
-template< class T, class U > FORCEINLINE T* CastChecked( const TSubclassOf<U>& Src, ECastCheckedType::Type CheckType = ECastCheckedType::NullChecked ) { return CastChecked<T>(*Src, CheckType); }
+template <typename To, typename From>
+FORCEINLINE TCopyQualifiersFromTo_T<From, To>* Cast(const TSubclassOf<From>& Src)
+{
+	return Cast<To>(*Src);
+}
+template <typename To, typename From>
+FORCEINLINE TCopyQualifiersFromTo_T<From, To>* CastChecked(const TSubclassOf<From>& Src)
+{
+	return CastChecked<To>(*Src);
+}
+template <typename To, typename From>
+FORCEINLINE TCopyQualifiersFromTo_T<From, To>* CastChecked(const TSubclassOf<From>& Src, ECastCheckedType::Type CheckType)
+{
+	return CastChecked<To>(*Src, CheckType);
+}
 
 #define DECLARE_CAST_BY_FLAG(ClassName) \
 	class ClassName; \
