Added support for TNotNull in Cast.

#rnx

diff --git a/Engine/Source/Runtime/CoreUObject/Public/Templates/Casts.h b/Engine/Source/Runtime/CoreUObject/Public/Templates/Casts.h
--- a/Engine/Source/Runtime/CoreUObject/Public/Templates/Casts.h
+++ b/Engine/Source/Runtime/CoreUObject/Public/Templates/Casts.h
@@ -4,6 +4,7 @@
 
 #include "Containers/UnrealString.h"
 #include "CoreTypes.h"
+#include "Misc/NotNull.h"
 #include "Templates/LosesQualifiersFromTo.h"
 #include "Templates/UnrealTemplate.h"
 #include "UObject/Class.h"
@@ -437,6 +438,30 @@ FORCEINLINE TCopyQualifiersFromTo_T<From, To>* CastChecked(const TSubclassOf<Fro
 	return CastChecked<To>(*Src, CheckType);
 }
 
+// TNotNull versions of the casts
+#if UE_ENABLE_NOTNULL_WRAPPER
+	template <typename T, typename U>
+	FORCEINLINE auto Cast(TNotNull<U> Ptr) -> decltype(Cast<T>((U)Ptr))
+	{
+		return Cast<T>((U)Ptr);
+	}
+	template <typename T, typename U>
+	FORCEINLINE auto ExactCast(TNotNull<U> Ptr) -> decltype(ExactCast<T>((U)Ptr))
+	{
+		return ExactCast<T>((U)Ptr);
+	}
+	template <typename T, typename U>
+	FORCEINLINE auto CastChecked(TNotNull<U> Ptr) -> decltype(CastChecked<T>((U)Ptr))
+	{
+		return CastChecked<T>((U)Ptr);
+	}
+	template <typename T, typename U>
+	FORCEINLINE auto CastChecked(TNotNull<U> Ptr, ECastCheckedType::Type CheckType) -> decltype(CastChecked<T>((U)Ptr, CheckType))
+	{
+		return CastChecked<T>((U)Ptr, CheckType);
+	}
+#endif
+
 #define DECLARE_CAST_BY_FLAG(ClassName) \
 	class ClassName; \
 	template <> \
