dstm:
- reworked physics query RPCs to execute transactionally so they can dereference remote objects, etc
- some of the physics RPCs now get pushed into a queue that we can service outside of pumping the network stack
- some of the RPCs now have a priority attached so that we can realize when a higher priority request is pending in the queue that we need to switch to working on
- object requests now inherit the priority of the originating work, so if we execute a physics query at priority 3, and it needs an object, then the object migration inherits the priority 3
- ExecuteTransactional now has an outer loop that can switch to processing higher priority work from the deferred RPC queue, where it abandons the current tracking for the current work and switches to executing the other thing. Probably could use more refactorings to make this a bit nicer.
- reworked the physics RPC manager. removed pumping of the network and basically turned it into a peer system like the object migration manager where we fill it with requests, then Tick it over and over to fire off requests and collect responses before allowing ExecuteTransactional to retry the work.

diff --git a/Engine/Source/Runtime/CoreUObject/Private/UObject/RemoteObject.cpp b/Engine/Source/Runtime/CoreUObject/Private/UObject/RemoteObject.cpp
--- a/Engine/Source/Runtime/CoreUObject/Private/UObject/RemoteObject.cpp
+++ b/Engine/Source/Runtime/CoreUObject/Private/UObject/RemoteObject.cpp
@@ -114,6 +114,14 @@ void InitRemoteObjects()
 	{
 		UE::RemoteObject::Transfer::RestoreRemoteObjectDataDelegate.BindStatic(&UE::RemoteObject::Serialization::Disk::LoadObjectFromDisk);
 	}
+	if (!UE::RemoteObject::Transfer::AnyDeferredRPCQueuesExecutableDelegate.IsBound())
+	{
+		UE::RemoteObject::Transfer::AnyDeferredRPCQueuesExecutableDelegate.BindStatic([](FRemoteServerId){ return false; });
+	}
+	if (!UE::RemoteObject::Transfer::PhysicsQueryTickDelegate.IsBound())
+	{
+		UE::RemoteObject::Transfer::PhysicsQueryTickDelegate.BindStatic([](){ return true; });
+	}
 
 	FCoreDelegates::OnEndFrame.AddStatic(&UE::RemoteObject::Transfer::WaitForServersToSync);
 }
diff --git a/Engine/Source/Runtime/CoreUObject/Private/UObject/RemoteObjectSerialization.cpp b/Engine/Source/Runtime/CoreUObject/Private/UObject/RemoteObjectSerialization.cpp
--- a/Engine/Source/Runtime/CoreUObject/Private/UObject/RemoteObjectSerialization.cpp
+++ b/Engine/Source/Runtime/CoreUObject/Private/UObject/RemoteObjectSerialization.cpp
@@ -64,7 +64,7 @@ FString GenerateRemoteObjectFilename(FRemoteObjectId ObjectId, FRemoteServerId D
 	return FPaths::Combine(*FPaths::ProjectSavedDir(), *(FString::Printf(TEXT("%s_%s.remote"), *ObjectId.ToString(), *DestinationServerId.ToString())));
 }
 
-void LoadObjectFromDisk(uint32 RequestToken, FRemoteObjectId ObjectId, FRemoteServerId CurrentOwnerServerId, FRemoteServerId DestinationServerId)
+void LoadObjectFromDisk(uint32 RequestToken, FRemoteServerId RequestPriority, FRemoteObjectId ObjectId, FRemoteServerId CurrentOwnerServerId, FRemoteServerId DestinationServerId)
 {
 	const FString Filename(GenerateRemoteObjectFilename(ObjectId, CurrentOwnerServerId));
 	TUniquePtr<FArchive> FileReader{ IFileManager::Get().CreateFileReader(*Filename) };
diff --git a/Engine/Source/Runtime/CoreUObject/Private/UObject/RemoteObjectSerialization.h b/Engine/Source/Runtime/CoreUObject/Private/UObject/RemoteObjectSerialization.h
--- a/Engine/Source/Runtime/CoreUObject/Private/UObject/RemoteObjectSerialization.h
+++ b/Engine/Source/Runtime/CoreUObject/Private/UObject/RemoteObjectSerialization.h
@@ -61,6 +61,6 @@ namespace UE::RemoteObject::Serialization
 
 namespace UE::RemoteObject::Serialization::Disk
 {
-	void LoadObjectFromDisk(uint32 RequestToken, FRemoteObjectId ObjectId, FRemoteServerId CurrentOwnerServerId, FRemoteServerId DestinationServerId);
+	void LoadObjectFromDisk(uint32 RequestToken, FRemoteServerId RequestPriority, FRemoteObjectId ObjectId, FRemoteServerId CurrentOwnerServerId, FRemoteServerId DestinationServerId);
 	void SaveObjectToDisk(uint32 RequestToken, FRemoteObjectId ObjectId, FRemoteServerId ServerId, FRemoteObjectData& ObjectData);
 }
\ No newline at end of file
diff --git a/Engine/Source/Runtime/CoreUObject/Private/UObject/RemoteObjectTransfer.cpp b/Engine/Source/Runtime/CoreUObject/Private/UObject/RemoteObjectTransfer.cpp
--- a/Engine/Source/Runtime/CoreUObject/Private/UObject/RemoteObjectTransfer.cpp
+++ b/Engine/Source/Runtime/CoreUObject/Private/UObject/RemoteObjectTransfer.cpp
@@ -20,10 +20,13 @@ namespace UE::RemoteObject::Transfer
 {
 const FRemoteServerId DatabaseId(0x3ff);
 
-TDelegate<void(uint32 RequestToken, FRemoteObjectId ObjectId, FRemoteServerId CurrentOwnerServerId, FRemoteServerId DestinationServerId)> RequestRemoteObjectDelegate;
+TDelegate<void(uint32 RequestToken, FRemoteServerId RequestPriority, FRemoteObjectId ObjectId, FRemoteServerId CurrentOwnerServerId, FRemoteServerId DestinationServerId)> RequestRemoteObjectDelegate;
 TDelegate<void(uint32 RequestToken, FRemoteObjectId ObjectId, FRemoteServerId DestinationServerId, FRemoteObjectData& ObjectData)> RemoteObjectTransferDelegate;
 TDelegate<void(uint32 RequestToken, FRemoteObjectId ObjectId, FRemoteServerId DestinationServerId)> RemoteObjectDeniedTransferDelegate;
-TDelegate<void()> FlushRemoteObjectRequestsDelegate;
+TDelegate<void()> TickNetworkDelegate;
+TDelegate<bool(FRemoteServerId)> AnyDeferredRPCQueuesExecutableDelegate;
+TDelegate<void(FRemoteServerId)> ExecuteDeferredRPCQueuesDelegate;
+TDelegate<void()> ExecuteAllDeferredRPCQueuesDelegate;
 TDelegate<void()> WaitForServersToSyncDelegate;
 
 TMulticastDelegate<void(const FRemoteObjectData& ObjectData, FRemoteServerId RemoteServerId)> OnObjectDataSentDelegate;
@@ -33,10 +36,11 @@ TMulticastDelegate<void(FRemoteObjectId ObjectId, FRemoteServerId RemoteServerId
 TMulticastDelegate<void(const TArray<UObject*>& MigratedObjects, FRemoteServerId RemoteServerId)> OnObjectsReceivedDelegate;
 
 TDelegate<void(uint32 RequestToken, FRemoteObjectId ObjectId, FRemoteServerId DatabaseId, FRemoteObjectData& ObjectData)> StoreRemoteObjectDataDelegate;
-TDelegate<void(uint32 RequestToken, FRemoteObjectId ObjectId, FRemoteServerId DatabaseId, FRemoteServerId DestinationServerId)> RestoreRemoteObjectDataDelegate;
+TDelegate<void(uint32 RequestToken, FRemoteServerId RequestPriority, FRemoteObjectId ObjectId, FRemoteServerId DatabaseId, FRemoteServerId DestinationServerId)> RestoreRemoteObjectDataDelegate;
 
-TDelegate<void()> FlushRemotePhysicsQueryDelegate;
-TDelegate<void()> ClearRemotePhysicsQueryQueueDelegate;
+TDelegate<void()> PhysicsQueryBeginRequestDelegate;
+TDelegate<bool()> PhysicsQueryTickDelegate;
+TDelegate<void()> PhysicsQueryEndRequestDelegate;
 
 TMulticastDelegate<void()> OnTransactionCompletedDelegate;
 TMulticastDelegate<void()> OnTransactionAbortedDelegate;
@@ -54,9 +58,10 @@ public:
 	TArray<FRemoteObjectId> RequestedObjects;
 	TAtomic<uint32> ActiveRequestToken;
 	TAtomic<uint32> NextActiveRequestToken;
+	FRemoteServerId ActiveRequestPriority;
 	int32 PendingRequests = 0;
 
-	void BeginRequest()
+	void BeginRequest(FRemoteServerId Priority)
 	{
 		check(!AutoRTFM::IsClosed());
 
@@ -74,6 +79,7 @@ public:
 		bool TokenExchangeSuccess = ActiveRequestToken.CompareExchange(ExpectedToken, (NextActiveRequestToken.Load() % 32000) + 1000);
 		check(TokenExchangeSuccess);
 
+		ActiveRequestPriority = Priority;
 		NextActiveRequestToken++;
 	}
 
@@ -91,20 +97,23 @@ public:
 		// clear the object lists
 		LockedObjects.Empty();
 		RequestedObjects.Empty();
+
+		// clear the priority
+		ActiveRequestPriority = FRemoteServerId();
 	}
 
-	void RequestRemoteObject(uint32 RequestToken, FRemoteObjectId ObjectId, FRemoteServerId CurrentOwnerServerId, FRemoteServerId DestinationServerId)
+	void RequestRemoteObject(uint32 RequestToken, FRemoteServerId RequestPriority, FRemoteObjectId ObjectId, FRemoteServerId CurrentOwnerServerId, FRemoteServerId DestinationServerId)
 	{
 		check(!AutoRTFM::IsClosed());
 
 		// Request an object from remote server
 		if (CurrentOwnerServerId != DatabaseId)
 		{
-			RequestRemoteObjectDelegate.Execute(RequestToken, ObjectId, CurrentOwnerServerId, DestinationServerId);
+			RequestRemoteObjectDelegate.Execute(RequestToken, RequestPriority, ObjectId, CurrentOwnerServerId, DestinationServerId);
 		}
 		else
 		{
-			RestoreRemoteObjectDataDelegate.Execute(RequestToken, ObjectId, DatabaseId, GetGlobalServerId());
+			RestoreRemoteObjectDataDelegate.Execute(RequestToken, RequestPriority, ObjectId, DatabaseId, GetGlobalServerId());
 		}
 	}
 
@@ -189,7 +198,7 @@ void RequestObjectData(FRemoteObjectId ObjectId, FRemoteServerId CurrentOwnerSer
 	check(!AutoRTFM::IsClosed());
 
 	check(GTransferQueue.ActiveRequestToken.Load() != 0);
-	GTransferQueue.RequestRemoteObject(GTransferQueue.ActiveRequestToken.Load(), ObjectId, CurrentOwnerServerId, DestinationServerId);
+	GTransferQueue.RequestRemoteObject(GTransferQueue.ActiveRequestToken.Load(), GTransferQueue.ActiveRequestPriority, ObjectId, CurrentOwnerServerId, DestinationServerId);
 }
 
 UObject* MigrateObjectFromObjectData(FRemoteObjectData& ObjectData, FRemoteServerId RemoteServerId)
@@ -237,7 +246,7 @@ void MigrateObjectFromRemoteServer(FRemoteObjectId ObjectId, FRemoteServerId Cur
 	{
 		// if we are not in a transaction (dev/test workflow), begin a
 		// temp request
-		GTransferQueue.BeginRequest();
+		GTransferQueue.BeginRequest(GetGlobalServerId());
 	}
 
 	UE_AUTORTFM_OPEN
@@ -271,6 +280,11 @@ void MigrateObjectFromRemoteServer(FRemoteObjectId ObjectId, FRemoteServerId Cur
 }
 
 void MigrateObjectToRemoteServer(uint32 RequestToken, UObject* Object, FRemoteServerId DestinationServerId)
+{
+	MigrateObjectToRemoteServerWithExplicitPriority(RequestToken, GetGlobalServerId(), Object, DestinationServerId);
+}
+
+void MigrateObjectToRemoteServerWithExplicitPriority(uint32 RequestToken, FRemoteServerId RequestPriority, UObject* Object, FRemoteServerId DestinationServerId)
 {
 	using namespace UE::RemoteObject::Serialization;
 	using namespace UE::RemoteObject::Private;
@@ -291,7 +305,7 @@ void MigrateObjectToRemoteServer(uint32 RequestToken, UObject* Object, FRemoteSe
 		if (LockedObjectId == ObjectId)
 		{
 			// do we give it up, or deny their request?
-			if(GetGlobalServerId() < DestinationServerId)
+			if(GTransferQueue.ActiveRequestPriority < RequestPriority)
 			{
 				// we are higher priority, send a deny response
 				bIsObjectLocked = true;
@@ -364,6 +378,11 @@ void RegisterRemoteObjectId(FRemoteObjectId Id, FRemoteServerId OwnerServerId)
 }
 
 void MigrateObjectToRemoteServer(uint32 RequestToken, FRemoteObjectId ObjectId, FRemoteServerId DestinationServerId)
+{
+	MigrateObjectToRemoteServerWithExplicitPriority(RequestToken, GetGlobalServerId(), ObjectId, DestinationServerId);
+}
+
+void MigrateObjectToRemoteServerWithExplicitPriority(uint32 RequestToken, FRemoteServerId RequestPriority, FRemoteObjectId ObjectId, FRemoteServerId DestinationServerId)
 {
 	using namespace UE::RemoteObject::Handle;
 	using namespace UE::RemoteObject::Serialization;
@@ -375,7 +394,7 @@ void MigrateObjectToRemoteServer(uint32 RequestToken, FRemoteObjectId ObjectId,
 	if (Object && !IsRemote(Object))
 	{
 		UE_LOG(LogRemoteObject, VeryVerbose, TEXT("    : MigrateObjectToRemoteServer(%u) found obj %s"), RequestToken, *ObjectId.ToString());
-		MigrateObjectToRemoteServer(RequestToken, Object, DestinationServerId);
+		MigrateObjectToRemoteServerWithExplicitPriority(RequestToken, RequestPriority, Object, DestinationServerId);
 	}
 	else
 	{
@@ -397,7 +416,7 @@ void MigrateObjectToRemoteServer(uint32 RequestToken, FRemoteObjectId ObjectId,
 		if (CurrentOwnerServerId != DestinationServerId)
 		{
 			UE_LOG(LogRemoteObject, Verbose, TEXT("    : MigrateObjectToRemoteServer(%u) forwarding request (we believe %s owns the Object)"), RequestToken, *CurrentOwnerServerId.ToString());
-			GTransferQueue.RequestRemoteObject(RequestToken, ObjectId, CurrentOwnerServerId, DestinationServerId);
+			GTransferQueue.RequestRemoteObject(RequestToken, RequestPriority, ObjectId, CurrentOwnerServerId, DestinationServerId);
 		}
 		else
 		{
@@ -443,141 +462,190 @@ void PrefetchRemoteObject(FRemoteObjectId ObjectId, UObject* DestinationOuter)
 	};
 }
 
-void ExecuteTransactional(const TFunction<void(void)>& Work)
+void ExecuteTransactionalWithExplicitPriority(FRemoteServerId WorkPriority, const TFunctionRef<void(void)>& Work)
 {
-	check(!AutoRTFM::IsClosed());
-
-	GTransferQueue.BeginRequest();
+	ExecuteTransactionalInternal(WorkPriority, Work);
+}
 
-	UE_LOG(LogRemoteObject, VeryVerbose, TEXT("-- ExecuteTransactional : request %u --"), GTransferQueue.ActiveRequestToken.Load());
+void ExecuteTransactional(const TFunctionRef<void(void)>& Work)
+{
+	ExecuteTransactionalWithExplicitPriority(GetGlobalServerId(), Work);
+}
 
-	// we might run through this loop after the work succeeds in the case
-	// that the work prefetches objects that it doesn't access
-	bool bWorkSucceeded = false;
+void ExecuteTransactionalInternal(FRemoteServerId WorkPriority, const TFunctionRef<void(void)>& Work)
+{
+	check(!AutoRTFM::IsClosed());
 
 	for (;;)
 	{
-		UE_LOG(LogRemoteObjectD, Verbose, TEXT("   : %d locked objs :"), GTransferQueue.LockedObjects.Num());
-		for (FRemoteObjectId ObjectId : GTransferQueue.LockedObjects)
+		// first check to see if there's any higher priority requests we should be servicing
+		if (AnyDeferredRPCQueuesExecutableDelegate.Execute(WorkPriority))
 		{
-			UE_LOG(LogRemoteObjectD, Verbose, TEXT("       %s"), *ObjectId.ToString());
+			// do those first...
+			ExecuteDeferredRPCQueuesDelegate.ExecuteIfBound(WorkPriority);
 		}
 
-		UE_LOG(LogRemoteObjectD, Verbose, TEXT("   : %d requested objs :"), GTransferQueue.RequestedObjects.Num());
-		for (FRemoteObjectId ObjectId : GTransferQueue.RequestedObjects)
-		{
-			UE_LOG(LogRemoteObjectD, Verbose, TEXT("       %s"), *ObjectId.ToString());
-		}
+		GTransferQueue.BeginRequest(WorkPriority);
 
-		TArray<FRemoteObjectId> ObjectsToRequest;
+		PhysicsQueryBeginRequestDelegate.ExecuteIfBound();
 
-		// try to migrate all RequestedObjects -> LockedObjects
-		for (FRemoteObjectId RequestedObjectId : GTransferQueue.RequestedObjects)
+		UE_LOG(LogRemoteObject, VeryVerbose, TEXT("-- ExecuteTransactional : request %u --"), GTransferQueue.ActiveRequestToken.Load());
+
+		// we might run through this loop after the work succeeds in the case
+		// that the work prefetches objects that it doesn't access
+		bool bWorkSucceeded = false;
+
+		for (;;)
 		{
-			if (GTransferQueue.LockedObjects.Find(RequestedObjectId) == INDEX_NONE)
+			if (AnyDeferredRPCQueuesExecutableDelegate.Execute(WorkPriority))
 			{
-				ObjectsToRequest.Add(RequestedObjectId);
+				break;
 			}
-		}
-
-		UE_LOG(LogRemoteObjectD, Verbose, TEXT("   : ObjectsToRequest: %d entries"), ObjectsToRequest.Num());
 
-		// send the requests... slowly for now (to avoid straining the network subsystem)
-		int32 RequestsSent = 0;
-		double LastWaitingPrintTime = FPlatformTime::Seconds();
-		for(;;)
-		{
-			// do we need to, and can we send more requests?
-			while ((RequestsSent < ObjectsToRequest.Num()) && (GTransferQueue.PendingRequests < 5))
+			UE_LOG(LogRemoteObjectD, Verbose, TEXT("   : %d locked objs :"), GTransferQueue.LockedObjects.Num());
+			for (FRemoteObjectId ObjectId : GTransferQueue.LockedObjects)
 			{
-				// yes, send one request
-				FRemoteObjectId RequestedObjectId = ObjectsToRequest[RequestsSent++];
+				UE_LOG(LogRemoteObjectD, Verbose, TEXT("       %s"), *ObjectId.ToString());
+			}
 
-				if (UE::RemoteObject::Handle::RequestObject(RequestedObjectId))
-				{
-					// true means we already have the object locally and didn't actually send a request for it
-					GTransferQueue.LockedObjects.Add(RequestedObjectId);
-					UE_LOG(LogRemoteObjectD, Verbose, TEXT("   : marking %s as locked, it is already local"), *RequestedObjectId.ToString());
-				}
-				else
-				{
-					GTransferQueue.PendingRequests++;
-					UE_LOG(LogRemoteObjectD, Verbose, TEXT("   : sent request for %s (%d pending requests)"), *RequestedObjectId.ToString(), GTransferQueue.PendingRequests);
-				}
+			UE_LOG(LogRemoteObjectD, Verbose, TEXT("   : %d requested objs :"), GTransferQueue.RequestedObjects.Num());
+			for (FRemoteObjectId ObjectId : GTransferQueue.RequestedObjects)
+			{
+				UE_LOG(LogRemoteObjectD, Verbose, TEXT("       %s"), *ObjectId.ToString());
 			}
 
-			// pump the network
-			FlushRemoteObjectRequestsDelegate.ExecuteIfBound();
+			TArray<FRemoteObjectId> ObjectsToRequest;
 
-			// are all requests sent?
-			if (RequestsSent == ObjectsToRequest.Num())
+			// try to migrate all RequestedObjects -> LockedObjects
+			for (FRemoteObjectId RequestedObjectId : GTransferQueue.RequestedObjects)
 			{
-				// yes, all requests are sent
-				// are they all complete?
-				if (GTransferQueue.PendingRequests == 0)
+				if (GTransferQueue.LockedObjects.Find(RequestedObjectId) == INDEX_NONE)
 				{
-					// yes, all done
-					break;
+					ObjectsToRequest.Add(RequestedObjectId);
 				}
 			}
 
-			if (FPlatformTime::Seconds() > LastWaitingPrintTime + 1.0)
-			{
-				UE_LOG(LogRemoteObject, VeryVerbose, TEXT("   : waiting for %d responses..."), GTransferQueue.PendingRequests);
-				LastWaitingPrintTime = FPlatformTime::Seconds();
-			}
-		}
-
-		// now we also might have migrated objects away, so check to see if LockedObjects == RequestedObjects
-		if(GTransferQueue.LockedObjects.Num() != GTransferQueue.RequestedObjects.Num())
-		{
-			UE_LOG(LogRemoteObject, VeryVerbose, TEXT("   : locked count %d mismatches requests %d, did we give up locked objects? trying again..."), GTransferQueue.LockedObjects.Num(), GTransferQueue.RequestedObjects.Num());
+			UE_LOG(LogRemoteObjectD, Verbose, TEXT("   : ObjectsToRequest: %d entries"), ObjectsToRequest.Num());
 
-			// try again!
-			continue;
-		}
+			// send the requests... slowly for now (to avoid straining the network subsystem)
+			int32 RequestsSent = 0;
+			double LastWaitingPrintTime = FPlatformTime::Seconds();
+			for(;;)
+			{
+				// do we need to, and can we send more requests?
+				while ((RequestsSent < ObjectsToRequest.Num()) && (GTransferQueue.PendingRequests < 5))
+				{
+					// yes, send one request
+					FRemoteObjectId RequestedObjectId = ObjectsToRequest[RequestsSent++];
+
+					if (UE::RemoteObject::Handle::RequestObject(RequestedObjectId))
+					{
+						// true means we already have the object locally and didn't actually send a request for it
+						GTransferQueue.LockedObjects.Add(RequestedObjectId);
+						UE_LOG(LogRemoteObjectD, Verbose, TEXT("   : marking %s as locked, it is already local"), *RequestedObjectId.ToString());
+					}
+					else
+					{
+						GTransferQueue.PendingRequests++;
+						UE_LOG(LogRemoteObjectD, Verbose, TEXT("   : sent request for %s (%d pending requests)"), *RequestedObjectId.ToString(), GTransferQueue.PendingRequests);
+					}
+				}
 
-		FlushRemotePhysicsQueryDelegate.ExecuteIfBound();
+				// pump the network
+				TickNetworkDelegate.ExecuteIfBound();
 
-		UE_LOG(LogRemoteObject, VeryVerbose, TEXT("   : READY: all %d requested objects are locked"), GTransferQueue.RequestedObjects.Num());
+				// are all requests sent?
+				if (RequestsSent == ObjectsToRequest.Num())
+				{
+					// yes, all requests are sent
+					// are they all complete?
+					if (GTransferQueue.PendingRequests == 0)
+					{
+						// yes, all done
+						break;
+					}
+				}
 
-		// all of the requested objects are locked, try to perform the work
-		if (!bWorkSucceeded)
-		{
-			AutoRTFM::ETransactionResult WorkTransactResult = AutoRTFM::Transact([&Work](){ Work(); });
+				double TimeNow = FPlatformTime::Seconds();
+				if (TimeNow > LastWaitingPrintTime + 1.0)
+				{
+					UE_LOG(LogRemoteObject, VeryVerbose, TEXT("   : waiting for %d responses..."), GTransferQueue.PendingRequests);
+					LastWaitingPrintTime = TimeNow;
+				}
+			}
 
-			if (WorkTransactResult == AutoRTFM::ETransactionResult::AbortedByRequest)
+			if (!PhysicsQueryTickDelegate.Execute())
 			{
-					// the work aborted, it should have inserted a new requested object, retry
-				UE_LOG(LogRemoteObject, VeryVerbose, TEXT("   : work ABORTED"));
-				OnTransactionAbortedDelegate.Broadcast();
+				// physics doesn't have all responses yet, try again
+				continue;
 			}
-			else
+
+			// now we also might have migrated objects away, so check to see if LockedObjects == RequestedObjects
+			if(GTransferQueue.LockedObjects.Num() != GTransferQueue.RequestedObjects.Num())
 			{
-				UE_LOG(LogRemoteObject, VeryVerbose, TEXT("   : work COMPLETED!"));
+				UE_LOG(LogRemoteObject, VeryVerbose, TEXT("   : locked count %d mismatches requests %d, did we give up locked objects? trying again..."), GTransferQueue.LockedObjects.Num(), GTransferQueue.RequestedObjects.Num());
 
-				bWorkSucceeded = true;
-				OnTransactionCompletedDelegate.Broadcast();
+				// try again!
+				continue;
+			}
 
-				ClearRemotePhysicsQueryQueueDelegate.ExecuteIfBound();
+			UE_LOG(LogRemoteObject, VeryVerbose, TEXT("   : READY: all %d requested objects are locked"), GTransferQueue.RequestedObjects.Num());
 
-				if (GTransferQueue.LockedObjects.Num() == GTransferQueue.RequestedObjects.Num())
+			// all of the requested objects are locked, try to perform the work
+			if (!bWorkSucceeded)
+			{
+				AutoRTFM::ETransactionResult WorkTransactResult = AutoRTFM::Transact([&Work](){ Work(); });
+
+				if (WorkTransactResult == AutoRTFM::ETransactionResult::AbortedByRequest)
 				{
-					// the work succeeded and no more object requests were made, we're all done!
-					break;
+						// the work aborted, it should have inserted a new requested object, retry
+					UE_LOG(LogRemoteObject, VeryVerbose, TEXT("   : work ABORTED"));
+					OnTransactionAbortedDelegate.Broadcast();
 				}
 				else
 				{
-					// the work is complete, but it added object requests (prefetches), so we
-					// run through the loop again in order to fetch those objects
-					continue;
+					UE_LOG(LogRemoteObject, VeryVerbose, TEXT("   : work COMPLETED!"));
+
+					bWorkSucceeded = true;
+					OnTransactionCompletedDelegate.Broadcast();
+
+					if (GTransferQueue.LockedObjects.Num() == GTransferQueue.RequestedObjects.Num())
+					{
+						// the work succeeded and no more object requests were made, we're all done!
+						break;
+					}
+					else
+					{
+						// the work is complete, but it added object requests (prefetches), so we
+						// run through the loop again in order to fetch those objects
+						continue;
+					}
 				}
 			}
 		}
-	}
 
-	UE_LOG(LogRemoteObject, VeryVerbose, TEXT("-- ExecuteTransactional %u DONE --"), GTransferQueue.ActiveRequestToken.Load());
-	GTransferQueue.EndRequest();
+		// if we get here and the work was successfully executed, we're done
+		// otherwise, we broke out of the request because there are higher
+		// priority pending requests that need to be serviced so we need
+		// to loop back to the top
+		if (bWorkSucceeded)
+		{
+			UE_LOG(LogRemoteObject, VeryVerbose, TEXT("-- ExecuteTransactional %u DONE --"), GTransferQueue.ActiveRequestToken.Load());
+		}
+		else
+		{
+			UE_LOG(LogRemoteObject, VeryVerbose, TEXT("-- ExecuteTransactional %u DONE but execution failed, retrying to service RPC requests --"), GTransferQueue.ActiveRequestToken.Load());
+		}
+
+		GTransferQueue.EndRequest();
+
+		PhysicsQueryEndRequestDelegate.ExecuteIfBound();
+
+		if (bWorkSucceeded)
+		{
+			break;
+		}
+	}
 }
 
 void WaitForServersToSync()
diff --git a/Engine/Source/Runtime/CoreUObject/Public/UObject/RemoteObjectTransfer.h b/Engine/Source/Runtime/CoreUObject/Public/UObject/RemoteObjectTransfer.h
--- a/Engine/Source/Runtime/CoreUObject/Public/UObject/RemoteObjectTransfer.h
+++ b/Engine/Source/Runtime/CoreUObject/Public/UObject/RemoteObjectTransfer.h
@@ -16,7 +16,8 @@ namespace UE::RemoteObject::Transfer
 	* Execute the Work under an transaction that can catch remote object requests
 	* as well as negotiating remote object transfers, retrying the transaction until it succeeds.
 	*/
-	COREUOBJECT_API void ExecuteTransactional(const TFunction<void(void)>& Work);
+	COREUOBJECT_API void ExecuteTransactional(const TFunctionRef<void(void)>& Work);
+	COREUOBJECT_API void ExecuteTransactionalWithExplicitPriority(FRemoteServerId WorkPriority, const TFunctionRef<void(void)>& Work);
 
 	/**
 	* (Advanced API) Adds the remote object to the current transaction's requested object list, causing
@@ -53,6 +54,7 @@ namespace UE::RemoteObject::Transfer
 	* @param DestinationServerId Id of a remote server to receive the object
 	*/
 	COREUOBJECT_API void MigrateObjectToRemoteServer(uint32 RequestToken, UObject* Object, FRemoteServerId DestinationServerId);
+	COREUOBJECT_API void MigrateObjectToRemoteServerWithExplicitPriority(uint32 RequestToken, FRemoteServerId RequestPriority, UObject* Object, FRemoteServerId DestinationServerId);
 
 	/**
 	* Migrates an object to a remote server
@@ -60,6 +62,7 @@ namespace UE::RemoteObject::Transfer
 	* @param DestinationServerId Id of a remote server to receive the object
 	*/
 	COREUOBJECT_API void MigrateObjectToRemoteServer(uint32 RequestToken, FRemoteObjectId Id, FRemoteServerId DestinationServerId);
+	COREUOBJECT_API void MigrateObjectToRemoteServerWithExplicitPriority(uint32 RequestToken, FRemoteServerId RequestPriority, FRemoteObjectId Id, FRemoteServerId DestinationServerId);
 	
 	/**
 	* Migrates an object from a remote server (temp function)
@@ -110,10 +113,18 @@ namespace UE::RemoteObject::Transfer
 	extern COREUOBJECT_API TDelegate<void(uint32 /*RequestToken*/, FRemoteObjectId /*ObjectId*/, FRemoteServerId /*DestinationServerId*/)> RemoteObjectDeniedTransferDelegate;
 
 	/** Delegate that requests remote object data from CurrentOwnerServerId to be transferred to DestinationServerId */
-	extern COREUOBJECT_API TDelegate<void(uint32 /*RequestToken*/, FRemoteObjectId /*ObjectId*/, FRemoteServerId /*CurrentOwnerServerId*/, FRemoteServerId /*DestinationServerId*/)> RequestRemoteObjectDelegate;
+	extern COREUOBJECT_API TDelegate<void(uint32 /*RequestToken*/, FRemoteServerId /*RequestPriority*/, FRemoteObjectId /*ObjectId*/, FRemoteServerId /*CurrentOwnerServerId*/, FRemoteServerId /*DestinationServerId*/)> RequestRemoteObjectDelegate;
 
-	/** Delegate that flushes remote object requests */
-	extern COREUOBJECT_API TDelegate<void()> FlushRemoteObjectRequestsDelegate;
+	/** Delegate that pumps the network and executes RPC handlers */
+	extern COREUOBJECT_API TDelegate<void()> TickNetworkDelegate;
+	
+	/** Check to see if any higher or equal priority work is pending in the RPC queue **/
+	extern COREUOBJECT_API TDelegate<bool(FRemoteServerId)> AnyDeferredRPCQueuesExecutableDelegate;
+	
+	/** Execute higher or equal priority work in the RPC queue **/
+	extern COREUOBJECT_API TDelegate<void(FRemoteServerId)> ExecuteDeferredRPCQueuesDelegate;
+	extern COREUOBJECT_API TDelegate<void()> ExecuteAllDeferredRPCQueuesDelegate;
+	
 	/** Delegate that syncs all servers */
 	extern COREUOBJECT_API TDelegate<void()> WaitForServersToSyncDelegate;
 
@@ -140,13 +151,12 @@ namespace UE::RemoteObject::Transfer
 	extern COREUOBJECT_API TDelegate<void(uint32 /*RequestToken*/, FRemoteObjectId /*ObjectId*/, FRemoteServerId /*DatabaseId*/, FRemoteObjectData& /*ObjectData*/)> StoreRemoteObjectDataDelegate;
 
 	/** Delegate that restores object data from a database. */
-	extern COREUOBJECT_API TDelegate<void(uint32 /*RequestToken*/, FRemoteObjectId /*ObjectId*/, FRemoteServerId /*DatabaseId*/, FRemoteServerId /*DestinationServerId*/)> RestoreRemoteObjectDataDelegate;
+	extern COREUOBJECT_API TDelegate<void(uint32 /*RequestToken*/, FRemoteServerId /*RequestPriority*/, FRemoteObjectId /*ObjectId*/, FRemoteServerId /*DatabaseId*/, FRemoteServerId /*DestinationServerId*/)> RestoreRemoteObjectDataDelegate;
 
-	/** Delegate that flushes pending physics query requests. */
-	extern COREUOBJECT_API TDelegate<void()> FlushRemotePhysicsQueryDelegate;
-	
-	/** Delegate that marks the end of a successful transaction. Allows the physics query queue to be cleared. */
-	extern COREUOBJECT_API TDelegate<void()> ClearRemotePhysicsQueryQueueDelegate;
+	/** Physics query delegates **/
+	extern COREUOBJECT_API TDelegate<void()> PhysicsQueryBeginRequestDelegate;
+	extern COREUOBJECT_API TDelegate<bool()> PhysicsQueryTickDelegate;
+	extern COREUOBJECT_API TDelegate<void()> PhysicsQueryEndRequestDelegate;
 
 	/** Delegate executed when a ExecuteTransactional observes a completed AutoRTFM transaction */
 	extern COREUOBJECT_API TMulticastDelegate<void()> OnTransactionCompletedDelegate;
@@ -154,8 +164,9 @@ namespace UE::RemoteObject::Transfer
 	/** Delegate executed when a ExecuteTransactional observes an aborted AutoRTFM transaction */
 	extern COREUOBJECT_API TMulticastDelegate<void()> OnTransactionAbortedDelegate;
 
-} // namespace UE::RemoveObject::Transfer
+	extern COREUOBJECT_API void ExecuteTransactionalInternal(FRemoteServerId WorkPriority, const TFunctionRef<void(void)>& Work);
 
+} // namespace UE::RemoveObject::Transfer
 
 struct FRemoteObjectReference
 {
diff --git a/Engine/Source/Runtime/CoreUObject/Public/UObject/RemoteObjectTypes.h b/Engine/Source/Runtime/CoreUObject/Public/UObject/RemoteObjectTypes.h
--- a/Engine/Source/Runtime/CoreUObject/Public/UObject/RemoteObjectTypes.h
+++ b/Engine/Source/Runtime/CoreUObject/Public/UObject/RemoteObjectTypes.h
@@ -40,6 +40,10 @@ struct FRemoteServerId
 	{
 		return Id < Other.Id;
 	}
+	bool operator <= (FRemoteServerId Other) const
+	{
+		return Id <= Other.Id;
+	}
 
 	friend COREUOBJECT_API FArchive& operator<<(FArchive& Ar, FRemoteServerId& Id);
 
