Making non-editor targets that compile against editor use the editor path to determine if a ustruct is derived from another one.

Some editor only modules, like modules from TEDS, relies on the editor code path for IsChildOf to make a runtime re-parenting of their struct types, and therefore they will not work properly when used in a UE program that compiles against editor code.

The reason is that the non-editor code path relies in a an array that contains the inherence chain, but that array gets out of date if the parent of a ustruct is changed at runtime

diff --git a/Engine/Source/Runtime/CoreUObject/Public/UObject/ObjectMacros.h b/Engine/Source/Runtime/CoreUObject/Public/UObject/ObjectMacros.h
--- a/Engine/Source/Runtime/CoreUObject/Public/UObject/ObjectMacros.h
+++ b/Engine/Source/Runtime/CoreUObject/Public/UObject/ObjectMacros.h
@@ -37,7 +37,7 @@ typedef	uint64 ScriptPointerType;
 #define USTRUCT_ISCHILDOF_STRUCTARRAY 2 // stores an array of parents per struct and uses this to compare - faster than 1 and thread-safe but can have issues with BP reinstancing and hot reload
 
 // USTRUCT_FAST_ISCHILDOF_IMPL sets which implementation of IsChildOf to use.
-#if UE_EDITOR
+#if WITH_EDITOR
 	// On editor, we use the outerwalk implementation because BP reinstancing and hot reload
 	// mess up the struct array
 	#define USTRUCT_FAST_ISCHILDOF_IMPL USTRUCT_ISCHILDOF_OUTERWALK
