Fix rare crash when reinstancing coronado

#rnx

diff --git a/Engine/Source/Runtime/CoreUObject/Private/UObject/OverriddenPropertySet.cpp b/Engine/Source/Runtime/CoreUObject/Private/UObject/OverriddenPropertySet.cpp
--- a/Engine/Source/Runtime/CoreUObject/Private/UObject/OverriddenPropertySet.cpp
+++ b/Engine/Source/Runtime/CoreUObject/Private/UObject/OverriddenPropertySet.cpp
@@ -241,6 +241,17 @@ FOverriddenPropertyNodeID::FOverriddenPropertyNodeID(const FProperty* Property)
 	}
 }
 
+FOverriddenPropertyNodeID::FOverriddenPropertyNodeID(const FOverriddenPropertyNodeID& ParentNodeID, const FOverriddenPropertyNodeID& NodeID)
+{
+	// Combine the 2 node ids
+	FStringBuilderBase SubPropertyKeyBuilder;
+	SubPropertyKeyBuilder = *ParentNodeID.ToString();
+	SubPropertyKeyBuilder.Append(TEXT("."));
+	SubPropertyKeyBuilder.Append(*NodeID.ToString());
+	Path = FName(SubPropertyKeyBuilder.ToString());
+	Object = NodeID.Object;
+}
+
 FOverriddenPropertyNodeID FOverriddenPropertyNodeID::RootNodeId()
 {
 	FOverriddenPropertyNodeID Result;
@@ -304,6 +315,24 @@ int32 FOverriddenPropertyNodeID::ToMapInternalIndex(FScriptMapHelper& MapHelper)
 	return INDEX_NONE;
 }
 
+bool FOverriddenPropertyNodeID::operator==(const FOverriddenPropertyNodeID& Other) const
+{
+	if (Path == Other.Path)
+	{
+		return true;
+	}
+
+	// After reinstanciation, we do not change the path we only patch the ptr.
+	// Unfortunately we do not have a stable id that is stable through reinstantiation,
+	// so are only way if to compare ptr.
+	if (Object && Other.Object && Object == Other.Object)
+	{
+		return true;
+	}
+
+	return false;
+}
+
 void FOverriddenPropertyNodeID::HandleObjectsReInstantiated(const TMap<UObject*, UObject*>& Map)
 {
 	if (!Object)
@@ -338,12 +367,7 @@ FOverriddenPropertyNode& FOverriddenPropertySet::FindOrAddNode(FOverriddenProper
 	}
 
 	// Not found add the node
-	FStringBuilderBase SubPropertyKeyBuilder;
-	SubPropertyKeyBuilder = *ParentNode.NodeID.ToString();
-	SubPropertyKeyBuilder.Append(TEXT("."));
-	SubPropertyKeyBuilder.Append(*NodeID.ToString());
-	SubNodeID.Path = FName(SubPropertyKeyBuilder.ToString());
-	SubNodeID.Object = NodeID.Object;
+	SubNodeID = FOverriddenPropertyNodeID(ParentNode.NodeID, NodeID);
 	const FSetElementId NewID = OverriddenPropertyNodes.Emplace(SubNodeID);
 	return OverriddenPropertyNodes.Get(NewID);
 }
@@ -1392,7 +1416,7 @@ void FOverriddenPropertySet::Reset()
 void FOverriddenPropertySet::HandleObjectsReInstantiated(const TMap<UObject*, UObject*>& Map)
 {
 #if WITH_EDITOR
-	// When their is a cached archetype, it is an indicator this object is about to be replaced
+	// When there is a cached archetype, it is an indicator this object is about to be replaced
 	// So no need to replace any ptr, otherwise we might not be able to reconstitute the right information
 	if (FEditorCacheArchetypeManager::Get().GetCachedArchetype(Owner))
 	{
@@ -1437,4 +1461,4 @@ const FOverriddenPropertyNode* FOverriddenPropertySet::GetOverriddenPropertyNode
 	}
 
 	return OverriddenPropertyNode;
-}
\ No newline at end of file
+}
diff --git a/Engine/Source/Runtime/CoreUObject/Private/UObject/PropertyArray.cpp b/Engine/Source/Runtime/CoreUObject/Private/UObject/PropertyArray.cpp
--- a/Engine/Source/Runtime/CoreUObject/Private/UObject/PropertyArray.cpp
+++ b/Engine/Source/Runtime/CoreUObject/Private/UObject/PropertyArray.cpp
@@ -287,23 +287,6 @@ void FArrayProperty::SerializeItem(FStructuredArchive::FSlot Slot, void* Value,
 								return i;
 							}
 						}
-
-						////////////////////////////////////////////////////////////////////////////
-						// HACK for CreateProjectFromTemplate, 
-						// Fallback to find object by name
-						FName ObjectNameToFind = Object->GetFName();
-						for (int i = 0; i < ArrayNum; ++i)
-						{
-							if (UObject* CurrentObject = InnerObjectProperty->GetObjectPropertyValue(ArrayHelper.GetElementPtr(i)))
-							{
-								if (CurrentObject->GetFName() == ObjectNameToFind)
-								{
-									return i;
-								}
-							}
-						}
-						// HACK for CreateProjectFromTemplate, fallback to find object by name
-						////////////////////////////////////////////////////////////////////////////
 					}
 					return INDEX_NONE;
 				};
@@ -484,26 +467,6 @@ void FArrayProperty::SerializeItem(FStructuredArchive::FSlot Slot, void* Value,
 								}
 							}
 
-							////////////////////////////////////////////////////////////////////////////
-							// HACK for CreateProjectFromTemplate, 
-							// Fallback to find object by name
-							if (ObjectToFind.Object)
-							{
-								FName ObjectNameToFind = ObjectToFind.Object->GetFName();
-								for (int i = 0; i < ArrayNum; ++i)
-								{
-									if (UObject* CurrentObject = InnerObjectProperty->GetObjectPropertyValue(ArrayHelper.GetElementPtr(i)); CurrentObject && (!bFilterExternalObjects || !CurrentObject->IsPackageExternal()))
-									{
-										if (CurrentObject->GetFName() == ObjectNameToFind)
-										{
-											return i;
-										}
-									}
-								}
-							}
-							// HACK for CreateProjectFromTemplate, fallback to find object by name
-							////////////////////////////////////////////////////////////////////////////
-
 							return INDEX_NONE;
 						};
 
@@ -518,17 +481,18 @@ void FArrayProperty::SerializeItem(FStructuredArchive::FSlot Slot, void* Value,
 								{
 								case EOverriddenPropertyOperation::Remove:
 									{
-										// We don't support removing external entities from overriden arrays yet
-										checkf(!Pair.Key.Object || !Pair.Key.Object->IsPackageExternal(), TEXT("Unsupported remove operation on item %s(0x%p)"), *GetNameSafe(Pair.Key.Object), Pair.Key.Object.Get());
-
 										const int32 DefaultIndex = FindObject(Pair.Key, DefaultsArrayHelper);
 										if (DefaultIndex != INDEX_NONE)
 										{
+											// We don't support removing external entities from overriden arrays yet
+											UObject* DefaultObject = InnerObjectProperty->GetObjectPropertyValue(DefaultsArrayHelper.GetElementPtr(DefaultIndex));
+											checkf(DefaultObject && !DefaultObject->IsPackageExternal(), TEXT("Unsupported remove operation on item %s(0x%p)"), *GetNameSafe(DefaultObject), DefaultObject);
+
 											RemovedIndices.Add(DefaultIndex);
 										}
 										else
 										{
-											UE_LOG(LogOverridableObject, VeryVerbose, TEXT("Unable to save deleted item %s(0x%p)"), *GetNameSafe(Pair.Key.Object), Pair.Key.Object.Get());
+											UE_LOG(LogOverridableObject, VeryVerbose, TEXT("Unable to save deleted item %s"), *Pair.Key.ToDebugString());
 										}
 										break;
 									}
@@ -537,14 +501,22 @@ void FArrayProperty::SerializeItem(FStructuredArchive::FSlot Slot, void* Value,
 										const int32 Index = FindObject(Pair.Key, ArrayHelper);
 										if (Index != INDEX_NONE)
 										{
-											if (!bFilterExternalObjects || !Pair.Key.Object->IsPackageExternal())
+											if (!bFilterExternalObjects)
 											{
 												AddedIndices.Add(Index);
 											}
+											else
+											{
+												UObject* Object = InnerObjectProperty->GetObjectPropertyValue(ArrayHelper.GetElementPtr(Index));
+												if(Object && !Object->IsPackageExternal())
+												{
+													AddedIndices.Add(Index);
+												}
+											}
 										}
 										else
 										{
-											UE_LOG(LogOverridableObject, VeryVerbose, TEXT("Unable to save added item %s(0x%p)"), *GetNameSafe(Pair.Key.Object), Pair.Key.Object.Get());
+											UE_LOG(LogOverridableObject, VeryVerbose, TEXT("Unable to save added item %s"), *Pair.Key.ToDebugString());
 										}
 										break;
 									}
diff --git a/Engine/Source/Runtime/CoreUObject/Public/UObject/OverriddenPropertySet.h b/Engine/Source/Runtime/CoreUObject/Public/UObject/OverriddenPropertySet.h
--- a/Engine/Source/Runtime/CoreUObject/Public/UObject/OverriddenPropertySet.h
+++ b/Engine/Source/Runtime/CoreUObject/Public/UObject/OverriddenPropertySet.h
@@ -203,6 +203,7 @@ struct FOverriddenPropertyNodeID
 		const int32 ObjectIndex = GUObjectArray.ObjectToIndex(&InObject);
 		Path = *(FString::Printf(TEXT("%d%d"), ObjectIndex, GUObjectArray.AllocateSerialNumber(ObjectIndex)));
 	}
+	FOverriddenPropertyNodeID(const FOverriddenPropertyNodeID& ParentNodeID, const FOverriddenPropertyNodeID& SubNodeID);
 
 	static COREUOBJECT_API FOverriddenPropertyNodeID RootNodeId();
 
@@ -210,16 +211,18 @@ struct FOverriddenPropertyNodeID
 	static COREUOBJECT_API FOverriddenPropertyNodeID FromMapKey(const FProperty* KeyProperty, const void* KeyData);
 	COREUOBJECT_API int32 ToMapInternalIndex(FScriptMapHelper& MapHelper) const;
 
-	bool operator==(const FOverriddenPropertyNodeID& Other) const
-	{
-		return Path == Other.Path || (Object && Other.Object && Object == Other.Object);
-	}
+	COREUOBJECT_API bool operator==(const FOverriddenPropertyNodeID& Other) const;
 
 	FString ToString() const
 	{
 		return Path.ToString();
 	}
 
+	FString ToDebugString() const
+	{
+		return Path.ToString() + FString::Printf(TEXT("(0x%p)"), Object.Get());
+	}
+
 	bool IsValid() const
 	{
 		return !Path.IsNone();
@@ -227,16 +230,23 @@ struct FOverriddenPropertyNodeID
 
 	void HandleObjectsReInstantiated(const TMap<UObject*, UObject*>& Map);
 
+	friend uint32 GetTypeHash(const FOverriddenPropertyNodeID& NodeID)
+	{
+		return GetTypeHash(NodeID.Path);
+	}
+
+private:
 	UPROPERTY()
 	FName Path;
 
+	/* NOTE: Not always valid can point to a destroyed or can be a stale ptr. Use with cautiousness
+	 * This can happen during reinstantiation. It is only there for ptr comparison in the ==.
+	 * There is no unique id kept between old and new when an object is reinstantiated, so we are keeping the ptr here.
+	 * When it is a ptr of an archetype, there are some cases we do not patch it because we need it in the CPFUO.
+	 * Check FOverriddenPropertySet::HandleObjectsReInstantiated special exception
+	 */
 	UPROPERTY()
 	TObjectPtr<const UObject> Object;
-
-	friend uint32 GetTypeHash(const FOverriddenPropertyNodeID& NodeID)
-	{
-		return GetTypeHash(NodeID.Path);
-	}
 };
 
 /*
@@ -251,7 +261,7 @@ struct FOverriddenPropertyNode
 {
 	GENERATED_BODY()
 
-	FOverriddenPropertyNode(FOverriddenPropertyNodeID InNodeID = FOverriddenPropertyNodeID())
+	FOverriddenPropertyNode(const FOverriddenPropertyNodeID& InNodeID = FOverriddenPropertyNodeID())
 		: NodeID(InNodeID)
 	{}
 
