Improved error messages for failure to construct a default object for an unfulfilled non-nullable property.  Currently, it gives an error about UObject being abstract when it's actually due to an interface property via ObjectMustImplement metadata.

diff --git a/Engine/Source/Runtime/CoreUObject/Private/UObject/PropertyBaseObject.cpp b/Engine/Source/Runtime/CoreUObject/Private/UObject/PropertyBaseObject.cpp
--- a/Engine/Source/Runtime/CoreUObject/Private/UObject/PropertyBaseObject.cpp
+++ b/Engine/Source/Runtime/CoreUObject/Private/UObject/PropertyBaseObject.cpp
@@ -723,7 +723,7 @@ bool FObjectPropertyBase::AllowObjectTypeReinterpretationTo(const FObjectPropert
 	return false;
 }
 
-static bool CanPropertyHoldObject(const FObjectPropertyBase* Prop, TObjectPtr<const UObject> Obj)
+static bool CanPropertyHoldObject(const FObjectPropertyBase* Prop, TObjectPtr<const UObject> Obj, FString* FailReason = nullptr)
 {
 	UClass* PropertyClass = Prop->PropertyClass;
 	checkSlow(PropertyClass);
@@ -731,35 +731,49 @@ static bool CanPropertyHoldObject(const FObjectPropertyBase* Prop, TObjectPtr<co
 	UClass* ObjectClass = Obj.GetClass();
 	checkSlow(ObjectClass);
 
+	UClass* InterfaceClassToCheck = nullptr;
+
 #if WITH_METADATA
 	// Check if the object has metadata that states that the object must implement a particular interface
 	UClass* MetaDataClass = Prop->GetOwnerProperty()->GetClassMetaData(TEXT("ObjectMustImplement"));
 	if (MetaDataClass)
 	{
-		if (ObjectClass->ImplementsInterface(MetaDataClass))
-		{
-			return true;
-		}
+		InterfaceClassToCheck = MetaDataClass;
 	}
 	else
 #endif // WITH_METADATA
 	{
 		// Check if the object is or is derived from the property class
-		if (ObjectClass->IsChildOf(PropertyClass) || ObjectClass->GetAuthoritativeClass()->IsChildOf(PropertyClass))
+		if (!ObjectClass->IsChildOf(PropertyClass) && !ObjectClass->GetAuthoritativeClass()->IsChildOf(PropertyClass))
 		{
-			return true;
+			if (!PropertyClass->IsChildOf<UInterface>())
+			{
+				if (FailReason)
+				{
+					*FailReason = FString::Printf(TEXT("object '%s' is incompatible with object property of type '%s'"), *ObjectClass->GetName(), *PropertyClass->GetName());
+				}
+				return false;
+			}
+
+			// Is it even possible to have an object property with an interface class set?  Check it anyway.
+			InterfaceClassToCheck = PropertyClass;
 		}
-		// Check if the object implements the interface - is it possible to have an object property with an interface class set?
-		if (ObjectClass->ImplementsInterface(PropertyClass))
+	}
+
+	// Check if the object implements the interface
+	if (InterfaceClassToCheck && !ObjectClass->ImplementsInterface(InterfaceClassToCheck))
+	{
+		if (FailReason)
 		{
-			return true;
+			*FailReason = FString::Printf(TEXT("class '%s' does not implement interface '%s'"), *ObjectClass->GetName(), *InterfaceClassToCheck->GetName());
 		}
+		return false;
 	}
 
-	return false;
+	return true;
 }
 
-UObject* FObjectPropertyBase::ConstructDefaultObjectValueIfNecessary(UObject* ExistingValue) const
+UObject* FObjectPropertyBase::ConstructDefaultObjectValueIfNecessary(UObject* ExistingValue, FString* OutFailureReason) const
 {
 	UObject* NewDefaultObjectValue = nullptr;
 
@@ -792,16 +806,31 @@ UObject* FObjectPropertyBase::ConstructDefaultObjectValueIfNecessary(UObject* Ex
 	}
 
 	// We cannot know which object to create if the class is abstract or if there is a required interface
-	if (!NewDefaultObjectValue && !PropertyClass->HasAnyClassFlags(CLASS_Abstract)
+	if (!NewDefaultObjectValue)
+	{
+		if (PropertyClass->HasAnyClassFlags(CLASS_Abstract))
+		{
+			if (OutFailureReason)
+			{
+				*OutFailureReason = FString::Printf(TEXT("class '%s' is abstract"), *PropertyClass->GetName());
+			}
+		}
 #if WITH_METADATA
-		&& !this->GetOwnerProperty()->GetClassMetaData(TEXT("ObjectMustImplement"))
+		else if (UClass* MetaDataClass = this->GetOwnerProperty()->GetClassMetaData(TEXT("ObjectMustImplement")))
+		{
+			if (OutFailureReason)
+			{
+				*OutFailureReason = FString::Printf(TEXT("interface '%s' cannot be instantiated"), *MetaDataClass->GetName());
+			}
+		}
 #endif // WITH_METADATA
-	)
-	{
-		// Existing value did not exist or it could not be used as a template
-		// Existing value may be null in case we were serializing an array of UObjects that failed to load (missing import class). Since the array is first pre-allocated with null values
-		// it will not have any existing objects to instantiate
-		NewDefaultObjectValue = NewObject<UObject>(Outer, PropertyClass);
+		else
+		{
+			// Existing value did not exist or it could not be used as a template
+			// Existing value may be null in case we were serializing an array of UObjects that failed to load (missing import class). Since the array is first pre-allocated with null values
+			// it will not have any existing objects to instantiate
+			NewDefaultObjectValue = NewObject<UObject>(Outer, PropertyClass);
+		}
 	}
 
 	return NewDefaultObjectValue;
@@ -846,7 +875,8 @@ void FObjectPropertyBase::CheckValidObject(void* ValueAddress, TObjectPtr<UObjec
 #endif // USE_CIRCULAR_DEPENDENCY_LOAD_DEFERRING
 
 	// Is a null class even possible here?
-	if ((PropertyClass != nullptr) && !CanPropertyHoldObject(this, Object))
+	FString HoldFailReasonString;
+	if ((PropertyClass != nullptr) && !CanPropertyHoldObject(this, Object, &HoldFailReasonString))
 	{
 		// we could be in the middle of replacing references to the 
 		// PropertyClass itself (in the middle of an FArchiveReplaceObjectRef 
@@ -883,17 +913,23 @@ void FObjectPropertyBase::CheckValidObject(void* ValueAddress, TObjectPtr<UObjec
 				using UE::CoreUObject::Private::GetNonNullableBehavior;
 
 				ENonNullableBehavior NonNullableBehavior = GetNonNullableBehavior();
+				FString DefaultValueFailureReason;
 				if (NonNullableBehavior == ENonNullableBehavior::CreateDefaultObjectIfPossible)
 				{
-					DefaultValue = ConstructDefaultObjectValueIfNecessary(OldValue);
+					DefaultValue = ConstructDefaultObjectValueIfNecessary(OldValue, &DefaultValueFailureReason);
+					if (!DefaultValueFailureReason.IsEmpty())
+					{
+						DefaultValueFailureReason.InsertAt(0, TEXT(" as "));
+					}
 				}
 
 				if (DefaultValue)
 				{
 					UE_LOG(LogProperty, Warning,
-						TEXT("Serialized %s for a non-nullable property of %s. Reference will be defaulted to %s (previously: %s).\n    ReferencingObject = %s\n    Property = %s\n    Item = %s"),
+						TEXT("Serialized %s for a non-nullable property of %s but %s. Reference will be defaulted to %s (previously: %s).\n    ReferencingObject = %s\n    Property = %s\n    Item = %s"),
 						*ObjectClass->GetFullName(),
 						*PropertyClass->GetFullName(),
+						*HoldFailReasonString,
 						*GetFullNameSafe(DefaultValue),
 						*GetFullNameSafe(OldValue),
 						*GetFullNameSafe(Outer),
@@ -904,9 +940,11 @@ void FObjectPropertyBase::CheckValidObject(void* ValueAddress, TObjectPtr<UObjec
 				else if (NonNullableBehavior == ENonNullableBehavior::LogWarning)
 				{
 					UE_LOG(LogProperty, Warning,
-						TEXT("Serialized %s for a non-nullable property of %s. Reference will be nulled (previously: %s) - will cause a runtime error if accessed.\n    ReferencingObject = %s\n    Property = %s\n    Item = %s"),
+						TEXT("Serialized %s for a non-nullable property of %s but %s. Reference will be nulled%s (previously: %s) - will cause a runtime error if accessed.\n    ReferencingObject = %s\n    Property = %s\n    Item = %s"),
 						*ObjectClass->GetFullName(),
 						*PropertyClass->GetFullName(),
+						*HoldFailReasonString,
+						*DefaultValueFailureReason,
 						*GetFullNameSafe(OldValue),
 						*GetFullNameSafe(Outer),
 						*GetFullName(),
@@ -916,10 +954,11 @@ void FObjectPropertyBase::CheckValidObject(void* ValueAddress, TObjectPtr<UObjec
 				else
 				{
 					UE_LOG(LogProperty, Error,
-						TEXT("Serialized %s for a non-nullable property of %s. Reference will be nulled%s (previously: %s) - will cause a runtime error if accessed.\n    ReferencingObject = %s\n    Property = %s\n    Item = %s"),
+						TEXT("Serialized %s for a non-nullable property of %s but %s. Reference will be nulled%s (previously: %s) - will cause a runtime error if accessed.\n    ReferencingObject = %s\n    Property = %s\n    Item = %s"),
 						*ObjectClass->GetFullName(),
 						*PropertyClass->GetFullName(),
-						(NonNullableBehavior == ENonNullableBehavior::CreateDefaultObjectIfPossible) ? *FString::Printf(TEXT(" as %s is abstract"), *PropertyClass->GetName()) : TEXT(""),
+						*HoldFailReasonString,
+						*DefaultValueFailureReason,
 						*GetFullNameSafe(OldValue),
 						*GetFullNameSafe(Outer),
 						*GetFullName(),
diff --git a/Engine/Source/Runtime/CoreUObject/Private/UObject/PropertyObject.cpp b/Engine/Source/Runtime/CoreUObject/Private/UObject/PropertyObject.cpp
--- a/Engine/Source/Runtime/CoreUObject/Private/UObject/PropertyObject.cpp
+++ b/Engine/Source/Runtime/CoreUObject/Private/UObject/PropertyObject.cpp
@@ -258,9 +258,14 @@ void FObjectProperty::PostSerializeObjectItem(FArchive& SerializingArchive, void
 		using UE::CoreUObject::Private::GetNonNullableBehavior;
 
 		ENonNullableBehavior NonNullableBehavior = GetNonNullableBehavior();
+		FString DefaultValueFailureReason;
 		if (NonNullableBehavior == ENonNullableBehavior::CreateDefaultObjectIfPossible)
 		{
-			DefaultValue = ConstructDefaultObjectValueIfNecessary(CurrentValue);
+			DefaultValue = ConstructDefaultObjectValueIfNecessary(CurrentValue, &DefaultValueFailureReason);
+			if (!DefaultValueFailureReason.IsEmpty())
+			{
+				DefaultValueFailureReason.InsertAt(0, TEXT(" as "));
+			}
 		}
 
 		if (DefaultValue)
@@ -275,17 +280,19 @@ void FObjectProperty::PostSerializeObjectItem(FArchive& SerializingArchive, void
 		else if (NonNullableBehavior == ENonNullableBehavior::LogWarning)
 		{
 			UE_LOG(LogProperty, Warning,
-				TEXT("Failed to serialize value for non-nullable property %s (previously: %s). Reference will be nulled - will cause a runtime error if accessed."),
+				TEXT("Failed to serialize value for non-nullable property %s (previously: %s). Reference will be nulled%s - will cause a runtime error if accessed."),
 				*GetFullNameSafe(CurrentValue),
-				*GetFullName()
+				*GetFullName(),
+				*DefaultValueFailureReason
 			);
 		}
 		else
 		{
 			UE_LOG(LogProperty, Error,
-				TEXT("Failed to serialize value for non-nullable property %s (previously: %s). Reference will be nulled - will cause a runtime error if accessed."),
+				TEXT("Failed to serialize value for non-nullable property %s (previously: %s). Reference will be nulled%s - will cause a runtime error if accessed."),
 				*GetFullNameSafe(CurrentValue),
-				*GetFullName()
+				*GetFullName(),
+				*DefaultValueFailureReason
 			);
 		}
 
diff --git a/Engine/Source/Runtime/CoreUObject/Public/UObject/UnrealType.h b/Engine/Source/Runtime/CoreUObject/Public/UObject/UnrealType.h
--- a/Engine/Source/Runtime/CoreUObject/Public/UObject/UnrealType.h
+++ b/Engine/Source/Runtime/CoreUObject/Public/UObject/UnrealType.h
@@ -2658,9 +2658,10 @@ protected:
 	 * Constructs a new object if the existing one is missing or is not compatible with the property class
 	 * Used for making sure non-nullable properties have valid values.
 	 * @param ExistingValue Previous object value (can be null)
-	 * @return Non-null object that was assigned to this property value address
+	 * @param OutFailureReason An optional output string describing the reason why null was returned.
+	 * @return Pointer that should be assigned to the property value, or null if a default value could not be created.
 	 */
-	UObject* ConstructDefaultObjectValueIfNecessary(UObject* ExistingValue) const;
+	UObject* ConstructDefaultObjectValueIfNecessary(UObject* ExistingValue, FString* OutFailureReason = nullptr) const;
 
 	// Disable false positive buffer overrun warning during pgoprofile linking step
 	PRAGMA_DISABLE_BUFFER_OVERRUN_WARNING
