CIS fix.

#rnx

diff --git a/Engine/Source/Runtime/CoreUObject/Private/UObject/PropertyBaseObject.cpp b/Engine/Source/Runtime/CoreUObject/Private/UObject/PropertyBaseObject.cpp
--- a/Engine/Source/Runtime/CoreUObject/Private/UObject/PropertyBaseObject.cpp
+++ b/Engine/Source/Runtime/CoreUObject/Private/UObject/PropertyBaseObject.cpp
@@ -5,6 +5,7 @@
 #include "UObject/NonNullPropertyUtils.h"
 #include "UObject/ObjectMacros.h"
 #include "Templates/Casts.h"
+#include "UObject/Interface.h"
 #include "UObject/Package.h"
 #include "UObject/UnrealType.h"
 #include "UObject/UnrealTypePrivate.h"
