dstm:
- reworked how ExecuteTransactional handles deferred RPC requests. We no longer call ExecuteTransactional recursively which could cause very large callstacks. Now we maintain a local list of work. We prime the list with the original work and then keep filling it in with higher or equal priority work that is discovered along the way. After we finish everything, we then also drain any remaining deferred work of any priority before finishing.
- we now also execute deferred RPC requests in reverse order, because (in the future) requests might have dependencies, and requests made later might be required in order to fulfill requests made earlier. so, we process from highest priority to lowest, but within a priority we process in last-in-first-out order
- tagged a few TArray Remove calls with EAllowShrinking::No and changed some Empty calls to Reset to avoid extra allocator work
- added a SetAbortReason api to send in a string for development so the logs can be more descriptive

diff --git a/Engine/Source/Runtime/CoreUObject/Private/UObject/RemoteObject.cpp b/Engine/Source/Runtime/CoreUObject/Private/UObject/RemoteObject.cpp
--- a/Engine/Source/Runtime/CoreUObject/Private/UObject/RemoteObject.cpp
+++ b/Engine/Source/Runtime/CoreUObject/Private/UObject/RemoteObject.cpp
@@ -114,9 +114,9 @@ void InitRemoteObjects()
 	{
 		UE::RemoteObject::Transfer::RestoreRemoteObjectDataDelegate.BindStatic(&UE::RemoteObject::Serialization::Disk::LoadObjectFromDisk);
 	}
-	if (!UE::RemoteObject::Transfer::AnyDeferredRPCQueuesExecutableDelegate.IsBound())
+	if (!UE::RemoteObject::Transfer::FetchDeferredRPCDelegate.IsBound())
 	{
-		UE::RemoteObject::Transfer::AnyDeferredRPCQueuesExecutableDelegate.BindStatic([](FRemoteServerId){ return false; });
+		UE::RemoteObject::Transfer::FetchDeferredRPCDelegate.BindStatic([](FRemoteServerId) { return TTuple<FRemoteServerId, TFunction<void(void)>>(); });
 	}
 	if (!UE::RemoteObject::Transfer::PhysicsQueryTickDelegate.IsBound())
 	{
diff --git a/Engine/Source/Runtime/CoreUObject/Private/UObject/RemoteObjectTransfer.cpp b/Engine/Source/Runtime/CoreUObject/Private/UObject/RemoteObjectTransfer.cpp
--- a/Engine/Source/Runtime/CoreUObject/Private/UObject/RemoteObjectTransfer.cpp
+++ b/Engine/Source/Runtime/CoreUObject/Private/UObject/RemoteObjectTransfer.cpp
@@ -24,9 +24,8 @@ TDelegate<void(uint32 RequestToken, FRemoteServerId RequestPriority, FRemoteObje
 TDelegate<void(uint32 RequestToken, FRemoteObjectId ObjectId, FRemoteServerId DestinationServerId, FRemoteObjectData& ObjectData)> RemoteObjectTransferDelegate;
 TDelegate<void(uint32 RequestToken, FRemoteObjectId ObjectId, FRemoteServerId DestinationServerId)> RemoteObjectDeniedTransferDelegate;
 TDelegate<void()> TickNetworkDelegate;
-TDelegate<bool(FRemoteServerId)> AnyDeferredRPCQueuesExecutableDelegate;
-TDelegate<void(FRemoteServerId)> ExecuteDeferredRPCQueuesDelegate;
-TDelegate<void()> ExecuteAllDeferredRPCQueuesDelegate;
+TDelegate<TTuple<FRemoteServerId, TFunction<void(void)>>(FRemoteServerId)> FetchDeferredRPCDelegate;
+TDelegate<TTuple<FRemoteServerId, TFunction<void(void)>>(void)> FetchNextDeferredRPCDelegate;
 TDelegate<void()> WaitForServersToSyncDelegate;
 
 TMulticastDelegate<void(const FRemoteObjectData& ObjectData, FRemoteServerId RemoteServerId)> OnObjectDataSentDelegate;
@@ -44,7 +43,9 @@ TDelegate<void()> PhysicsQueryEndRequestDelegate;
 
 TMulticastDelegate<void()> OnTransactionCompletedDelegate;
 TMulticastDelegate<void()> OnTransactionAbortedDelegate;
-	
+
+FString DebugLastAbortReason;
+
 UObject* MigrateObjectFromObjectData(FRemoteObjectData& ObjectData, FRemoteServerId RemoteServerId);
 
 /**
@@ -95,8 +96,8 @@ public:
 		check(PendingRequests == 0);
 
 		// clear the object lists
-		LockedObjects.Empty();
-		RequestedObjects.Empty();
+		LockedObjects.Reset();
+		RequestedObjects.Reset();
 
 		// clear the priority
 		ActiveRequestPriority = FRemoteServerId();
@@ -267,6 +268,7 @@ void MigrateObjectFromRemoteServer(FRemoteObjectId ObjectId, FRemoteServerId Cur
 	if (AutoRTFM::IsClosed())
 	{
 		// if we are in a transaction, abort so the outer handler can renegotiate object transfers and retry
+		SetAbortReason(TEXT("object migration"));
 		AutoRTFM::AbortTransaction();
 	}
 	else
@@ -318,7 +320,7 @@ void MigrateObjectToRemoteServerWithExplicitPriority(uint32 RequestToken, FRemot
 				UE_LOG(LogRemoteObject, Verbose, TEXT("     : MigrateObjectToRemoteServer(%u) GIVING UP our locked %s"), RequestToken, *LockedObjectId.ToString());
 
 				// unlock the object from ourselves
-				GTransferQueue.LockedObjects.RemoveAt(LockedObjectIndex);
+				GTransferQueue.LockedObjects.RemoveAt(LockedObjectIndex, EAllowShrinking::No);
 			}
 
 			break;
@@ -472,24 +474,45 @@ void ExecuteTransactional(const TFunctionRef<void(void)>& Work)
 	ExecuteTransactionalWithExplicitPriority(GetGlobalServerId(), Work);
 }
 
-void ExecuteTransactionalInternal(FRemoteServerId WorkPriority, const TFunctionRef<void(void)>& Work)
+void ExecuteTransactionalInternal(FRemoteServerId InWorkPriority, const TFunctionRef<void(void)>& InWork)
 {
 	check(!AutoRTFM::IsClosed());
 
-	for (;;)
+	TArray<TTuple<FRemoteServerId, TFunction<void(void)>>> PendingWorks;
+	// it's not possible to construct a TFunction from a TFunctionRef, but it's 
+	// safe to do in this context so we just wrap the TFunctionRef in a lambda
+	PendingWorks.Add(TTuple<FRemoteServerId, TFunction<void(void)>>(InWorkPriority, [&InWork](){ InWork(); }));
+
+	uint32 LocalIterationNumber = 0;
+
+	while (PendingWorks.Num() > 0)
 	{
-		// first check to see if there's any higher priority requests we should be servicing
-		if (AnyDeferredRPCQueuesExecutableDelegate.Execute(WorkPriority))
+		LocalIterationNumber++;
+
+		// first check for higher priority work and insert it at the front
+		// this causes requests received later to be put at the start, because
+		// earlier requests could be dependent on later ones completing
+		for(;;)
 		{
-			// do those first...
-			ExecuteDeferredRPCQueuesDelegate.ExecuteIfBound(WorkPriority);
+			TTuple<FRemoteServerId, TFunction<void(void)>> NewWork = FetchDeferredRPCDelegate.Execute(PendingWorks[0].Get<0>());
+			if(NewWork.Get<0>().IsValid())
+			{
+				PendingWorks.Insert(NewWork, 0);
+			}
+			else
+			{
+				break;
+			}
 		}
 
-		GTransferQueue.BeginRequest(WorkPriority);
+		TTuple<FRemoteServerId, TFunction<void(void)>> ExecutingWork = PendingWorks[0];
+		FRemoteServerId ExecutingWorkPriority = ExecutingWork.Get<0>();
+
+		GTransferQueue.BeginRequest(ExecutingWorkPriority);
 
 		PhysicsQueryBeginRequestDelegate.ExecuteIfBound();
 
-		UE_LOG(LogRemoteObject, VeryVerbose, TEXT("-- ExecuteTransactional : request %u --"), GTransferQueue.ActiveRequestToken.Load());
+		UE_LOG(LogRemoteObject, VeryVerbose, TEXT("-- ExecuteTransactional : iteration %u request %u pri %s --"), LocalIterationNumber, GTransferQueue.ActiveRequestToken.Load(), *ExecutingWorkPriority.ToString());
 
 		// we might run through this loop after the work succeeds in the case
 		// that the work prefetches objects that it doesn't access
@@ -497,7 +520,24 @@ void ExecuteTransactionalInternal(FRemoteServerId WorkPriority, const TFunctionR
 
 		for (;;)
 		{
-			if (AnyDeferredRPCQueuesExecutableDelegate.Execute(WorkPriority))
+			// we need to check to see if there is any higher priority work
+			// that we need to context switch to
+			bool bRequiresContextSwitch = false;
+			for (;;)
+			{
+				TTuple<FRemoteServerId, TFunction<void(void)>> NewWork = FetchDeferredRPCDelegate.Execute(ExecutingWorkPriority);
+				if (NewWork.Get<0>().IsValid())
+				{
+					PendingWorks.Insert(NewWork, 0);
+					bRequiresContextSwitch = true;
+				}
+				else
+				{
+					break;
+				}
+			}
+
+			if (bRequiresContextSwitch)
 			{
 				break;
 			}
@@ -594,12 +634,14 @@ void ExecuteTransactionalInternal(FRemoteServerId WorkPriority, const TFunctionR
 			// all of the requested objects are locked, try to perform the work
 			if (!bWorkSucceeded)
 			{
-				AutoRTFM::ETransactionResult WorkTransactResult = AutoRTFM::Transact([&Work](){ Work(); });
+				DebugLastAbortReason = TEXT("");
+
+				AutoRTFM::ETransactionResult WorkTransactResult = AutoRTFM::Transact([&ExecutingWork](){ ExecutingWork.Get<1>()(); });
 
 				if (WorkTransactResult == AutoRTFM::ETransactionResult::AbortedByRequest)
 				{
-						// the work aborted, it should have inserted a new requested object, retry
-					UE_LOG(LogRemoteObject, VeryVerbose, TEXT("   : work ABORTED"));
+					// the work aborted, it should have inserted a new requested object, retry
+					UE_LOG(LogRemoteObject, VeryVerbose, TEXT("   : work ABORTED (%s)"), *DebugLastAbortReason);
 					OnTransactionAbortedDelegate.Broadcast();
 				}
 				else
@@ -634,7 +676,7 @@ void ExecuteTransactionalInternal(FRemoteServerId WorkPriority, const TFunctionR
 		}
 		else
 		{
-			UE_LOG(LogRemoteObject, VeryVerbose, TEXT("-- ExecuteTransactional %u DONE but execution failed, retrying to service RPC requests --"), GTransferQueue.ActiveRequestToken.Load());
+			UE_LOG(LogRemoteObject, VeryVerbose, TEXT("-- ExecuteTransactional %u ABANDONED, context switching to higher priority work --"), GTransferQueue.ActiveRequestToken.Load());
 		}
 
 		GTransferQueue.EndRequest();
@@ -643,7 +685,21 @@ void ExecuteTransactionalInternal(FRemoteServerId WorkPriority, const TFunctionR
 
 		if (bWorkSucceeded)
 		{
-			break;
+			// remove this work from the list
+			PendingWorks.RemoveAt(0, EAllowShrinking::No);
+		}
+
+		if (PendingWorks.Num() == 0)
+		{
+			// if we have no more pending work, then we are done executing what
+			// we came here to execute. but before we start going back to find more
+			// work for ourselves, try to process any pending RPCs by searching
+			// for any priority available to help unblock other servers
+			TTuple<FRemoteServerId, TFunction<void(void)>> NewWork = FetchDeferredRPCDelegate.Execute(FRemoteServerId());
+			if (NewWork.Get<0>().IsValid())
+			{
+				PendingWorks.Add(NewWork);
+			}
 		}
 	}
 }
@@ -654,6 +710,14 @@ void WaitForServersToSync()
 	WaitForServersToSyncDelegate.ExecuteIfBound();
 }
 
+void SetAbortReason(FStringView Reason)
+{
+	UE_AUTORTFM_OPEN
+	{
+		DebugLastAbortReason = Reason;
+	};
+}
+
 } // namespace UE::RemoteObject::Transfer
 
 namespace UE::RemoteObject::Transfer::Private
diff --git a/Engine/Source/Runtime/CoreUObject/Public/UObject/RemoteObjectTransfer.h b/Engine/Source/Runtime/CoreUObject/Public/UObject/RemoteObjectTransfer.h
--- a/Engine/Source/Runtime/CoreUObject/Public/UObject/RemoteObjectTransfer.h
+++ b/Engine/Source/Runtime/CoreUObject/Public/UObject/RemoteObjectTransfer.h
@@ -19,6 +19,11 @@ namespace UE::RemoteObject::Transfer
 	COREUOBJECT_API void ExecuteTransactional(const TFunctionRef<void(void)>& Work);
 	COREUOBJECT_API void ExecuteTransactionalWithExplicitPriority(FRemoteServerId WorkPriority, const TFunctionRef<void(void)>& Work);
 
+	/**
+	* (Dev API) Sets a friendly string indicating why we are aborting the current transaction
+	*/
+	COREUOBJECT_API void SetAbortReason(FStringView Reason);
+
 	/**
 	* (Advanced API) Adds the remote object to the current transaction's requested object list, causing
 	* the object to be migrated to this server. Does not abort the transaction if the object isn't yet local
@@ -118,13 +123,10 @@ namespace UE::RemoteObject::Transfer
 	/** Delegate that pumps the network and executes RPC handlers */
 	extern COREUOBJECT_API TDelegate<void()> TickNetworkDelegate;
 	
-	/** Check to see if any higher or equal priority work is pending in the RPC queue **/
-	extern COREUOBJECT_API TDelegate<bool(FRemoteServerId)> AnyDeferredRPCQueuesExecutableDelegate;
-	
-	/** Execute higher or equal priority work in the RPC queue **/
-	extern COREUOBJECT_API TDelegate<void(FRemoteServerId)> ExecuteDeferredRPCQueuesDelegate;
-	extern COREUOBJECT_API TDelegate<void()> ExecuteAllDeferredRPCQueuesDelegate;
-	
+	/** Fetch the next deferred RPC request **/
+	extern COREUOBJECT_API TDelegate<TTuple<FRemoteServerId, TFunction<void(void)>>(FRemoteServerId)> FetchDeferredRPCDelegate;
+	extern COREUOBJECT_API TDelegate<TTuple<FRemoteServerId, TFunction<void(void)>>(void)> FetchNextDeferredRPCDelegate;
+
 	/** Delegate that syncs all servers */
 	extern COREUOBJECT_API TDelegate<void()> WaitForServersToSyncDelegate;
 
