Make more of the asset registry work under Verse transactions.

diff --git a/Engine/Source/Runtime/CoreUObject/Private/AssetRegistry/AssetData.cpp b/Engine/Source/Runtime/CoreUObject/Private/AssetRegistry/AssetData.cpp
--- a/Engine/Source/Runtime/CoreUObject/Private/AssetRegistry/AssetData.cpp
+++ b/Engine/Source/Runtime/CoreUObject/Private/AssetRegistry/AssetData.cpp
@@ -13,7 +13,7 @@
 #include "Misc/AsciiSet.h"
 #include "Misc/AssetRegistryInterface.h"
 #include "Misc/PathViews.h"
-#include "Misc/ScopeRWLock.h"
+#include "Misc/TransactionallySafeRWLock.h"
 #include "Serialization/CompactBinary.h"
 #include "Serialization/CompactBinarySerialization.h"
 #include "Serialization/CompactBinaryWriter.h"
@@ -236,7 +236,7 @@ public:
 			FChunkArrayRegistryEntry Entry(MoveTemp(InChunkIDs));
 			const uint32 Hash = GetTypeHash(Entry);
 			{
-				const FReadScopeLock ReadLock(Lock);
+				const UE::TReadScopeLock _(Lock);
 				Index = ChunkArrays.FindIdByHash(Hash, Entry);
 			}
 
@@ -245,7 +245,7 @@ public:
 				// Two threads may hit this point simultaneously and race to insert, and if we just emplace, the second thread will
 				// replace the array inserted by the first. We never want this to happen since that would invalidate array views produced
 				// since, so we have to do the look-up again. Insertions should be rare in general, so the cost shouldn't be noticeable.
-				const FWriteScopeLock WriteLock(Lock);
+				const UE::TWriteScopeLock _(Lock);
 				Index = ChunkArrays.FindIdByHash(Hash, Entry);
 				if (!Index.IsValidId())
 				{
@@ -261,7 +261,7 @@ public:
 		FAssetData::FChunkArrayView ChunkIDs;
 		if (ChunkArrayRegistryHandle.IsValidId())
 		{
-			FReadScopeLock ReadLock(Lock);
+			const UE::TReadScopeLock _(Lock);
 			ChunkIDs = ChunkArrays[ChunkArrayRegistryHandle].GetChunkIDs();
 		}
 		return ChunkIDs;
@@ -281,7 +281,7 @@ public:
 private:
 	// We will only ever add elements to this set, so we can return persistent indices
 	TSet<FChunkArrayRegistryEntry> ChunkArrays;
-	mutable FRWLock Lock;
+	mutable FTransactionallySafeRWLock Lock;
 } GChunkArrayRegistry;
 
 } // namespace UE::AssetRegistry
@@ -1179,7 +1179,7 @@ public:
 		Algo::Sort(InVersions);
 		uint32 Hash = GetTypeHash(InVersions);
 		{
-			FReadScopeLock ScopeLock(Lock);
+			const UE::TReadScopeLock _(Lock);
 			TArray<FPackageCustomVersion>* Existing = RegisteredValues.FindByHash(Hash, InVersions);
 			if (Existing)
 			{
@@ -1192,7 +1192,7 @@ public:
 			}
 		}
 		{
-			FWriteScopeLock ScopeLock(Lock);
+			const UE::TWriteScopeLock _(Lock);
 			TArray<FPackageCustomVersion>& Existing = RegisteredValues.FindOrAddByHash(Hash, MoveTemp(InVersions));
 			Result.Ptr = TConstArrayView<FPackageCustomVersion>(Existing);
 			return Result;
@@ -1201,7 +1201,7 @@ public:
 
 private:
 	TSet<TArray<FPackageCustomVersion>> RegisteredValues;
-	FRWLock Lock;
+	FTransactionallySafeRWLock Lock;
 } GFPackageCustomVersionRegistry;
 
 FPackageCustomVersionsHandle FPackageCustomVersionsHandle::FindOrAdd(TConstArrayView<FCustomVersion> InVersions)
