[AutoRTFM] Make FUObjectHashTables::CriticalSection transactionally safe

This allows us to drop the open around archetype hash tables lock, which is important as it is causing writes in both the open and closed, which can lead to memory corruption on abort.

diff --git a/Engine/Source/Runtime/CoreUObject/Private/UObject/UObjectArchetype.cpp b/Engine/Source/Runtime/CoreUObject/Private/UObject/UObjectArchetype.cpp
--- a/Engine/Source/Runtime/CoreUObject/Private/UObject/UObjectArchetype.cpp
+++ b/Engine/Source/Runtime/CoreUObject/Private/UObject/UObjectArchetype.cpp
@@ -95,72 +95,68 @@ UObject* GetArchetypeFromRequiredInfoImpl(const UClass* Class, const UObject* Ou
 		if (Outer
 			&& Outer->GetClass() != UPackage::StaticClass()) // packages cannot have subobjects
 		{
-			// Validation temporarily disabled until SOL-7251 is investigated
-			UE_AUTORTFM_OPEN_NO_VALIDATION
-			{
-				// Get a lock on the UObject hash tables for the duration of the GetArchetype operation
-				FScopedUObjectHashTablesLock HashTablesLock;
+			// Get a lock on the UObject hash tables for the duration of the GetArchetype operation
+			FScopedUObjectHashTablesLock HashTablesLock;
 
-				UObject* ArchetypeToSearch = nullptr;
+			UObject* ArchetypeToSearch = nullptr;
 #if UE_CACHE_ARCHETYPE
-				ArchetypeToSearch = GetArchetypeImpl(Outer, Policy);
+			ArchetypeToSearch = GetArchetypeImpl(Outer, Policy);
 #if UE_VERIFY_CACHED_ARCHETYPE
-				{
-					UObject* VerifyArchetype = GetArchetypeFromRequiredInfoImpl(Outer->GetClass(), Outer->GetOuter(), Outer->GetFName(), Outer->GetFlags(), bUseUpToDateClass, Policy);
-					checkf(ArchetypeToSearch == VerifyArchetype, TEXT("Cached archetype mismatch, expected: %s, cached: %s"), *GetFullNameSafe(VerifyArchetype), *GetFullNameSafe(ArchetypeToSearch));
-				}
+			{
+				UObject* VerifyArchetype = GetArchetypeFromRequiredInfoImpl(Outer->GetClass(), Outer->GetOuter(), Outer->GetFName(), Outer->GetFlags(), bUseUpToDateClass, Policy);
+				checkf(ArchetypeToSearch == VerifyArchetype, TEXT("Cached archetype mismatch, expected: %s, cached: %s"), *GetFullNameSafe(VerifyArchetype), *GetFullNameSafe(ArchetypeToSearch));
+			}
 #endif // UE_VERIFY_CACHED_ARCHETYPE
 #else
-			
+		
 #if WITH_EDITOR
-				ArchetypeToSearch = Policy ? Policy->GetArchetype(Outer) : nullptr;
+			ArchetypeToSearch = Policy ? Policy->GetArchetype(Outer) : nullptr;
 #endif
 
-				if (!ArchetypeToSearch)
-				{
-					ArchetypeToSearch = GetArchetypeFromRequiredInfoImpl(Outer->GetClass(), Outer->GetOuter(), Outer->GetFName(), Outer->GetFlags(), bUseUpToDateClass, Policy);
-				}
+			if (!ArchetypeToSearch)
+			{
+				ArchetypeToSearch = GetArchetypeFromRequiredInfoImpl(Outer->GetClass(), Outer->GetOuter(), Outer->GetFName(), Outer->GetFlags(), bUseUpToDateClass, Policy);
+			}
 #endif // UE_CACHE_ARCHETYPE
-				UObject* MyArchetype = static_cast<UObject*>(FindObjectWithOuter(ArchetypeToSearch, Class, Name));
-				if (MyArchetype)
-				{
-					Result = MyArchetype; // found that my outers archetype had a matching component, that must be my archetype
-				}
-				else if (!!(ObjectFlags & RF_InheritableComponentTemplate) && Outer->IsA<UClass>())
+			UObject* MyArchetype = static_cast<UObject*>(FindObjectWithOuter(ArchetypeToSearch, Class, Name));
+			if (MyArchetype)
+			{
+				Result = MyArchetype; // found that my outers archetype had a matching component, that must be my archetype
+			}
+			else if (!!(ObjectFlags & RF_InheritableComponentTemplate) && Outer->IsA<UClass>())
+			{
+				const UClass* OuterSuperClass = static_cast<const UClass*>(Outer)->GetSuperClass();
+				for (const UClass* SuperClassArchetype = bUseUpToDateClass && OuterSuperClass ? OuterSuperClass->GetAuthoritativeClass() : OuterSuperClass;
+					SuperClassArchetype && SuperClassArchetype->HasAllClassFlags(CLASS_CompiledFromBlueprint);
+					SuperClassArchetype = SuperClassArchetype->GetSuperClass())
 				{
-					const UClass* OuterSuperClass = static_cast<const UClass*>(Outer)->GetSuperClass();
-					for (const UClass* SuperClassArchetype = bUseUpToDateClass && OuterSuperClass ? OuterSuperClass->GetAuthoritativeClass() : OuterSuperClass;
-						SuperClassArchetype && SuperClassArchetype->HasAllClassFlags(CLASS_CompiledFromBlueprint);
-						SuperClassArchetype = SuperClassArchetype->GetSuperClass())
+					if (GEventDrivenLoaderEnabled && EVENT_DRIVEN_ASYNC_LOAD_ACTIVE_AT_RUNTIME)
 					{
-						if (GEventDrivenLoaderEnabled && EVENT_DRIVEN_ASYNC_LOAD_ACTIVE_AT_RUNTIME)
-						{
-							if (SuperClassArchetype->HasAnyFlags(RF_NeedLoad))
-							{
-								UE_LOG(LogClass, Fatal, TEXT("%s had RF_NeedLoad when searching supers for an archetype of %s in %s"), *GetFullNameSafe(ArchetypeToSearch), *GetFullNameSafe(Class), *GetFullNameSafe(Outer));
-							}
-						}
-						Result = static_cast<UObject*>(FindObjectWithOuter(SuperClassArchetype, Class, Name));
-						// We can have invalid archetypes halfway through the hierarchy, keep looking if it's pending kill or transient
-						if (IsValid(Result) && !Result->HasAnyFlags(RF_Transient))
+						if (SuperClassArchetype->HasAnyFlags(RF_NeedLoad))
 						{
-							break;
+							UE_LOG(LogClass, Fatal, TEXT("%s had RF_NeedLoad when searching supers for an archetype of %s in %s"), *GetFullNameSafe(ArchetypeToSearch), *GetFullNameSafe(Class), *GetFullNameSafe(Outer));
 						}
 					}
+					Result = static_cast<UObject*>(FindObjectWithOuter(SuperClassArchetype, Class, Name));
+					// We can have invalid archetypes halfway through the hierarchy, keep looking if it's pending kill or transient
+					if (IsValid(Result) && !Result->HasAnyFlags(RF_Transient))
+					{
+						break;
+					}
 				}
-				else
+			}
+			else
+			{
+				if (GEventDrivenLoaderEnabled && EVENT_DRIVEN_ASYNC_LOAD_ACTIVE_AT_RUNTIME)
 				{
-					if (GEventDrivenLoaderEnabled && EVENT_DRIVEN_ASYNC_LOAD_ACTIVE_AT_RUNTIME)
+					if (ArchetypeToSearch->HasAnyFlags(RF_NeedLoad))
 					{
-						if (ArchetypeToSearch->HasAnyFlags(RF_NeedLoad))
-						{
-							UE_LOG(LogClass, Fatal, TEXT("%s had RF_NeedLoad when searching for an archetype of %s in %s"), *GetFullNameSafe(ArchetypeToSearch), *GetFullNameSafe(Class), *GetFullNameSafe(Outer));
-						}
+						UE_LOG(LogClass, Fatal, TEXT("%s had RF_NeedLoad when searching for an archetype of %s in %s"), *GetFullNameSafe(ArchetypeToSearch), *GetFullNameSafe(Class), *GetFullNameSafe(Outer));
 					}
-
-					Result = ArchetypeToSearch->GetClass()->FindArchetype(Class, Name);
 				}
-			};
+
+				Result = ArchetypeToSearch->GetClass()->FindArchetype(Class, Name);
+			}
 		}
 
 		if (!Result)
diff --git a/Engine/Source/Runtime/CoreUObject/Private/UObject/UObjectHash.cpp b/Engine/Source/Runtime/CoreUObject/Private/UObject/UObjectHash.cpp
--- a/Engine/Source/Runtime/CoreUObject/Private/UObject/UObjectHash.cpp
+++ b/Engine/Source/Runtime/CoreUObject/Private/UObject/UObjectHash.cpp
@@ -5,6 +5,7 @@
 =============================================================================*/
 
 #include "UObject/UObjectHash.h"
+#include "Misc/TransactionallySafeCriticalSection.h"
 #include "UObject/Class.h"
 #include "UObject/GarbageCollectionGlobals.h"
 #include "UObject/Package.h"
@@ -379,7 +380,7 @@ struct TBucketMapLock
 class FUObjectHashTables
 {
 	/** Critical section that guards against concurrent adds from multiple threads */
-	FCriticalSection CriticalSection;
+	FTransactionallySafeCriticalSection CriticalSection;
 
 public:
 
