Remove obsolete doc-comment about bools being 32-bit.

This comment likely dates back to the usage  of the UBOOL type. The C++ bool type does not truncate high bits when coercing its results.

diff --git a/Engine/Source/Runtime/CoreUObject/Public/UObject/UnrealType.h b/Engine/Source/Runtime/CoreUObject/Public/UObject/UnrealType.h
--- a/Engine/Source/Runtime/CoreUObject/Public/UObject/UnrealType.h
+++ b/Engine/Source/Runtime/CoreUObject/Public/UObject/UnrealType.h
@@ -1216,11 +1216,8 @@ public:
 		PropertyFlags &= ~NewFlags;
 	}
 	/**
-	 * Used to safely check whether any of the passed in flags are set. This is required
-	 * as PropertyFlags currently is a 64 bit data type and bool is a 32 bit data type so
-	 * simply using PropertyFlags&CPF_MyFlagBiggerThanMaxInt won't work correctly when
-	 * assigned directly to an bool.
-	 *
+	 * Used to safely check whether any of the passed in flags are set.
+	 * 
 	 * @param FlagsToCheck	Object flags to check for.
 	 *
 	 * @return				true if any of the passed in flags are set, false otherwise  (including no flags passed in).
@@ -1230,11 +1227,8 @@ public:
 		return (PropertyFlags & FlagsToCheck) != 0 || FlagsToCheck == CPF_AllFlags;
 	}
 	/**
-	 * Used to safely check whether all of the passed in flags are set. This is required
-	 * as PropertyFlags currently is a 64 bit data type and bool is a 32 bit data type so
-	 * simply using PropertyFlags&CPF_MyFlagBiggerThanMaxInt won't work correctly when
-	 * assigned directly to an bool.
-	 *
+	 * Used to safely check whether all of the passed in flags are set.
+	 * 
 	 * @param FlagsToCheck	Object flags to check for
 	 *
 	 * @return true if all of the passed in flags are set (including no flags passed in), false otherwise
