Fixes to the GetOverriddenState not returning the correct information

diff --git a/Engine/Source/Editor/UnrealEd/Private/EditorObject.cpp b/Engine/Source/Editor/UnrealEd/Private/EditorObject.cpp
--- a/Engine/Source/Editor/UnrealEd/Private/EditorObject.cpp
+++ b/Engine/Source/Editor/UnrealEd/Private/EditorObject.cpp
@@ -725,7 +725,7 @@ static const TCHAR* ImportProperties(
 				{
 					if (OverriddenOperation.IsSet())
 					{
-						FOverridableManager::Get().SetOverriddenProperties(*ComponentTemplate, OverriddenOperation.GetValue(), /*bNeedsSubObjectTemplateInstantiation*/true);
+						FOverridableManager::Get().RestoreOverrideOperation(ComponentTemplate, OverriddenOperation.GetValue(), /*bNeedsSubObjectTemplateInstantiation*/true);
 					}
 					else
 					{
@@ -1342,7 +1342,7 @@ static const TCHAR* ImportCreateSubObjectsStep(
 				{
 					if (OverriddenOperation.IsSet())
 					{
-						FOverridableManager::Get().SetOverriddenProperties(*ComponentTemplate, OverriddenOperation.GetValue(), /*bNeedsSubObjectTemplateInstantiation*/true);
+						FOverridableManager::Get().RestoreOverrideOperation(ComponentTemplate, OverriddenOperation.GetValue(), /*bNeedsSubObjectTemplateInstantiation*/true);
 					}
 					else
 					{
diff --git a/Engine/Source/Editor/UnrealEd/Private/Factories/EditorFactories.cpp b/Engine/Source/Editor/UnrealEd/Private/Factories/EditorFactories.cpp
--- a/Engine/Source/Editor/UnrealEd/Private/Factories/EditorFactories.cpp
+++ b/Engine/Source/Editor/UnrealEd/Private/Factories/EditorFactories.cpp
@@ -5387,7 +5387,7 @@ void FCustomizableTextObjectFactory::ProcessBuffer(UObject* InParent, EObjectFla
 				{
 					if (OverriddenOperation.IsSet())
 					{
-						FOverridableManager::Get().SetOverriddenProperties(*CreatedObject, OverriddenOperation.GetValue(), /*bNeedsSubObjectTemplateInstantiation*/true);
+						FOverridableManager::Get().RestoreOverrideOperation(CreatedObject, OverriddenOperation.GetValue(), /*bNeedsSubObjectTemplateInstantiation*/true);
 					}
 					else
 					{
diff --git a/Engine/Source/Runtime/CoreUObject/Private/UObject/Class.cpp b/Engine/Source/Runtime/CoreUObject/Private/UObject/Class.cpp
--- a/Engine/Source/Runtime/CoreUObject/Private/UObject/Class.cpp
+++ b/Engine/Source/Runtime/CoreUObject/Private/UObject/Class.cpp
@@ -1595,7 +1595,7 @@ void UStruct::SerializeVersionedTaggedProperties(FStructuredArchive::FSlot Slot,
 					}
 					if (!ControlContext.OverriddenProperties)
 					{
-						ControlContext.OverriddenProperties = FOverridableManager::Get().SetOverriddenProperties(*(UObject*)Data, Operation, /*bNeedsSubobjectTemplateInstantiation*/true);
+						ControlContext.OverriddenProperties = FOverridableManager::Get().RestoreOverrideOperation((UObject*)Data, Operation, /*bNeedsSubobjectTemplateInstantiation*/true);
 					}
 					else
 					{
@@ -1807,15 +1807,7 @@ void UStruct::SerializeVersionedTaggedProperties(FStructuredArchive::FSlot Slot,
 						{
 							if (FOverriddenPropertySet* OverriddenProperties = FOverridableSerializationLogic::GetOverriddenProperties())
 							{
-								// No need to restore none operations
-								if (Tag.OverrideOperation != EOverriddenPropertyOperation::None)
-								{
-									// Prevent marking as replaced the properties that are always overridden
-									if (Tag.OverrideOperation != EOverriddenPropertyOperation::Replace || !Property->HasAnyPropertyFlags(CPF_ExperimentalAlwaysOverriden))
-									{
-										OverriddenProperties->SetOverriddenPropertyOperation(Tag.OverrideOperation, UnderlyingArchive.GetSerializedPropertyChain(), Property);
-									}
-								}
+								OverriddenProperties->RestoreOverriddenPropertyOperation(Tag.OverrideOperation, UnderlyingArchive.GetSerializedPropertyChain(), Property);
 							}
 						}
 
diff --git a/Engine/Source/Runtime/CoreUObject/Private/UObject/OverridableManager.cpp b/Engine/Source/Runtime/CoreUObject/Private/UObject/OverridableManager.cpp
--- a/Engine/Source/Runtime/CoreUObject/Private/UObject/OverridableManager.cpp
+++ b/Engine/Source/Runtime/CoreUObject/Private/UObject/OverridableManager.cpp
@@ -107,10 +107,18 @@ const FOverriddenPropertySet* FOverridableManager::GetOverriddenProperties(const
 #endif // WITH_EDITORONLY_DATA
 }
 
-FOverriddenPropertySet* FOverridableManager::SetOverriddenProperties(UObject& Object, EOverriddenPropertyOperation Operation, const bool bNeedsSubobjectTemplateInstantiation)
+
+FOverriddenPropertySet* FOverridableManager::RestoreOverrideOperation(TNotNull<UObject*> Object, EOverriddenPropertyOperation Operation, const bool bNeedsSubobjectTemplateInstantiation)
 {
+	// When sub property overrides are going to be restored, it will turn this object back to modified, so let's not restore that operation.
+	// This allows cleanup of the modified state if we endup not overriding any sub properties
+	if (Operation == EOverriddenPropertyOperation::Modified)
+	{
+		Operation = EOverriddenPropertyOperation::None;
+	}
+
 #if WITH_EDITORONLY_DATA
-	FOverriddenPropertySet& ObjectOverriddenProperties = OverriddenObjectAnnotations.FindOrAdd(Object);
+	FOverriddenPropertySet& ObjectOverriddenProperties = OverriddenObjectAnnotations.FindOrAdd(*Object);
 	ObjectOverriddenProperties.Reset();
 	ObjectOverriddenProperties.SetOverriddenPropertyOperation(Operation, /*CurrentPropertyChain*/nullptr, /*Property*/nullptr);
 	ObjectOverriddenProperties.bNeedsSubobjectTemplateInstantiation = bNeedsSubobjectTemplateInstantiation;
@@ -118,11 +126,11 @@ FOverriddenPropertySet* FOverridableManager::SetOverriddenProperties(UObject& Ob
 #else
 	if (bNeedsSubobjectTemplateInstantiation)
 	{
-		NeedsSubobjectTemplateInstantiation.Set(&Object);
+		NeedsSubobjectTemplateInstantiation.Set(Object);
 	}
 	else
 	{
-		NeedsSubobjectTemplateInstantiation.Clear(&Object);
+		NeedsSubobjectTemplateInstantiation.Clear(Object);
 	}
 	return nullptr;
 #endif // WITH_EDITORONLY_DATA
@@ -134,32 +142,65 @@ EOverriddenState FOverridableManager::GetOverriddenState(UObject& Object)
 #if WITH_EDITORONLY_DATA
 	if(const FOverriddenPropertySet* OverriddenProperties = GetOverriddenProperties(Object))
 	{
-		// Consider any object that its template is a CDO as added.
-		if (UObject* Archetype = Object.GetArchetype())
+		// Transient should not affect the overridden state
+		if (Object.HasAnyFlags(RF_Transient))
 		{
-			if (Archetype->HasAnyFlags(RF_ClassDefaultObject))
-			{
-				return EOverriddenState::Added;
-			}
+			return EOverriddenState::NoOverrides;
+		}
+
+		// Looking if the Archetype of the object was a CDO is not working in the case of a delete and readd.
+		// So now we explicitly use the bWasAdded to remember if an object was added.
+		if (OverriddenProperties->WasAdded())
+		{
+			return EOverriddenState::Added;
 		}
 
 		const EOverriddenPropertyOperation Operation = OverriddenProperties->GetOverriddenPropertyOperation((FArchiveSerializedPropertyChain*)nullptr, (FProperty*)nullptr);
-		if (Operation != EOverriddenPropertyOperation::None)
+		switch (Operation)
 		{
-			return Operation == EOverriddenPropertyOperation::Replace ? EOverriddenState::AllOverridden : EOverriddenState::HasOverrides;
+			case EOverriddenPropertyOperation::Replace:
+				return EOverriddenState::AllOverridden;
+			case EOverriddenPropertyOperation::Modified:
+				return EOverriddenState::HasOverrides;
+			case EOverriddenPropertyOperation::None:
+				break;
+			default:
+				checkf(false, TEXT("Unsupported operation(%s) on object(%p:%s)"), *GetOverriddenOperationString(Operation), &Object, *Object.GetName());
 		}
 
-		// Need to check subobjects to 
-		TArray<UObject*> InstancedSubObjects;
-		constexpr bool bIncludeNestedObjects = false;
-		GetObjectsWithOuter(&Object, InstancedSubObjects, bIncludeNestedObjects);
-		for (UObject* InstancedSubObject : InstancedSubObjects)
+		auto GetSubObjectState = [this](UObject& Object, auto& Recurse) -> EOverriddenState
 		{
-			checkf(InstancedSubObject, TEXT("Expecting non null SubObjects"));
-			if (GetOverriddenState(*InstancedSubObject) != EOverriddenState::NoOverrides)
+			// Need to check subobjects
+			for (TPropertyValueIterator<FObjectProperty> ObjPropIt(Object.GetClass(), &Object); ObjPropIt; ++ObjPropIt)
 			{
-				return EOverriddenState::HasOverrides;
+				if (UObject* InstancedSubObject = ObjPropIt.Key()->GetObjectPropertyValue(ObjPropIt.Value()))
+				{
+					if (InstancedSubObject->IsIn(&Object))
+					{
+						if (const FOverriddenPropertySet* OverriddenProperties = GetOverriddenProperties(*InstancedSubObject))
+						{
+							const EOverriddenPropertyOperation Operation = OverriddenProperties->GetOverriddenPropertyOperation((FArchiveSerializedPropertyChain*)nullptr, (FProperty*)nullptr);
+							if (Operation != EOverriddenPropertyOperation::None)
+							{
+								return EOverriddenState::HasOverrides;
+							}
+
+							if (Recurse(*InstancedSubObject, Recurse) != EOverriddenState::NoOverrides)
+							{
+								return EOverriddenState::HasOverrides;
+							}
+						}
+					}
+				}
 			}
+
+			return EOverriddenState::NoOverrides;
+		};
+
+		EOverriddenState SubObjectState = GetSubObjectState(Object, GetSubObjectState);
+		if (SubObjectState != EOverriddenState::NoOverrides)
+		{
+			return SubObjectState;
 		}
 	}
 #endif // WITH_EDITORONLY_DATA
@@ -190,6 +231,9 @@ void FOverridableManager::OverrideInstancedSubObject(const UObject& Object, UObj
 void FOverridableManager::PropagateOverrideToInstancedSubObjects(UObject& Object)
 {
 #if WITH_EDITORONLY_DATA
+	// Do not use TPropertyValueIterator<FObjectProperty> to be consistent with PropagateClearOverridesToInstancedSubObjects
+	// which cannot use TPropertyValueIterator<FObjectProperty> because the object might not have 
+	// the object class setup correctly when this is called from PostInit
 	TArray<UObject*> InstancedSubObjects;
 	constexpr bool bIncludeNestedObjects = false;
 	GetObjectsWithOuter(&Object, InstancedSubObjects, bIncludeNestedObjects);
@@ -286,6 +330,9 @@ void FOverridableManager::ClearInstancedSubObjectOverrides(const UObject& Object
 void FOverridableManager::PropagateClearOverridesToInstancedSubObjects(UObject& Object)
 {
 #if WITH_EDITORONLY_DATA
+	// Cannot use TPropertyValueIterator<FObjectProperty> here because the object might not have 
+	// the object class setup correctly when this is called from PostInit.
+	// @todo: need to figure out why this is a problem. Maybe we should not even this method during PostInit.
 	TArray<UObject*> InstancedSubObjects;
 	constexpr bool bIncludeNestedObjects = false;
 	GetObjectsWithOuter(&Object, InstancedSubObjects, bIncludeNestedObjects);
@@ -312,7 +359,7 @@ void FOverridableManager::SerializeOverriddenProperties(UObject& Object, FStruct
 
 		if (ArchiveState.IsLoading())
 		{
-			OverriddenProperties = SetOverriddenProperties(Object, Operation, /*bNeedsSubobjectTemplateInstantiation*/false);
+			OverriddenProperties = RestoreOverrideOperation(&Object, Operation, /*bNeedsSubobjectTemplateInstantiation*/false);
 			checkf(OverriddenProperties, TEXT("Expecting a overridden property set returned"));
 		}
 
diff --git a/Engine/Source/Runtime/CoreUObject/Private/UObject/OverriddenPropertySet.cpp b/Engine/Source/Runtime/CoreUObject/Private/UObject/OverriddenPropertySet.cpp
--- a/Engine/Source/Runtime/CoreUObject/Private/UObject/OverriddenPropertySet.cpp
+++ b/Engine/Source/Runtime/CoreUObject/Private/UObject/OverriddenPropertySet.cpp
@@ -511,8 +511,8 @@ bool FOverriddenPropertySet::ClearOverriddenProperty(FOverriddenPropertyNode& Pa
 	const void* SubValuePtr = Data;
 	FOverriddenPropertyNode* OverriddenPropertyNode = &ParentPropertyNode;
 	int32 ArrayIndex = INDEX_NONE;
-	TArray<FOverriddenPropertyNodeID> PropertyNodeIDPath;
-	PropertyNodeIDPath.Push(OverriddenPropertyNode->NodeID);
+	TArray<FOverriddenPropertyNodeID> TraversedNodes;
+	TraversedNodes.Push(OverriddenPropertyNode->NodeID);
 	while (PropertyIterator && (!OverriddenPropertyNode || OverriddenPropertyNode->Operation != EOverriddenPropertyOperation::Replace))
 	{
 		ArrayIndex = INDEX_NONE;
@@ -527,7 +527,7 @@ bool FOverriddenPropertySet::ClearOverriddenProperty(FOverriddenPropertyNode& Pa
 			{
 				CurrentOverriddenPropertyNode = OverriddenPropertyNodes.Find(*CurrentPropKey);
 				checkf(CurrentOverriddenPropertyNode, TEXT("Expecting a node"));
-				PropertyNodeIDPath.Push(CurrentOverriddenPropertyNode->NodeID);
+				TraversedNodes.Push(CurrentOverriddenPropertyNode->NodeID);
 			}
 		}
 
@@ -664,29 +664,31 @@ bool FOverriddenPropertySet::ClearOverriddenProperty(FOverriddenPropertyNode& Pa
 		++PropertyIterator;
 	}
 
-	auto CleanupClearedNodes = [this, &PropertyNodeIDPath]()
+	auto CleanupClearedNodes = [this, &TraversedNodes]()
 	{
-		// Need to cleanup up the chain of property nodes if they endup empty
-		FOverriddenPropertyNodeID ChildPropertyNodeID;
-		while (FOverriddenPropertyNode* CurrentPropertyNode = !PropertyNodeIDPath.IsEmpty() ? OverriddenPropertyNodes.Find(PropertyNodeIDPath.Top()) : nullptr)
+		// Go through each traversed property in reversed order to do cleanup
+		// We need to continue the cleanup until there is more overrides than just the one we are removing
+		FOverriddenPropertyNodeID LastCleanedNode;
+		while (FOverriddenPropertyNode* CurrentNode = !TraversedNodes.IsEmpty() ? OverriddenPropertyNodes.Find(TraversedNodes.Top()) : nullptr)
 		{
-			PropertyNodeIDPath.Pop();
-			if (CurrentPropertyNode->SubPropertyNodeKeys.Num() > 1)
+			TraversedNodes.Pop();
+
+			if (LastCleanedNode.IsValid())
 			{
-				// Now need to remove the child from this node
-				if (ChildPropertyNodeID.IsValid())
-				{
-					const FOverriddenPropertyNodeID* NodeToRemove = CurrentPropertyNode->SubPropertyNodeKeys.FindKey(ChildPropertyNodeID);
-					checkf(NodeToRemove, TEXT("Expecting a node"));
-					CurrentPropertyNode->SubPropertyNodeKeys.Remove(*NodeToRemove);
+				const FOverriddenPropertyNodeID* NodeToRemove = CurrentNode->SubPropertyNodeKeys.FindKey(LastCleanedNode);
+				checkf(NodeToRemove, TEXT("Expecting to find the last cleaned node"));
 
-					verifyf(OverriddenPropertyNodes.Remove(ChildPropertyNodeID), TEXT("Expecting the node to be removed"));
+				// In the case there are other overrides, just cleanup that node and stop.
+				if (CurrentNode->SubPropertyNodeKeys.Num() > 1)
+				{
+					CurrentNode->SubPropertyNodeKeys.Remove(*NodeToRemove);
+					verifyf(OverriddenPropertyNodes.Remove(LastCleanedNode), TEXT("Expecting the node to be removed"));
+					break;
 				}
-				break;
 			}
 
-			RemoveOverriddenSubProperties(*CurrentPropertyNode);
-			ChildPropertyNodeID = CurrentPropertyNode->NodeID;
+			RemoveOverriddenSubProperties(*CurrentNode);
+			LastCleanedNode = CurrentNode->NodeID;
 		}
 	};
 
@@ -749,7 +751,15 @@ void FOverriddenPropertySet::NotifyPropertyChange(FOverriddenPropertyNode* Paren
 
 	ON_SCOPE_EXIT
 	{
-		if (bNeedsCleanup && ParentPropertyNode && SubPropertyNode && SubPropertyNode->SubPropertyNodeKeys.IsEmpty())
+		if (!ParentPropertyNode || Notification != EPropertyNotificationType::PostEdit)
+		{
+			return;
+		}
+
+		if (SubPropertyNode && SubPropertyNode->SubPropertyNodeKeys.IsEmpty() &&
+			(bNeedsCleanup || 
+			 SubPropertyNode->Operation == EOverriddenPropertyOperation::None || 
+			 SubPropertyNode->Operation == EOverriddenPropertyOperation::Modified))
 		{
 			FOverriddenPropertyNodeID RemovedNodeID;
 			if (ParentPropertyNode->SubPropertyNodeKeys.RemoveAndCopyValue(Property, RemovedNodeID))
@@ -828,17 +838,10 @@ void FOverriddenPropertySet::NotifyPropertyChange(FOverriddenPropertyNode* Paren
 								FOverriddenPropertyNode& AddedSubObjectNode = FindOrAddNode(*SubPropertyNode, AddedSubObjectID);
 								AddedSubObjectNode.Operation = EOverriddenPropertyOperation::Add;
 
-								// Check if this could be a readd
-								UObject* RemovedSubObjectArchetype = AddedSubObject->GetArchetype();
-								if(RemovedSubObjectArchetype && !RemovedSubObjectArchetype->HasAnyFlags(RF_ClassDefaultObject))
+								// Notify the subobject that it was added
+								if (FOverriddenPropertySet* AddedSubObjectOverriddenProperties = OverridableManager.GetOverriddenProperties(*AddedSubObject))
 								{
-									const FOverriddenPropertyNodeID RemovedSubObjectID(*RemovedSubObjectArchetype);
-									FOverriddenPropertyNodeID RemovedNodeID;
-									if (SubPropertyNode->SubPropertyNodeKeys.RemoveAndCopyValue(RemovedSubObjectID, RemovedNodeID))
-									{
-										verifyf(OverriddenPropertyNodes.Remove(RemovedNodeID), TEXT("Expecting the node to be removed"));
-										bNeedsCleanup = true;
-									}
+									AddedSubObjectOverriddenProperties->bWasAdded = true;
 								}
 							}
 						}
@@ -1298,6 +1301,23 @@ FOverriddenPropertyNode* FOverriddenPropertySet::SetOverriddenPropertyOperation(
 	return SetOverriddenPropertyOperation(Operation, OverriddenPropertyNodes.FindOrAdd(RootNodeID), CurrentPropertyChain, Property);
 }
 
+FOverriddenPropertyNode* FOverriddenPropertySet::RestoreOverriddenPropertyOperation(EOverriddenPropertyOperation Operation, const FArchiveSerializedPropertyChain* CurrentPropertyChain, FProperty* Property)
+{
+	// 'None' and 'Modified' operations are not needed to be restored on a property,
+	// because 'None' is equal to the node not existing and
+	// 'Modified' will be restored when the sub property overrides will be restored successfully
+	if (Operation != EOverriddenPropertyOperation::None && Operation != EOverriddenPropertyOperation::Modified)
+	{
+		// Prevent marking as replaced the properties that are always overridden
+		if (!Property || Operation != EOverriddenPropertyOperation::Replace || !Property->HasAnyPropertyFlags(CPF_ExperimentalAlwaysOverriden))
+		{
+			return SetOverriddenPropertyOperation(Operation, CurrentPropertyChain, Property);
+		}
+	}
+
+	return nullptr;
+}
+
 const FOverriddenPropertyNode* FOverriddenPropertySet::GetOverriddenPropertyNode(const FArchiveSerializedPropertyChain* CurrentPropertyChain) const
 {
 	if (const FOverriddenPropertyNode* RootNode = OverriddenPropertyNodes.Find(RootNodeID))
@@ -1394,6 +1414,23 @@ FOverriddenPropertyNode* FOverriddenPropertySet::SetSubPropertyOperation(EOverri
 	return &OverriddenPropertyNode;
 }
 
+FOverriddenPropertyNode* FOverriddenPropertySet::SetSubObjectOperation(EOverriddenPropertyOperation Operation, FOverriddenPropertyNode& Node, TNotNull<UObject*> SubObject)
+{
+	const FOverriddenPropertyNodeID SubObjectID(*SubObject);
+	FOverriddenPropertyNode* SubObjectNode = SetSubPropertyOperation(Operation, Node, SubObjectID);
+
+	if (Operation == EOverriddenPropertyOperation::Add)
+	{
+		// Notify the subobject that it was added
+		if (FOverriddenPropertySet* AddedSubObjectOverriddenProperties = FOverridableManager::Get().GetOverriddenProperties(*SubObject))
+		{
+			AddedSubObjectOverriddenProperties->bWasAdded = true;
+		}
+	}
+
+	return SubObjectNode;
+}
+
 bool FOverriddenPropertySet::IsCDOOwningProperty(const FProperty& Property) const
 {
 	checkf(Owner, TEXT("Expecting a valid overridable owner"));
diff --git a/Engine/Source/Runtime/CoreUObject/Private/UObject/Property.cpp b/Engine/Source/Runtime/CoreUObject/Private/UObject/Property.cpp
--- a/Engine/Source/Runtime/CoreUObject/Private/UObject/Property.cpp
+++ b/Engine/Source/Runtime/CoreUObject/Private/UObject/Property.cpp
@@ -1742,7 +1742,7 @@ const TCHAR* FProperty::ImportSingleProperty( const TCHAR* Str, void* DestData,
 			{
 				checkf(Path, TEXT("Expecting a path"));
 				FArchiveSerializedPropertyChain Chain = Path->ToSerializedPropertyChain();
-				OverriddenProperties->SetOverriddenPropertyOperation(Operation.GetValue(), &Chain, /*Property*/nullptr);
+				OverriddenProperties->RestoreOverriddenPropertyOperation(Operation.GetValue(), &Chain, /*Property*/nullptr);
 			}
 		}
 
@@ -2013,8 +2013,7 @@ const TCHAR* FProperty::ImportSingleProperty( const TCHAR* Str, void* DestData,
 									if (FOverriddenPropertyNode* ArrayOverriddenPropertyNode = OverriddenProperties ? OverriddenProperties->SetOverriddenPropertyOperation(EOverriddenPropertyOperation::Modified, &Chain, /*Property*/nullptr) : nullptr)
 									{
 										// Rebuild the overridden info
-										const FOverriddenPropertyNodeID RemovedSubObjectID(*RemovedSubObject);
-										OverriddenProperties->SetSubPropertyOperation(EOverriddenPropertyOperation::Remove, *ArrayOverriddenPropertyNode, RemovedSubObjectID);
+										OverriddenProperties->SetSubObjectOperation(EOverriddenPropertyOperation::Remove, *ArrayOverriddenPropertyNode, RemovedSubObject);
 									}
 								}
 							}
@@ -2060,8 +2059,7 @@ const TCHAR* FProperty::ImportSingleProperty( const TCHAR* Str, void* DestData,
 										if (FOverriddenPropertyNode* ArrayOverriddenPropertyNode = OverriddenProperties ? OverriddenProperties->SetOverriddenPropertyOperation(EOverriddenPropertyOperation::Modified, &Chain, /*Property*/nullptr) : nullptr)
 										{
 											// Rebuild the overridden info
-											const FOverriddenPropertyNodeID AddedSubObjectID(*AddedSubObject);
-											OverriddenProperties->SetSubPropertyOperation(EOverriddenPropertyOperation::Add, *ArrayOverriddenPropertyNode, AddedSubObjectID);
+											OverriddenProperties->SetSubObjectOperation(EOverriddenPropertyOperation::Add, *ArrayOverriddenPropertyNode, AddedSubObject);
 										}
 									}
 								}
diff --git a/Engine/Source/Runtime/CoreUObject/Private/UObject/PropertyArray.cpp b/Engine/Source/Runtime/CoreUObject/Private/UObject/PropertyArray.cpp
--- a/Engine/Source/Runtime/CoreUObject/Private/UObject/PropertyArray.cpp
+++ b/Engine/Source/Runtime/CoreUObject/Private/UObject/PropertyArray.cpp
@@ -332,8 +332,7 @@ void FArrayProperty::SerializeItem(FStructuredArchive::FSlot Slot, void* Value,
 							if (FOverriddenPropertyNode* ArrayOverriddenPropertyNode = OverriddenProperties ? OverriddenProperties->SetOverriddenPropertyOperation(EOverriddenPropertyOperation::Modified, UnderlyingArchive.GetSerializedPropertyChain(), /*Property*/nullptr) : nullptr)
 							{
 								// Rebuild the overridden info
-								const FOverriddenPropertyNodeID RemovedSubObjectID(*RemovedSubObject);
-								OverriddenProperties->SetSubPropertyOperation(EOverriddenPropertyOperation::Remove, *ArrayOverriddenPropertyNode, RemovedSubObjectID);
+								OverriddenProperties->SetSubObjectOperation(EOverriddenPropertyOperation::Remove, *ArrayOverriddenPropertyNode, RemovedSubObject);
 							}
 						}
 					}
@@ -406,8 +405,7 @@ void FArrayProperty::SerializeItem(FStructuredArchive::FSlot Slot, void* Value,
 							if (FOverriddenPropertyNode* ArrayOverriddenPropertyNode = OverriddenProperties ? OverriddenProperties->SetOverriddenPropertyOperation(EOverriddenPropertyOperation::Modified, UnderlyingArchive.GetSerializedPropertyChain(), /*Property*/nullptr) : nullptr)
 							{
 								// Rebuild the overridden info
-								const FOverriddenPropertyNodeID AddedSubObjectID(*AddedSubObject);
-								OverriddenProperties->SetSubPropertyOperation(EOverriddenPropertyOperation::Add, *ArrayOverriddenPropertyNode, AddedSubObjectID);
+								OverriddenProperties->SetSubObjectOperation(EOverriddenPropertyOperation::Add, *ArrayOverriddenPropertyNode, AddedSubObject);
 							}
 						}
 					}
diff --git a/Engine/Source/Runtime/CoreUObject/Public/UObject/OverridableManager.h b/Engine/Source/Runtime/CoreUObject/Public/UObject/OverridableManager.h
--- a/Engine/Source/Runtime/CoreUObject/Public/UObject/OverridableManager.h
+++ b/Engine/Source/Runtime/CoreUObject/Public/UObject/OverridableManager.h
@@ -159,12 +159,13 @@ public:
 	COREUOBJECT_API const FOverriddenPropertySet* GetOverriddenProperties(const UObject& Object);
 
 	/**
-	 * Set the override operation on this object and it will enable it if it wasn't already enabled
+	 * Restore the override operation from a saved state on this object, and it will enable it if it wasn't already enabled
+	 * @note: It will not restore a Modified state has this will be done automatically when sub properties overrides are restored
 	 * @param Object to set the override operation on
 	 * @param Operation the override operation to set on the object
 	 * @param bNeedsSubobjectTemplateInstantiation set to true if it will need a sub object template instantiation during postload, false will not do it.
 	 * @return the overridden properties of the object */
-	COREUOBJECT_API FOverriddenPropertySet* SetOverriddenProperties(UObject& Object, EOverriddenPropertyOperation Operation, const bool bNeedsSubobjectTemplateInstantiation);
+	COREUOBJECT_API FOverriddenPropertySet* RestoreOverrideOperation(TNotNull<UObject*> Object, EOverriddenPropertyOperation Operation, const bool bNeedsSubobjectTemplateInstantiation);
 
 	/**
 	 * Retrieve the overridden state for the specified object
diff --git a/Engine/Source/Runtime/CoreUObject/Public/UObject/OverriddenPropertySet.h b/Engine/Source/Runtime/CoreUObject/Public/UObject/OverriddenPropertySet.h
--- a/Engine/Source/Runtime/CoreUObject/Public/UObject/OverriddenPropertySet.h
+++ b/Engine/Source/Runtime/CoreUObject/Public/UObject/OverriddenPropertySet.h
@@ -356,6 +356,15 @@ public:
 	 * @return the node containing the information of the overridden property */
 	FOverriddenPropertyNode* SetOverriddenPropertyOperation(EOverriddenPropertyOperation Operation, const FArchiveSerializedPropertyChain* CurrentPropertyChain, FProperty* Property);
 
+	/**
+	 * Restore the overridable operation of the current property from the serialized property chain and the specified property
+	 * @note: This will not restore modified state, has restoring sub properties will do it anyway. 
+	 * @param Operation to set for this property
+	 * @param CurrentPropertyChain leading to the property being serialized if any, null or empty it will return the node of the object itself
+	 * @param Property being serialized if any, if null it will fallback on the last property of the chain
+	 * @return the node containing the information of the overridden property */
+	FOverriddenPropertyNode* RestoreOverriddenPropertyOperation(EOverriddenPropertyOperation Operation, const FArchiveSerializedPropertyChain* CurrentPropertyChain, FProperty* Property);
+
 	/**
 	 * Retrieve the overridden property node from the serialized property chain
 	 * @param CurrentPropertyChain leading to the property being serialized if any, null or empty it will return the node of the object itself
@@ -376,6 +385,13 @@ public:
 	 * @return the node to the sub property */
 	FOverriddenPropertyNode* SetSubPropertyOperation(EOverriddenPropertyOperation Operation, FOverriddenPropertyNode& Node, FOverriddenPropertyNodeID NodeID);
 
+	/**
+	 * Set the overridable operation of a sub object of the specified node.
+	 * @param Operation to set for this property
+	 * @param Node from where the sub property is owned by
+	 * @param SubObject of the operation
+	 * @return the node to the sub property */
+	FOverriddenPropertyNode* SetSubObjectOperation(EOverriddenPropertyOperation Operation, FOverriddenPropertyNode& Node, TNotNull<UObject*> SubObject);
 	/**
 	 * Check if this is an overridden property set of a CDO and that this property is owned by the class of this CDO
 	 * NOTE: this is used to know if a property should be serialized to keep its default CDO value.
@@ -389,6 +405,12 @@ public:
 	void Reset();
 	void HandleObjectsReInstantiated(const TMap<UObject*, UObject*>& Map);
 
+	/* return whether this object is considered added or not */
+	bool WasAdded() const
+	{
+		return bWasAdded;
+	}
+
 protected:
 
 	FOverriddenPropertyNode& FindOrAddNode(FOverriddenPropertyNode& ParentPropertyNode, FOverriddenPropertyNodeID NodeID);
@@ -403,11 +425,13 @@ protected:
 
 	void RemoveOverriddenSubProperties(FOverriddenPropertyNode& PropertyIterator);
 
-
 private:
 	UPROPERTY()
 	TObjectPtr<UObject> Owner = nullptr;
 
+	UPROPERTY()
+	bool bWasAdded = false;
+
 	UPROPERTY()
 	TSet<FOverriddenPropertyNode> OverriddenPropertyNodes;
 
