[AsyncLoading2]

- Fix leak on AsyncLoadingThread caused in 38546438. Objects created while loading while there is any package currently loading is currently assumed to be part of loading and relies on the loader to remove the EInternalObjectsFlag::Async flag from the objects so they may be collected.
- Promoting the warning for adding a new object to the thread's context package after the package has finished postloading. This shouldn't be possible unless some new callback is added so we should error and correct the problem immediately if it ever occurs.

#rnx

diff --git a/Engine/Source/Runtime/CoreUObject/Private/Serialization/AsyncLoading2.cpp b/Engine/Source/Runtime/CoreUObject/Private/Serialization/AsyncLoading2.cpp
--- a/Engine/Source/Runtime/CoreUObject/Private/Serialization/AsyncLoading2.cpp
+++ b/Engine/Source/Runtime/CoreUObject/Private/Serialization/AsyncLoading2.cpp
@@ -9125,12 +9125,11 @@ void FAsyncLoadingThread2::NotifyConstructedDuringAsyncLoading(UObject* Object,
 	FAsyncPackage2* ContextAsyncPackage = (FAsyncPackage2*)ThreadContext.AsyncPackage;
 	UPackage* ObjectPackage = Object->GetPackage();
 
-	if (!ContextAsyncPackage || !ObjectPackage || ObjectPackage == GetTransientPackage())
+	if (!ContextAsyncPackage || !ObjectPackage)
 	{
 		// Something is creating objects on the async loading thread outside of the actual async loading code
 		// e.g. ShaderCodeLibrary::OnExternalReadCallback doing FTaskGraphInterface::Get().WaitUntilTaskCompletes(Event);
-		// Or the object doesn't have a package, in which case we didn't load it from disk
-		// Or the object is now in the transient package, and is no longer needed
+		// Or the object doesn't have a package, in which case we didn't load it from disk. Created objects not part of the package shoudl be cleaned up by their creator
 		UE_LOG(LogStreaming, VeryVerbose, TEXT("Constructed object '%s' during loading will not be post-loaded due to its package '%s' (null and transient package sub-objects are ignored)."), *Object->GetFullName(), (ObjectPackage ? *ObjectPackage->GetFullName() : TEXT("null")));
 		return;
 	}
@@ -9179,7 +9178,7 @@ void FAsyncLoadingThread2::NotifyConstructedDuringAsyncLoading(UObject* Object,
 	}
 	else
 	{
-		UE_LOG(LogStreaming, Warning, TEXT("Object '%s' was created while the FUObjectThreadContext's async package '%s(0x%s)' is done post-loading objects. This object may be left in a partially loaded state."),
+		UE_LOG(LogStreaming, Error, TEXT("Object '%s' was created while the FUObjectThreadContext's async package '%s(0x%s)' is done post-loading objects. This should not happening (perhaps a FAsyncPackageScope2 is missing?) This object may be left in a partially loaded state and may leak."),
 			*Object->GetFullName(), *ContextAsyncPackage->Desc.UPackageName.ToString(), *LexToString(ContextAsyncPackage->Desc.UPackageId));
 	}
 }
