Add support for specifying flags for properties in a propertybag

#tests PIE, Editor

diff --git a/Engine/Source/Runtime/CoreUObject/Private/StructUtils/PropertyBag.cpp b/Engine/Source/Runtime/CoreUObject/Private/StructUtils/PropertyBag.cpp
--- a/Engine/Source/Runtime/CoreUObject/Private/StructUtils/PropertyBag.cpp
+++ b/Engine/Source/Runtime/CoreUObject/Private/StructUtils/PropertyBag.cpp
@@ -3535,8 +3535,7 @@ const UPropertyBag* UPropertyBag::GetOrCreateFromDescs(const TConstArrayView<FPr
 				NewProperty->SetMetaData(*PropertyDescMetaData.Key.ToString(), *PropertyDescMetaData.Value);
 			}
 #endif
-			
-			NewProperty->SetPropertyFlags(CPF_Edit);
+			NewProperty->SetPropertyFlags((EPropertyFlags)Desc.PropertyFlags);
 			NewBag->AddCppProperty(NewProperty);
 			Desc.CachedProperty = NewProperty; 
 		}
diff --git a/Engine/Source/Runtime/CoreUObject/Public/StructUtils/PropertyBag.h b/Engine/Source/Runtime/CoreUObject/Public/StructUtils/PropertyBag.h
--- a/Engine/Source/Runtime/CoreUObject/Public/StructUtils/PropertyBag.h
+++ b/Engine/Source/Runtime/CoreUObject/Public/StructUtils/PropertyBag.h
@@ -6,6 +6,7 @@
 #include "StructUtils/StructView.h"
 #include "Templates/ValueOrError.h"
 #include "Containers/StaticArray.h"
+#include "UObject/ObjectMacros.h"
 #include "PropertyBag.generated.h"
 
 /** Property bag property type, loosely based on BluePrint pin types. */
@@ -220,6 +221,8 @@ struct COREUOBJECT_API FPropertyBagPropertyDesc
 {
 	GENERATED_BODY()
 
+	static_assert(std::is_same_v<std::underlying_type_t<EPropertyFlags>, uint64>, "FPropertyBagPropertyDesc::PropertyFlag does not match EPropertyFlags type");
+
 	FPropertyBagPropertyDesc() = default;
 	FPropertyBagPropertyDesc(const FName InName, const FProperty* InSourceProperty);
 	FPropertyBagPropertyDesc(const FName InName, const EPropertyBagPropertyType InValueType, const UObject* InValueTypeObject = nullptr)
@@ -228,19 +231,21 @@ struct COREUOBJECT_API FPropertyBagPropertyDesc
 		, ValueType(InValueType)
 	{
 	}
-	FPropertyBagPropertyDesc(const FName InName, const EPropertyBagContainerType InContainerType, const EPropertyBagPropertyType InValueType, const UObject* InValueTypeObject = nullptr)
+	FPropertyBagPropertyDesc(const FName InName, const EPropertyBagContainerType InContainerType, const EPropertyBagPropertyType InValueType, const UObject* InValueTypeObject = nullptr, EPropertyFlags InPropertyFlags = CPF_Edit)
 		: ValueTypeObject(InValueTypeObject)
 		, Name(InName)
 		, ValueType(InValueType)
 		, ContainerTypes(InContainerType)
+		, PropertyFlags((uint64)InPropertyFlags)
 	{
 	}
 
-	FPropertyBagPropertyDesc(const FName InName, const FPropertyBagContainerTypes& InNestedContainers, const EPropertyBagPropertyType InValueType, UObject* InValueTypeObject = nullptr)
+	FPropertyBagPropertyDesc(const FName InName, const FPropertyBagContainerTypes& InNestedContainers, const EPropertyBagPropertyType InValueType, UObject* InValueTypeObject = nullptr, EPropertyFlags InPropertyFlags = CPF_Edit)
 		: ValueTypeObject(InValueTypeObject)
 		, Name(InName)
 		, ValueType(InValueType)
 		, ContainerTypes(InNestedContainers)
+		, PropertyFlags((uint64)InPropertyFlags)
 	{
 	}
 
@@ -282,6 +287,10 @@ struct COREUOBJECT_API FPropertyBagPropertyDesc
 	UPROPERTY(EditAnywhere, Category="Default")
 	FPropertyBagContainerTypes ContainerTypes;
 
+	/** Flags that will get copied over to this property. uint64 since EPropertyFlags is not UEnum */
+	UPROPERTY(EditAnywhere, Category = "Default")
+	uint64 PropertyFlags = (uint64)CPF_Edit;
+
 #if WITH_EDITORONLY_DATA
 	/** Editor-only metadata for CachedProperty */
 	UPROPERTY(EditAnywhere, Category="Default")
