Migrate internal object flags when migrating remote objects between servers

diff --git a/Engine/Source/Runtime/CoreUObject/Private/UObject/RemoteObjectSerialization.cpp b/Engine/Source/Runtime/CoreUObject/Private/UObject/RemoteObjectSerialization.cpp
--- a/Engine/Source/Runtime/CoreUObject/Private/UObject/RemoteObjectSerialization.cpp
+++ b/Engine/Source/Runtime/CoreUObject/Private/UObject/RemoteObjectSerialization.cpp
@@ -414,6 +414,8 @@ FArchive& operator<<(FArchive& Ar, FRemoteObjectHeader& Header)
 	Ar << Header.Class;
 	Ar << Header.Outer;
 	Ar << Header.Archetype;
+	Ar << Header.InternalFlags;
+
 	return Ar;
 }
 
@@ -429,6 +431,7 @@ void SerializeRemoteObject(UObject* Object, FArchive& Ar, FRemoteObjectHeader& O
 	UObject* Outer = FObjectHandleUtils::GetNonAccessTrackedOuterNoResolve(Object);
 	OutHeader.Outer = Outer;
 	OutHeader.Archetype = Object->GetArchetype();
+	OutHeader.InternalFlags = (int32)(Object->GetInternalFlags() & EInternalObjectFlags::Garbage);
 
 	Object->Serialize(Ar);
 }
@@ -669,8 +672,19 @@ UObject* ConstructRemoteObject(const FRemoteObjectHeader& Header, FRemoteObjectC
 	checkf(FRemoteObjectId(Object) == Header.RemoteId, TEXT("Created an object with a different ID:%s than requested:%s"), *FRemoteObjectId(Object).ToString(), *Header.RemoteId.ToString());
 
 	MarkAsLocal(Object);
-	// It's possible that one of the migrated objects was still in memory on this server and it was marked as Garbage in PostMigrate so clear the flag in case we migrated it back
-	Object->ClearGarbage();
+
+	EInternalObjectFlags InternalFlags = (EInternalObjectFlags)Header.InternalFlags;
+	if (!(InternalFlags & EInternalObjectFlags::Garbage))
+	{
+		// It's possible that one of the migrated objects was still in memory on this server and it was marked as Garbage in PostMigrate so clear the flag in case we migrated it back
+		Object->ClearGarbage();
+	}
+	if (InternalFlags != EInternalObjectFlags::None)
+	{
+		// Set any relavant internal flags the object might have had set when it was migrated. It's possible we're setting the Garbage flag here but it's not impossible to migrate a destroyed actor
+		// and we need to retain that flag in such case.
+		Object->SetInternalFlags(InternalFlags);
+	}
 
 	return Object;
 }
diff --git a/Engine/Source/Runtime/CoreUObject/Private/UObject/RemoteObjectSerialization.h b/Engine/Source/Runtime/CoreUObject/Private/UObject/RemoteObjectSerialization.h
--- a/Engine/Source/Runtime/CoreUObject/Private/UObject/RemoteObjectSerialization.h
+++ b/Engine/Source/Runtime/CoreUObject/Private/UObject/RemoteObjectSerialization.h
@@ -18,6 +18,7 @@ namespace UE::RemoteObject::Serialization
 		TObjectPtr<UClass> Class;
 		TObjectPtr<UObject> Outer;
 		TObjectPtr<UObject> Archetype;
+		int32 InternalFlags = 0;
 
 		/** Transient (SerialNumber is local-only but it's stored here for convenience when deserializing object data) */
 		int32 SerialNumber = 0;
