[UObjectBaseUtility]

Remove the internal Async flag when objects are marked as garbage. The Async flag prevents the object from being GC'd if the object was created off the game thread. When explicitly marking the object as garbage we can make the object collectable for the user since it matches their intent and removes the need of knowing the thread their code is running on.

#rnx

diff --git a/Engine/Source/Runtime/CoreUObject/Public/UObject/UObjectBaseUtility.h b/Engine/Source/Runtime/CoreUObject/Public/UObject/UObjectBaseUtility.h
--- a/Engine/Source/Runtime/CoreUObject/Public/UObject/UObjectBaseUtility.h
+++ b/Engine/Source/Runtime/CoreUObject/Public/UObject/UObjectBaseUtility.h
@@ -182,6 +182,10 @@ public:
 		UE_AUTORTFM_ONCOMMIT(this)
 		{
 			AtomicallySetFlags(RF_MirroredGarbage);
+
+			// If we explicitly marked the object as garbage, remove the async flag so it's visible to the GC
+			AtomicallyClearInternalFlags(EInternalObjectFlags::Async);
+
 			GUObjectArray.IndexToObject(InternalIndex)->SetGarbage();
 		};
 	}
