[Cook - Optional] Change the list of Additional optional asset to be SoftObjectClass instead of UClass, since they might not be loaded the first time the list is read (and the list is contructed only once)

* Useful for cooked editor to allow for EditorOnly data to be cooked, like Editor Only components on BP actors.

#rnx

diff --git a/Engine/Source/Runtime/CoreUObject/Private/UObject/SavePackage/SaveContext.cpp b/Engine/Source/Runtime/CoreUObject/Private/UObject/SavePackage/SaveContext.cpp
--- a/Engine/Source/Runtime/CoreUObject/Private/UObject/SavePackage/SaveContext.cpp
+++ b/Engine/Source/Runtime/CoreUObject/Private/UObject/SavePackage/SaveContext.cpp
@@ -277,7 +277,8 @@ const TSet<UObject*>& FSaveContext::GetPostSaveObjectsToSerialize() const
 
 namespace
 {
-	TArray<UClass*> AutomaticOptionalInclusionAssetTypeList;
+	// Use Soft Class Ptr to gather all asset types even if they are not yet loaded.
+	TArray<TSoftClassPtr<UObject>> AutomaticOptionalInclusionAssetTypeList;
 }
 
 void FSaveContext::SetupHarvestingRealms()
@@ -289,32 +290,34 @@ void FSaveContext::SetupHarvestingRealms()
 	CurrentHarvestingRealm = IsCooking() ? ESaveRealm::Game : ESaveRealm::Editor;
 
 	// Generate the automatic optional context inclusion asset list
-	static bool bAssetListGenerated = [](TArray<UClass*>& OutAssetList)
-	{
-		TArray<FString> AssetList;
-		GConfig->GetArray(TEXT("CookSettings"), TEXT("AutomaticOptionalInclusionAssetType"), AssetList, GEditorIni);
-		for (const FString& AssetType : AssetList)
+	static bool bAssetListGenerated = [](TArray<TSoftClassPtr<UObject>>& OutAssetList)
 		{
-			if (UClass* AssetClass = FindObject<UClass>(nullptr, *AssetType, true))
-			{
-				OutAssetList.Add(AssetClass);
-			}
-			else
+			TArray<FString> AssetList;
+			GConfig->GetArray(TEXT("CookSettings"), TEXT("AutomaticOptionalInclusionAssetType"), AssetList, GEditorIni);
+			for (const FString& AssetType : AssetList)
 			{
-				UE_LOG(LogSavePackage, Warning, TEXT("The asset type '%s' was not found while building the allowlist for automatic optional data inclusion list."), *AssetType);
+				TSoftClassPtr<UObject> SoftClassPath(AssetType);
+				if (!SoftClassPath.IsNull())
+				{
+					OutAssetList.Add(MoveTemp(SoftClassPath));
+				}
+				else
+				{
+					UE_LOG(LogSavePackage, Warning, TEXT("The asset type '%s' found while building the allowlist for automatic optional data inclusion list is not a valid class path."), *AssetType);
+				}
 			}
-		}
-		return true;
-	}(AutomaticOptionalInclusionAssetTypeList);
+			return true;
+		}(AutomaticOptionalInclusionAssetTypeList);
 
 	if (bAssetListGenerated && Asset)
 	{
 		// if the asset type itself is a class (ie. BP) use that to check for auto optional
 		UClass* AssetType = Cast<UClass>(Asset);
 		AssetType = AssetType ? AssetType : Asset->GetClass();
-		bool bAllowedClass = Algo::FindByPredicate(AutomaticOptionalInclusionAssetTypeList, [AssetType](const UClass* InAssetClass)
+		bool bAllowedClass = Algo::FindByPredicate(AutomaticOptionalInclusionAssetTypeList, [AssetType](const TSoftClassPtr<UObject>& InAssetClass)
 			{
-				return AssetType->IsChildOf(InAssetClass);
+				const UClass* AssetClass = InAssetClass.Get();
+				return AssetClass && AssetType->IsChildOf(AssetClass);
 			}) != nullptr;
 		bIsSaveAutoOptional = IsCooking() && IsSaveOptional() && bAllowedClass;
 	}
