IsRemote() methods for TWeakObjectPtr and TObjectPtr.

Motivation: we sometimes need to skip unnecessary object migrations in existing runtime code and it's cumbersome to wrap every instance in #if UE_WITH_REMOTE_OBJECT_HANDLE macro.
TObjectPtr already has IsResolved() (== !IsRemote()) and we've been using it but I thought it would be nice to add IsRemote() for consistency and to disambiguate use cases of IsResolved() that apply to remote objects only.

diff --git a/Engine/Source/Runtime/Core/Public/UObject/WeakObjectPtrTemplates.h b/Engine/Source/Runtime/Core/Public/UObject/WeakObjectPtrTemplates.h
--- a/Engine/Source/Runtime/Core/Public/UObject/WeakObjectPtrTemplates.h
+++ b/Engine/Source/Runtime/Core/Public/UObject/WeakObjectPtrTemplates.h
@@ -265,6 +265,11 @@ public:
 	}
 #endif // UE_WITH_REMOTE_OBJECT_HANDLE
 
+	FORCEINLINE bool IsRemote() const
+	{
+		return WeakPtr.IsRemote();
+	}
+
 	/**
 	 * Weak object pointer serialization, this forwards to FArchive::operator<<(struct FWeakObjectPtr&) or an override
 	 */
diff --git a/Engine/Source/Runtime/CoreUObject/Private/UObject/WeakObjectPtr.cpp b/Engine/Source/Runtime/CoreUObject/Private/UObject/WeakObjectPtr.cpp
--- a/Engine/Source/Runtime/CoreUObject/Private/UObject/WeakObjectPtr.cpp
+++ b/Engine/Source/Runtime/CoreUObject/Private/UObject/WeakObjectPtr.cpp
@@ -189,7 +189,7 @@ FUObjectItem* FWeakObjectPtr::TryResolveRemoteObject() const
 
 	if (ResolvedObject)
 	{		
-		if (IsRemote(ResolvedObject))
+		if (UE::RemoteObject::Handle::IsRemote(ResolvedObject))
 		{
 			// Object memory is still on this server but it's marked as remote so we need to resolve it
 			ResolvedObject = ResolveObject(ResolvedObject, ERemoteReferenceType::Weak);
diff --git a/Engine/Source/Runtime/CoreUObject/Public/UObject/ObjectPtr.h b/Engine/Source/Runtime/CoreUObject/Public/UObject/ObjectPtr.h
--- a/Engine/Source/Runtime/CoreUObject/Public/UObject/ObjectPtr.h
+++ b/Engine/Source/Runtime/CoreUObject/Public/UObject/ObjectPtr.h
@@ -197,6 +197,12 @@ public:
 
 	FORCEINLINE bool IsResolved() const { return IsObjectHandleResolved(Handle); }
 
+	/**
+	 * FObjectPtr::IsRemote is only used when UE_WITH_REMOTE_OBJECT_HANDLE is true, and is mutually exclusive with other uses of unresolved FObjectPtr.
+	 * When using UE_WITH_REMOTE_OBJECT_HANDLE, unresolved FObjectPtrs mean the object is not locally available and is present on another server, see @RemoteObject.h
+	*/
+	FORCEINLINE bool IsRemote() const { return !IsResolved(); }
+
 	// Gets the PathName of the object without resolving the object reference.
 	// @TODO OBJPTR: Deprecate this.
 	FString GetPath() const { return GetPathName(); }
@@ -723,6 +729,7 @@ public:
 	FORCEINLINE bool operator!() const { return ObjectPtr_Private::IsObjectPtrNull<T>(ObjectPtr); }
 	explicit FORCEINLINE operator bool() const { return !ObjectPtr_Private::IsObjectPtrNull<T>(ObjectPtr); }
 	FORCEINLINE bool IsResolved() const { return ObjectPtr.IsResolved(); }
+	FORCEINLINE bool IsRemote() const { return ObjectPtr.IsRemote(); }
 	FORCEINLINE FString GetPath() const { return ObjectPtr.GetPath(); }
 	FORCEINLINE FString GetPathName() const { return ObjectPtr.GetPathName(); }
 	FORCEINLINE FName GetFName() const { return ObjectPtr.GetFName(); }
diff --git a/Engine/Source/Runtime/CoreUObject/Public/UObject/RemoteObject.h b/Engine/Source/Runtime/CoreUObject/Public/UObject/RemoteObject.h
--- a/Engine/Source/Runtime/CoreUObject/Public/UObject/RemoteObject.h
+++ b/Engine/Source/Runtime/CoreUObject/Public/UObject/RemoteObject.h
@@ -8,10 +8,18 @@
 
 class UObject;
 
+/**
+ * Remote objects are unique UObjects that are referenced by a local server instance but their memory is actually owned by (exists on) another server.
+ * 
+ * It's possible that an object is remote but its memory hasn't been freed yet (UObject with EInternalObjectFlags::Remote flag that hasn't been GC'd yet). 
+ * In such case any attempt to access that object through TObjectPtr will result in its memory being migrated from a remote server to a local server. 
+ * Remote object memory is freed in the next GC pass after the object has been migrated and any existing references to that object (must be referenced by TObjectPtr) 
+ * will be updated by GC to point to the remote object's FRemoteObjectStub. 
+ */
 namespace UE::RemoteObject
 {
 	/**
-	* Returns a unique id associated with this process
+	* Returns a unique id associated with this (server) process
 	*/
 	COREUOBJECT_API FRemoteServerId GetGlobalServerId();
 
@@ -20,6 +28,7 @@ namespace UE::RemoteObject
 namespace UE::RemoteObject::Handle
 {
 	// Structure that holds basic information about a remote object
+	// This is what FObjectPtr that references a remote object actually points to after the remote object's memory has been claimed by GC
 	struct FRemoteObjectStub
 	{
 		FRemoteObjectId Id;
diff --git a/Engine/Source/Runtime/CoreUObject/Public/UObject/WeakObjectPtr.h b/Engine/Source/Runtime/CoreUObject/Public/UObject/WeakObjectPtr.h
--- a/Engine/Source/Runtime/CoreUObject/Public/UObject/WeakObjectPtr.h
+++ b/Engine/Source/Runtime/CoreUObject/Public/UObject/WeakObjectPtr.h
@@ -311,6 +311,19 @@ public:
 	* Returns true if this weak pointer resolves to the specified object
 	*/
 	COREUOBJECT_API bool HasSameObject(const UObject* Other) const;
+
+	/**
+	* Returns true if the object this weak pointer represents is owned by another server
+	*/
+	FORCEINLINE bool IsRemote() const
+	{
+		return UE::RemoteObject::Handle::IsRemote(ObjectRemoteId);
+	}
+#else
+	FORCEINLINE bool IsRemote() const
+	{
+		return false;
+	}
 #endif // UE_WITH_REMOTE_OBJECT_HANDLE
 
 	/** Hash function. */
