[CoreUObject] Implement TBaseStructure for all FInterpCurve variants

* Also add "ElementType" typedef in the FInterpCurve to know the underlying type

#rnx

diff --git a/Engine/Source/Runtime/Core/Public/Math/InterpCurve.h b/Engine/Source/Runtime/Core/Public/Math/InterpCurve.h
--- a/Engine/Source/Runtime/Core/Public/Math/InterpCurve.h
+++ b/Engine/Source/Runtime/Core/Public/Math/InterpCurve.h
@@ -35,6 +35,7 @@ public:
 	float LoopKeyOffset;
 
 public:
+	using ElementType = T;
 
 	/** Default constructor. */
 	FInterpCurve() 
diff --git a/Engine/Source/Runtime/CoreUObject/Private/UObject/Class.cpp b/Engine/Source/Runtime/CoreUObject/Private/UObject/Class.cpp
--- a/Engine/Source/Runtime/CoreUObject/Private/UObject/Class.cpp
+++ b/Engine/Source/Runtime/CoreUObject/Private/UObject/Class.cpp
@@ -7674,6 +7674,42 @@ UScriptStruct* TBaseStructure<FFallbackStruct>::Get()
 	return ScriptStruct;
 }
 
+UScriptStruct* TBaseStructure<FInterpCurveFloat>::Get()
+{
+	static UScriptStruct* ScriptStruct = StaticGetBaseStructureInternal(TEXT("InterpCurveFloat"));
+	return ScriptStruct;
+}
+
+UScriptStruct* TBaseStructure<FInterpCurveVector2D>::Get()
+{
+	static UScriptStruct* ScriptStruct = StaticGetBaseStructureInternal(TEXT("InterpCurveVector2D"));
+	return ScriptStruct;
+}
+
+UScriptStruct* TBaseStructure<FInterpCurveVector>::Get()
+{
+	static UScriptStruct* ScriptStruct = StaticGetBaseStructureInternal(TEXT("InterpCurveVector"));
+	return ScriptStruct;
+}
+
+UScriptStruct* TBaseStructure<FInterpCurveQuat>::Get()
+{
+	static UScriptStruct* ScriptStruct = StaticGetBaseStructureInternal(TEXT("InterpCurveQuat"));
+	return ScriptStruct;
+}
+
+UScriptStruct* TBaseStructure<FInterpCurveTwoVectors>::Get()
+{
+	static UScriptStruct* ScriptStruct = StaticGetBaseStructureInternal(TEXT("InterpCurveTwoVectors"));
+	return ScriptStruct;
+}
+
+UScriptStruct* TBaseStructure<FInterpCurveLinearColor>::Get()
+{
+	static UScriptStruct* ScriptStruct = StaticGetBaseStructureInternal(TEXT("InterpCurveLinearColor"));
+	return ScriptStruct;
+}
+
 UScriptStruct* TBaseStructure<FInterpCurvePointFloat>::Get()
 {
 	static UScriptStruct* ScriptStruct = StaticGetBaseStructureInternal(TEXT("InterpCurvePointFloat"));
diff --git a/Engine/Source/Runtime/CoreUObject/Public/UObject/Class.h b/Engine/Source/Runtime/CoreUObject/Public/UObject/Class.h
--- a/Engine/Source/Runtime/CoreUObject/Public/UObject/Class.h
+++ b/Engine/Source/Runtime/CoreUObject/Public/UObject/Class.h
@@ -4320,6 +4320,36 @@ template<> struct TBaseStructure<FFallbackStruct>
 	static COREUOBJECT_API UScriptStruct* Get();
 };
 
+template<> struct TBaseStructure<FInterpCurveFloat>
+{
+	static COREUOBJECT_API UScriptStruct* Get();
+};
+
+template<> struct TBaseStructure<FInterpCurveVector2D>
+{
+	static COREUOBJECT_API UScriptStruct* Get();
+};
+
+template<> struct TBaseStructure<FInterpCurveVector>
+{
+	static COREUOBJECT_API UScriptStruct* Get();
+};
+
+template<> struct TBaseStructure<FInterpCurveQuat>
+{
+	static COREUOBJECT_API UScriptStruct* Get();
+};
+
+template<> struct TBaseStructure<FInterpCurveTwoVectors>
+{
+	static COREUOBJECT_API UScriptStruct* Get();
+};
+
+template<> struct TBaseStructure<FInterpCurveLinearColor>
+{
+	static COREUOBJECT_API UScriptStruct* Get();
+};
+
 template<> struct TBaseStructure<FInterpCurvePointFloat>
 {
 	static COREUOBJECT_API UScriptStruct* Get();
