[PropertyBag] fixed validations for Class and SoftClass properties.
- added unit tests
- added typed validated version of SetSoftObjectPath when an object pointer is provided

diff --git a/Engine/Source/Developer/StructUtilsTestSuite/Private/PropertyBagTest.cpp b/Engine/Source/Developer/StructUtilsTestSuite/Private/PropertyBagTest.cpp
--- a/Engine/Source/Developer/StructUtilsTestSuite/Private/PropertyBagTest.cpp
+++ b/Engine/Source/Developer/StructUtilsTestSuite/Private/PropertyBagTest.cpp
@@ -243,6 +243,82 @@ struct FTest_Class : FAITestBase
 };
 IMPLEMENT_AI_INSTANT_TEST(FTest_Class, "System.StructUtils.PropertyBag.Class");
 
+struct FTest_SubclassOf : FAITestBase
+{
+	virtual bool InstantTest() override
+	{
+		static const FName PropertyName = GET_MEMBER_NAME_CHECKED(FTestStructWithSubClassOf, ClassProperty);
+		const UScriptStruct& Struct = *FTestStructWithSubClassOf::StaticStruct();
+		const FClassProperty& Property = *CastFieldChecked<FClassProperty>(Struct.FindPropertyByName(PropertyName));
+
+		FInstancedPropertyBag Bag;
+		Bag.AddProperty(PropertyName, &Property);
+		const FPropertyBagPropertyDesc* BagPropertyDesc = Bag.FindPropertyDescByName(PropertyName);
+		AITEST_NOT_NULL(TEXT("Expecting to find inserted property"), BagPropertyDesc);
+
+#if WITH_EDITOR // FPropertyBagPropertyDesc::MetaClass is in WITH_EDITORONLY_DATA
+		AITEST_NOT_NULL(TEXT("Created bag property MetaClass pointer"), BagPropertyDesc->MetaClass.Get());
+		AITEST_EQUAL(TEXT("Created bag property MetaClass type"), BagPropertyDesc->MetaClass->GetFName(), Property.MetaClass->GetFName());
+#endif // WITH_EDITOR
+
+		AITEST_NOT_NULL(TEXT("Created bag property ValueTypeObject pointer"), BagPropertyDesc->ValueTypeObject.Get());
+		AITEST_EQUAL(TEXT("Created bag property ValueTypeObject type"), BagPropertyDesc->ValueTypeObject->GetFName(), Property.MetaClass->GetFName());
+
+		const FString ResultSuccess = UEnum::GetValueAsString(EPropertyBagResult::Success);
+
+		EPropertyBagResult SetResult = Bag.SetValueClass(PropertyName, UBagTestObject1::StaticClass());
+		AITEST_EQUAL(TEXT("SetValueClass result for allowed class")
+			, FStringView(UEnum::GetValueAsString(SetResult))
+			, FStringView(ResultSuccess));
+
+		SetResult = Bag.SetValueClass(PropertyName, UClass::StaticClass());
+		AITEST_NOT_EQUAL(TEXT("SetValueClass result for non-allowed class")
+			, FStringView(UEnum::GetValueAsString(SetResult))
+			, FStringView(ResultSuccess));
+
+		return true;
+	}
+};
+IMPLEMENT_AI_INSTANT_TEST(FTest_SubclassOf, "System.StructUtils.PropertyBag.SubclassOf");
+
+struct FTest_SoftClassPtr : FAITestBase
+{
+	virtual bool InstantTest() override
+	{
+		static const FName PropertyName = GET_MEMBER_NAME_CHECKED(FTestStructWithSubClassOf, SoftClassProperty);
+		const UScriptStruct& Struct = *FTestStructWithSubClassOf::StaticStruct();
+		const FSoftClassProperty& Property = *CastFieldChecked<FSoftClassProperty>(Struct.FindPropertyByName(PropertyName));
+
+		FInstancedPropertyBag Bag;
+		Bag.AddProperty(PropertyName, &Property);
+		const FPropertyBagPropertyDesc* BagPropertyDesc = Bag.FindPropertyDescByName(PropertyName);
+		AITEST_NOT_NULL(TEXT("Expecting to find inserted property"), BagPropertyDesc);
+
+#if WITH_EDITOR // FPropertyBagPropertyDesc::MetaClass is in WITH_EDITORONLY_DATA
+		AITEST_NOT_NULL(TEXT("Created bag property MetaClass pointer"), BagPropertyDesc->MetaClass.Get());
+		AITEST_EQUAL(TEXT("Created bag property MetaClass type"), BagPropertyDesc->MetaClass->GetFName(), Property.MetaClass->GetFName());
+#endif // WITH_EDITOR
+
+		AITEST_NOT_NULL(TEXT("Created bag property ValueTypeObject pointer"), BagPropertyDesc->ValueTypeObject.Get());
+		AITEST_EQUAL(TEXT("Created bag property ValueTypeObject type"), BagPropertyDesc->ValueTypeObject->GetFName(), Property.MetaClass->GetFName());
+
+		const FString ResultSuccess = UEnum::GetValueAsString(EPropertyBagResult::Success);
+
+		EPropertyBagResult SetResult = Bag.SetValueSoftPath(PropertyName, UBagTestObject1::StaticClass());
+		AITEST_EQUAL(TEXT("SetValueSoftPath result for allowed class")
+			, FStringView(UEnum::GetValueAsString(SetResult))
+			, FStringView(ResultSuccess));
+
+		SetResult = Bag.SetValueSoftPath(PropertyName, UClass::StaticClass());
+		AITEST_NOT_EQUAL(TEXT("SetValueSoftPath result for non-allowed class")
+			, FStringView(UEnum::GetValueAsString(SetResult))
+			, FStringView(ResultSuccess));
+
+		return true;
+	}
+};
+IMPLEMENT_AI_INSTANT_TEST(FTest_SoftClassPtr, "System.StructUtils.PropertyBag.SoftClassPtr");
+
 struct FTest_Enum : FAITestBase
 {
 	virtual bool InstantTest() override
diff --git a/Engine/Source/Developer/StructUtilsTestSuite/Public/StructUtilsTestTypes.h b/Engine/Source/Developer/StructUtilsTestSuite/Public/StructUtilsTestTypes.h
--- a/Engine/Source/Developer/StructUtilsTestSuite/Public/StructUtilsTestTypes.h
+++ b/Engine/Source/Developer/StructUtilsTestSuite/Public/StructUtilsTestTypes.h
@@ -2,6 +2,7 @@
 
 #pragma once
 
+#include "Templates/SubclassOf.h"
 #include "StructUtils/PropertyBag.h"
 #include "StructUtilsTestTypes.generated.h"
 
@@ -207,3 +208,15 @@ public:
 	UPROPERTY()
 	FInstancedStruct Value;
 };
+
+USTRUCT()
+struct FTestStructWithSubClassOf
+{
+	GENERATED_BODY()
+
+	UPROPERTY()
+	TSubclassOf<UBagTestObject1> ClassProperty;
+
+	UPROPERTY()
+	TSoftClassPtr<UBagTestObject1> SoftClassProperty;
+};
\ No newline at end of file
diff --git a/Engine/Source/Runtime/CoreUObject/Private/StructUtils/PropertyBag.cpp b/Engine/Source/Runtime/CoreUObject/Private/StructUtils/PropertyBag.cpp
--- a/Engine/Source/Runtime/CoreUObject/Private/StructUtils/PropertyBag.cpp
+++ b/Engine/Source/Runtime/CoreUObject/Private/StructUtils/PropertyBag.cpp
@@ -202,7 +202,7 @@ namespace UE::StructUtils::Private
 		{
 			if (const FClassProperty* ClassProperty = CastField<FClassProperty>(InSourceProperty))
 			{
-				return ClassProperty->PropertyClass;
+				return ClassProperty->MetaClass;
 			}
 
 			return ObjectProperty->PropertyClass;
@@ -211,7 +211,7 @@ namespace UE::StructUtils::Private
 		{
 			if (const FSoftClassProperty* SoftClassProperty = CastField<FSoftClassProperty>(InSourceProperty))
 			{
-				return SoftClassProperty->PropertyClass;
+				return SoftClassProperty->MetaClass;
 			}
 
 			return SoftObjectProperty->PropertyClass;
@@ -1206,6 +1206,8 @@ namespace UE::StructUtils::Private
 		}
 		else
 		{
+			check(Desc->ValueType == EPropertyBagPropertyType::Class
+				|| Desc->ValueType == EPropertyBagPropertyType::SoftClass);
 			const UClass* Class = Cast<UClass>(InValue);
 			const UClass* PropClass = nullptr;
 
@@ -1229,7 +1231,7 @@ namespace UE::StructUtils::Private
 		return EPropertyBagResult::Success;
 	}
 
-	EPropertyBagResult SetPropertyValueAsSoftPath(const FPropertyBagPropertyDesc* Desc, void* Address, const FSoftObjectPath& InValue)
+	EPropertyBagResult SetPropertyValueAsSoftPath(const FPropertyBagPropertyDesc* Desc, void* Address, const FSoftObjectPath& InPath, const UObject* InObject = nullptr)
 	{
 		if (Desc == nullptr || Desc->CachedProperty == nullptr)
 		{
@@ -1253,7 +1255,35 @@ namespace UE::StructUtils::Private
 		check(SoftObjectProperty->PropertyClass);
 		check(Desc->ValueTypeObject);
 
-		SoftObjectProperty->SetPropertyValue(Address, FSoftObjectPtr(InValue));
+		// Can perform extra validation when source object is provided
+		if (InObject)
+		{
+			if (Desc->ValueType == EPropertyBagPropertyType::SoftObject)
+			{
+				if (CanCastTo(InObject->GetClass(), SoftObjectProperty->PropertyClass) == false)
+				{
+					return EPropertyBagResult::TypeMismatch;
+				}
+			}
+			else
+			{
+				check(Desc->ValueType == EPropertyBagPropertyType::SoftClass);
+				const UClass* Class = CastChecked<UClass>(InObject);
+				const UClass* PropClass = nullptr;
+
+				if (const FSoftClassProperty* SoftClassProperty = CastFieldChecked<FSoftClassProperty>(Desc->CachedProperty))
+				{
+					PropClass = SoftClassProperty->MetaClass;
+				}
+			
+				if (!Class || !PropClass || !Class->IsChildOf(PropClass))
+				{
+					return EPropertyBagResult::TypeMismatch;
+				}
+			}
+		}
+
+		SoftObjectProperty->SetPropertyValue(Address, FSoftObjectPtr(InPath));
 
 		return EPropertyBagResult::Success;
 	}
@@ -1455,6 +1485,10 @@ FPropertyBagPropertyDesc::FPropertyBagPropertyDesc(const FName InName, const FPr
 	{
 		MetaClass = ClassProperty->MetaClass;
 	}
+	else if (const FSoftClassProperty* SoftClassProperty = CastField<FSoftClassProperty>(InSourceProperty))
+	{
+		MetaClass = SoftClassProperty->MetaClass;
+	}
 
 	if (const TMap<FName, FString>* SourcePropertyMetaData = InSourceProperty->GetMetaDataMap())
 	{
@@ -2479,6 +2513,12 @@ EPropertyBagResult FInstancedPropertyBag::SetValueSoftPath(const FName Name, con
 	return UE::StructUtils::Private::SetPropertyValueAsSoftPath(Desc, GetMutableValueAddress(Desc), InValue);
 }
 
+EPropertyBagResult FInstancedPropertyBag::SetValueSoftPath(const FName Name, const UObject* InValue)
+{
+	const FPropertyBagPropertyDesc* Desc = FindPropertyDescByName(Name);
+	return UE::StructUtils::Private::SetPropertyValueAsSoftPath(Desc, GetMutableValueAddress(Desc), InValue, InValue);
+}
+
 EPropertyBagResult FInstancedPropertyBag::SetValueSerializedString(const FName Name, const FString& InValue)
 {
 	const FPropertyBagPropertyDesc* Desc = FindPropertyDescByName(Name);
@@ -2667,6 +2707,15 @@ EPropertyBagResult FInstancedPropertyBag::SetValueSoftPath(const FPropertyBagPro
 	return UE::StructUtils::Private::SetPropertyValueAsSoftPath(&Desc, GetMutableValueAddress(&Desc), InValue);
 }
 
+EPropertyBagResult FInstancedPropertyBag::SetValueSoftPath(const FPropertyBagPropertyDesc& Desc, const UObject* InValue)
+{
+	if (!OwnsPropertyDesc(Desc))
+	{
+		return EPropertyBagResult::PropertyNotFound;
+	}
+	return UE::StructUtils::Private::SetPropertyValueAsSoftPath(&Desc, GetMutableValueAddress(&Desc), InValue, InValue);
+}
+
 TValueOrError<FPropertyBagArrayRef, EPropertyBagResult> FInstancedPropertyBag::GetMutableArrayRef(const FName Name)
 {
 	const FPropertyBagPropertyDesc* Desc = FindPropertyDescByName(Name);
@@ -3279,6 +3328,11 @@ EPropertyBagResult FPropertyBagArrayRef::SetValueSoftPath(const int32 Index, con
 	return UE::StructUtils::Private::SetPropertyValueAsSoftPath(&ValueDesc, GetMutableAddress(Index), InValue);
 }
 
+EPropertyBagResult FPropertyBagArrayRef::SetValueSoftPath(const int32 Index, const UObject* InValue)
+{
+	return UE::StructUtils::Private::SetPropertyValueAsSoftPath(&ValueDesc, GetMutableAddress(Index), InValue, InValue);
+}
+
 //----------------------------------------------------------------//
 //  FPropertyBagSetRef
 //----------------------------------------------------------------//
diff --git a/Engine/Source/Runtime/CoreUObject/Public/StructUtils/PropertyBag.h b/Engine/Source/Runtime/CoreUObject/Public/StructUtils/PropertyBag.h
--- a/Engine/Source/Runtime/CoreUObject/Public/StructUtils/PropertyBag.h
+++ b/Engine/Source/Runtime/CoreUObject/Public/StructUtils/PropertyBag.h
@@ -659,6 +659,7 @@ struct COREUOBJECT_API FInstancedPropertyBag
 	EPropertyBagResult SetValueObject(const FName Name, UObject* InValue);
 	EPropertyBagResult SetValueClass(const FName Name, UClass* InValue);
 	EPropertyBagResult SetValueSoftPath(const FName Name, const FSoftObjectPath& InValue);
+	EPropertyBagResult SetValueSoftPath(const FName Name, const UObject* InValue);
 
 	/**
 	 * Sets property value from a serialized representation of the value. If the string value provided
@@ -711,6 +712,7 @@ struct COREUOBJECT_API FInstancedPropertyBag
 	EPropertyBagResult SetValueObject(const FPropertyBagPropertyDesc& Desc, UObject* InValue);
 	EPropertyBagResult SetValueClass(const FPropertyBagPropertyDesc& Desc, UClass* InValue);
 	EPropertyBagResult SetValueSoftPath(const FPropertyBagPropertyDesc& Desc, const FSoftObjectPath& InValue);
+	EPropertyBagResult SetValueSoftPath(const FPropertyBagPropertyDesc& Desc, const UObject* InValue);
 
 	/** Sets enum value specified type. */
 	template <typename T>
@@ -955,6 +957,7 @@ public:
 	EPropertyBagResult SetValueObject(const int32 Index, UObject* InValue);
 	EPropertyBagResult SetValueClass(const int32 Index, UClass* InValue);
 	EPropertyBagResult SetValueSoftPath(const int32 Index, const FSoftObjectPath& InValue);
+	EPropertyBagResult SetValueSoftPath(const int32 Index, const UObject* InValue);
 
 	/** Sets enum value specified type. */
 	template <typename T>
