Use MarkAsGarbage() when setting EInternalObjectFlags::Garbage on a migrated object

diff --git a/Engine/Source/Runtime/CoreUObject/Private/UObject/RemoteObjectSerialization.cpp b/Engine/Source/Runtime/CoreUObject/Private/UObject/RemoteObjectSerialization.cpp
--- a/Engine/Source/Runtime/CoreUObject/Private/UObject/RemoteObjectSerialization.cpp
+++ b/Engine/Source/Runtime/CoreUObject/Private/UObject/RemoteObjectSerialization.cpp
@@ -673,16 +673,23 @@ UObject* ConstructRemoteObject(const FRemoteObjectHeader& Header, FRemoteObjectC
 
 	MarkAsLocal(Object);
 
+	// Update internal flags on the migrated object. It's possible the object being migrated already existed in memory on this server and had the EInternalObjectFlags::Garbage flag set.
+	// Unless the migrated version also had this flag set we need to clear it.
+	// It's also possible that the local object didn't have this flag set but the migrated one has so we need to set it on this server too (it's not impossible to migrate objects marked as garbage)
 	EInternalObjectFlags InternalFlags = (EInternalObjectFlags)Header.InternalFlags;
+	// Clearing and setting the garbage flag needs to happen through dedicated functions
 	if (!(InternalFlags & EInternalObjectFlags::Garbage))
 	{
-		// It's possible that one of the migrated objects was still in memory on this server and it was marked as Garbage in PostMigrate so clear the flag in case we migrated it back
 		Object->ClearGarbage();
 	}
+	else
+	{
+		Object->MarkAsGarbage();
+		InternalFlags &= ~EInternalObjectFlags::Garbage;
+	}
+	// Any other internal flags can be set with SetInternalFlags 
 	if (InternalFlags != EInternalObjectFlags::None)
 	{
-		// Set any relavant internal flags the object might have had set when it was migrated. It's possible we're setting the Garbage flag here but it's not impossible to migrate a destroyed actor
-		// and we need to retain that flag in such case.
 		Object->SetInternalFlags(InternalFlags);
 	}
 
