Whem migrating objects that are default subobjects, migrate their owner as well

diff --git a/Engine/Source/Runtime/CoreUObject/Private/UObject/RemoteObjectSerialization.cpp b/Engine/Source/Runtime/CoreUObject/Private/UObject/RemoteObjectSerialization.cpp
--- a/Engine/Source/Runtime/CoreUObject/Private/UObject/RemoteObjectSerialization.cpp
+++ b/Engine/Source/Runtime/CoreUObject/Private/UObject/RemoteObjectSerialization.cpp
@@ -22,6 +22,12 @@
 #define UE_WITH_REMOTEOBJECTWRITER_DEBUGGING 0
 #endif
 
+int32 GRemoteObjectsMigrateFullHierarchy = 1;
+static FAutoConsoleVariableRef CVarRemoteObjectsMigrateFullHierarchy(
+	TEXT("ro.MigrateFullHierarchy"),
+	GRemoteObjectsMigrateFullHierarchy,
+	TEXT("Whether remote objects that are default subobjects should be always migrated with their parent objects"));
+
 using FNameIndexType = FRemoteObjectPathName::FNameIndexType;
 
 FArchive& operator<<(FArchive& Ar, FRemoteObjectPathName& PathName)
@@ -716,7 +722,7 @@ void ResolvePathNames(FRemoteObjectData& InObjectData, TArray<UObject*>& OutReso
 }
 
 
-FRemoteObjectData* SerializeObjectData(UObject* Object, TSet<UObject*>& OutObjects)
+FRemoteObjectData* SerializeObjectData(UObject* InObject, TSet<UObject*>& OutObjects)
 {
 	using namespace UE::RemoteObject::Handle;
 
@@ -726,7 +732,21 @@ FRemoteObjectData* SerializeObjectData(UObject* Object, TSet<UObject*>& OutObjec
 	TArray<FRemoteObjectHeader> ObjectHeaders;
 	int32 Version = 0;
 	int64 HeaderOffset = 0;
+	FRemoteObjectId RequestedObjectId(InObject);
+	UObject* Object = InObject;
 
+	if (GRemoteObjectsMigrateFullHierarchy)
+	{
+		// If the object we want to migrate is a default subobject, make sure we migrate its parent too 
+		// so that it can safely be reconstructed on the other end (otherwise its outer may be missing)
+		while (Object && Object->IsDefaultSubobject())
+		{
+			Object = Object->GetOuter();
+		}
+		// Sanity check because by definition default subobjects should always have an outer
+		checkf(Object, TEXT("Attempted to migrate default subobject (%s) hierarchy but it doesn't have an outer or all outers are also default subobjects"), *GetPathNameSafe(InObject));
+	}
+	
 	{
 		FArchiveRemoteObjectWriter Ar(Object, *ObjectData, ObjectsToSerialize, OutObjects);
 		ObjectsToSerialize.Add(Object);
@@ -735,22 +755,37 @@ FRemoteObjectData* SerializeObjectData(UObject* Object, TSet<UObject*>& OutObjec
 		{
 			FRemoteObjectWriterScope Scope(Ar, TEXT("Header"));
 			Ar << Version;
+			Ar << RequestedObjectId;
 			OffsetOfHeaderOffset = Ar.Tell();
 			Ar << HeaderOffset;
 		}
 
-		for (int32 ObjectIndex = 0; ObjectIndex < ObjectsToSerialize.Num(); ++ObjectIndex)
+		int32 SerializedObjectIndex = 0;
+		bool bSerializedRequestedObject = false;
+
+		do
 		{
-			UObject* ObjectToSerialize = ObjectsToSerialize[ObjectIndex];
-			if (!OutObjects.Contains(ObjectToSerialize))
+			for (; SerializedObjectIndex < ObjectsToSerialize.Num(); ++SerializedObjectIndex)
+			{
+				UObject* ObjectToSerialize = ObjectsToSerialize[SerializedObjectIndex];
+				if (!OutObjects.Contains(ObjectToSerialize))
+				{
+					FRemoteObjectWriterScope Scope(Ar, ObjectToSerialize);
+					OutObjects.Add(ObjectToSerialize);
+					FRemoteObjectHeader& Header = ObjectHeaders.Emplace_GetRef();
+					SerializeRemoteObject(ObjectToSerialize, Ar, Header);
+					Ar.AddNameToNameMap(Header.Name);
+				}
+			}
+
+			bSerializedRequestedObject = OutObjects.Contains(InObject);
+			if (!bSerializedRequestedObject)
 			{
-				FRemoteObjectWriterScope Scope(Ar, ObjectToSerialize);
-				OutObjects.Add(ObjectToSerialize);
-				FRemoteObjectHeader& Header = ObjectHeaders.Emplace_GetRef();
-				SerializeRemoteObject(ObjectToSerialize, Ar, Header);
-				Ar.AddNameToNameMap(Header.Name);
+				// InObject was a default subobject (see GRemoteObjectsMigrateFullHierarchy) but when we serialized its parent 
+				// it turned out that the parent had no direct reference to InObject in which case we need to manually add InObject to ObjectsToSerialize list
+				ObjectsToSerialize.Add(InObject);
 			}
-		}
+		} while (!bSerializedRequestedObject);
 
 		HeaderOffset = Ar.Tell();
 		Ar.Seek(OffsetOfHeaderOffset);
@@ -766,21 +801,24 @@ FRemoteObjectData* SerializeObjectData(UObject* Object, TSet<UObject*>& OutObjec
 	return ObjectData;
 }
 
-void DeserializeObjectData(FRemoteObjectData& ObjectData, TArray<UObject*>& OutObjects)
+int32 DeserializeObjectData(FRemoteObjectData& ObjectData, TArray<UObject*>& OutObjects)
 {
 	using namespace UE::RemoteObject::Handle;
 	using namespace UE::RemoteObject::Private;
 
 	int32 Version = 0;
+	FRemoteObjectId RequestedObjectId;
 	int64 HeaderOffset = 0;
 	TArray<FRemoteObjectHeader> ObjectHeaders;
 	TArray<UObject*> ResolvedPathNameObjects;
+	int32 RequestedObjectIndex = -1;
 
 	ResolvePathNames(ObjectData, ResolvedPathNameObjects);
 
 	FArchiveRemoteObjectReader Ar(ObjectData, ResolvedPathNameObjects);
 
 	Ar << Version;
+	Ar << RequestedObjectId;
 	Ar << HeaderOffset;
 
 	const int64 ObjectDataOffset = Ar.Tell();
@@ -816,9 +854,14 @@ void DeserializeObjectData(FRemoteObjectData& ObjectData, TArray<UObject*>& OutO
 		}
 
 		// Deserialize all objects
-		for (UObject* Object : OutObjects)
+		for (int32 ObjectIndex = 0; ObjectIndex < OutObjects.Num(); ++ObjectIndex)
 		{
+			UObject* Object = OutObjects[ObjectIndex];
 			Object->Serialize(Ar);
+			if (RequestedObjectId == FRemoteObjectId(Object))
+			{
+				RequestedObjectIndex = ObjectIndex;
+			}
 		}
 
 		// Patch any unresolved remote references
@@ -828,6 +871,8 @@ void DeserializeObjectData(FRemoteObjectData& ObjectData, TArray<UObject*>& OutO
 			Object->Serialize(PatchAr);
 		}
 	}
+	checkf(RequestedObjectIndex >= 0, TEXT("Received remote object data but the requested object (%s) was not deserialized"), *RequestedObjectId.ToString());
+	return RequestedObjectIndex;
 }
 
 } // namespace UE::RemoteObject::Serialization
\ No newline at end of file
diff --git a/Engine/Source/Runtime/CoreUObject/Private/UObject/RemoteObjectSerialization.h b/Engine/Source/Runtime/CoreUObject/Private/UObject/RemoteObjectSerialization.h
--- a/Engine/Source/Runtime/CoreUObject/Private/UObject/RemoteObjectSerialization.h
+++ b/Engine/Source/Runtime/CoreUObject/Private/UObject/RemoteObjectSerialization.h
@@ -54,8 +54,22 @@ namespace UE::RemoteObject::Serialization
 		const FRemoteObjectConstructionParams* Find(FName InName, UObject* InOuter) const;
 	};
 
+	/**
+	* Serializes an object and its subobject (or if the object is a default subobject, its parent and the parent's subobjects)
+	* @param Object Object to be serialized
+	* @param OutObjects All objects that have been serialized (Object and its subobjects and/or parent)
+	* @return Remote object data representing the serialized objects
+	*/
 	FRemoteObjectData* SerializeObjectData(UObject* Object, TSet<UObject*>& OutObjects);
-	void DeserializeObjectData(FRemoteObjectData& ObjectData, TArray<UObject*>& OutReceivedObjects);
+
+	/**
+	* Deserializes remote object data
+	* @param ObjectData to be deserialized
+	* @param OutReceivedObjects All deserialized objects
+	* @return Index of an object in OutReceivedObjects that was the main object the migration request was triggered for 
+	* (usually 0 but if a migration requests a default subobject then its parent is also migrated and the return value will be > 0)
+	*/
+	int32 DeserializeObjectData(FRemoteObjectData& ObjectData, TArray<UObject*>& OutReceivedObjects);
 
 } // namespace UE::RemoveObject::Serialization
 
diff --git a/Engine/Source/Runtime/CoreUObject/Private/UObject/RemoteObjectTransfer.cpp b/Engine/Source/Runtime/CoreUObject/Private/UObject/RemoteObjectTransfer.cpp
--- a/Engine/Source/Runtime/CoreUObject/Private/UObject/RemoteObjectTransfer.cpp
+++ b/Engine/Source/Runtime/CoreUObject/Private/UObject/RemoteObjectTransfer.cpp
@@ -211,7 +211,7 @@ UObject* MigrateObjectFromObjectData(FRemoteObjectData& ObjectData, FRemoteServe
 	OnObjectDataReceivedDelegate.Broadcast(ObjectData, RemoteServerId);
 
 	TArray<UObject*> ReceivedObjects;
-	DeserializeObjectData(ObjectData, ReceivedObjects);
+	int32 RequestedObjectIndex = DeserializeObjectData(ObjectData, ReceivedObjects);
 
 	checkf(ReceivedObjects.Num() > 0, TEXT("PeerId:%s unable to deserialize object data (%d bytes)"), *GetGlobalServerId().ToString(), ObjectData.Bytes.Num());
 
@@ -235,7 +235,7 @@ UObject* MigrateObjectFromObjectData(FRemoteObjectData& ObjectData, FRemoteServe
 	// Notify that objects have been migrated from a remote server
 	OnObjectsReceivedDelegate.Broadcast(ReceivedObjects, RemoteServerId);
 
-	return ReceivedObjects[0];
+	return ReceivedObjects[RequestedObjectIndex];
 } 
 
 void MigrateObjectFromRemoteServer(FRemoteObjectId ObjectId, FRemoteServerId CurrentOwnerServerId, UObject* DestinationOuter)
