Fixed IDO serialization within CPFUO

#rnx

diff --git a/Engine/Source/Runtime/CoreUObject/Private/UObject/PropertyBagRepository.cpp b/Engine/Source/Runtime/CoreUObject/Private/UObject/PropertyBagRepository.cpp
--- a/Engine/Source/Runtime/CoreUObject/Private/UObject/PropertyBagRepository.cpp
+++ b/Engine/Source/Runtime/CoreUObject/Private/UObject/PropertyBagRepository.cpp
@@ -819,7 +819,7 @@ UObject* FPropertyBagRepository::CreateInstanceDataObject(UObject* Owner, FArchi
 	// enable impersonation so that the IDO gets loaded instead of Owner
 	TGuardValue<bool> ScopedImpersonateProperties(LoadContext->bImpersonateProperties, true);
 
-	if (FLinkerLoad* Linker = Owner->GetLinker())
+	if (FLinkerLoad* Linker = Owner->GetLinker(); Linker == &Archive)
 	{
 		Owner->SetFlags(RF_NeedLoad);
 		{
