FieldPath: Fix crash when calling ToString after a field is invalidated without a subsequent call to ResolveField

diff --git a/Engine/Source/Runtime/CoreUObject/Private/UObject/FieldPath.cpp b/Engine/Source/Runtime/CoreUObject/Private/UObject/FieldPath.cpp
--- a/Engine/Source/Runtime/CoreUObject/Private/UObject/FieldPath.cpp
+++ b/Engine/Source/Runtime/CoreUObject/Private/UObject/FieldPath.cpp
@@ -281,9 +281,9 @@ FString FFieldPath::ToString() const
 	FString Result;
 	if (UStruct* Owner = ResolvedOwner.Get())
 	{
-		if (ResolvedField)
+		if (FField* LocalResolvedField = GetTyped(FField::StaticClass()))
 		{
-			Result = ResolvedField->GetPathName();
+			Result = LocalResolvedField->GetPathName();
 		}
 		else
 		{
