Multiserver connections with remote objects support enabled will now serialize references to UObjects (either from TObjectPtr or TWeakObjectPtr UPROPERTY members or NetSerialize functions) as FRemoteObjectReference.

This is because remote objects are unique on each of the connected servers and should not be replicated the same way as server -> client connections do.

diff --git a/Engine/Plugins/Runtime/MultiServerReplication/Source/MultiServerReplication/Private/MultiServerBeaconClient.cpp b/Engine/Plugins/Runtime/MultiServerReplication/Source/MultiServerReplication/Private/MultiServerBeaconClient.cpp
--- a/Engine/Plugins/Runtime/MultiServerReplication/Source/MultiServerReplication/Private/MultiServerBeaconClient.cpp
+++ b/Engine/Plugins/Runtime/MultiServerReplication/Source/MultiServerReplication/Private/MultiServerBeaconClient.cpp
@@ -13,6 +13,12 @@
 
 #include UE_INLINE_GENERATED_CPP_BY_NAME(MultiServerBeaconClient)
 
+int32 GMultiServerAllowRemoteObjectReferences = 1;
+static FAutoConsoleVariableRef CVarMultiServerAllowRemoteObjectReferences(
+	TEXT("multiserver.AllowRemoteObjectReferences"),
+	GMultiServerAllowRemoteObjectReferences,
+	TEXT("Whether references to UObjects are replicated as remote object references between servers"));
+
 AMultiServerBeaconClient::AMultiServerBeaconClient() :
 	Super()
 {
@@ -117,6 +123,16 @@ void AMultiServerBeaconClient::OnFailure()
 	Super::OnFailure();
 }
 
+bool AMultiServerBeaconClient::InitBase()
+{
+	if (Super::InitBase() && NetDriver)
+	{
+		NetDriver->SetUsingRemoteObjectReferences(!!UE_WITH_REMOTE_OBJECT_HANDLE && GMultiServerAllowRemoteObjectReferences);
+		return true;
+	}
+	return false;
+}
+
 void AMultiServerBeaconClient::ServerUpdateLevelVisibility_Implementation(const FUpdateLevelVisibilityLevelInfo& LevelVisibility)
 {
 	if (GetConnectionState() == EBeaconConnectionState::Open)
diff --git a/Engine/Plugins/Runtime/MultiServerReplication/Source/MultiServerReplication/Private/MultiServerBeaconHost.cpp b/Engine/Plugins/Runtime/MultiServerReplication/Source/MultiServerReplication/Private/MultiServerBeaconHost.cpp
--- a/Engine/Plugins/Runtime/MultiServerReplication/Source/MultiServerReplication/Private/MultiServerBeaconHost.cpp
+++ b/Engine/Plugins/Runtime/MultiServerReplication/Source/MultiServerReplication/Private/MultiServerBeaconHost.cpp
@@ -7,6 +7,8 @@
 
 #include UE_INLINE_GENERATED_CPP_BY_NAME(MultiServerBeaconHost)
 
+extern int32 GMultiServerAllowRemoteObjectReferences;
+
 AMultiServerBeaconHost::AMultiServerBeaconHost(const FObjectInitializer& ObjectInitializer) :
 	Super(ObjectInitializer)
 {
@@ -39,6 +41,7 @@ bool AMultiServerBeaconHost::InitHost()
 				NetDriver->InitialConnectTimeout = BeaconConnectionInitialTimeout;
 				NetDriver->ConnectionTimeout = BeaconConnectionTimeout;
 				NetDriver->SetReplicateTransactionally(false);
+				NetDriver->SetUsingRemoteObjectReferences(!!UE_WITH_REMOTE_OBJECT_HANDLE && GMultiServerAllowRemoteObjectReferences);
 				return true;
 			}
 			else
diff --git a/Engine/Plugins/Runtime/MultiServerReplication/Source/MultiServerReplication/Public/MultiServerBeaconClient.h b/Engine/Plugins/Runtime/MultiServerReplication/Source/MultiServerReplication/Public/MultiServerBeaconClient.h
--- a/Engine/Plugins/Runtime/MultiServerReplication/Source/MultiServerReplication/Public/MultiServerBeaconClient.h
+++ b/Engine/Plugins/Runtime/MultiServerReplication/Source/MultiServerReplication/Public/MultiServerBeaconClient.h
@@ -31,6 +31,10 @@ public:
 	virtual void OnFailure() override;
 //~ End AOnlineBeaconClient interface
 
+//~ Begin AOnlineBeacon interface
+	virtual bool InitBase() override;
+//~ End AOnlineBeacon interface
+
 	/**
 	 * Attempt to connect to another specified MultiServer node.
 	 * AOnlineBeaconClient::OnFailure will be called immediately if there's a problem
diff --git a/Engine/Source/Runtime/Core/Private/Misc/NetworkVersion.cpp b/Engine/Source/Runtime/Core/Private/Misc/NetworkVersion.cpp
--- a/Engine/Source/Runtime/Core/Private/Misc/NetworkVersion.cpp
+++ b/Engine/Source/Runtime/Core/Private/Misc/NetworkVersion.cpp
@@ -339,4 +339,8 @@ void FNetworkVersion::DescribeNetworkRuntimeFeaturesBitset(EEngineNetworkRuntime
 	{
 		OutVerboseDescription.Append(TEXT("GenericReplication"));
 	}
+	if (EnumHasAnyFlags(FeaturesBitflag, EEngineNetworkRuntimeFeatures::UseRemoteObjectReferences))
+	{
+		OutVerboseDescription.Append(TEXT(", UseRemoteObjectReferences"));
+	}
 }
\ No newline at end of file
diff --git a/Engine/Source/Runtime/Core/Public/Misc/EngineNetworkCustomVersion.h b/Engine/Source/Runtime/Core/Public/Misc/EngineNetworkCustomVersion.h
--- a/Engine/Source/Runtime/Core/Public/Misc/EngineNetworkCustomVersion.h
+++ b/Engine/Source/Runtime/Core/Public/Misc/EngineNetworkCustomVersion.h
@@ -90,6 +90,7 @@ struct FEngineNetworkCustomVersion
 		RepMoveOptionalAcceleration = 35,		// Bump version to support serialization changes to RepMove for optional Acceleration
 		CustomExports = 36,						// Bump version to support CustomExports (such as NetTokens used for supporting "Dynamic Replication" of GameplayTags -- see FGameplayTag::NetSerialize_Packed)		// -----<new versions can be added above this line>-------------------------------------------------
 		MontagePlayCountSerialization = 37,		// Bump version to support serialization changes for dynamic montage play count 
+		RemoteObjectReferences = 38,			// Bump version to support serialization changes for remote object references 
 
 		// -----<new versions can be added above this line>-------------------------------------------------
 		VersionPlusOne,
diff --git a/Engine/Source/Runtime/Core/Public/Misc/NetworkVersion.h b/Engine/Source/Runtime/Core/Public/Misc/NetworkVersion.h
--- a/Engine/Source/Runtime/Core/Public/Misc/NetworkVersion.h
+++ b/Engine/Source/Runtime/Core/Public/Misc/NetworkVersion.h
@@ -45,6 +45,7 @@ enum class EEngineNetworkRuntimeFeatures : uint16
 {
 	None = 0,
 	IrisEnabled = 1 << None, // Are we running the Iris or Generic replication system
+	UseRemoteObjectReferences = 1 << IrisEnabled, // Are UObjects being serialized as remote object references
 };
 ENUM_CLASS_FLAGS(EEngineNetworkRuntimeFeatures);
 
diff --git a/Engine/Source/Runtime/CoreUObject/Private/UObject/CoreNet.cpp b/Engine/Source/Runtime/CoreUObject/Private/UObject/CoreNet.cpp
--- a/Engine/Source/Runtime/CoreUObject/Private/UObject/CoreNet.cpp
+++ b/Engine/Source/Runtime/CoreUObject/Private/UObject/CoreNet.cpp
@@ -7,6 +7,7 @@
 #include "UObject/CoreNet.h"
 #include "UObject/PropertyOptional.h"
 #include "UObject/UnrealType.h"
+#include "UObject/RemoteObjectTransfer.h"
 #include "Misc/EngineNetworkCustomVersion.h"
 
 DEFINE_LOG_CATEGORY_STATIC(LogCoreNet, Log, All);
@@ -440,6 +441,10 @@ FArchive& FNetBitWriter::operator<<(class FName& N)
 
 FArchive& FNetBitWriter::operator<<(UObject*& Object)
 {
+#if UE_WITH_REMOTE_OBJECT_HANDLE
+	UE_CLOG(PackageMap->IsUsingRemoteObjectReferences(), LogCoreNet, Warning, TEXT("Sending raw object reference when NetDriver expects remote references (%s)"), *GetFullNameSafe(Object));
+#endif
+
 	PackageMap->SerializeObject(*this, UObject::StaticClass(), Object);
 	return *this;
 }
@@ -491,11 +496,29 @@ FArchive& FNetBitWriter::operator<<(FSoftObjectPtr& Value)
 
 FArchive& FNetBitWriter::operator<<(FObjectPtr& Value)
 {
+#if UE_WITH_REMOTE_OBJECT_HANDLE
+	if (PackageMap->IsUsingRemoteObjectReferences())
+	{
+		FRemoteObjectReference RemoteReference(Value);
+		bool bResult = false;
+		RemoteReference.NetSerialize(*this, PackageMap, bResult);
+		return *this;
+	}
+#endif
 	return FArchiveUObject::SerializeObjectPtr(*this, Value);
 }
 
 FArchive& FNetBitWriter::operator<<(struct FWeakObjectPtr& WeakObjectPtr)
 {
+#if UE_WITH_REMOTE_OBJECT_HANDLE
+	if (PackageMap->IsUsingRemoteObjectReferences())
+	{
+		FRemoteObjectReference RemoteReference(WeakObjectPtr);
+		bool bResult = false;
+		RemoteReference.NetSerialize(*this, PackageMap, bResult);
+		return *this;
+	}
+#endif
 	return FArchiveUObject::SerializeWeakObjectPtr(*this, WeakObjectPtr);
 }
 
@@ -523,6 +546,10 @@ FNetBitReader::~FNetBitReader() = default;
 
 FArchive& FNetBitReader::operator<<(UObject*& Object)
 {
+#if UE_WITH_REMOTE_OBJECT_HANDLE
+	UE_CLOG(PackageMap->IsUsingRemoteObjectReferences(), LogCoreNet, Warning, TEXT("Received raw object reference when NetDriver expects remote references (%s)"), *GetFullNameSafe(Object));
+#endif
+
 	PackageMap->SerializeObject(*this, UObject::StaticClass(), Object);
 	return *this;
 }
@@ -594,11 +621,31 @@ FArchive& FNetBitReader::operator<<(FSoftObjectPtr& Value)
 
 FArchive& FNetBitReader::operator<<(FObjectPtr& Value)
 {
+#if UE_WITH_REMOTE_OBJECT_HANDLE
+	if (PackageMap->IsUsingRemoteObjectReferences())
+	{
+		FRemoteObjectReference RemoteReference;
+		bool bResult = false;
+		RemoteReference.NetSerialize(*this, PackageMap, bResult);
+		Value = RemoteReference.ToObjectPtr();
+		return *this;
+	}
+#endif
 	return FArchiveUObject::SerializeObjectPtr(*this, Value);
 }
 
 FArchive& FNetBitReader::operator<<(struct FWeakObjectPtr& WeakObjectPtr)
 {
+#if UE_WITH_REMOTE_OBJECT_HANDLE
+	if (PackageMap->IsUsingRemoteObjectReferences())
+	{
+		FRemoteObjectReference RemoteReference;
+		bool bResult = false;
+		RemoteReference.NetSerialize(*this, PackageMap, bResult);
+		WeakObjectPtr = RemoteReference.ToWeakPtr();
+		return *this;
+	}
+#endif
 	return FArchiveUObject::SerializeWeakObjectPtr(*this, WeakObjectPtr);
 }
 
diff --git a/Engine/Source/Runtime/CoreUObject/Private/UObject/PropertyObject.cpp b/Engine/Source/Runtime/CoreUObject/Private/UObject/PropertyObject.cpp
--- a/Engine/Source/Runtime/CoreUObject/Private/UObject/PropertyObject.cpp
+++ b/Engine/Source/Runtime/CoreUObject/Private/UObject/PropertyObject.cpp
@@ -9,6 +9,8 @@
 #include "UObject/LinkerPlaceholderExportObject.h"
 #include "UObject/LinkerPlaceholderClass.h"
 #include "UObject/PropertyHelper.h"
+#include "UObject/RemoteObjectTransfer.h"
+#include "UObject/CoreNet.h"
 
 /*-----------------------------------------------------------------------------
 	FObjectProperty.
@@ -333,6 +335,28 @@ void FObjectProperty::PostSerializeObjectItem(FArchive& SerializingArchive, void
 	}
 }
 
+#if UE_WITH_REMOTE_OBJECT_HANDLE
+bool FObjectProperty::NetSerializeItem(FArchive& Ar, UPackageMap* Map, void* Data, TArray<uint8>* MetaData) const
+{
+	const bool bUseRemoteObjectReference = Map && Map->IsUsingRemoteObjectReferences() && HasAnyPropertyFlags(CPF_TObjectPtr);
+	if (bUseRemoteObjectReference)
+	{
+		TObjectPtr<UObject>* ObjectPtr = GetPropertyValuePtr(Data);
+
+		FRemoteObjectReference RemoteReference(*ObjectPtr);
+		bool bResult = false;
+		RemoteReference.NetSerialize(Ar, Map, bResult);
+
+		if (Ar.IsLoading())
+		{
+			*ObjectPtr = TObjectPtr<UObject>(RemoteReference.ToObjectPtr());
+		}
+		return bResult;
+	}
+	return Super::NetSerializeItem(Ar, Map, Data, MetaData);
+}
+#endif // UE_WITH_REMOTE_OBJECT_HANDLE
+
 const TCHAR* FObjectProperty::ImportText_Internal(const TCHAR* Buffer, void* ContainerOrPropertyPtr, EPropertyPointerType PropertyPointerType, UObject* OwnerObject, int32 PortFlags, FOutputDevice* ErrorText) const
 {
 	const TCHAR* Result = TFObjectPropertyBase<TObjectPtr<UObject>>::ImportText_Internal(Buffer, ContainerOrPropertyPtr, PropertyPointerType, OwnerObject, PortFlags, ErrorText);
diff --git a/Engine/Source/Runtime/CoreUObject/Private/UObject/PropertyWeakObjectPtr.cpp b/Engine/Source/Runtime/CoreUObject/Private/UObject/PropertyWeakObjectPtr.cpp
--- a/Engine/Source/Runtime/CoreUObject/Private/UObject/PropertyWeakObjectPtr.cpp
+++ b/Engine/Source/Runtime/CoreUObject/Private/UObject/PropertyWeakObjectPtr.cpp
@@ -5,6 +5,8 @@
 #include "UObject/UnrealType.h"
 #include "UObject/LinkerPlaceholderExportObject.h"
 #include "UObject/LinkerPlaceholderClass.h"
+#include "UObject/RemoteObjectTransfer.h"
+#include "UObject/CoreNet.h"
 
 /*-----------------------------------------------------------------------------
 	FWeakObjectProperty.
@@ -104,6 +106,28 @@ void FWeakObjectProperty::SerializeItem( FStructuredArchive::FSlot Slot, void* V
 	}
 }
 
+#if UE_WITH_REMOTE_OBJECT_HANDLE
+bool FWeakObjectProperty::NetSerializeItem(FArchive& Ar, UPackageMap* Map, void* Data, TArray<uint8>* MetaData) const
+{
+	const bool bUseRemoteObjectReference = Map && Map->IsUsingRemoteObjectReferences();
+	if (bUseRemoteObjectReference)
+	{
+		FWeakObjectPtr* WeakPtr = GetPropertyValuePtr(Data);
+
+		FRemoteObjectReference RemoteReference(FObjectPtr(WeakPtr->GetRemoteId()));
+		bool bResult = false;
+		RemoteReference.NetSerialize(Ar, Map, bResult);
+
+		if (Ar.IsLoading())
+		{
+			*WeakPtr = RemoteReference.ToWeakPtr();
+		}
+		return bResult;
+	}
+	return Super::NetSerializeItem(Ar, Map, Data, MetaData);
+}
+#endif // UE_WITH_REMOTE_OBJECT_HANDLE
+
 UObject* FWeakObjectProperty::GetObjectPropertyValue(const void* PropertyValueAddress) const
 {
 	return GetPropertyValue(PropertyValueAddress).Get();
diff --git a/Engine/Source/Runtime/CoreUObject/Private/UObject/RemoteObjectTransfer.cpp b/Engine/Source/Runtime/CoreUObject/Private/UObject/RemoteObjectTransfer.cpp
--- a/Engine/Source/Runtime/CoreUObject/Private/UObject/RemoteObjectTransfer.cpp
+++ b/Engine/Source/Runtime/CoreUObject/Private/UObject/RemoteObjectTransfer.cpp
@@ -827,6 +827,42 @@ UObject* FRemoteObjectReference::Resolve() const
 	return ToObjectPtr().Get();
 }
 
+bool FRemoteObjectReference::NetSerialize(FArchive& Ar, class UPackageMap* Map, bool& bOutSuccess)
+{
+	// Serialize remote object id
+	Ar << *this;
+
+	// Serialize the pathname of the object if possible which will then be used on the receiving end to find assets loaded by both servers in memory
+	FString PathName;
+	if (Ar.IsSaving() && ObjectId.IsValid())
+	{		
+		if (UObject* ExistingObject = StaticFindObjectFastInternal(ObjectId))
+		{
+			PathName = ExistingObject->GetPathName();
+		}
+	}
+
+	Ar << PathName;
+
+	if (Ar.IsLoading() && PathName.Len())
+	{
+		// We currently can't differentiate between assets that exist on both servers which we don't want to migrate
+		// so try to find the object in memory first by its id and if it doesn't exist try to find the object by pathname
+		// and if such object does exist use its remote id to resolve this reference to avoid migrating the asset from another server
+		UObject* ExistingObject = StaticFindObjectFastInternal(ObjectId);
+		if (!ExistingObject)
+		{
+			ExistingObject = StaticFindObject(UObject::StaticClass(), nullptr, *PathName);
+			if (ExistingObject)
+			{
+				ObjectId = FRemoteObjectId(ExistingObject);
+			}
+		}
+	}
+	bOutSuccess = true;
+	return bOutSuccess;
+}
+
 FArchive& operator<<(FArchive& Ar, FRemoteObjectReference& Ref)
 {
 	Ar << Ref.ObjectId;
diff --git a/Engine/Source/Runtime/CoreUObject/Public/UObject/CoreNet.h b/Engine/Source/Runtime/CoreUObject/Public/UObject/CoreNet.h
--- a/Engine/Source/Runtime/CoreUObject/Public/UObject/CoreNet.h
+++ b/Engine/Source/Runtime/CoreUObject/Public/UObject/CoreNet.h
@@ -232,6 +232,7 @@ class UPackageMap : public UObject
 	virtual UObject*		GetObjectFromNetGUID( const FNetworkGUID& NetGUID, const bool bIgnoreMustBeMapped ) { return NULL; }
 	virtual FNetworkGUID	GetNetGUIDFromObject( const UObject* InObject) const { return FNetworkGUID(); }
 	virtual bool			IsGUIDBroken( const FNetworkGUID& NetGUID, const bool bMustBeRegistered ) const { return false; }
+	virtual bool			IsUsingRemoteObjectReferences() const { return false; }
 
 	COREUOBJECT_API virtual void Serialize(FArchive& Ar) override;
 
diff --git a/Engine/Source/Runtime/CoreUObject/Public/UObject/RemoteObjectTransfer.h b/Engine/Source/Runtime/CoreUObject/Public/UObject/RemoteObjectTransfer.h
--- a/Engine/Source/Runtime/CoreUObject/Public/UObject/RemoteObjectTransfer.h
+++ b/Engine/Source/Runtime/CoreUObject/Public/UObject/RemoteObjectTransfer.h
@@ -216,12 +216,7 @@ public:
 		return true;
 	}
 
-	bool NetSerialize(FArchive& Ar, class UPackageMap* Map, bool& bOutSuccess)
-	{
-		Ar << *this;
-		bOutSuccess = true;
-		return bOutSuccess;
-	}
+	COREUOBJECT_API bool NetSerialize(FArchive& Ar, class UPackageMap* Map, bool& bOutSuccess);
 
 	friend COREUOBJECT_API FArchive& operator<<(FArchive& Ar, FRemoteObjectReference& Ref);
 };
diff --git a/Engine/Source/Runtime/CoreUObject/Public/UObject/UnrealType.h b/Engine/Source/Runtime/CoreUObject/Public/UObject/UnrealType.h
--- a/Engine/Source/Runtime/CoreUObject/Public/UObject/UnrealType.h
+++ b/Engine/Source/Runtime/CoreUObject/Public/UObject/UnrealType.h
@@ -2877,6 +2877,9 @@ class COREUOBJECT_API FObjectProperty : public TFObjectPropertyBase<TObjectPtr<U
 
 	// FProperty interface
 	virtual void SerializeItem(FStructuredArchive::FSlot Slot, void* Value, void const* Defaults) const override;
+#if UE_WITH_REMOTE_OBJECT_HANDLE
+	virtual bool NetSerializeItem(FArchive& Ar, UPackageMap* Map, void* Data, TArray<uint8>* MetaData = nullptr) const override;
+#endif // UE_WITH_REMOTE_OBJECT_HANDLE
 	virtual void EmitReferenceInfo(UE::GC::FSchemaBuilder& Schema, int32 BaseOffset, TArray<const FStructProperty*>& EncounteredStructProps, UE::GC::FPropertyStack& DebugPath) override;
 	virtual const TCHAR* ImportText_Internal(const TCHAR* Buffer, void* ContainerOrPropertyPtr, EPropertyPointerType PropertyPointerType, UObject* OwnerObject, int32 PortFlags, FOutputDevice* ErrorText) const override;
 	virtual EConvertFromTypeResult ConvertFromType(const FPropertyTag& Tag, FStructuredArchive::FSlot Slot, uint8* Data, UStruct* DefaultsStruct, const uint8* Defaults) override;
@@ -2985,6 +2988,9 @@ class COREUOBJECT_API FWeakObjectProperty : public TFObjectPropertyBase<FWeakObj
 
 	// FProperty interface
 	virtual void SerializeItem(FStructuredArchive::FSlot Slot, void* Value, void const* Defaults) const override;
+#if UE_WITH_REMOTE_OBJECT_HANDLE
+	virtual bool NetSerializeItem(FArchive& Ar, UPackageMap* Map, void* Data, TArray<uint8>* MetaData = nullptr) const override;
+#endif
 protected:
 	virtual void LinkInternal(FArchive& Ar) override;
 private:
diff --git a/Engine/Source/Runtime/Engine/Classes/Engine/NetDriver.h b/Engine/Source/Runtime/Engine/Classes/Engine/NetDriver.h
--- a/Engine/Source/Runtime/Engine/Classes/Engine/NetDriver.h
+++ b/Engine/Source/Runtime/Engine/Classes/Engine/NetDriver.h
@@ -2245,6 +2245,15 @@ public:
 		bReplicateTransactionally = bInReplicateTransactionally;
 	}
 
+	/** Returns true if references to UObjects should be serialized as FRemoteObjectReference */
+	bool IsUsingRemoteObjectReferences() const { return bUseRemoteObjectRefs; }
+
+	/** When using remote objects with UE_WITH_REMOTE_OBJECT_HANDLE, determines whether or not references to UObjects (either from TObjectPtr or TWeakObjectPtr) are serialized as FRemoteObjectReference */
+	void SetUsingRemoteObjectReferences(bool bValue)
+	{
+		bUseRemoteObjectRefs = bValue;
+	}
+
 protected:
 	
 	/** Stream of random numbers to be used by this instance of UNetDriver */
@@ -2383,4 +2392,7 @@ private:
 
 	/** When using remote objects with UE_WITH_REMOTE_OBJECT_HANDLE, determines whether or not property replication happens within a transaction */
 	bool bReplicateTransactionally = true;
+
+	/** When using remote objects with UE_WITH_REMOTE_OBJECT_HANDLE, determines whether or not references to UObjects (either from TObjectPtr or TWeakObjectPtr) are serialized as FRemoteObjectReference */
+	bool bUseRemoteObjectRefs = false;
 };
diff --git a/Engine/Source/Runtime/Engine/Classes/Engine/PackageMapClient.h b/Engine/Source/Runtime/Engine/Classes/Engine/PackageMapClient.h
--- a/Engine/Source/Runtime/Engine/Classes/Engine/PackageMapClient.h
+++ b/Engine/Source/Runtime/Engine/Classes/Engine/PackageMapClient.h
@@ -456,6 +456,7 @@ public:
 	ENGINE_API virtual bool SerializeNewActor( FArchive& Ar, class UActorChannel *Channel, class AActor*& Actor) override;
 	
 	ENGINE_API virtual bool WriteObject( FArchive& Ar, UObject* InOuter, FNetworkGUID NetGUID, FString ObjName ) override;
+	ENGINE_API virtual bool IsUsingRemoteObjectReferences() const override;
 
 	// UPackageMapClient Connection specific methods
 
diff --git a/Engine/Source/Runtime/Engine/Private/NetDriver.cpp b/Engine/Source/Runtime/Engine/Private/NetDriver.cpp
--- a/Engine/Source/Runtime/Engine/Private/NetDriver.cpp
+++ b/Engine/Source/Runtime/Engine/Private/NetDriver.cpp
@@ -2304,6 +2304,12 @@ EEngineNetworkRuntimeFeatures UNetDriver::GetNetworkRuntimeFeatures() const
 		EnumAddFlags(NetDriverFeatures, EEngineNetworkRuntimeFeatures::IrisEnabled);
 	}
 
+	const bool bWithRemoteObjectReferences = !!UE_WITH_REMOTE_OBJECT_HANDLE;
+	if (bWithRemoteObjectReferences && bUseRemoteObjectRefs)	
+	{
+		EnumAddFlags(NetDriverFeatures, EEngineNetworkRuntimeFeatures::UseRemoteObjectReferences);
+	}
+
 	return NetDriverFeatures;
 }
 
diff --git a/Engine/Source/Runtime/Engine/Private/PackageMapClient.cpp b/Engine/Source/Runtime/Engine/Private/PackageMapClient.cpp
--- a/Engine/Source/Runtime/Engine/Private/PackageMapClient.cpp
+++ b/Engine/Source/Runtime/Engine/Private/PackageMapClient.cpp
@@ -466,6 +466,18 @@ bool UPackageMapClient::WriteObject( FArchive& Ar, UObject* ObjOuter, FNetworkGU
 	return true;
 }
 
+bool UPackageMapClient::IsUsingRemoteObjectReferences() const
+{
+	if (Connection)
+	{
+		if (UNetDriver* Driver = Connection->Driver.Get())
+		{
+			return Driver->IsUsingRemoteObjectReferences();
+		}
+	}
+	return false;
+}
+
 /**
  *	Standard method of serializing a new actor.
  *		For static actors, this will just be a single call to SerializeObject, since they can be referenced by their path name.
