Fixes for IDOs

- Fixed IDO creation to serialize script properties directly into the IDO rather than re-serializing the entire owner with only script properties redirected to the IDO.
- Fixed IDO CopyTaggedProperties to set bImpersonateProperties only on IDOs, otherwise copying to an IDO from its owner will copy from the IDO to itself due to impersonation.
- Fixed calculation of ExpectedSerialSize when serializing only script properties. It would previously fail if anything was serialized before script properties.
- Added tracking of script serialization offsets to FUObjectSerializeContext to avoid requiring the archive to track it and pass it through to IDO creation.
- Added an assert to tagged property serialization when the tag size is invalid, to avoid infinite loop with a tag size of zero.
- Also set bImpersonateProperties within FScopedObjectSerializeContext for loading an IDO.
- Renamed FScopedIDOSerializationContext to FScopedObjectSerializeContext.

diff --git a/Engine/Source/Runtime/CoreUObject/Internal/UObject/InstanceDataObjectUtils.h b/Engine/Source/Runtime/CoreUObject/Internal/UObject/InstanceDataObjectUtils.h
--- a/Engine/Source/Runtime/CoreUObject/Internal/UObject/InstanceDataObjectUtils.h
+++ b/Engine/Source/Runtime/CoreUObject/Internal/UObject/InstanceDataObjectUtils.h
@@ -46,6 +46,22 @@ UClass* CreateInstanceDataObjectClass(const FPropertyPathNameTree* PropertyTree,
 /** Copy whether each property was set by serialization from one IDO to another. */
 void CopyPropertyValueSerializedData(const FFieldVariant& OldField, void* OldDataPtr, const FFieldVariant& NewField, void* NewDataPtr);
 
+/** Returns whether the object is an IDO. */
+UE_API bool IsInstanceDataObject(const UObject* Object);
+
+/**
+ * Creates an IDO for Owner and serializes it from Archive.
+ *
+ * Creates an IDO class containing the union of the fields in Owner and any associated unknown property tree.
+ *
+ * @param Owner   Object from which to construct an InstanceDataObject.
+ * @param Archive   Archive from which to serialize script properties for the object.
+ * @param StartOffset   Offset to seek Archive to before serializing script properties for the object.
+ * @param EndOffset   Offset to expect Archive to be at after serializing script properties for the object.
+ * @return The InstanceDataObject for Owner.
+ */
+UE_API UObject* CreateInstanceDataObject(UObject* Owner, FArchive& Ar, int64 StartOffset, int64 EndOffset);
+
 /** Returns the IDO for the object if it has one, otherwise returns the object. */
 UE_API UObject* ResolveInstanceDataObject(UObject* Object);
 
diff --git a/Engine/Source/Runtime/CoreUObject/Internal/UObject/PropertyBagRepository.h b/Engine/Source/Runtime/CoreUObject/Internal/UObject/PropertyBagRepository.h
--- a/Engine/Source/Runtime/CoreUObject/Internal/UObject/PropertyBagRepository.h
+++ b/Engine/Source/Runtime/CoreUObject/Internal/UObject/PropertyBagRepository.h
@@ -64,12 +64,13 @@ public:
 	 * Creates an InstanceDataObject class containing the union of the fields in the owner and its associated
 	 * unknown property tree.
 	 *
-	 * @param Owner		Object from which to construct an InstanceDataObject.
-	 * @param Archive	Archive to serialize the IDO from. May be the linker for the Owner.
-	 * @param Offset	Offset to seek Archive to before serializing the object. Ignored if Archive is the linker.
-	 * @return			The InstanceDataObject object, UClass derived from associated unknown property tree.
+	 * @param Owner			Object from which to construct an InstanceDataObject.
+	 * @param Archive		Archive from which to serialize script properties for the object.
+	 * @param StartOffset	Offset to seek Archive to before serializing script properties for the object.
+	 * @param EndOffset		Offset to expect the Archive to be at after serializing script properties for the object.
+	 * @return				The InstanceDataObject object, UClass derived from associated unknown property tree.
 	 */
-	UObject* CreateInstanceDataObject(UObject* Owner, FArchive& Archive, int64 Offset);
+	UObject* CreateInstanceDataObject(UObject* Owner, FArchive& Archive, int64 StartOffset, int64 EndOffset);
 	UObject* DuplicateInstanceDataObject(UObject* SourceOwner, UObject* DestOwner);
 
 	// called at the end of postload to copy data from Owner to its IDO
@@ -157,32 +158,4 @@ private:
 
 #endif // WITH_EDITORONLY_DATA
 
-// construct this context in the same scope as an object is being serialized to support InstanceDataObjects.
-// if saving, this will simply set flags. If loading, an IDO will be constructed at the end of the scope when needed
-struct FScopedIDOSerializationContext
-{
-#if WITH_EDITORONLY_DATA
-	COREUOBJECT_API FScopedIDOSerializationContext(UObject* Object, FArchive& Archive);
-	COREUOBJECT_API ~FScopedIDOSerializationContext();
-
-private:
-	FArchive& Archive;
-	UObject* const Object;
-	const int64 PreSerializeOffset;
-	UObject* SavedSerializedObject;
-	bool bSavedTrackSerializedPropertyPath;
-	bool bSavedTrackInitializedProperties;
-	bool bSavedTrackSerializedProperties;
-	bool bSavedTrackUnknownProperties;
-	bool bSavedTrackUnknownEnumNames;
-	bool bSavedImpersonateProperties;
-	bool bCreateIDO;
-#else
-	inline FScopedIDOSerializationContext(UObject* Object, FArchive& Archive) {}
-#endif
-
-	FScopedIDOSerializationContext(const FScopedIDOSerializationContext&) = delete;
-	FScopedIDOSerializationContext& operator=(const FScopedIDOSerializationContext&) = delete;
-};
-
 } // UE
diff --git a/Engine/Source/Runtime/CoreUObject/Internal/UObject/UObjectSerializeContext.h b/Engine/Source/Runtime/CoreUObject/Internal/UObject/UObjectSerializeContext.h
new file mode 100644
--- /dev/null
+++ b/Engine/Source/Runtime/CoreUObject/Internal/UObject/UObjectSerializeContext.h
@@ -0,0 +1,50 @@
+// Copyright Epic Games, Inc. All Rights Reserved.
+
+#pragma once
+
+#include "CoreTypes.h"
+
+#define UE_API COREUOBJECT_API
+
+class FArchive;
+class UObject;
+
+namespace UE
+{
+
+/**
+* An object that configures the FUObjectSerializeContext for serialization of the object with the archive.
+*
+* Construct this on the stack within the scope that the object will be serialized.
+* Objects that support creation of an InstanceDataObject (IDO) on load will create the IDO when this is destructed.
+*/
+struct FScopedObjectSerializeContext
+{
+	UE_API FScopedObjectSerializeContext(UObject* Object, FArchive& Archive);
+	UE_API ~FScopedObjectSerializeContext();
+
+	FScopedObjectSerializeContext(const FScopedObjectSerializeContext&) = delete;
+	FScopedObjectSerializeContext& operator=(const FScopedObjectSerializeContext&) = delete;
+
+private:
+#if WITH_EDITORONLY_DATA
+	FArchive& Archive;
+	UObject* const Object;
+#endif
+	UObject* SavedSerializedObject;
+#if WITH_EDITORONLY_DATA
+	int64 SavedSerializedObjectScriptStartOffset;
+	int64 SavedSerializedObjectScriptEndOffset;
+	bool bSavedTrackSerializedPropertyPath;
+	bool bSavedTrackInitializedProperties;
+	bool bSavedTrackSerializedProperties;
+	bool bSavedTrackUnknownProperties;
+	bool bSavedTrackUnknownEnumNames;
+	bool bSavedImpersonateProperties;
+	bool bCreateInstanceDataObject;
+#endif
+};
+
+} // UE
+
+#undef UE_API
diff --git a/Engine/Source/Runtime/CoreUObject/Private/UObject/Class.cpp b/Engine/Source/Runtime/CoreUObject/Private/UObject/Class.cpp
--- a/Engine/Source/Runtime/CoreUObject/Private/UObject/Class.cpp
+++ b/Engine/Source/Runtime/CoreUObject/Private/UObject/Class.cpp
@@ -1934,6 +1934,9 @@ void UStruct::SerializeVersionedTaggedProperties(FStructuredArchive::FSlot Slot,
 				}
 				else if (Tag.Size != Loaded)
 				{
+					checkf(Tag.Size > 0, TEXT("Invalid property tag size %d in %s of %s of type %s. Package: %s. Property: '%s'. Type: '%s'."),
+						Tag.Size, *Tag.Name.ToString(), *GetName(), *WriteToString<64>(Tag.GetType()), *UnderlyingArchive.GetArchiveName(),
+						Property ? *Property->GetName() : TEXT(""), *GetFullName());
 					UnderlyingArchive.Seek(StartOfProperty + Tag.Size);
 				}
 			}
@@ -6080,9 +6083,15 @@ bool UClass::ImplementsInterface( const class UClass* SomeInterface ) const
  */
 void UClass::SerializeDefaultObject(UObject* Object, FStructuredArchive::FSlot Slot)
 {
-	// tell the archive that it's allowed to load data for transient properties
+	FUObjectSerializeContext* SerializeContext = FUObjectThreadContext::Get().GetSerializeContext();
 	FArchive& UnderlyingArchive = Slot.GetUnderlyingArchive();
 
+#if WITH_EDITORONLY_DATA
+	if (SerializeContext->SerializedObject == Object)
+	{
+		SerializeContext->SerializedObjectScriptStartOffset = UnderlyingArchive.Tell();
+	}
+#endif
 	UnderlyingArchive.MarkScriptSerializationStart(Object);
 	UnderlyingArchive.StartSerializingDefaults();
 
@@ -6103,6 +6112,12 @@ void UClass::SerializeDefaultObject(UObject* Object, FStructuredArchive::FSlot S
 	}
 	UnderlyingArchive.StopSerializingDefaults();
 	UnderlyingArchive.MarkScriptSerializationEnd(Object);
+#if WITH_EDITORONLY_DATA
+	if (SerializeContext->SerializedObject == Object)
+	{
+		SerializeContext->SerializedObjectScriptEndOffset = UnderlyingArchive.Tell();
+	}
+#endif
 }
 
 void UClass::SerializeSparseClassData(FStructuredArchive::FSlot Slot)
diff --git a/Engine/Source/Runtime/CoreUObject/Private/UObject/InstanceDataObjectUtils.cpp b/Engine/Source/Runtime/CoreUObject/Private/UObject/InstanceDataObjectUtils.cpp
--- a/Engine/Source/Runtime/CoreUObject/Private/UObject/InstanceDataObjectUtils.cpp
+++ b/Engine/Source/Runtime/CoreUObject/Private/UObject/InstanceDataObjectUtils.cpp
@@ -249,7 +249,7 @@ namespace UE
 	static const FName NAME_IDOMapKey(ANSITEXTVIEW("Key"));
 	static const FName NAME_IDOMapValue(ANSITEXTVIEW("Value"));
 
-	bool bEnableIDOSupport = false;
+	bool bEnableIDOSupport = true;
 	FAutoConsoleVariableRef EnableIDOSupportCVar(
 		TEXT("IDO.Enable"),
 		bEnableIDOSupport,
@@ -289,11 +289,16 @@ namespace UE
 
 	bool IsInstanceDataObjectSupportEnabled(const UObject* InObject)
 	{
-		if(!IsInstanceDataObjectSupportEnabled() || !InObject)
+		if (!IsInstanceDataObjectSupportEnabled() || !InObject)
 		{
 			return false;
 		}
 
+		if (IsInstanceDataObject(InObject))
+		{
+			return true;
+		}
+
 		// Property bag placeholder objects are always enabled for IDO support
 		if (FPropertyBagRepository::IsPropertyBagPlaceholderObject(InObject))
 		{
@@ -302,7 +307,7 @@ namespace UE
 
 		// Assume that if this object has an IDO that it's enabled. This assumption is important for objects
 		// that were reparented into the transient package but still need their loose properties CPFUOed to new instances
-		if (UE::FPropertyBagRepository::Get().HasInstanceDataObject(InObject))
+		if (FPropertyBagRepository::Get().HasInstanceDataObject(InObject))
 		{
 			return true;
 		}
@@ -832,12 +837,10 @@ namespace UE
 	void CopyTaggedProperties(const UObject* Source, UObject* Dest)
 	{
 		FUObjectSerializeContext* SerializeContext = FUObjectThreadContext::Get().GetSerializeContext();
-		TGuardValue<bool> ImpersonatePropertiesScope(SerializeContext->bImpersonateProperties, true);
 
-		const bool bIsCDO = Source->HasAnyFlags(RF_ClassDefaultObject);
-		
 		// Track only initialized properties when copying. This is required to skip uninitialized properties
 		// during saving and to mark initialized properties during loading.
+		const bool bIsCDO = Source->HasAnyFlags(RF_ClassDefaultObject);
 		TGuardValue<bool> ScopedTrackInitializedProperties(SerializeContext->bTrackInitializedProperties, !bIsCDO);
 		TGuardValue<bool> ScopedTrackSerializedProperties(SerializeContext->bTrackSerializedProperties, false);
 		TGuardValue<bool> ScopedTrackUnknownProperties(SerializeContext->bTrackUnknownProperties, false);
@@ -846,16 +849,22 @@ namespace UE
 		TArray<uint8> Buffer;
 		Buffer.Reserve(Source->GetClass()->GetStructureSize());
 
-		FObjectWriter Writer(Buffer);
-		FSerializingDefaultsScope WriterDefaultsScope(Writer, Source);
-		Writer.ArNoDelta = true;
-		Source->GetClass()->SerializeTaggedProperties(Writer, (uint8*)Source, Source->GetClass(), nullptr);
+		{
+			FObjectWriter Writer(Buffer);
+			FSerializingDefaultsScope WriterDefaultsScope(Writer, Source);
+			Writer.ArNoDelta = true;
+			TGuardValue<bool> ImpersonatePropertiesScope(SerializeContext->bImpersonateProperties, IsInstanceDataObject(Source));
+			Source->GetClass()->SerializeTaggedProperties(Writer, (uint8*)Source, Source->GetClass(), nullptr);
+		}
 
-		FObjectReader Reader(Buffer);
-		FSerializingDefaultsScope ReaderDefaultsScope(Reader, Dest);
-		Reader.ArMergeOverrides = true;
-		Reader.ArPreserveArrayElements = true;
-		Dest->GetClass()->SerializeTaggedProperties(Reader, (uint8*)Dest, Dest->GetClass(), nullptr);
+		{
+			FObjectReader Reader(Buffer);
+			FSerializingDefaultsScope ReaderDefaultsScope(Reader, Dest);
+			Reader.ArMergeOverrides = true;
+			Reader.ArPreserveArrayElements = true;
+			TGuardValue<bool> ImpersonatePropertiesScope(SerializeContext->bImpersonateProperties, IsInstanceDataObject(Dest));
+			Dest->GetClass()->SerializeTaggedProperties(Reader, (uint8*)Dest, Dest->GetClass(), nullptr);
+		}
 	}
 
 	static void SetClassFlags(UClass* IDOClass, const UClass* OwnerClass)
@@ -1015,11 +1024,21 @@ namespace UE
 		}
 	}
 
-	UObject* ResolveInstanceDataObject(UObject* Object)
-	{
-		UObject* InstanceDataObject = FPropertyBagRepository::Get().FindInstanceDataObject(Object);
-		return InstanceDataObject ? InstanceDataObject : Object;
-	}
+bool IsInstanceDataObject(const UObject* Object)
+{
+	return Object && Object->GetClass()->UObject::IsA(UInstanceDataObjectClass::StaticClass());
+}
+
+UObject* CreateInstanceDataObject(UObject* Owner, FArchive& Ar, int64 StartOffset, int64 EndOffset)
+{
+	return FPropertyBagRepository::Get().CreateInstanceDataObject(Owner, Ar, StartOffset, EndOffset);
+}
+
+UObject* ResolveInstanceDataObject(UObject* Object)
+{
+	UObject* InstanceDataObject = FPropertyBagRepository::Get().FindInstanceDataObject(Object);
+	return InstanceDataObject ? InstanceDataObject : Object;
+}
 
 } // UE
 
diff --git a/Engine/Source/Runtime/CoreUObject/Private/UObject/LinkerLoad.cpp b/Engine/Source/Runtime/CoreUObject/Private/UObject/LinkerLoad.cpp
--- a/Engine/Source/Runtime/CoreUObject/Private/UObject/LinkerLoad.cpp
+++ b/Engine/Source/Runtime/CoreUObject/Private/UObject/LinkerLoad.cpp
@@ -23,6 +23,7 @@
 #include "UObject/PropertyBagRepository.h"
 #include "UObject/SoftObjectPath.h"
 #include "UObject/UObjectHash.h"
+#include "UObject/UObjectSerializeContext.h"
 #include "Misc/PackageName.h"
 #include "Blueprint/BlueprintSupport.h"
 #include "Misc/PackageAccessTrackingOps.h"
@@ -415,7 +416,6 @@ void FLinkerLoad::CreateActiveRedirectsMap(const FString& GEngineIniName)
 FScopedCreateImportCounter::FScopedCreateImportCounter(FLinkerLoad* Linker, int32 Index)
 {
 	LoadContext = FUObjectThreadContext::Get().GetSerializeContext();
-	check(LoadContext);
 
 	// Remember the old linker and index
 	PreviousLinker = LoadContext->SerializedImportLinker;
@@ -444,7 +444,6 @@ struct FScopedCreateExportCounter
 	FScopedCreateExportCounter(FLinkerLoad* Linker, int32 Index)
 	{
 		LoadContext = FUObjectThreadContext::Get().GetSerializeContext();
-		check(LoadContext);
 
 		// Remember the old linker and index
 		PreviousLinker = LoadContext->SerializedExportLinker;
@@ -699,16 +698,13 @@ void FLinkerLoad::PRIVATE_PatchNewObjectIntoExport(int32 OldExportIndex, UObject
 			// If the object was in the ObjLoaded queue (exported, but not yet serialized), swap out for our new object
 			FUObjectSerializeContext* LoadContext = InLoadContext ? InLoadContext
 				: FUObjectThreadContext::Get().GetSerializeContext();
-			if (LoadContext)
+			if (!LoadContext->PRIVATE_PatchNewObjectIntoExport(OldObject, NewObject))
 			{
-				if (!LoadContext->PRIVATE_PatchNewObjectIntoExport(OldObject, NewObject))
+				// Make sure the new object gets PostLoad called on it if it needs it:
+				// it wasn't in the queue so add it ObjLoaded list
+				if (OldObjectFlags & RF_NeedPostLoad)
 				{
-					// Make sure the new object gets PostLoad called on it if it needs it:
-					// it wasn't in the queue so add it ObjLoaded list
-					if (OldObjectFlags & RF_NeedPostLoad)
-					{
-						LoadContext->AddLoadedObject(NewObject);
-					}
+					LoadContext->AddLoadedObject(NewObject);
 				}
 			}
 		}
@@ -3181,7 +3177,6 @@ void FLinkerLoad::GatherImportDependencies(int32 ImportIndex, TSet<FDependencyRe
 	}
 
 	FUObjectSerializeContext* LoadContext = FUObjectThreadContext::Get().GetSerializeContext();
-	check(LoadContext);
 
 	BeginLoad(LoadContext, TEXT("GatherImportDependencies"));
 
@@ -3385,7 +3380,6 @@ FLinkerLoad::EVerifyResult FLinkerLoad::VerifyImport(int32 ImportIndex)
 
 						// now, fake our Import to be what the redirector pointed to
 						Import.XObject = Redir->DestinationObject;
-						check(CurrentLoadContext);
 						CurrentLoadContext->IncrementImportCount();
 						FLinkerManager::Get().AddLoaderWithNewImports(this);
 					}
@@ -3855,19 +3849,16 @@ bool FLinkerLoad::VerifyImportInner(const int32 ImportIndex, FString& WarningSuf
 #if WITH_EDITOR
 				else
 				{
+					// If we're serializing a redirector's destination object, validate/create the outer package object if it's a missing type. If this import
+					// represents a non-native type object that's no longer valid, this will allow exports of that type to still be serialized to a property bag,
+					// by creating a placeholder type object in its place. This way we won't lose any previously-serialized data for exports missing their type.
 					const FUObjectSerializeContext* SerializeContext = FUObjectThreadContext::Get().GetSerializeContext();
-					if (ensure(SerializeContext))
+					const UObjectRedirector* Redirector = Cast<UObjectRedirector>(SerializeContext->SerializedObject);
+					if (Redirector && Redirector->IsSerializingDestinationObject() && TryCreatePlaceholderClassImport(ImportIndex))
 					{
-						// If we're serializing a redirector's destination object, validate/create the outer package object if it's a missing type. If this import
-						// represents a non-native type object that's no longer valid, this will allow exports of that type to still be serialized to a property bag,
-						// by creating a placeholder type object in its place. This way we won't lose any previously-serialized data for exports missing their type.
-						const UObjectRedirector* Redirector = Cast<UObjectRedirector>(SerializeContext->SerializedObject);
-						if (Redirector && Redirector->IsSerializingDestinationObject() && TryCreatePlaceholderClassImport(ImportIndex))
-						{
-							// We don't need to do a package search for the import below since we've just created it. Return false to signal that there is no failure.
-							check(Import.XObject);
-							return false;
-						}
+						// We don't need to do a package search for the import below since we've just created it. Return false to signal that there is no failure.
+						check(Import.XObject);
+						return false;
 					}
 				}
 #endif
@@ -4082,7 +4073,6 @@ bool FLinkerLoad::VerifyImportInner(const int32 ImportIndex, FString& WarningSuf
 				// except if we are looking for _the_ package...in which case we are looking for TmpPkg, so we are done
 				Import.XObject = TmpPkg;
 				FUObjectSerializeContext* CurrentLoadContext = FUObjectThreadContext::Get().GetSerializeContext();
-				check(CurrentLoadContext);
 				CurrentLoadContext->IncrementImportCount();
 				FLinkerManager::Get().AddLoaderWithNewImports(this);
 				return false;
@@ -4141,7 +4131,6 @@ bool FLinkerLoad::VerifyImportInner(const int32 ImportIndex, FString& WarningSuf
 				{
 					Import.XObject = FindObject;
 					FUObjectSerializeContext* CurrentLoadContext = FUObjectThreadContext::Get().GetSerializeContext();
-					check(CurrentLoadContext);
 					CurrentLoadContext->IncrementImportCount();
 					FLinkerManager::Get().AddLoaderWithNewImports(this);
 				}
@@ -4645,8 +4634,6 @@ void FLinkerLoad::Preload( UObject* Object )
 	// Preload the object if necessary.
 	if (Object->HasAnyFlags(RF_NeedLoad))
 	{
-		FUObjectSerializeContext* CurrentLoadContext = FUObjectThreadContext::Get().GetSerializeContext();
-
 		if (Object->GetLinker() == this)
 		{
 			TRACE_CPUPROFILER_EVENT_SCOPE(FLinkerLoad::Preload);
@@ -4655,7 +4642,8 @@ void FLinkerLoad::Preload( UObject* Object )
 #endif // LOADTIMEPROFILERTRACE_ENABLED
 
 			UClass* Cls = Cast<UClass>(Object);
-			checkf(!GEventDrivenLoaderEnabled || !bLockoutLegacyOperations || !EVENT_DRIVEN_ASYNC_LOAD_ACTIVE_AT_RUNTIME, TEXT("Invalid call to FLinkerLoad::Preload while using the EDL. '%s' should have been reported via GetPreloadDependencies instead."), *Object->GetPathName());
+			checkf(!GEventDrivenLoaderEnabled || !bLockoutLegacyOperations || !EVENT_DRIVEN_ASYNC_LOAD_ACTIVE_AT_RUNTIME,
+				TEXT("Invalid call to FLinkerLoad::Preload while using the EDL. '%s' should have been reported via GetPreloadDependencies instead."), *Object->GetPathName());
 #if USE_CIRCULAR_DEPENDENCY_LOAD_DEFERRING
 			bool const bIsNonNativeObject = !Object->GetOutermost()->HasAnyPackageFlags(PKG_CompiledIn);
 			// we can determine that this is a blueprint class/struct by checking if it 
@@ -4731,13 +4719,13 @@ void FLinkerLoad::Preload( UObject* Object )
 				bool bSerializeOnlyScriptProperties = false;
 				const bool bDoesSavedClassMatchActualClass = DoesSavedClassMatchActualClass(ExportIndex);
 				FGuardValue_Bitfield(bIsLoadingToPropertyBagObject, UE::FPropertyBagRepository::IsPropertyBagPlaceholderObject(Object));
-				if(UEVer() >= EUnrealEngineObjectUE5Version::SCRIPT_SERIALIZATION_OFFSET)
+				if (UEVer() >= EUnrealEngineObjectUE5Version::SCRIPT_SERIALIZATION_OFFSET)
 				{
 					if (bIsLoadingToPropertyBagObject || !bDoesSavedClassMatchActualClass)
 					{
 						// note: script start/end offsets are relative to the export's offset in the file
 						StartPos += Export.ScriptSerializationStartOffset;
-						ExpectedSerialSize = Export.ScriptSerializationEndOffset;
+						ExpectedSerialSize = Export.ScriptSerializationEndOffset - Export.ScriptSerializationStartOffset;
 						bSerializeOnlyScriptProperties = true;	// signals that we can safely narrow the load to SerializeScriptProperties()
 					}
 				}
@@ -4785,8 +4773,7 @@ void FLinkerLoad::Preload( UObject* Object )
 						ESoftObjectPathSerializeType::AlwaysSerialize);
 #endif
 
-					// Toggle support for IDOs
-					UE::FScopedIDOSerializationContext IDOLoadContext(Object, *this);
+					UE::FScopedObjectSerializeContext ObjectSerializeScope(Object, *this);
 
 					if (Object->HasAnyFlags(RF_ClassDefaultObject))
 					{
@@ -4813,11 +4800,6 @@ void FLinkerLoad::Preload( UObject* Object )
 							return;
 						}
 #endif // USE_CIRCULAR_DEPENDENCY_LOAD_DEFERRING
-						
-						check(CurrentLoadContext);
-						// Maintain the current SerializedObject.
-						UObject* PrevSerializedObject = CurrentLoadContext->SerializedObject;
-						CurrentLoadContext->SerializedObject = Object;
 
 #if WITH_EDITOR && WITH_TEXT_ARCHIVE_SUPPORT
 						if (IsTextFormat())
@@ -4842,27 +4824,19 @@ void FLinkerLoad::Preload( UObject* Object )
 						}
 
 						Object->SetFlags(RF_LoadCompleted);
-						CurrentLoadContext->SerializedObject = PrevSerializedObject;
 					}
 					else
 					{
-
 #if WITH_EDITOR
 						static const FName NAME_UObjectSerialize = FName(TEXT("UObject::Serialize, Name, ClassName"));
 						FArchive::FScopeAddDebugData P(*this, NAME_UObjectSerialize);
 						FArchive::FScopeAddDebugData N(*this, Object->GetFName());
 						FArchive::FScopeAddDebugData C(*this, Object->GetClass()->GetFName());
-#endif
-						check(CurrentLoadContext);
-#if WITH_EDITOR
+
 						SCOPED_LOADTIMER_TEXT(*WriteToString<128>(GetClassTraceScope(Object), TEXTVIEW("_LoadSerialize")));
 						TRACE_CPUPROFILER_EVENT_SCOPE_TEXT_ON_CHANNEL(*Object->GetFullName(), AssetLoadTimeChannel);
 #endif
 
-						// Maintain the current SerializedObject.
-						UObject* PrevSerializedObject = CurrentLoadContext->SerializedObject;
-						CurrentLoadContext->SerializedObject = Object;
-
 #if WITH_EDITOR && WITH_TEXT_ARCHIVE_SUPPORT
 						if (IsTextFormat())
 						{
@@ -4907,7 +4881,6 @@ void FLinkerLoad::Preload( UObject* Object )
 						check(!bIsSerializingScriptProperties);
 #endif
 						Object->SetFlags(RF_LoadCompleted);
-						CurrentLoadContext->SerializedObject = PrevSerializedObject;
 					}
 				}
 
@@ -5042,8 +5015,7 @@ void FLinkerLoad::Preload( UObject* Object )
 				{
 					Object->SetFlags(RF_NeedPostLoad|RF_WasLoaded);
 					check(LinkerRoot && LinkerRoot == Object->GetOutermost());
-					check(CurrentLoadContext);
-					CurrentLoadContext->AddLoadedObject(Object);
+					FUObjectThreadContext::Get().GetSerializeContext()->AddLoadedObject(Object);
 				}
 			}
 		}
@@ -5300,7 +5272,6 @@ UObject* FLinkerLoad::CreateExport( int32 Index )
 			// will find and return an existing object if one exists and only create a new one if there doesn't.
 			Export.Object = CreatePackage( *Export.ObjectName.ToString() );
 			check(Export.Object);
-			check(CurrentLoadContext);
 			CurrentLoadContext->IncrementForcedExportCount();
 			FLinkerManager::Get().AddLoaderWithForcedExports(this);
 		}
@@ -5458,7 +5429,6 @@ UObject* FLinkerLoad::CreateExport( int32 Index )
 			// Mark that we need to dissociate forced exports later on if we are a forced export.
 			if (Export.bForcedExport)
 			{
-				check(CurrentLoadContext);
 				CurrentLoadContext->IncrementForcedExportCount();
 				FLinkerManager::Get().AddLoaderWithForcedExports(this);
 			}
@@ -5473,7 +5443,6 @@ UObject* FLinkerLoad::CreateExport( int32 Index )
 
 				// If this object was allocated but never loaded (components created by a constructor) make sure it gets loaded
 				// Don't do this for any packages that have previously fully loaded as they may have in memory changes
-				check(CurrentLoadContext);
 				CurrentLoadContext->AddLoadedObject(Export.Object);
 				if (!Export.Object->HasAnyFlags(RF_LoadCompleted) && (!LinkerRoot->IsFullyLoaded() || IsBlueprintFinalizationPending()))
 				{
@@ -5505,7 +5474,6 @@ UObject* FLinkerLoad::CreateExport( int32 Index )
 				LoadClass = UObjectRedirector::StaticClass();
 				// Create new import for UObjectRedirector class
 				FObjectImport& RedirectorImport = ImportMap.Emplace_GetRef(UObjectRedirector::StaticClass());
-				check(CurrentLoadContext);
 				CurrentLoadContext->IncrementImportCount();
 				FLinkerManager::Get().AddLoaderWithNewImports(this);				
 				Export.ClassIndex = FPackageIndex::FromImport(ImportMap.Num() - 1);
@@ -5527,7 +5495,6 @@ UObject* FLinkerLoad::CreateExport( int32 Index )
 				Export.Object = ActualObjectWithTheName;
 				Export.Object->SetLinker(this, Index);
 				Export.Object->SetFlags(RF_NeedLoad|RF_NeedPostLoad|RF_NeedPostLoadSubobjects|RF_WasLoaded);
-				check(CurrentLoadContext);
 				CurrentLoadContext->AddLoadedObject(Export.Object);
 				return ActualObjectWithTheName;
 			}
@@ -5708,7 +5675,6 @@ UObject* FLinkerLoad::CreateExport( int32 Index )
 				// so add the object to the list of objects that need to be loaded, which will be processed
 				// in EndLoad()
 				Export.Object->SetLinker( this, Index );
-				check(CurrentLoadContext);
 				CurrentLoadContext->AddLoadedObject(Export.Object);
 			}
 		}
@@ -5750,7 +5716,7 @@ UObject* FLinkerLoad::CreateExport( int32 Index )
 						// fully regenerated before another (there is no changing 
 						// that); so dependencies need to be recompiled later (with
 						// all the regenerated classes in place)
-						FScopedClassDependencyGather DependencyHelper(ClassObject, FUObjectThreadContext::Get().GetSerializeContext());
+						FScopedClassDependencyGather DependencyHelper(ClassObject, CurrentLoadContext);
 #endif //WITH_EDITOR
 
 						ClassObject->Bind();
@@ -5773,7 +5739,6 @@ UObject* FLinkerLoad::CreateExport( int32 Index )
 			// Mark that we need to dissociate forced exports later on.
 			if( Export.bForcedExport )
 			{
-				check(CurrentLoadContext);
 				CurrentLoadContext->IncrementForcedExportCount();
 				FLinkerManager::Get().AddLoaderWithForcedExports(this);
 			}
@@ -5920,7 +5885,6 @@ UObject* FLinkerLoad::CreateImport( int32 Index )
 						{
 							// Associate import and indicate that we associated an import for later cleanup.
 							Import.XObject = FindObject;
-							check(CurrentLoadContext);
 							CurrentLoadContext->IncrementImportCount();
 							FLinkerManager::Get().AddLoaderWithNewImports(this);
 						}
@@ -5970,7 +5934,6 @@ UObject* FLinkerLoad::CreateImport( int32 Index )
 					}
 				}
 #endif
-				check(CurrentLoadContext);
 				CurrentLoadContext->IncrementImportCount();
 				FLinkerManager::Get().AddLoaderWithNewImports(this);
 			}
@@ -6193,7 +6156,6 @@ void FLinkerLoad::Detach()
 	if (!FPlatformProperties::HasEditorOnlyData())
 	{
 		FUObjectSerializeContext* CurrentLoadContext = FUObjectThreadContext::Get().GetSerializeContext();
-		check(CurrentLoadContext);
 		CurrentLoadContext->RemoveDelayedLinkerClosePackage(this);
 	}
 
@@ -6441,13 +6403,16 @@ void FLinkerLoad::MarkScriptSerializationStart( const UObject* Obj )
 		if (ExportMap.IsValidIndex(Index))
 		{
 			FObjectExport& Export = ExportMap[Index];
+			const int64 RelativeSerialOffset = Tell() - Export.SerialOffset;
 			if (!UseUnversionedPropertySerialization() && UEVer() >= EUnrealEngineObjectUE5Version::SCRIPT_SERIALIZATION_OFFSET)
 			{
-				checkf(Export.ScriptSerializationStartOffset == (Tell() - Export.SerialOffset), TEXT("Serialized script property start offset does not match offset during deserialization"));
+				checkf(Export.ScriptSerializationStartOffset == RelativeSerialOffset,
+					TEXT("Serialized script property start offset %" INT64_FMT " does not match offset during deserialization %" INT64_FMT " for object %s in %s."),
+					Export.ScriptSerializationStartOffset, RelativeSerialOffset, *Export.ObjectName.ToString(), *LinkerRoot->GetName());
 			}
 			else
 			{
-				Export.ScriptSerializationStartOffset = Tell() - Export.SerialOffset;
+				Export.ScriptSerializationStartOffset = RelativeSerialOffset;
 			}
 		}
 	}
@@ -6467,13 +6432,16 @@ void FLinkerLoad::MarkScriptSerializationEnd( const UObject* Obj )
 		if (ExportMap.IsValidIndex(Index))
 		{
 			FObjectExport& Export = ExportMap[Index];
+			const int64 RelativeSerialOffset = Tell() - Export.SerialOffset;
 			if (!UseUnversionedPropertySerialization() && UEVer() >= EUnrealEngineObjectUE5Version::SCRIPT_SERIALIZATION_OFFSET)
 			{
-				checkf(Export.ScriptSerializationEndOffset == (Tell() - Export.SerialOffset), TEXT("Serialized script property end offset does not match offset during deserialization"));
+				checkf(Export.ScriptSerializationEndOffset == RelativeSerialOffset,
+					TEXT("Serialized script property end offset %" INT64_FMT " does not match offset during deserialization %" INT64_FMT " for object %s in %s."),
+					Export.ScriptSerializationEndOffset, RelativeSerialOffset, *Export.ObjectName.ToString(), *LinkerRoot->GetName());
 			}
 			else
 			{
-				Export.ScriptSerializationEndOffset = Tell() - Export.SerialOffset;
+				Export.ScriptSerializationEndOffset = RelativeSerialOffset;
 			}
 		}
 	}
@@ -6602,7 +6570,6 @@ UObject* FLinkerLoad::GetArchetypeFromLoader(const UObject* Obj)
 	if (GEventDrivenLoaderEnabled)
 	{
 		FUObjectSerializeContext* CurrentLoadContext = FUObjectThreadContext::Get().GetSerializeContext();
-		check(CurrentLoadContext);
 		check(!TemplateForGetArchetypeFromLoader || CurrentLoadContext->SerializedObject == Obj);
 		return TemplateForGetArchetypeFromLoader;
 	}
diff --git a/Engine/Source/Runtime/CoreUObject/Private/UObject/Obj.cpp b/Engine/Source/Runtime/CoreUObject/Private/UObject/Obj.cpp
--- a/Engine/Source/Runtime/CoreUObject/Private/UObject/Obj.cpp
+++ b/Engine/Source/Runtime/CoreUObject/Private/UObject/Obj.cpp
@@ -1890,8 +1890,15 @@ namespace UE::Private
 
 void UObject::SerializeScriptProperties( FStructuredArchive::FSlot Slot ) const
 {
+	FUObjectSerializeContext* SerializeContext = FUObjectThreadContext::Get().GetSerializeContext();
 	FArchive& UnderlyingArchive = Slot.GetUnderlyingArchive();
 
+#if WITH_EDITORONLY_DATA
+	if (SerializeContext->SerializedObject == this)
+	{
+		SerializeContext->SerializedObjectScriptStartOffset = UnderlyingArchive.Tell();
+	}
+#endif
 	UnderlyingArchive.MarkScriptSerializationStart(this);
 	if( HasAnyFlags(RF_ClassDefaultObject) )
 	{
@@ -1962,6 +1969,12 @@ void UObject::SerializeScriptProperties( FStructuredArchive::FSlot Slot ) const
 		UnderlyingArchive.StopSerializingDefaults();
 	}
 	UnderlyingArchive.MarkScriptSerializationEnd(this);
+#if WITH_EDITORONLY_DATA
+	if (SerializeContext->SerializedObject == this)
+	{
+		SerializeContext->SerializedObjectScriptEndOffset = UnderlyingArchive.Tell();
+	}
+#endif
 }
 
 
diff --git a/Engine/Source/Runtime/CoreUObject/Private/UObject/PropertyBagRepository.cpp b/Engine/Source/Runtime/CoreUObject/Private/UObject/PropertyBagRepository.cpp
--- a/Engine/Source/Runtime/CoreUObject/Private/UObject/PropertyBagRepository.cpp
+++ b/Engine/Source/Runtime/CoreUObject/Private/UObject/PropertyBagRepository.cpp
@@ -18,6 +18,7 @@
 #include "UObject/PropertyStateTracking.h"
 #include "UObject/UObjectGlobals.h"
 #include "UObject/UObjectHash.h"
+#include "UObject/UObjectSerializeContext.h"
 #include "UObject/UObjectThreadContext.h"
 #include "Templates/UnrealTemplate.h"
 
@@ -784,7 +785,7 @@ FString FPropertyBagRepository::GetReferencerName() const
 	return TEXT("FPropertyBagRepository");
 }
 
-UObject* FPropertyBagRepository::CreateInstanceDataObject(UObject* Owner, FArchive& Archive, int64 PreSerializeOffset)
+UObject* FPropertyBagRepository::CreateInstanceDataObject(UObject* Owner, FArchive& Archive, int64 StartOffset, int64 EndOffset)
 {
 	FPropertyBagRepositoryLock LockRepo(this);
 	FPropertyBagAssociationData& BagData = AssociatedData.FindOrAdd(Owner);
@@ -793,120 +794,41 @@ UObject* FPropertyBagRepository::CreateInstanceDataObject(UObject* Owner, FArchi
 		return BagData.InstanceDataObject;
 	}
 
+	// TODO: Should the IDO or its class be in another package?
+	UObject* Outer = GetTransientPackage();
+
+	// Construct the class for the IDO.
 	TSharedPtr<FPropertyPathNameTree> PropertyTree = FUnknownPropertyTree(Owner).Find();
 	const FUnknownEnumNames EnumNames(Owner);
 	const FUnknownEnumNames* EnumNamesPtr = !EnumNames.IsEmpty() ? &EnumNames : nullptr;
-	// construct InstanceDataObject class
-	// TODO: should we put the InstanceDataObject or its class in a package?
-	const UClass* InstanceDataObjectClass = CreateInstanceDataObjectClass(PropertyTree.Get(), EnumNamesPtr, Owner->GetClass(), GetTransientPackage());
-
-	BagData.bNeedsFixup = StructContainsLooseProperties(InstanceDataObjectClass);
-
-	UObject* Outer = GetTransientPackage();
+	const UClass* InstanceDataObjectClass = CreateInstanceDataObjectClass(PropertyTree.Get(), EnumNamesPtr, Owner->GetClass(), Outer);
 
-	// construct InstanceDataObject object
+	// Construct the IDO.
 	FStaticConstructObjectParameters Params(InstanceDataObjectClass);
 	Params.SetFlags |= EObjectFlags::RF_Transactional;
 	Params.Name = MakeUniqueObjectName(Outer, Params.Class, Owner->GetFName());
 	Params.Outer = Outer;
-	UObject* InstanceDataObjectObject = StaticConstructObject_Internal(Params);
-	BagData.InstanceDataObject = InstanceDataObjectObject;
-	InstanceDataObjectToOwner.Add(InstanceDataObjectObject, Owner);
-	
-	// setup load context to mark properties the that were set by serialization
-	FUObjectSerializeContext* LoadContext = FUObjectThreadContext::Get().GetSerializeContext();
-	TGuardValue<bool> ScopedTrackSerializedProperties(LoadContext->bTrackSerializedProperties, true);
-	// enable impersonation so that the IDO gets loaded instead of Owner
-	TGuardValue<bool> ScopedImpersonateProperties(LoadContext->bImpersonateProperties, true);
+	UObject* InstanceDataObject = StaticConstructObject_Internal(Params);
 
-	if (FLinkerLoad* Linker = Owner->GetLinker(); Linker == &Archive)
-	{
-		Owner->SetFlags(RF_NeedLoad);
-		{
-			FGuardValue_Bitfield(Linker->ArMergeOverrides, true);
-			Linker->Preload(Owner);
-		}
-	}
-	else
+	// Load the IDO.
+	const int64 OffsetToRestore = Archive.Tell();
+	FScopedObjectSerializeContext ObjectSerializeContext(InstanceDataObject, Archive);
+	Archive.Seek(StartOffset);
 	{
-		const int64 PostSerializeOffset = Archive.Tell();
-		Archive.Seek(PreSerializeOffset);
-
-		{
-			FGuardValue_Bitfield(Archive.ArMergeOverrides, true);
-			Owner->Serialize(Archive);
-		}
-
-		if (!ensureMsgf(Archive.Tell() == PostSerializeOffset,
-						TEXT("Serializing %s into its IDO consumed %" INT64_FMT " bytes when %" INT64_FMT " bytes were expected."),
-						*Owner->GetPathName(), Archive.Tell() - PreSerializeOffset, PostSerializeOffset - PreSerializeOffset))
-		{
-			Archive.Seek(PostSerializeOffset);
-		}
+		FGuardValue_Bitfield(Archive.ArMergeOverrides, true);
+		InstanceDataObject->SerializeScriptProperties(Archive);
 	}
+	ensureMsgf(Archive.Tell() == EndOffset,
+		TEXT("Serializing %s into its IDO consumed %" INT64_FMT " bytes when %" INT64_FMT " bytes were expected."),
+		*Owner->GetPathName(), Archive.Tell() - StartOffset, EndOffset - StartOffset);
+	Archive.Seek(OffsetToRestore);
 
+	InstanceDataObjectToOwner.Add(InstanceDataObject, Owner);
+	BagData.InstanceDataObject = InstanceDataObject;
+	BagData.bNeedsFixup = StructContainsLooseProperties(InstanceDataObjectClass);
 	return BagData.InstanceDataObject;
 }
 
-FScopedIDOSerializationContext::FScopedIDOSerializationContext(UObject* InObject, FArchive& InArchive)
-	: Archive(InArchive)
-	, Object(InObject)
-	, PreSerializeOffset(InArchive.Tell())
-{
-	FUObjectSerializeContext* SerializeContext = FUObjectThreadContext::Get().GetSerializeContext();
-
-	// Save the parts of the serialize context that are modified below.
-	SavedSerializedObject = SerializeContext->SerializedObject;
-	bSavedTrackSerializedPropertyPath = SerializeContext->bTrackSerializedPropertyPath;
-	bSavedTrackInitializedProperties = SerializeContext->bTrackInitializedProperties;
-	bSavedTrackSerializedProperties = SerializeContext->bTrackSerializedProperties;
-	bSavedTrackUnknownProperties = SerializeContext->bTrackUnknownProperties;
-	bSavedTrackUnknownEnumNames = SerializeContext->bTrackUnknownEnumNames;
-	bSavedImpersonateProperties = SerializeContext->bImpersonateProperties;
-
-	// Disable if cooking because extra data within impersonated objects will not be understood outside of the editor.
-	const bool bHasIDOSupport = !InArchive.IsCooking() && IsInstanceDataObjectSupportEnabled(InObject);
-
-	const bool bIsLoading = InArchive.IsLoading();
-
-	bCreateIDO = bIsLoading && bHasIDOSupport && !SerializeContext->bImpersonateProperties &&
-		!InObject->GetClass()->HasAnyClassFlags(CLASS_NewerVersionExists);
-
-	SerializeContext->SerializedObject = InObject;
-
-	// Enable tracking any property or enum that does not match the current schema when creating an IDO.
-	SerializeContext->bTrackSerializedPropertyPath = bCreateIDO;
-	SerializeContext->bTrackUnknownProperties = bCreateIDO;
-	SerializeContext->bTrackUnknownEnumNames = bCreateIDO;
-
-	UObject* TargetObject = SerializeContext->bImpersonateProperties ? ResolveInstanceDataObject(InObject) : InObject;
-	SerializeContext->bTrackInitializedProperties = bIsLoading && FInitializedPropertyValueState(TargetObject).IsTracking();
-	SerializeContext->bTrackSerializedProperties = bIsLoading && FSerializedPropertyValueState(TargetObject).IsTracking();
-
-	if (!bIsLoading)
-	{
-		SerializeContext->bImpersonateProperties = bHasIDOSupport;
-	}
-}
-
-FScopedIDOSerializationContext::~FScopedIDOSerializationContext()
-{
-	if (bCreateIDO)
-	{
-		FPropertyBagRepository::Get().CreateInstanceDataObject(Object, Archive, PreSerializeOffset);
-	}
-
-	// Restore the parts of the serialize context that the constructor modified.
-	FUObjectSerializeContext* SerializeContext = FUObjectThreadContext::Get().GetSerializeContext();
-	SerializeContext->SerializedObject = SavedSerializedObject;
-	SerializeContext->bTrackSerializedPropertyPath = bSavedTrackSerializedPropertyPath;
-	SerializeContext->bTrackInitializedProperties = bSavedTrackInitializedProperties;
-	SerializeContext->bTrackSerializedProperties = bSavedTrackSerializedProperties;
-	SerializeContext->bTrackUnknownProperties = bSavedTrackUnknownProperties;
-	SerializeContext->bTrackUnknownEnumNames = bSavedTrackUnknownEnumNames;
-	SerializeContext->bImpersonateProperties = bSavedImpersonateProperties;
-}
-
 // Not sure this is necessary.
 void FPropertyBagRepository::ShrinkMaps()
 {
diff --git a/Engine/Source/Runtime/CoreUObject/Private/UObject/SavePackage/PackageHarvester.cpp b/Engine/Source/Runtime/CoreUObject/Private/UObject/SavePackage/PackageHarvester.cpp
--- a/Engine/Source/Runtime/CoreUObject/Private/UObject/SavePackage/PackageHarvester.cpp
+++ b/Engine/Source/Runtime/CoreUObject/Private/UObject/SavePackage/PackageHarvester.cpp
@@ -7,11 +7,11 @@
 #include "Interfaces/ITargetPlatform.h"
 #include "Misc/PackageAccessTracking.h"
 #include "UObject/OverridableManager.h"
-#include "UObject/PropertyBagRepository.h"
 #include "UObject/SavePackage/SaveContext.h"
 #include "UObject/SavePackage/SavePackageUtilities.h"
 #include "UObject/UObjectGlobals.h"
 #include "UObject/UObjectHash.h"
+#include "UObject/UObjectSerializeContext.h"
 
 EObjectMark GenerateMarksForObject(const UObject* InObject, FSaveContext& SaveContext)
 {
@@ -477,7 +477,7 @@ void FPackageHarvester::ProcessExport(const FExportWithContext& InProcessContext
 {
 	UObject* Export = InProcessContext.Export;
 	FExportScope HarvesterScope(*this, InProcessContext);
-	UE::FScopedIDOSerializationContext IDOSaveContext(Export, *this);
+	UE::FScopedObjectSerializeContext ObjectSerializeContext(Export, *this);
 
 	// Harvest its class; warn if class is not included
 	UClass* Class = Export->GetClass();
diff --git a/Engine/Source/Runtime/CoreUObject/Private/UObject/SavePackage2.cpp b/Engine/Source/Runtime/CoreUObject/Private/UObject/SavePackage2.cpp
--- a/Engine/Source/Runtime/CoreUObject/Private/UObject/SavePackage2.cpp
+++ b/Engine/Source/Runtime/CoreUObject/Private/UObject/SavePackage2.cpp
@@ -3,7 +3,6 @@
 
 #include "UObject/SavePackage.h"
 #include "UObject/LinkerLoad.h"
-#include "UObject/PropertyBagRepository.h"
 
 #if UE_WITH_SAVEPACKAGE
 #include "AssetRegistry/AssetData.h"
@@ -43,11 +42,13 @@
 #include "UObject/ObjectRedirector.h"
 #include "UObject/ObjectSaveContext.h"
 #include "UObject/Package.h"
+#include "UObject/PropertyBagRepository.h"
 #include "UObject/SavePackage/PackageHarvester.h"
 #include "UObject/SavePackage/SaveContext.h"
 #include "UObject/SavePackage/SavePackageUtilities.h"
 #include "UObject/UObjectGlobals.h"
 #include "UObject/UObjectHash.h"
+#include "UObject/UObjectSerializeContext.h"
 
 #if ENABLE_COOK_STATS
 #include "ProfilingDebugging/ScopedTimers.h"
@@ -2133,7 +2134,7 @@ ESavePackageResult WritePackageTextHeader(FStructuredArchive::FRecord& Structure
 		SCOPED_SAVETIMER(UPackage_Save_SaveExport);
 		SCOPED_SAVETIMER_TEXT(*WriteToString<128>(GetClassTraceScope(Export.Object), TEXT("_SaveSerialize")));
 
-		UE::FScopedIDOSerializationContext IDOSaveContext(Export.Object, Linker);
+		UE::FScopedObjectSerializeContext ObjectSerializeContext(Export.Object, Linker);
 
 		Export.SerialOffset = ExportsArchive.Tell();
 		Linker.CurrentlySavingExport = FPackageIndex::FromExport(ExportIndex);
@@ -2203,7 +2204,7 @@ ESavePackageResult WriteExports(FStructuredArchive::FRecord& StructuredArchiveRo
 		{
 			SCOPED_SAVETIMER(UPackage_Save_SaveExport);
 
-			UE::FScopedIDOSerializationContext IDOSaveContext(Export.Object, ExportsRecord.GetUnderlyingArchive());
+			UE::FScopedObjectSerializeContext ObjectSerializeContext(Export.Object, ExportsRecord.GetUnderlyingArchive());
 
 			// Save the object data.
 			Export.SerialOffset = Linker->Tell();
diff --git a/Engine/Source/Runtime/CoreUObject/Private/UObject/UObjectThreadContext.cpp b/Engine/Source/Runtime/CoreUObject/Private/UObject/UObjectThreadContext.cpp
--- a/Engine/Source/Runtime/CoreUObject/Private/UObject/UObjectThreadContext.cpp
+++ b/Engine/Source/Runtime/CoreUObject/Private/UObject/UObjectThreadContext.cpp
@@ -1,12 +1,15 @@
 // Copyright Epic Games, Inc. All Rights Reserved.
 
-/*=============================================================================
-	UObjectThreadContext.cpp: Unreal object globals
-=============================================================================*/
-
 #include "UObject/UObjectThreadContext.h"
-#include "UObject/Object.h"
+
 #include "UObject/LinkerLoad.h"
+#include "UObject/Object.h"
+#include "UObject/UObjectSerializeContext.h"
+
+#if WITH_EDITORONLY_DATA
+#include "UObject/InstanceDataObjectUtils.h"
+#include "UObject/PropertyStateTracking.h"
+#endif
 
 DEFINE_LOG_CATEGORY(LogUObjectThreadContext);
 
@@ -120,3 +123,78 @@ void FUObjectSerializeContext::DetachFromLinkers()
 {
 	check(!GEventDrivenLoaderEnabled);
 }
+
+namespace UE
+{
+
+FScopedObjectSerializeContext::FScopedObjectSerializeContext(UObject* InObject, FArchive& InArchive)
+#if WITH_EDITORONLY_DATA
+	: Archive(InArchive)
+	, Object(InObject)
+#endif
+{
+	FUObjectSerializeContext* SerializeContext = FUObjectThreadContext::Get().GetSerializeContext();
+
+	// Save the parts of the serialize context that are modified below.
+	SavedSerializedObject = SerializeContext->SerializedObject;
+#if WITH_EDITORONLY_DATA
+	SavedSerializedObjectScriptStartOffset = SerializeContext->SerializedObjectScriptStartOffset;
+	SavedSerializedObjectScriptEndOffset = SerializeContext->SerializedObjectScriptEndOffset;
+	bSavedTrackSerializedPropertyPath = SerializeContext->bTrackSerializedPropertyPath;
+	bSavedTrackInitializedProperties = SerializeContext->bTrackInitializedProperties;
+	bSavedTrackSerializedProperties = SerializeContext->bTrackSerializedProperties;
+	bSavedTrackUnknownProperties = SerializeContext->bTrackUnknownProperties;
+	bSavedTrackUnknownEnumNames = SerializeContext->bTrackUnknownEnumNames;
+	bSavedImpersonateProperties = SerializeContext->bImpersonateProperties;
+
+	SerializeContext->SerializedObject = Object;
+
+	const bool bIsLoading = Archive.IsLoading();
+	const bool bIsIDO = IsInstanceDataObject(Object);
+
+	// Disable if cooking because any extra data in the IDO will not be understood outside of the editor.
+	const bool bSupportsIDO = !Archive.IsCooking() && IsInstanceDataObjectSupportEnabled(Object);
+
+	// Disable if a newer version of the class exists because there is no point in creating an IDO of an obsolete type.
+	bCreateInstanceDataObject = bIsLoading && !bIsIDO && bSupportsIDO && !Object->GetClass()->HasAnyClassFlags(CLASS_NewerVersionExists);
+
+	// Enable tracking any property or enum that does not match the current schema when creating an IDO.
+	SerializeContext->bTrackSerializedPropertyPath = bCreateInstanceDataObject;
+	SerializeContext->bTrackUnknownProperties = bCreateInstanceDataObject;
+	SerializeContext->bTrackUnknownEnumNames = bCreateInstanceDataObject;
+
+	// Enable tracking of property value state on load if the object supports it.
+	SerializeContext->bTrackInitializedProperties = bIsLoading && FInitializedPropertyValueState(Object).IsTracking();
+	SerializeContext->bTrackSerializedProperties = bIsLoading && FSerializedPropertyValueState(Object).IsTracking();
+
+	// Enable impersonation when loading to an IDO or saving an object that may have an IDO.
+	SerializeContext->bImpersonateProperties = bIsLoading ? bIsIDO : bSupportsIDO;
+#endif // WITH_EDITORONLY_DATA
+}
+
+FScopedObjectSerializeContext::~FScopedObjectSerializeContext()
+{
+	FUObjectSerializeContext* SerializeContext = FUObjectThreadContext::Get().GetSerializeContext();
+
+#if WITH_EDITORONLY_DATA
+	if (bCreateInstanceDataObject)
+	{
+		CreateInstanceDataObject(Object, Archive, SerializeContext->SerializedObjectScriptStartOffset, SerializeContext->SerializedObjectScriptEndOffset);
+	}
+#endif
+
+	// Restore the parts of the serialize context that the constructor modified.
+	SerializeContext->SerializedObject = SavedSerializedObject;
+#if WITH_EDITORONLY_DATA
+	SerializeContext->SerializedObjectScriptStartOffset = SavedSerializedObjectScriptStartOffset;
+	SerializeContext->SerializedObjectScriptEndOffset = SavedSerializedObjectScriptEndOffset;
+	SerializeContext->bTrackSerializedPropertyPath = bSavedTrackSerializedPropertyPath;
+	SerializeContext->bTrackInitializedProperties = bSavedTrackInitializedProperties;
+	SerializeContext->bTrackSerializedProperties = bSavedTrackSerializedProperties;
+	SerializeContext->bTrackUnknownProperties = bSavedTrackUnknownProperties;
+	SerializeContext->bTrackUnknownEnumNames = bSavedTrackUnknownEnumNames;
+	SerializeContext->bImpersonateProperties = bSavedImpersonateProperties;
+#endif
+}
+
+} // UE
diff --git a/Engine/Source/Runtime/CoreUObject/Public/UObject/UObjectThreadContext.h b/Engine/Source/Runtime/CoreUObject/Public/UObject/UObjectThreadContext.h
--- a/Engine/Source/Runtime/CoreUObject/Public/UObject/UObjectThreadContext.h
+++ b/Engine/Source/Runtime/CoreUObject/Public/UObject/UObjectThreadContext.h
@@ -175,6 +175,10 @@ public:
 	/** Points to the most recently used Linker for serialization by CreateExport() */
 	FLinkerLoad* SerializedExportLinker;
 #if WITH_EDITORONLY_DATA
+	/** The absolute offset of the start of the serialized script properties for the serialized object, or -1. */
+	UE_INTERNAL int64 SerializedObjectScriptStartOffset = -1;
+	/** The absolute offset of the end of the serialized script properties for the serialized object, or -1. */
+	UE_INTERNAL int64 SerializedObjectScriptEndOffset = -1;
 	/** Path to the property currently being serialized */
 	UE_INTERNAL UE::FPropertyPathName SerializedPropertyPath;
 	/** True when SerializedPropertyPath is being tracked during serialization. */
diff --git a/Engine/Source/Runtime/Engine/Private/UnrealEngine.cpp b/Engine/Source/Runtime/Engine/Private/UnrealEngine.cpp
--- a/Engine/Source/Runtime/Engine/Private/UnrealEngine.cpp
+++ b/Engine/Source/Runtime/Engine/Private/UnrealEngine.cpp
@@ -128,6 +128,7 @@ UnrealEngine.cpp: Implements the UEngine class and helpers.
 #include "Misc/DelayedAutoRegister.h"
 #include <exception>
 #include "AssetCompilingManager.h"
+#include "UObject/UObjectSerializeContext.h"
 
 #if UE_WITH_IRIS
 #include "Iris/IrisConfig.h"
@@ -17390,8 +17391,7 @@ public:
 		}
 #endif // USE_STABLE_LOCALIZATION_KEYS
 
-		UE::FScopedIDOSerializationContext IDOSaveContext(SrcObject, *this);
-
+		UE::FScopedObjectSerializeContext ObjectSerializeContext(SrcObject, *this);
 		SrcObject->Serialize(*this);
 	}
 
@@ -17458,9 +17458,7 @@ public:
 		}
 #endif // USE_STABLE_LOCALIZATION_KEYS
 
-		// Enable IDO, if needed
-		UE::FScopedIDOSerializationContext IDOLoadContext(Obj, *this);
-		
+		UE::FScopedObjectSerializeContext ObjectSerializeContext(Obj, *this);
 		Obj->Serialize(*this);
 	}
 
@@ -17558,9 +17556,7 @@ public:
 		}
 #endif // USE_STABLE_LOCALIZATION_KEYS
 
-		// Enable IDO, if needed
-		UE::FScopedIDOSerializationContext IDOLoadContext(DstObject, *this);
-		
+		UE::FScopedObjectSerializeContext ObjectSerializeContext(DstObject, *this);
 		DstObject->Serialize(*this);
 	}
 
@@ -17907,8 +17903,7 @@ void UEngine::CopyPropertiesForUnrelatedObjects(UObject* OldObject, UObject* New
 			}
 		#endif // USE_STABLE_LOCALIZATION_KEYS
 
-			UE::FScopedIDOSerializationContext IDOSaveContext(OldInstance, SubObjWriter);
-
+			UE::FScopedObjectSerializeContext ObjectSerializeContext(OldInstance, SubObjWriter);
 			OldInstance->Serialize(SubObjWriter);
 		}
 	}
