[AsyncLoader]
- Make the old loader the default for loose cooked files due to deadlocks being reported in 5.5 when using ZenLoader by default

[RN] Make the old loader the default for loose cooked files due to deadlocks being reported in 5.5 when using ZenLoader by default

diff --git a/Engine/Source/Runtime/CoreUObject/Private/Serialization/AsyncPackageLoader.cpp b/Engine/Source/Runtime/CoreUObject/Private/Serialization/AsyncPackageLoader.cpp
--- a/Engine/Source/Runtime/CoreUObject/Private/Serialization/AsyncPackageLoader.cpp
+++ b/Engine/Source/Runtime/CoreUObject/Private/Serialization/AsyncPackageLoader.cpp
@@ -17,7 +17,7 @@
 #include "HAL/IConsoleManager.h"
 #include "AutoRTFM/AutoRTFM.h"
 
-static bool GUseOldLoaderAsFallback = false;
+static bool GUseOldLoaderAsFallback = true;
 static FAutoConsoleVariableRef CVarGUseOldLoaderAsFallback(
 	TEXT("s.UseOldLoaderAsFallback"),
 	GUseOldLoaderAsFallback,
