Add an additional accepted value for cvar gc.GarbageReferenceTrackingEnabled.
Settings gc.GarbageReferenceTrackingEnabled to 3 will report all garbage references (as setting 1) and then terminate immediately by fatal error if any were reported.

diff --git a/Engine/Source/Runtime/CoreUObject/Private/UObject/GarbageCollection.cpp b/Engine/Source/Runtime/CoreUObject/Private/UObject/GarbageCollection.cpp
--- a/Engine/Source/Runtime/CoreUObject/Private/UObject/GarbageCollection.cpp
+++ b/Engine/Source/Runtime/CoreUObject/Private/UObject/GarbageCollection.cpp
@@ -249,7 +249,10 @@ int32 GGarbageReferenceTrackingEnabled = 0;
 static FAutoConsoleVariableRef CGarbageReferenceTrackingEnabled(
 	TEXT("gc.GarbageReferenceTrackingEnabled"),
 	GGarbageReferenceTrackingEnabled,
-	TEXT("Causes the Garbage Collector to track and log unreleased garbage objects. If 1, will dump every reference. If 2, will dump a sample of the references to highlight problematic properties."),
+	TEXT("Causes the Garbage Collector to track and log unreleased garbage objects. ")
+	TEXT("If 1, will dump every reference. ")
+	TEXT("If 2, will dump a sample of the references to highlight problematic properties. ")
+	TEXT("If 3, will dump every reference and then emit a fatal error to terminate the process. "),
 	ECVF_Default
 );
 
@@ -5276,7 +5279,7 @@ static void DumpGarbageReferencers(TConstArrayView<TUniquePtr<FWorkerContext>> C
 	}
 
 	// Dump all garbage references in detail
-	if (GGarbageReferenceTrackingEnabled == 1)
+	if (GGarbageReferenceTrackingEnabled == 1 || GGarbageReferenceTrackingEnabled == 3)
 	{
 		int32 NumGarbageReferences = 0;
 		for (const TUniquePtr<FWorkerContext>& Context : Contexts)
@@ -5292,6 +5295,7 @@ static void DumpGarbageReferencers(TConstArrayView<TUniquePtr<FWorkerContext>> C
 			Context->GarbageReferences.Reset();
 		}
 		
+		UE_CLOG(GGarbageReferenceTrackingEnabled == 3 && NumGarbageReferences != 0, LogGarbage, Fatal, TEXT("Reported %d garbage references"), NumGarbageReferences);
 		UE_LOG(LogGarbage, Log, TEXT("Reported %d garbage references"), NumGarbageReferences);
 	}
 	// Dump one reference for each source property
