Fix doc-comments for `ICppStructOps::Identical` method.

These may have been copy-pasted from the neighboring method and discussed copying the object, but this is a comparison method.

diff --git a/Engine/Source/Runtime/CoreUObject/Public/UObject/Class.h b/Engine/Source/Runtime/CoreUObject/Public/UObject/Class.h
--- a/Engine/Source/Runtime/CoreUObject/Public/UObject/Class.h
+++ b/Engine/Source/Runtime/CoreUObject/Public/UObject/Class.h
@@ -1168,8 +1168,8 @@ public:
 			return GetCapabilities().HasIdentical;
 		}
 		/** 
-		 * Compare this structure 
-		 * @return true if the copy was handled, otherwise it will fall back to FStructProperty::Identical
+		 * Compare this structure
+		 * @return true if the comparison was performed. bOutResult contains the result of the comparison.
 		 */
 		virtual bool Identical(const void* A, const void* B, uint32 PortFlags, bool& bOutResult) = 0;
 
