Extend dynamic subobject instancing to include container/optional types.

diff --git a/Engine/Source/Runtime/CoreUObject/Private/UObject/PropertyArray.cpp b/Engine/Source/Runtime/CoreUObject/Private/UObject/PropertyArray.cpp
--- a/Engine/Source/Runtime/CoreUObject/Private/UObject/PropertyArray.cpp
+++ b/Engine/Source/Runtime/CoreUObject/Private/UObject/PropertyArray.cpp
@@ -1166,7 +1166,8 @@ void FArrayProperty::FinishDestroyInternal( void* Data ) const
  */
 void FArrayProperty::InstanceSubobjects(void* Data, void const* DefaultData, TNotNull<UObject*> InOwner, FObjectInstancingGraph* InstanceGraph)
 {
-	if (Data && Inner->ContainsInstancedObjectProperty())
+	if (Data
+		&& (Inner->ContainsInstancedObjectProperty() || (InstanceGraph && InstanceGraph->CanUseDynamicInstancing())))
 	{
 		FScriptArrayHelper ArrayHelper(this, Data);
 		FScriptArrayHelper DefaultArrayHelper(this, DefaultData);
diff --git a/Engine/Source/Runtime/CoreUObject/Private/UObject/PropertyMap.cpp b/Engine/Source/Runtime/CoreUObject/Private/UObject/PropertyMap.cpp
--- a/Engine/Source/Runtime/CoreUObject/Private/UObject/PropertyMap.cpp
+++ b/Engine/Source/Runtime/CoreUObject/Private/UObject/PropertyMap.cpp
@@ -1302,8 +1302,9 @@ void FMapProperty::InstanceSubobjects(void* Data, void const* DefaultData, TNotN
 		return;
 	}
 
-	const bool bInstancedKey   = KeyProp  ->ContainsInstancedObjectProperty();
-	const bool bInstancedValue = ValueProp->ContainsInstancedObjectProperty();
+	const bool bUsesDynamicInstancing = InstanceGraph && InstanceGraph->CanUseDynamicInstancing();
+	const bool bInstancedKey   = KeyProp  ->ContainsInstancedObjectProperty() || bUsesDynamicInstancing;
+	const bool bInstancedValue = ValueProp->ContainsInstancedObjectProperty() || bUsesDynamicInstancing;
 
 	if (!bInstancedKey && !bInstancedValue)
 	{
diff --git a/Engine/Source/Runtime/CoreUObject/Private/UObject/PropertyOptional.cpp b/Engine/Source/Runtime/CoreUObject/Private/UObject/PropertyOptional.cpp
--- a/Engine/Source/Runtime/CoreUObject/Private/UObject/PropertyOptional.cpp
+++ b/Engine/Source/Runtime/CoreUObject/Private/UObject/PropertyOptional.cpp
@@ -470,7 +470,9 @@ void FOptionalProperty::FinishDestroyInternal( void* Data ) const
 
 void FOptionalProperty::InstanceSubobjects(void* Data, void const* DefaultData, TNotNull<UObject*> InOwner, struct FObjectInstancingGraph* InstanceGraph)
 {
-	if (Data && IsSet(Data) && ValueProperty->ContainsInstancedObjectProperty())
+	if (Data
+		&& IsSet(Data)
+		&& (ValueProperty->ContainsInstancedObjectProperty() || (InstanceGraph && InstanceGraph->CanUseDynamicInstancing())))
 	{
 		ValueProperty->InstanceSubobjects(Data, DefaultData, InOwner, InstanceGraph);
 	}
diff --git a/Engine/Source/Runtime/CoreUObject/Private/UObject/PropertySet.cpp b/Engine/Source/Runtime/CoreUObject/Private/UObject/PropertySet.cpp
--- a/Engine/Source/Runtime/CoreUObject/Private/UObject/PropertySet.cpp
+++ b/Engine/Source/Runtime/CoreUObject/Private/UObject/PropertySet.cpp
@@ -808,7 +808,8 @@ void FSetProperty::InstanceSubobjects(void* Data, void const* DefaultData, TNotN
 		return;
 	}
 
-	const bool bInstancedElement = ElementProp->ContainsInstancedObjectProperty();
+	const bool bUsesDynamicInstancing = InstanceGraph && InstanceGraph->CanUseDynamicInstancing();
+	const bool bInstancedElement = ElementProp->ContainsInstancedObjectProperty() || bUsesDynamicInstancing;
 
 	if (!bInstancedElement)
 	{
