[AutoRTFM] Allow IsAsyncLoadingPackages to run safely in the open.

This fixes an error affecting the Verse presubmit tests.

The memory barrier wasn't necessary when dealing with thread-safe counters. Checking two atomic counters (while making no changes) is safe to do from closed code by running in the open.

Additionally, simplified a neighboring function to use ALWAYS_OPEN instead of requiring an OPEN block.

diff --git a/Engine/Source/Runtime/CoreUObject/Private/Serialization/AsyncLoadingThread.h b/Engine/Source/Runtime/CoreUObject/Private/Serialization/AsyncLoadingThread.h
--- a/Engine/Source/Runtime/CoreUObject/Private/Serialization/AsyncLoadingThread.h
+++ b/Engine/Source/Runtime/CoreUObject/Private/Serialization/AsyncLoadingThread.h
@@ -472,16 +472,13 @@ public:
 	}
 
 	/** Gets the current state of async loading */
-	static bool GetIsInAsyncLoadingTick()
+	UE_AUTORTFM_ALWAYS_OPEN static bool GetIsInAsyncLoadingTick()
 	{
-		bool Result = false;
-		UE_AUTORTFM_OPEN{ Result = !!AsyncLoadingTickCounter.GetValue(); };
-		return Result;
+		return !!AsyncLoadingTickCounter.GetValue();
 	}
 
-	FORCEINLINE bool IsAsyncLoadingPackages() override
+	UE_AUTORTFM_ALWAYS_OPEN FORCEINLINE bool IsAsyncLoadingPackages() override
 	{
-		FPlatformMisc::MemoryBarrier();
 		return QueuedPackagesCounter.GetValue() != 0 || ExistingAsyncPackagesCounter.GetValue() != 0;
 	}
 
