Mutli-Server Changes:
- Introducing the concept of Borrowing vs. Ownership Changes when migrating Objects.
- Borrowed Actors are to be returned to the Owning Server thus when loaning out Actors we don't remove them from local Tick functions.
- Changing Actor to not do InitialOverlaps on PostMigration because we're going to assume Physics has already run and taken care of this.
- Lots of minor changes around the codebase.  If you see something like GetController() changed to !Controller, it's to prevent an automatic migration when we dereference the TObjectPtr.
- TickTaskManager has some subtle changes when UE_WITH_REMOTE_OBJECT_HANDLE is defined because the Task actually stays behind when you Deregister its Tick Functions.  That would otherwise cause migrated Actors to potentially be ticked multiple times in a frame if reregistered.

Part of #jira UE-233940

diff --git a/Engine/Source/Runtime/CoreUObject/Private/UObject/RemoteObjectTransfer.cpp b/Engine/Source/Runtime/CoreUObject/Private/UObject/RemoteObjectTransfer.cpp
--- a/Engine/Source/Runtime/CoreUObject/Private/UObject/RemoteObjectTransfer.cpp
+++ b/Engine/Source/Runtime/CoreUObject/Private/UObject/RemoteObjectTransfer.cpp
@@ -46,7 +46,7 @@ TMulticastDelegate<void()> OnTransactionAbortedDelegate;
 
 FString DebugLastAbortReason;
 
-UObject* MigrateObjectFromObjectData(FRemoteObjectData& ObjectData, FRemoteServerId RemoteServerId);
+UObject* MigrateObjectFromObjectData(FRemoteObjectData& ObjectData, FRemoteServerId RemoteServerId, EObjectMigrationFlags ObjectFlags);
 
 /**
  * Remote object transfer queue. 
@@ -157,7 +157,9 @@ public:
 
 		UE_LOG(LogRemoteObject, Verbose, TEXT("     : FulfillReceiveRequest(%u) %s (was requested: %d, token matches: %d)"), RequestToken, *ObjectId.ToString(), bFoundRequest, bMatchesActiveRequest);
 
-		UObject* Result = MigrateObjectFromObjectData(*ObjectData, RemoteServerId);
+		const bool bAssignedOwnership = (RequestToken == PushMigrationRequestToken);
+		EObjectMigrationFlags MigrationFlags = bAssignedOwnership ? EObjectMigrationFlags::Acquire : EObjectMigrationFlags::Receive;
+		UObject* Result = MigrateObjectFromObjectData(*ObjectData, RemoteServerId, MigrationFlags);
 		checkf(UE::CoreUObject::Private::FObjectHandleUtils::GetRemoteId(Result) == ObjectId, TEXT("PeerId:%s requested object %s migration from PeerId:%s but received %s"),
 			*GetGlobalServerId().ToString(), *ObjectId.ToString(), *RemoteServerId.ToString(), *UE::CoreUObject::Private::FObjectHandleUtils::GetRemoteId(Result).ToString());
 
@@ -202,7 +204,7 @@ void RequestObjectData(FRemoteObjectId ObjectId, FRemoteServerId CurrentOwnerSer
 	GTransferQueue.RequestRemoteObject(GTransferQueue.ActiveRequestToken.Load(), GTransferQueue.ActiveRequestPriority, ObjectId, CurrentOwnerServerId, DestinationServerId);
 }
 
-UObject* MigrateObjectFromObjectData(FRemoteObjectData& ObjectData, FRemoteServerId RemoteServerId)
+UObject* MigrateObjectFromObjectData(FRemoteObjectData& ObjectData, FRemoteServerId RemoteServerId, EObjectMigrationFlags MigrationFlags)
 {
 	using namespace UE::RemoteObject::Serialization;
 	check(!AutoRTFM::IsClosed());
@@ -226,9 +228,13 @@ UObject* MigrateObjectFromObjectData(FRemoteObjectData& ObjectData, FRemoteServe
 	}
 
 #if UE_WITH_REMOTE_OBJECT_HANDLE
+	FPostMigrateParams PostMigrateParams;
+	PostMigrateParams.RequestingServerId = RemoteServerId;
+	PostMigrateParams.DestinationServerId = GetGlobalServerId();
+	PostMigrateParams.MigrationFlags = MigrationFlags;
 	for (UObject* ReceivedObject : ReceivedObjects)
 	{
-		ReceivedObject->PostMigrate(GetGlobalServerId());
+		ReceivedObject->PostMigrate(PostMigrateParams);
 	}
 #endif
 
@@ -263,10 +269,20 @@ void MigrateObjectFromRemoteServer(FRemoteObjectId ObjectId, FRemoteServerId Cur
 		UE_LOG(LogRemoteObject, Verbose, TEXT("Aborting transaction, need obj id %s on server %s"),
 			*ObjectId.ToString(),
 			*CurrentOwnerServerId.ToString());
+		
+		if (UE_LOG_ACTIVE(LogRemoteObject, VeryVerbose))
+		{
+			ANSICHAR HumanReadableString[8192] = { 0 };
+			constexpr int32 SkipNumCalls = 5;
+			FGenericPlatformStackWalk::StackWalkAndDump(HumanReadableString, sizeof(HumanReadableString), SkipNumCalls);
+			UE_LOG(LogRemoteObject, VeryVerbose, TEXT("Callstack: %hs"), HumanReadableString);
+		}
 	};
 
 	if (AutoRTFM::IsClosed())
 	{
+		// TODO: DO we want to log when we're pulling over an object we've NEVER known about?!
+		// 
 		// if we are in a transaction, abort so the outer handler can renegotiate object transfers and retry
 		SetAbortReason(TEXT("object migration"));
 		AutoRTFM::AbortTransaction();
@@ -340,7 +356,7 @@ void MigrateObjectToRemoteServerWithExplicitPriority(uint32 RequestToken, FRemot
 		for(UObject* SO : SentObjects)
 		{
 			FRemoteObjectId SubObjectId = UE::CoreUObject::Private::FObjectHandleUtils::GetRemoteId(SO);
-			UE_LOG(LogRemoteObjectD, Verbose, TEXT("         sent subobject : %s"), *SubObjectId.ToString());
+			UE_LOG(LogRemoteObject, Verbose, TEXT("         sent subobject : %s [%s]"), *SubObjectId.ToString(), *GetNameSafe(SO));
 		}
 
 		// Notify that object data has been migrated (actually will be but SendObjectData does a MoveTemp on the data)
@@ -358,10 +374,16 @@ void MigrateObjectToRemoteServerWithExplicitPriority(uint32 RequestToken, FRemot
 			// Notify that objects have been migrated to a remote server
 			OnObjectsSentDelegate.Broadcast(SentObjects, DestinationServerId);
 
+			// Are we loaning it out or changing ownership?
+			const bool bOwnershipChange = (RequestToken == PushMigrationRequestToken);
+			FPostMigrateParams PostMigrateParams;
+			PostMigrateParams.RequestingServerId = RequestPriority;
+			PostMigrateParams.DestinationServerId = DestinationServerId;
+			PostMigrateParams.MigrationFlags = bOwnershipChange ? EObjectMigrationFlags::Reassign : EObjectMigrationFlags::Loan;
 			for (UObject* SentObject : SentObjects)
 			{
 #if UE_WITH_REMOTE_OBJECT_HANDLE
-				SentObject->PostMigrate(DestinationServerId);
+				SentObject->PostMigrate(PostMigrateParams);
 #endif
 				MarkAsRemote(SentObject, DestinationServerId);
 			}
@@ -512,7 +534,7 @@ void ExecuteTransactionalInternal(FRemoteServerId InWorkPriority, const TFunctio
 
 		PhysicsQueryBeginRequestDelegate.ExecuteIfBound();
 
-		UE_LOG(LogRemoteObject, VeryVerbose, TEXT("-- ExecuteTransactional : iteration %u request %u pri %s --"), LocalIterationNumber, GTransferQueue.ActiveRequestToken.Load(), *ExecutingWorkPriority.ToString());
+		UE_LOG(LogRemoteObjectD, VeryVerbose, TEXT("-- ExecuteTransactional : iteration %u request %u pri %s --"), LocalIterationNumber, GTransferQueue.ActiveRequestToken.Load(), *ExecutingWorkPriority.ToString());
 
 		// we might run through this loop after the work succeeds in the case
 		// that the work prefetches objects that it doesn't access
@@ -609,7 +631,7 @@ void ExecuteTransactionalInternal(FRemoteServerId InWorkPriority, const TFunctio
 				double TimeNow = FPlatformTime::Seconds();
 				if (TimeNow > LastWaitingPrintTime + 1.0)
 				{
-					UE_LOG(LogRemoteObject, VeryVerbose, TEXT("   : waiting for %d responses..."), GTransferQueue.PendingRequests);
+					UE_LOG(LogRemoteObjectD, VeryVerbose, TEXT("   : waiting for %d responses..."), GTransferQueue.PendingRequests);
 					LastWaitingPrintTime = TimeNow;
 				}
 			}
@@ -623,13 +645,13 @@ void ExecuteTransactionalInternal(FRemoteServerId InWorkPriority, const TFunctio
 			// now we also might have migrated objects away, so check to see if LockedObjects == RequestedObjects
 			if(GTransferQueue.LockedObjects.Num() != GTransferQueue.RequestedObjects.Num())
 			{
-				UE_LOG(LogRemoteObject, VeryVerbose, TEXT("   : locked count %d mismatches requests %d, did we give up locked objects? trying again..."), GTransferQueue.LockedObjects.Num(), GTransferQueue.RequestedObjects.Num());
+				UE_LOG(LogRemoteObjectD, VeryVerbose, TEXT("   : locked count %d mismatches requests %d, did we give up locked objects? trying again..."), GTransferQueue.LockedObjects.Num(), GTransferQueue.RequestedObjects.Num());
 
 				// try again!
 				continue;
 			}
 
-			UE_LOG(LogRemoteObject, VeryVerbose, TEXT("   : READY: all %d requested objects are locked"), GTransferQueue.RequestedObjects.Num());
+			UE_LOG(LogRemoteObjectD, VeryVerbose, TEXT("   : READY: all %d requested objects are locked"), GTransferQueue.RequestedObjects.Num());
 
 			// all of the requested objects are locked, try to perform the work
 			if (!bWorkSucceeded)
@@ -641,12 +663,12 @@ void ExecuteTransactionalInternal(FRemoteServerId InWorkPriority, const TFunctio
 				if (WorkTransactResult == AutoRTFM::ETransactionResult::AbortedByRequest)
 				{
 					// the work aborted, it should have inserted a new requested object, retry
-					UE_LOG(LogRemoteObject, VeryVerbose, TEXT("   : work ABORTED (%s)"), *DebugLastAbortReason);
+					UE_LOG(LogRemoteObjectD, VeryVerbose, TEXT("   : work ABORTED (%s)"), *DebugLastAbortReason);
 					OnTransactionAbortedDelegate.Broadcast();
 				}
 				else
 				{
-					UE_LOG(LogRemoteObject, VeryVerbose, TEXT("   : work COMPLETED!"));
+					UE_LOG(LogRemoteObjectD, VeryVerbose, TEXT("   : work COMPLETED!"));
 
 					bWorkSucceeded = true;
 					OnTransactionCompletedDelegate.Broadcast();
@@ -672,11 +694,11 @@ void ExecuteTransactionalInternal(FRemoteServerId InWorkPriority, const TFunctio
 		// to loop back to the top
 		if (bWorkSucceeded)
 		{
-			UE_LOG(LogRemoteObject, VeryVerbose, TEXT("-- ExecuteTransactional %u DONE --"), GTransferQueue.ActiveRequestToken.Load());
+			UE_LOG(LogRemoteObjectD, VeryVerbose, TEXT("-- ExecuteTransactional %u DONE --"), GTransferQueue.ActiveRequestToken.Load());
 		}
 		else
 		{
-			UE_LOG(LogRemoteObject, VeryVerbose, TEXT("-- ExecuteTransactional %u ABANDONED, context switching to higher priority work --"), GTransferQueue.ActiveRequestToken.Load());
+			UE_LOG(LogRemoteObjectD, VeryVerbose, TEXT("-- ExecuteTransactional %u ABANDONED, context switching to higher priority work --"), GTransferQueue.ActiveRequestToken.Load());
 		}
 
 		GTransferQueue.EndRequest();
@@ -706,7 +728,7 @@ void ExecuteTransactionalInternal(FRemoteServerId InWorkPriority, const TFunctio
 
 void WaitForServersToSync()
 {
-	UE_LOG(LogRemoteObject, VeryVerbose, TEXT("-- WaitForServersToSync --"));
+	UE_LOG(LogRemoteObjectD, VeryVerbose, TEXT("-- WaitForServersToSync --"));
 	WaitForServersToSyncDelegate.ExecuteIfBound();
 }
 
diff --git a/Engine/Source/Runtime/CoreUObject/Public/UObject/Object.h b/Engine/Source/Runtime/CoreUObject/Public/UObject/Object.h
--- a/Engine/Source/Runtime/CoreUObject/Public/UObject/Object.h
+++ b/Engine/Source/Runtime/CoreUObject/Public/UObject/Object.h
@@ -537,7 +537,7 @@ public:
 	/**
 	 * Called when an object is about to be migrated and after it's been migrated to another process.
 	 */
-	virtual void PostMigrate(FRemoteServerId DestinationServerId) {}
+	virtual void PostMigrate(const struct FPostMigrateParams& PostMigrateParams) {}
 #endif
 
 	/**
diff --git a/Engine/Source/Runtime/CoreUObject/Public/UObject/RemoteObjectTransfer.h b/Engine/Source/Runtime/CoreUObject/Public/UObject/RemoteObjectTransfer.h
--- a/Engine/Source/Runtime/CoreUObject/Public/UObject/RemoteObjectTransfer.h
+++ b/Engine/Source/Runtime/CoreUObject/Public/UObject/RemoteObjectTransfer.h
@@ -8,6 +8,39 @@
 #include "UObject/Class.h"
 #include "Async/Async.h"
 
+enum class EObjectMigrationFlags : uint8
+{
+	Invalid 		= 0,
+	Receive			= (1 << 0),
+	Send 			= (1 << 1),
+	ChangeOwnership = (1 << 2),
+
+	/** Receive an object without changing ownership (implies it will be sent back to the server that loaned it) */
+	Borrow = Receive,
+
+	/** Loan an object while retaining ownership (implies the server loaning wants it back) */
+	Loan = Send,
+
+	/** Acquired objects mean the receiving side owns the object */
+	Acquire = Receive | ChangeOwnership,
+
+	/** Reassign objects to relinquish ownership and assign them to the destination server */
+	Reassign = Send | ChangeOwnership,
+};
+ENUM_CLASS_FLAGS(EObjectMigrationFlags);
+
+struct FPostMigrateParams
+{
+	/** The server ID that initiated the migration request */
+	FRemoteServerId RequestingServerId;
+
+	/** The server ID that the request is destined for */
+	FRemoteServerId DestinationServerId;
+
+	/** The flags that determine the semantics of the object migration request */
+	EObjectMigrationFlags MigrationFlags;
+};
+
 namespace UE::RemoteObject::Transfer
 {
 	constexpr uint32 PushMigrationRequestToken = 0xFFFFFFFFu;
diff --git a/Engine/Source/Runtime/Engine/Classes/GameFramework/Actor.h b/Engine/Source/Runtime/Engine/Classes/GameFramework/Actor.h
--- a/Engine/Source/Runtime/Engine/Classes/GameFramework/Actor.h
+++ b/Engine/Source/Runtime/Engine/Classes/GameFramework/Actor.h
@@ -2335,7 +2335,7 @@ public:
 	ENGINE_API virtual void PostLoad() override;
 	ENGINE_API virtual void PostLoadSubobjects( FObjectInstancingGraph* OuterInstanceGraph ) override;
 #if UE_WITH_REMOTE_OBJECT_HANDLE
-	ENGINE_API virtual void PostMigrate(FRemoteServerId DestinationServerId) override;
+	ENGINE_API virtual void PostMigrate(const struct FPostMigrateParams& PostMigrateParams) override;
 #endif
 	ENGINE_API virtual void BeginDestroy() override;
 #if !UE_STRIP_DEPRECATED_PROPERTIES
@@ -3224,7 +3224,7 @@ public:
 	FORCEINLINE bool HasDeferredComponentRegistration() const { return bHasDeferredComponentRegistration; }
 
 	/** Returns true if Actor has a registered root component */
-	ENGINE_API bool HasValidRootComponent();
+	ENGINE_API bool HasValidRootComponent() const;
 
 	/** 
 	 * Unregister all currently registered components
diff --git a/Engine/Source/Runtime/Engine/Private/Actor.cpp b/Engine/Source/Runtime/Engine/Private/Actor.cpp
--- a/Engine/Source/Runtime/Engine/Private/Actor.cpp
+++ b/Engine/Source/Runtime/Engine/Private/Actor.cpp
@@ -62,6 +62,10 @@
 #include "Engine/DamageEvents.h"
 #include "Engine/SCS_Node.h"
 
+#if UE_WITH_REMOTE_OBJECT_HANDLE
+#include "UObject/RemoteObjectTransfer.h"
+#endif
+
 #if WITH_EDITOR
 #include "FoliageHelper.h"
 #include "LevelInstance/LevelInstanceSubsystem.h"
@@ -960,7 +964,10 @@ void AActor::Serialize(FArchive& Ar)
 #if UE_WITH_REMOTE_OBJECT_HANDLE
 	if (Ar.IsMigratingRemoteObjects())
 	{
+		bool bLocalActorInitialized = IsActorInitialized();
+
 		Ar << CreationTime;
+		Ar << bLocalActorInitialized;
 
 		float CurrentLifespan = 0.0f; 
 		if (Ar.IsSaving())
@@ -972,6 +979,7 @@ void AActor::Serialize(FArchive& Ar)
 		{
 			// When migrating actors are in a state after BeginPlay has been called but we need to call SetLifeSpan to initialize local timer
 			SetLifeSpan(CurrentLifespan);
+			bActorInitialized = bLocalActorInitialized;
 		}
 	}
 #endif // UE_WITH_REMOTE_OBJECT_HANDLE
@@ -1144,54 +1152,75 @@ void AActor::PostLoadSubobjects(FObjectInstancingGraph* OuterInstanceGraph)
 
 static USceneComponent* FixupNativeActorComponents(AActor* Actor);
 
-void AActor::PostMigrate(FRemoteServerId DestinationServerId)
+void AActor::PostMigrate(const FPostMigrateParams& PostMigrateParams)
 {
-	Super::PostMigrate(DestinationServerId);
+	Super::PostMigrate(PostMigrateParams);
+
+	// We are using this as a proxy for having ownership until we sort out those mechanics
+	const bool bHasOwnership = bTickFunctionsRegistered;
+	FRemoteServerId LocalServerId = UE::RemoteObject::GetGlobalServerId();
+
+	UE_LOG(LogActor, Warning, TEXT("%hs: %s %d (LocalSrvId: %s. ReqSrvId: %s. DstSrvId: %s. HasOwnership: %d)"), __func__, *GetName(), static_cast<int>(PostMigrateParams.MigrationFlags), *LocalServerId.ToString(), *PostMigrateParams.RequestingServerId.ToString(), *PostMigrateParams.DestinationServerId.ToString(), bHasOwnership);
+
+	UWorld* const World = GetWorld();
+	checkf(World, TEXT("Migrated Actor %s (%s) doesn't have a World"), *GetPathName(), *FRemoteObjectId(this).ToString());
 
 	TGuardValue<ITransaction*> SuppressTransaction(GUndo, nullptr);
-	if (DestinationServerId != UE::RemoteObject::GetGlobalServerId())
+	if (EnumHasAnyFlags(PostMigrateParams.MigrationFlags, EObjectMigrationFlags::Send))
 	{
-		// Remove the actor from the world on this server as it is about to be migrated to another server.
-		// Migrated actors are unique to the server instance that currently owns their memory and are not replicated to other servers.
-		if (UWorld* const World = GetWorld())
-		{
-			ClearComponentOverlaps();
-			UE::Net::Private::FNetPropertyConditionManager::Get().NotifyObjectDestroyed(this);
+		ensureMsgf(PostMigrateParams.DestinationServerId != LocalServerId, TEXT("%hs: Sending Object %s from %s to %s but destination is the same as the Local Server Id."), __func__, *GetName(), *PostMigrateParams.RequestingServerId.ToString(), *PostMigrateParams.DestinationServerId.ToString());
 
-			// Notify net drivers that this actor has been destroyed.
-			if (FWorldContext* Context = GEngine->GetWorldContextFromWorld(World))
+		const bool bRelinquishOwnership = !bHasOwnership || EnumHasAnyFlags(PostMigrateParams.MigrationFlags, EObjectMigrationFlags::ChangeOwnership);
+		if (bRelinquishOwnership)
+		{
+			// Remove the actor from the world on this server as it is about to be migrated to another server.
+			// Migrated actors are unique to the server instance that currently owns their memory and are not replicated to other servers.
+			if (World)
 			{
-				for (FNamedNetDriver& Driver : Context->ActiveNetDrivers)
+				ClearComponentOverlaps();
+				UE::Net::Private::FNetPropertyConditionManager::Get().NotifyObjectDestroyed(this);
+
+				// Notify net drivers that this actor has been destroyed.
+				if (FWorldContext* Context = GEngine->GetWorldContextFromWorld(World))
 				{
-					if (Driver.NetDriver != nullptr && Driver.NetDriver->ShouldReplicateActor(this))
+					for (FNamedNetDriver& Driver : Context->ActiveNetDrivers)
 					{
-						Driver.NetDriver->NotifyActorDestroyed(this);
+						if (Driver.NetDriver != nullptr && Driver.NetDriver->ShouldReplicateActor(this))
+						{
+							Driver.NetDriver->NotifyActorDestroyed(this);
+						}
 					}
 				}
+
+				// Remove the actor from the actor list.
+				World->RemoveActor(this, /*bShouldModifyLevel*/ false);
 			}
 
-			// Remove the actor from the actor list.
-			World->RemoveActor(this, /*bShouldModifyLevel*/ false);
-		}
+			// Disabling Ticks prevents them from happening *this frame* while unregistering them stops them from happening *next frame*
+			PrimaryActorTick.SetTickFunctionEnable(false);
+			ForEachComponent(true, [](UActorComponent* ActorComponent)
+			{
+				ActorComponent->PrimaryComponentTick.SetTickFunctionEnable(false);
+			});
 
-		// Clean up the actor's components.
-		UnregisterAllComponents();
+			// Clean up the actor's components.
+			UnregisterAllComponents();
 
-		// Mark the actor and its direct components as pending kill.
-		MarkAsGarbage();
-		MarkComponentsAsGarbage();
+			// Unregister the actor's tick function
+			RegisterAllActorTickFunctions(/*bRegisterTickFunctions*/ false, /*bIncludeComponents*/ true);
 
-		// Unregister the actor's tick function
-		RegisterAllActorTickFunctions(/*bRegisterTickFunctions*/ false, /*bIncludeComponents*/ true);
+			// Mark the actor and its direct components as pending kill.
+			MarkComponentsAsGarbage();
+			MarkAsGarbage();
+		}
 	}
-	else
+	else if (EnumHasAnyFlags(PostMigrateParams.MigrationFlags, EObjectMigrationFlags::Receive))
 	{
+		ensureMsgf(PostMigrateParams.DestinationServerId == UE::RemoteObject::GetGlobalServerId(), TEXT("Unexpected DestinationServerId (%s was not the local ServerId: %s)"), *PostMigrateParams.DestinationServerId.ToString(), *LocalServerId.ToString());
+
 		// Add the actor that was migrated from another server to this server's world
 		GetLevel()->TryAddActorToList(this, /*bAddUnique*/false);
 
-		UWorld* const World = GetWorld();
-		checkf(World != nullptr, TEXT("Migrated Actor %s (%s) doesn't have a World"), *GetPathName(), *FRemoteObjectId(this).ToString());
-
 		// Set the actor's world transform if it has a native rootcomponent.
 		const USceneComponent* SceneRootComponent = FixupNativeActorComponents(this);
 		ensureMsgf(SceneRootComponent == RootComponent, TEXT("%hs: We were expecting RootComponent (%s) to be equal to FixupNativeActorComponents (%s). Transform fix-up code would be required to handle this case."), __func__, *GetNameSafe(RootComponent), *GetNameSafe(SceneRootComponent));
@@ -1211,30 +1240,23 @@ void AActor::PostMigrate(FRemoteServerId DestinationServerId)
 		BeginReplication();
 #endif // UE_WITH_IRIS
 
-		RegisterAllActorTickFunctions(true, false);
-
-		TInlineComponentArray<UActorComponent*> Components;
-		GetComponents(Components);
-
-		for (UActorComponent* Component : Components)
+		const bool bAcquiredOwnership = !bHasOwnership && EnumHasAnyFlags(PostMigrateParams.MigrationFlags, EObjectMigrationFlags::ChangeOwnership);
+		if (bAcquiredOwnership)
 		{
-			if (Component->IsRegistered())
-			{
-				Component->RegisterAllComponentTickFunctions(true);
-
-				if (USceneComponent* SceneComponent = Cast<USceneComponent>(Component))
-				{
-					// UpdateInitialOverlaps needs this flag cleared or it won't run
-					SceneComponent->ClearSkipUpdateOverlaps();
-				}
-			}
+			RegisterAllActorTickFunctions(true, true);
 		}
 
 		ActorHasBegunPlay = EActorBeginPlayState::HasBegunPlay;
 
-		UpdateInitialOverlaps(/*bFromLevelStreaming*/ false);
-		
-		World->AddNetworkActor(this);
+		if (bAcquiredOwnership)
+		{
+			// We are going to assume that the Physics queries on the previous owning server had taken multi-server physics into account
+			// thus we do not actually need to check for Initial Overlaps anymore
+			// ForEachSceneComponent: SceneComponent->ClearSkipUpdateOverlaps();
+			// UpdateInitialOverlaps(/*bFromLevelStreaming*/ false);
+
+			World->AddNetworkActor(this);
+		}
 	}
 }
 #endif // UE_WITH_REMOTE_OBJECT_HANDLE
@@ -5941,7 +5963,7 @@ bool AActor::IncrementalRegisterComponents(int32 NumComponentsToRegister, FRegis
 	return false;
 }
 
-bool AActor::HasValidRootComponent()
+bool AActor::HasValidRootComponent() const
 { 
 	return (RootComponent != nullptr && RootComponent->IsRegistered()); 
 }
diff --git a/Engine/Source/Runtime/Engine/Private/Components/ActorComponent.cpp b/Engine/Source/Runtime/Engine/Private/Components/ActorComponent.cpp
--- a/Engine/Source/Runtime/Engine/Private/Components/ActorComponent.cpp
+++ b/Engine/Source/Runtime/Engine/Private/Components/ActorComponent.cpp
@@ -1424,7 +1424,11 @@ void FActorComponentTickFunction::ExecuteTick(float DeltaTime, enum ELevelTick T
 
 FString FActorComponentTickFunction::DiagnosticMessage()
 {
+#if	UE_WITH_REMOTE_OBJECT_HANDLE
+	return (!Target.IsRemote() ? Target->GetFullName() : Target.GetRemoteId().ToString()) + TEXT("[TickComponent]");
+#else
 	return Target->GetFullName() + TEXT("[TickComponent]");
+#endif
 }
 
 FName FActorComponentTickFunction::DiagnosticContext(bool bDetailed)
diff --git a/Engine/Source/Runtime/Engine/Private/NetConnection.cpp b/Engine/Source/Runtime/Engine/Private/NetConnection.cpp
--- a/Engine/Source/Runtime/Engine/Private/NetConnection.cpp
+++ b/Engine/Source/Runtime/Engine/Private/NetConnection.cpp
@@ -2896,7 +2896,12 @@ bool UNetConnection::ReadPacketInfo(FBitReader& Reader, bool bHasPacketInfoPaylo
 
 		if (PlayerController)
 		{
-			PlayerController->UpdatePing(NewLag);
+			// PlayerController is valid and local (as opposed to currently loaned out)
+			UE_CLOG(PlayerController.IsRemote(), LogNet, Warning, TEXT("%hs: PlayerController is unexpectedly a RemoteObject"), __func__);
+			if (!PlayerController.IsRemote())
+			{
+				PlayerController->UpdatePing(NewLag);
+			}
 		}
 
 		if (NetPing.IsValid())
diff --git a/Engine/Source/Runtime/Engine/Private/TickTaskManager.cpp b/Engine/Source/Runtime/Engine/Private/TickTaskManager.cpp
--- a/Engine/Source/Runtime/Engine/Private/TickTaskManager.cpp
+++ b/Engine/Source/Runtime/Engine/Private/TickTaskManager.cpp
@@ -309,10 +309,13 @@ public:
 #if UE_WITH_REMOTE_OBJECT_HANDLE
 			auto ExecuteTickWork = [this, CurrentThread, &MyCompletionGraphEvent]()
 			{
+				// !IsCompletionHandleValid is an indication we had previously been ticked this frame and then migrated back
+				if (Target->IsCompletionHandleValid() && Target->IsTickFunctionEnabled())
+				{
 #endif
-			Target->ExecuteTick(Target->CalculateDeltaTime(Context.DeltaSeconds, Context.World), Context.TickType, CurrentThread, MyCompletionGraphEvent);
-
+				Target->ExecuteTick(Target->CalculateDeltaTime(Context.DeltaSeconds, Context.World), Context.TickType, CurrentThread, MyCompletionGraphEvent);
 #if UE_WITH_REMOTE_OBJECT_HANDLE
+				}
 			};
 
 			if (Target->bRunTransactionally)
@@ -327,7 +330,13 @@ public:
 
 			LIGHTWEIGHT_TIME_GUARD_END(FTickFunctionTask, Target->DiagnosticMessage());
 		}
-		Target->ClearTaskInformation();  // This is stale and a good time to clear it for safety
+#if UE_WITH_REMOTE_OBJECT_HANDLE
+		const bool bTargetIsValid = (Target->GetLastTickGameTime() != -1.0f); // it's a hack to check for null InternalData
+		if (bTargetIsValid)
+#endif
+		{
+			Target->ClearTaskInformation();  // This is stale and a good time to clear it for safety
+		}
 	}
 };
 
