Downgrade "Reachable Garbage Object" warnings to display. They are currently firing spuriously during cooking of UWorld packages. We will reenable after fixing those warnings.
#rnx

diff --git a/Engine/Source/Runtime/CoreUObject/Private/UObject/GarbageCollection.cpp b/Engine/Source/Runtime/CoreUObject/Private/UObject/GarbageCollection.cpp
--- a/Engine/Source/Runtime/CoreUObject/Private/UObject/GarbageCollection.cpp
+++ b/Engine/Source/Runtime/CoreUObject/Private/UObject/GarbageCollection.cpp
@@ -5287,9 +5287,10 @@ static void DumpGarbageReferencers(TConstArrayView<TUniquePtr<FWorkerContext>> C
 			for (FGarbageReferenceInfo& GarbageReference : Context->GarbageReferences)
 			{
 				// Note: EngineTest tests rely on the format of this message, see TestGarbageElimination
-				UE_LOG(LogGarbage, Warning, TEXT("Reachable garbage object: %s"), *GarbageReference.GarbageObject->GetFullName());
-				UE_LOG(LogGarbage, Warning, TEXT("Referenced by:            %s"), *GarbageReference.GetReferencingObjectInfo());
-				UE_LOG(LogGarbage, Warning, TEXT(""));
+				// TODO: Restore severity of this message to Warning once we have fixed the spurious occurrence of the warning during cook.
+				UE_LOG(LogGarbage, Display, TEXT("Reachable garbage object: %s"), *GarbageReference.GarbageObject->GetFullName());
+				UE_LOG(LogGarbage, Display, TEXT("Referenced by:            %s"), *GarbageReference.GetReferencingObjectInfo());
+				UE_LOG(LogGarbage, Display, TEXT(""));
 				NumGarbageReferences++;
 			}
 			Context->GarbageReferences.Reset();
@@ -5312,9 +5313,10 @@ static void DumpGarbageReferencers(TConstArrayView<TUniquePtr<FWorkerContext>> C
 				if (Seen.Contains(Key) == false)
 				{
 					// Note: EngineTest tests rely on the format of this message, see TestGarbageElimination
-					UE_LOG(LogGarbage, Warning, TEXT("Reachable garbage object: %s"), *GarbageReference.GarbageObject->GetFullName());
-					UE_LOG(LogGarbage, Warning, TEXT("Referenced by:            %s"), *GarbageReference.GetReferencingObjectInfo());
-					UE_LOG(LogGarbage, Warning, TEXT(""));
+					// TODO: Restore severity of this message to Warning once we have fixed the spurious occurrence of the warning during cook.
+					UE_LOG(LogGarbage, Display, TEXT("Reachable garbage object: %s"), *GarbageReference.GarbageObject->GetFullName());
+					UE_LOG(LogGarbage, Display, TEXT("Referenced by:            %s"), *GarbageReference.GetReferencingObjectInfo());
+					UE_LOG(LogGarbage, Display, TEXT(""));
 					Seen.Add(Key);
 					ReportedGarbageReferences++;
 				}
