[CoreUObject]
* Fixed non-unity/pch compile error

#skipci

diff --git a/Engine/Source/Runtime/CoreUObject/Private/UObject/RemoteObjectTransfer.cpp b/Engine/Source/Runtime/CoreUObject/Private/UObject/RemoteObjectTransfer.cpp
--- a/Engine/Source/Runtime/CoreUObject/Private/UObject/RemoteObjectTransfer.cpp
+++ b/Engine/Source/Runtime/CoreUObject/Private/UObject/RemoteObjectTransfer.cpp
@@ -1,6 +1,7 @@
 // Copyright Epic Games, Inc. All Rights Reserved.
 
 #include "UObject/RemoteObjectTransfer.h"
+#include "GenericPlatform/GenericPlatformStackWalk.h"
 #include "UObject/RemoteObjectPrivate.h"
 #include "UObject/RemoteObjectSerialization.h"
 #include "UObject/ObjectPtr.h"
