Fix AutoRTFM compile error: Attributes 'noinline and alwaysinline' are incompatible!

diff --git a/Engine/Source/Runtime/CoreUObject/Private/Serialization/AsyncLoadingThread.h b/Engine/Source/Runtime/CoreUObject/Private/Serialization/AsyncLoadingThread.h
--- a/Engine/Source/Runtime/CoreUObject/Private/Serialization/AsyncLoadingThread.h
+++ b/Engine/Source/Runtime/CoreUObject/Private/Serialization/AsyncLoadingThread.h
@@ -477,7 +477,7 @@ public:
 		return !!AsyncLoadingTickCounter.GetValue();
 	}
 
-	UE_AUTORTFM_ALWAYS_OPEN FORCEINLINE bool IsAsyncLoadingPackages() override
+	UE_AUTORTFM_ALWAYS_OPEN bool IsAsyncLoadingPackages() override
 	{
 		return QueuedPackagesCounter.GetValue() != 0 || ExistingAsyncPackagesCounter.GetValue() != 0;
 	}
