InstancedStruct: Correctly disable delta serialization by passing value as defaults, not nullptr

diff --git a/Engine/Source/Runtime/CoreUObject/Private/StructUtils/InstancedStruct.cpp b/Engine/Source/Runtime/CoreUObject/Private/StructUtils/InstancedStruct.cpp
--- a/Engine/Source/Runtime/CoreUObject/Private/StructUtils/InstancedStruct.cpp
+++ b/Engine/Source/Runtime/CoreUObject/Private/StructUtils/InstancedStruct.cpp
@@ -240,8 +240,8 @@ bool FInstancedStruct::ExportTextItem(FString& ValueStr, FInstancedStruct const&
 	if (const UScriptStruct* StructTypePtr = GetScriptStruct())
 	{
 		ValueStr += StructTypePtr->GetPathName();
-		// Force the default to nullptr to disable delta serialization because we reset the memory in import text.
-		StructTypePtr->ExportText(ValueStr, GetMemory(), nullptr, Parent, PortFlags, ExportRootScope);
+		// Force the default to Value to disable delta serialization because we reset the memory in import text.
+		StructTypePtr->ExportText(ValueStr, GetMemory(), GetMemory(), Parent, PortFlags, ExportRootScope);
 	}
 	else
 	{
