PR #12689: Fix typo in UObjectArray.h

#rnx

diff --git a/Engine/Source/Runtime/CoreUObject/Public/UObject/UObjectArray.h b/Engine/Source/Runtime/CoreUObject/Public/UObject/UObjectArray.h
--- a/Engine/Source/Runtime/CoreUObject/Public/UObject/UObjectArray.h
+++ b/Engine/Source/Runtime/CoreUObject/Public/UObject/UObjectArray.h
@@ -830,8 +830,8 @@ public:
 		/**
 		* Provides notification that a UObjectBase has been added to the uobject array
 		 *
-		 * @param Object object that has been destroyed
-		 * @param Index	index of object that is being deleted
+		 * @param Object object that has been added
+		 * @param Index	index of object that is being added
 		 */
 		virtual void NotifyUObjectCreated(const class UObjectBase *Object, int32 Index)=0;
 
