Migration code from old math types to new math types

diff --git a/Engine/Source/Runtime/CoreUObject/Private/UObject/OverriddenPropertySet.cpp b/Engine/Source/Runtime/CoreUObject/Private/UObject/OverriddenPropertySet.cpp
--- a/Engine/Source/Runtime/CoreUObject/Private/UObject/OverriddenPropertySet.cpp
+++ b/Engine/Source/Runtime/CoreUObject/Private/UObject/OverriddenPropertySet.cpp
@@ -114,6 +114,11 @@ bool FOverridableSerializationLogic::HasCapabilities(ECapabilities InCapabilitie
 	return (Capabilities & InCapabilities) == InCapabilities;
 }
 
+FOverriddenPropertySet* FOverridableSerializationLogic::GetOverriddenPropertiesSlow()
+{
+	return OverriddenProperties;
+}
+
 EOverriddenPropertyOperation FOverridableSerializationLogic::GetOverriddenPropertyOperation(const FArchive& Ar, FProperty* Property /*= nullptr*/, uint8* DataPtr /*= nullptr*/, uint8* DefaultValue /*= nullptr*/)
 {
 	const FArchiveSerializedPropertyChain* CurrentPropertyChain = Ar.GetSerializedPropertyChain();
diff --git a/Engine/Source/Runtime/CoreUObject/Public/UObject/OverriddenPropertySet.h b/Engine/Source/Runtime/CoreUObject/Public/UObject/OverriddenPropertySet.h
--- a/Engine/Source/Runtime/CoreUObject/Public/UObject/OverriddenPropertySet.h
+++ b/Engine/Source/Runtime/CoreUObject/Public/UObject/OverriddenPropertySet.h
@@ -77,6 +77,13 @@ public:
 		return OverriddenProperties;
 	}
 
+	/**
+	 * Call during the serialization of an object to get its overriden properties
+	 * Note: Expects the current serialized object to use overridable serialization
+	 * Note this is not re-entrant and it stores information in a thread local storage
+	 * @return the overridden properties of the current object being serialized */
+	COREUOBJECT_API static FOverriddenPropertySet* GetOverriddenPropertiesSlow();
+
 	/**
 	 * Retrieve from the Archive and the current property the overridden property operation to know if it has to be serialized or not
 	 * @param Ar currently being used to serialize the current object (will be used to retrieve the current property serialized path)
@@ -411,6 +418,10 @@ public:
 		return bWasAdded;
 	}
 
+	TObjectPtr<UObject> GetOwner() const
+	{
+		return Owner;
+	}
 protected:
 
 	FOverriddenPropertyNode& FindOrAddNode(FOverriddenPropertyNode& ParentPropertyNode, FOverriddenPropertyNodeID NodeID);
