[ZenLoader]
- Allow preemption of package creation when the GC is waiting

#rnx

diff --git a/Engine/Source/Runtime/CoreUObject/Private/Serialization/AsyncLoading2.cpp b/Engine/Source/Runtime/CoreUObject/Private/Serialization/AsyncLoading2.cpp
--- a/Engine/Source/Runtime/CoreUObject/Private/Serialization/AsyncLoading2.cpp
+++ b/Engine/Source/Runtime/CoreUObject/Private/Serialization/AsyncLoading2.cpp
@@ -463,6 +463,13 @@ static FAutoConsoleVariableRef CVarGStallDetectorIdleLoops(
 	ECVF_Default
 );
 
+static bool GAsyncLoading2_AllowPreemptingPackagesDuringGC = true;
+static FAutoConsoleVariableRef CVar_AllowPreemptingPackagesDuringGC(
+	TEXT("s.AllowPreemptingPackagesDuringGC"),
+	GAsyncLoading2_AllowPreemptingPackagesDuringGC,
+	TEXT("Allow the async loading thread to get pre-empted by garbage collection while it's creating packages."),
+	ECVF_Default);
+
 #if USING_INSTRUMENTATION
 static bool GDetectRaceDuringLoading = false;
 static FAutoConsoleVariableRef CVarDetectRaceDuringLoading(
@@ -4100,7 +4107,7 @@ private:
 	void IncludePackageInSyncLoadContextRecursive(FAsyncLoadingThreadState2& ThreadState, uint64 ContextId, FAsyncPackage2* Package);
 	void UpdateSyncLoadContext(FAsyncLoadingThreadState2& ThreadState, bool bAutoHandleSyncLoadContext = true);
 
-	bool CreateAsyncPackagesFromQueue(FAsyncLoadingThreadState2& ThreadState);
+	bool CreateAsyncPackagesFromQueue(FAsyncLoadingThreadState2& ThreadState, bool bForceExitForGarbageCollect = false);
 
 	FAsyncPackage2* CreateAsyncPackage(FAsyncLoadingThreadState2& ThreadState, const FAsyncPackageDesc2& Desc)
 	{
@@ -4682,7 +4689,7 @@ UE_TRACE_EVENT_BEGIN(CUSTOM_LOADTIMER_LOG, CreateAsyncPackage, NoSync)
 	UE_TRACE_EVENT_FIELD(UE::Trace::WideString, PackageName)
 UE_TRACE_EVENT_END()
 
-bool FAsyncLoadingThread2::CreateAsyncPackagesFromQueue(FAsyncLoadingThreadState2& ThreadState)
+bool FAsyncLoadingThread2::CreateAsyncPackagesFromQueue(FAsyncLoadingThreadState2& ThreadState, bool bForceExitForGarbageCollect)
 {
 	TRACE_CPUPROFILER_EVENT_SCOPE(CreateAsyncPackagesFromQueue);
 	
@@ -4690,7 +4697,7 @@ bool FAsyncLoadingThread2::CreateAsyncPackagesFromQueue(FAsyncLoadingThreadState
 	TGuardValue GuardVisibilityFilter(FUObjectThreadContext::Get().AsyncVisibilityFilter, EInternalObjectFlags::None);
 
 	bool bPackagesCreated = false;
-	const int32 TimeSliceGranularity = ThreadState.UseTimeLimit() ? 4 : MAX_int32;
+	const int32 TimeSliceGranularity = (ThreadState.UseTimeLimit() || bForceExitForGarbageCollect) ? 4 : MAX_int32;
 
 	FIoBatch IoBatch = IoDispatcher.NewBatch();
 
@@ -8731,7 +8738,7 @@ uint32 FAsyncLoadingThread2::Run()
 
 					if (QueuedPackagesCounter || !PendingPackages.IsEmpty())
 					{
-						if (CreateAsyncPackagesFromQueue(ThreadState))
+						if (CreateAsyncPackagesFromQueue(ThreadState, GAsyncLoading2_AllowPreemptingPackagesDuringGC))
 						{
 							// Fall through to FAsyncLoadEventQueue2 processing unless we need to suspend
 							if (SuspendRequestedCount.load(std::memory_order_relaxed) > 0 || IsGarbageCollectionWaiting())
