Restore "Reachable Garbage Object" log message severity to warning. Backout of CL 38953715.
#rnx

diff --git a/Engine/Source/Runtime/CoreUObject/Private/UObject/GarbageCollection.cpp b/Engine/Source/Runtime/CoreUObject/Private/UObject/GarbageCollection.cpp
--- a/Engine/Source/Runtime/CoreUObject/Private/UObject/GarbageCollection.cpp
+++ b/Engine/Source/Runtime/CoreUObject/Private/UObject/GarbageCollection.cpp
@@ -5287,10 +5287,9 @@ static void DumpGarbageReferencers(TConstArrayView<TUniquePtr<FWorkerContext>> C
 			for (FGarbageReferenceInfo& GarbageReference : Context->GarbageReferences)
 			{
 				// Note: EngineTest tests rely on the format of this message, see TestGarbageElimination
-				// TODO: Restore severity of this message to Warning once we have fixed the spurious occurrence of the warning during cook.
-				UE_LOG(LogGarbage, Display, TEXT("Reachable garbage object: %s"), *GarbageReference.GarbageObject->GetFullName());
-				UE_LOG(LogGarbage, Display, TEXT("Referenced by:            %s"), *GarbageReference.GetReferencingObjectInfo());
-				UE_LOG(LogGarbage, Display, TEXT(""));
+				UE_LOG(LogGarbage, Warning, TEXT("Reachable garbage object: %s"), *GarbageReference.GarbageObject->GetFullName());
+				UE_LOG(LogGarbage, Warning, TEXT("Referenced by:            %s"), *GarbageReference.GetReferencingObjectInfo());
+				UE_LOG(LogGarbage, Warning, TEXT(""));
 				NumGarbageReferences++;
 			}
 			Context->GarbageReferences.Reset();
@@ -5313,10 +5312,9 @@ static void DumpGarbageReferencers(TConstArrayView<TUniquePtr<FWorkerContext>> C
 				if (Seen.Contains(Key) == false)
 				{
 					// Note: EngineTest tests rely on the format of this message, see TestGarbageElimination
-					// TODO: Restore severity of this message to Warning once we have fixed the spurious occurrence of the warning during cook.
-					UE_LOG(LogGarbage, Display, TEXT("Reachable garbage object: %s"), *GarbageReference.GarbageObject->GetFullName());
-					UE_LOG(LogGarbage, Display, TEXT("Referenced by:            %s"), *GarbageReference.GetReferencingObjectInfo());
-					UE_LOG(LogGarbage, Display, TEXT(""));
+					UE_LOG(LogGarbage, Warning, TEXT("Reachable garbage object: %s"), *GarbageReference.GarbageObject->GetFullName());
+					UE_LOG(LogGarbage, Warning, TEXT("Referenced by:            %s"), *GarbageReference.GetReferencingObjectInfo());
+					UE_LOG(LogGarbage, Warning, TEXT(""));
 					Seen.Add(Key);
 					ReportedGarbageReferences++;
 				}
