Fixed TCppStructOps::Destruct by avoiding an unnecessary virtual dispatch, which also fixes a warning when calling a non-virtual destructor on a struct with virtual functions. (-Wdelete-non-abstract-non-virtual-dtor)

diff --git a/Engine/Source/Runtime/CoreUObject/Public/UObject/Class.h b/Engine/Source/Runtime/CoreUObject/Public/UObject/Class.h
--- a/Engine/Source/Runtime/CoreUObject/Public/UObject/Class.h
+++ b/Engine/Source/Runtime/CoreUObject/Public/UObject/Class.h
@@ -1403,7 +1403,7 @@ public:
 			check(!(TTraits::WithNoDestructor || TIsPODType<CPPSTRUCT>::Value)); // don't call this if we have indicated it is not necessary
 			// that could have been an if statement, but we might as well force optimization above the virtual call
 			// could also not attempt to call the destructor for types where this is not possible, but I didn't do that here
-			((CPPSTRUCT*)Dest)->~CPPSTRUCT();
+			DestructItem((CPPSTRUCT*)Dest);
 		}
 		virtual bool Serialize(FArchive& Ar, void *Data) override
 		{
