[PropertyBag] fixed typo in ensure message

diff --git a/Engine/Source/Runtime/CoreUObject/Private/StructUtils/PropertyBag.cpp b/Engine/Source/Runtime/CoreUObject/Private/StructUtils/PropertyBag.cpp
--- a/Engine/Source/Runtime/CoreUObject/Private/StructUtils/PropertyBag.cpp
+++ b/Engine/Source/Runtime/CoreUObject/Private/StructUtils/PropertyBag.cpp
@@ -431,7 +431,7 @@ namespace UE::StructUtils::Private
 			}
 			break;
 		default:
-			ensureMsgf(false, TEXT("Unhandled stype %s"), *UEnum::GetValueAsString(Desc.ValueType));
+			ensureMsgf(false, TEXT("Unhandled type %s"), *UEnum::GetValueAsString(Desc.ValueType));
 		}
 
 		return nullptr;
