Generate server-unique names when constructing UObjects with remote object support enabled and no name for the new object is specified

diff --git a/Engine/Source/Runtime/CoreUObject/Private/UObject/RemoteObject.cpp b/Engine/Source/Runtime/CoreUObject/Private/UObject/RemoteObject.cpp
--- a/Engine/Source/Runtime/CoreUObject/Private/UObject/RemoteObject.cpp
+++ b/Engine/Source/Runtime/CoreUObject/Private/UObject/RemoteObject.cpp
@@ -158,6 +158,20 @@ UE::RemoteObject::Handle::FRemoteObjectStub* FindRemoteObjectStub(FRemoteObjectI
 	return ObjectMaps->FindRemoteObjectStub(ObjectId);
 }
 
+FName GetServerBaseNameForUniqueName(const UClass* Class)
+{
+	using namespace UE::RemoteObject;
+
+	checkf(Class, TEXT("Unable to generate base name for a unique object name without the object's Class"));
+
+	// Packages follow different naming rules than other UObjects and ATM we're not migrating packages so fall back to Class->GetFName()
+	if (GetGlobalServerId().IsValid() && Class->GetFName() != NAME_Package)
+	{
+		return *FString::Printf(TEXT("%s_S%s"), *Class->GetFName().GetPlainNameString(), *GetGlobalServerId().ToString());
+	}
+	return Class->GetFName();
+}
+
 FUnsafeToMigrateScope::FUnsafeToMigrateScope()
 {
 	UnsafeToMigrateObjects++;
diff --git a/Engine/Source/Runtime/CoreUObject/Private/UObject/RemoteObjectPrivate.h b/Engine/Source/Runtime/CoreUObject/Private/UObject/RemoteObjectPrivate.h
--- a/Engine/Source/Runtime/CoreUObject/Private/UObject/RemoteObjectPrivate.h
+++ b/Engine/Source/Runtime/CoreUObject/Private/UObject/RemoteObjectPrivate.h
@@ -43,6 +43,11 @@ namespace UE::RemoteObject::Private
 	*/
 	UE::RemoteObject::Handle::FRemoteObjectStub* FindRemoteObjectStub(FRemoteObjectId ObjectId);
 
+	/**
+	* Gets a base FName that will be used to generate a unique object name (see MakeUniqueObjectName)
+	*/
+	FName GetServerBaseNameForUniqueName(const UClass* Class);
+
 	struct FUnsafeToMigrateScope
 	{
 		FUnsafeToMigrateScope();
diff --git a/Engine/Source/Runtime/CoreUObject/Private/UObject/UObjectGlobals.cpp b/Engine/Source/Runtime/CoreUObject/Private/UObject/UObjectGlobals.cpp
--- a/Engine/Source/Runtime/CoreUObject/Private/UObject/UObjectGlobals.cpp
+++ b/Engine/Source/Runtime/CoreUObject/Private/UObject/UObjectGlobals.cpp
@@ -87,6 +87,7 @@
 #include "AutoRTFM/AutoRTFM.h"
 #include "UObject/PropertyOptional.h"
 #include "UObject/ObjectHandlePrivate.h"
+#include "UObject/RemoteObjectPrivate.h"
 #include "UObject/RemoteObjectSerialization.h"
 #include "UObject/VerseStringProperty.h"
 #include "VerseVM/VVMVerseClass.h"
@@ -2826,7 +2827,14 @@ FName MakeUniqueObjectName(UObject* Parent, const UClass* Class, FName InBaseNam
 
 	CSV_SCOPED_TIMING_STAT(UObject, MakeUniqueObjectName);
 	check(Class);
-	const FName BaseName = InBaseName.IsNone() ? Class->GetFName() : InBaseName;
+	const FName BaseName = InBaseName.IsNone() ?
+#if UE_WITH_REMOTE_OBJECT_HANDLE
+		UE::RemoteObject::Private::GetServerBaseNameForUniqueName(Class)
+#else
+		Class->GetFName()
+#endif
+		: InBaseName;
+
 	const FName ClassPathName = InBaseName.IsNone() ? *Class->GetPathName() : InBaseName;
 
 	if (!!(Options & EUniqueObjectNameOptions::UniversallyUnique))
