Add TryPin to weak object pointer.  Currently calling Pin from the render thread could result in a deadlock with the garbage collection waiting for the render thread to finish. TryPin will return false when garbage collection is active, which allows for a system to skip pinning when it is not safe to do so.

diff --git a/Engine/Source/Runtime/Core/Public/UObject/WeakObjectPtrTemplates.h b/Engine/Source/Runtime/Core/Public/UObject/WeakObjectPtrTemplates.h
--- a/Engine/Source/Runtime/Core/Public/UObject/WeakObjectPtrTemplates.h
+++ b/Engine/Source/Runtime/Core/Public/UObject/WeakObjectPtrTemplates.h
@@ -155,6 +155,31 @@ public:
 		return StrongPtr;
 	}
 
+	/**
+	 * Pin the weak pointer and get a strongptr.
+	 * @param bEvenIfPendingKill if this is true, pendingkill objects are considered valid
+	 * @return nullptr if this object is gone or the weak pointer is explicitly null, otherwise a valid uobject pointer
+	 */
+	FORCEINLINE bool TryPin(TStrongObjectPtr<T>& OutResult, bool bEvenIfPendingKill) const
+	{
+		using ObjectType = std::remove_pointer_t<decltype(TWeakObjectPtrBase::Get())>;
+		TStrongObjectPtr<ObjectType> UObjectPtr;
+		if (TWeakObjectPtrBase::TryPin(bEvenIfPendingKill, UObjectPtr))
+		{
+			OutResult.Attach((T*)UObjectPtr.Detach());
+			return true;
+		}
+		return false;
+	}
+
+	/**
+	 * Pin the weak pointer as a strong ptr. This is an optimized version implying bEvenIfPendingKill=false.
+	 */
+	FORCEINLINE bool TryPin(TStrongObjectPtr<T>& OutResult /*, bool bEvenIfPendingKill = false */ ) const
+	{
+		return TryPin(OutResult, false);
+	}
+
 	/** Deferences the weak pointer even if its marked RF_Unreachable. This is needed to resolve weak pointers during GC (such as ::AddReferenceObjects) */
 	FORCEINLINE T* GetEvenIfUnreachable() const
 	{
diff --git a/Engine/Source/Runtime/CoreUObject/Private/UObject/GCScopeLock.h b/Engine/Source/Runtime/CoreUObject/Private/UObject/GCScopeLock.h
--- a/Engine/Source/Runtime/CoreUObject/Private/UObject/GCScopeLock.h
+++ b/Engine/Source/Runtime/CoreUObject/Private/UObject/GCScopeLock.h
@@ -69,6 +69,28 @@ public:
 			} while (!bLocked);
 		}
 	}
+
+	/** Attempt async Lock **/
+	bool TryLockAsync()
+	{
+		bool bLocked = true;
+		if (!IsInGameThread())
+		{
+			bLocked = false;
+			if (GCCounter.GetValue() == 0)
+			{
+				FScopeLock CriticalLock(&Critical);
+				if (GCCounter.GetValue() == 0)
+				{
+					AsyncCounter.Increment();
+					bLocked = true;
+				}
+			}
+		}
+
+		return bLocked;
+	}
+
 	/** Release lock from non-game thread */
 	void UnlockAsync()
 	{
diff --git a/Engine/Source/Runtime/CoreUObject/Private/UObject/GarbageCollection.cpp b/Engine/Source/Runtime/CoreUObject/Private/UObject/GarbageCollection.cpp
--- a/Engine/Source/Runtime/CoreUObject/Private/UObject/GarbageCollection.cpp
+++ b/Engine/Source/Runtime/CoreUObject/Private/UObject/GarbageCollection.cpp
@@ -198,6 +198,19 @@ FGCScopeGuard::~FGCScopeGuard()
 	FGCCSyncObject::Get().UnlockAsync();
 }
 
+FGCScopeTryGuard::FGCScopeTryGuard()
+{
+	bLockSuccessful = FGCCSyncObject::Get().TryLockAsync();
+}
+
+FGCScopeTryGuard::~FGCScopeTryGuard()
+{
+	if(bLockSuccessful)
+	{
+		FGCCSyncObject::Get().UnlockAsync();
+	}
+}
+
 bool IsGarbageCollectionLocked()
 {
 	return FGCCSyncObject::Get().IsAsyncLocked();
diff --git a/Engine/Source/Runtime/CoreUObject/Private/UObject/WeakObjectPtr.cpp b/Engine/Source/Runtime/CoreUObject/Private/UObject/WeakObjectPtr.cpp
--- a/Engine/Source/Runtime/CoreUObject/Private/UObject/WeakObjectPtr.cpp
+++ b/Engine/Source/Runtime/CoreUObject/Private/UObject/WeakObjectPtr.cpp
@@ -161,6 +161,47 @@ TStrongObjectPtr<UObject> FWeakObjectPtr::Internal_Pin(bool bEvenIfGarbage) cons
 	return TStrongObjectPtr<UObject>(((ObjectItem != nullptr) && GUObjectArray.IsValid(ObjectItem, bEvenIfGarbage)) ? (UObject*)ObjectItem->Object : nullptr);
 }
 
+bool FWeakObjectPtr::TryPin(bool bEvenIfGarbage, class TStrongObjectPtr<UObject>& OutResult) const
+{
+	return Internal_TryPin(bEvenIfGarbage, OutResult);
+}
+
+bool FWeakObjectPtr::TryPin(class TStrongObjectPtr<UObject>& OutResult /*, bool bEvenIfGarbage = false*/) const
+{
+	return Internal_TryPin(false, OutResult);
+}
+
+bool FWeakObjectPtr::TryPinEvenIfUnreachable(class TStrongObjectPtr<UObject>& OutResult) const
+{
+	FGCScopeTryGuard GCScopeGuard;
+	if(!GCScopeGuard.LockSucceeded())
+	{
+		return false;
+	}
+
+	UObject* Result = nullptr;
+	if (Internal_IsValid(true, true))
+	{
+		FUObjectItem* ObjectItem = GUObjectArray.IndexToObject(GetObjectIndex_Private(), true);
+		Result = static_cast<UObject*>(ObjectItem->Object);
+	}
+	OutResult = TStrongObjectPtr<UObject>(Result);
+	return true;
+}
+
+bool FWeakObjectPtr::Internal_TryPin(bool bEvenIfGarbage, TStrongObjectPtr<UObject>& OutResult) const
+{
+	FGCScopeTryGuard GCScopeGuard;
+	if (!GCScopeGuard.LockSucceeded())
+	{
+		return false;
+	}
+
+	FUObjectItem* const ObjectItem = Internal_GetObjectItem();
+	OutResult = TStrongObjectPtr<UObject>(((ObjectItem != nullptr) && GUObjectArray.IsValid(ObjectItem, bEvenIfGarbage)) ? (UObject*)ObjectItem->Object : nullptr);
+	return true;
+}
+
 void FWeakObjectPtr::Serialize(FArchive& Ar)
 {
 	FArchiveUObject::SerializeWeakObjectPtr(Ar, *this);
diff --git a/Engine/Source/Runtime/CoreUObject/Public/UObject/GarbageCollection.h b/Engine/Source/Runtime/CoreUObject/Public/UObject/GarbageCollection.h
--- a/Engine/Source/Runtime/CoreUObject/Public/UObject/GarbageCollection.h
+++ b/Engine/Source/Runtime/CoreUObject/Public/UObject/GarbageCollection.h
@@ -115,6 +115,18 @@ public:
 	COREUOBJECT_API ~FGCScopeGuard();
 };
 
+class FGCScopeTryGuard
+{
+public:
+	COREUOBJECT_API FGCScopeTryGuard();
+	COREUOBJECT_API ~FGCScopeTryGuard();
+
+	bool LockSucceeded() const { return bLockSuccessful;}
+
+private:
+	bool bLockSuccessful = false;
+};
+
 class FGCObject;
 
 /** Information about references to objects marked as Garbage that's gather by the Garbage Collector */
diff --git a/Engine/Source/Runtime/CoreUObject/Public/UObject/WeakObjectPtr.h b/Engine/Source/Runtime/CoreUObject/Public/UObject/WeakObjectPtr.h
--- a/Engine/Source/Runtime/CoreUObject/Public/UObject/WeakObjectPtr.h
+++ b/Engine/Source/Runtime/CoreUObject/Public/UObject/WeakObjectPtr.h
@@ -245,6 +245,28 @@ public:
 	 */
 	COREUOBJECT_API class TStrongObjectPtr<UObject> PinEvenIfUnreachable() const;
 
+	/**
+	 * Attempt to get a strong object ptr to the weak pointer, but only if garbage collection is not in progress.
+	 * @param bEvenIfGarbage if this is true, Garbage objects are considered valid
+	 * @param OutResult will be invalid if this object is gone, or the weak pointer is explicitly null, or garbage collection is in progress. Otherwise a valid TStrongObjectPtr
+	 * @return true if garbage collection was not in progress, and OutResult was successfully captured, false if garbage collection was in progress and OutResult was not captured
+	 */
+	COREUOBJECT_API bool TryPin(bool bEvenIfGarbage, class TStrongObjectPtr<UObject>& OutResult) const;
+
+	/**
+	 * Attempt to get a strong object ptr to the weak pointer, but only if garbage collection is not in progress. This is an optimized version implying bEvenIfGarbage=false.
+	 * @param OutResult will be invalid if this object is gone, or the weak pointer is explicitly null, or garbage collection is in progress. Otherwise a valid TStrongObjectPtr
+	 * @return  true if garbage collection was not in progress, and OutResult was successfully captured, false if garbage collection was in progress and OutResult was not captured
+	 */
+	COREUOBJECT_API bool TryPin(class TStrongObjectPtr<UObject>& OutResult /*, bool bEvenIfGarbage = false*/) const;
+
+	/*
+	 * Attempt to get a strong object ptr even if it is marked as Garbage or Unreachable, but only if garbage collection is not in progress
+	 * @param OutResult will be invalid if this object is gone, or the weak pointer is explicitly null, or garbage collection is in progress. Otherwise a valid TStrongObjectPtr
+	 * @return  true if garbage collection was not in progress, and OutResult was successfully captured, false if garbage collection was in progress and OutResult was not captured 
+	 */
+	COREUOBJECT_API bool TryPinEvenIfUnreachable(class TStrongObjectPtr<UObject>& OutResult) const;
+
 	// This is explicitly not added to avoid resolving weak pointers too often - use Get() once in a function.
 	explicit operator bool() const = delete;
 
@@ -515,6 +537,7 @@ private:
 	}
 
 	COREUOBJECT_API TStrongObjectPtr<UObject> Internal_Pin(bool bEvenIfGarbage) const;
+	COREUOBJECT_API bool Internal_TryPin(bool bEvenIfGarbage, TStrongObjectPtr<UObject>& OutResult) const;
 	
 #if UE_WEAKOBJECTPTR_ZEROINIT_FIX
 	int32		ObjectIndex = UE::Core::Private::InvalidWeakObjectIndex;
