Instanced Struct: Fix crash on program teardown when an Instanced Struct is used as a static variable with a non-null underlying script struct.

Might need further investigation, see: https://udn.unrealengine.com/s/question/0D5QP00000aIj3m0AC/finstancedstruct-calling-reset-with-a-garbage-scriptstruct

diff --git a/Engine/Source/Runtime/CoreUObject/Private/StructUtils/InstancedStruct.cpp b/Engine/Source/Runtime/CoreUObject/Private/StructUtils/InstancedStruct.cpp
--- a/Engine/Source/Runtime/CoreUObject/Private/StructUtils/InstancedStruct.cpp
+++ b/Engine/Source/Runtime/CoreUObject/Private/StructUtils/InstancedStruct.cpp
@@ -108,7 +108,9 @@ void FInstancedStruct::Reset()
 	if (uint8* Memory = GetMutableMemory())
 	{
 		check(StructMemory != nullptr);
-		if (ScriptStruct != nullptr)
+
+		// We check if the struct is still valid, otherwise `static` Instanced Structs can crash here during teardown
+		if (ScriptStruct != nullptr && !ScriptStruct->HasAnyFlags(RF_BeginDestroyed | RF_FinishDestroyed | RF_MirroredGarbage))
 		{
 			ScriptStruct->DestroyStruct(GetMutableMemory());
 		}
