Cooker: Broadcast CookStarted from CookOnTheFly in addition to CookByTheBook. Some external systems need to make preparations for CookOnTheFly cooks.
#rnx

diff --git a/Engine/Source/Editor/UnrealEd/Private/CookOnTheFlyServer.cpp b/Engine/Source/Editor/UnrealEd/Private/CookOnTheFlyServer.cpp
--- a/Engine/Source/Editor/UnrealEd/Private/CookOnTheFlyServer.cpp
+++ b/Engine/Source/Editor/UnrealEd/Private/CookOnTheFlyServer.cpp
@@ -766,6 +766,9 @@ bool UCookOnTheFlyServer::StartCookOnTheFly(FCookOnTheFlyStartupOptions InCookOn
 	InitializePollables();
 
 	const bool bInitialized = CookOnTheFlyRequestManager->Initialize();
+
+	BroadcastCookStarted();
+
 	return bInitialized;
 #else
 	return false;
diff --git a/Engine/Source/Editor/UnrealEd/Private/Cooker/PackageTracker.h b/Engine/Source/Editor/UnrealEd/Private/Cooker/PackageTracker.h
--- a/Engine/Source/Editor/UnrealEd/Private/Cooker/PackageTracker.h
+++ b/Engine/Source/Editor/UnrealEd/Private/Cooker/PackageTracker.h
@@ -197,7 +197,7 @@ public:
 	FThreadSafeUnsolicitedPackagesList UnsolicitedCookedPackages;
 	FThreadSafeQueue<FRecompileShaderRequest> RecompileRequests;
 
-	/** Packages to never cook - entries are localpaths in FPaths::MakeStandardFilename format. */
+	/** Packages to never cook - entries are LongPackageNames. */
 	FThreadSafeSet<FName> NeverCookPackageList;
 	TFastPointerMap<const ITargetPlatform*, TSet<FName>> PlatformSpecificNeverCookPackages;
 
diff --git a/Engine/Source/Runtime/CoreUObject/Public/UObject/ICookInfo.h b/Engine/Source/Runtime/CoreUObject/Public/UObject/ICookInfo.h
--- a/Engine/Source/Runtime/CoreUObject/Public/UObject/ICookInfo.h
+++ b/Engine/Source/Runtime/CoreUObject/Public/UObject/ICookInfo.h
@@ -161,7 +161,6 @@ public:
 	 */
 	virtual void UnregisterCollector(IMPCollector* Collector) = 0;
 
-
 	/**
 	 * Gets the current cultures that are being cooked.
 	 */
@@ -179,8 +178,11 @@ public:
 	static COREUOBJECT_API FCookInfoEvent CookByTheBookStarted;
 	UE_DEPRECATED(5.5, "Use CookFinished, possibly restricting to the case CookInfo.GetCookType() == ECookType::ByTheBook.")
 	static COREUOBJECT_API FCookInfoEvent CookByTheBookFinished;
+	/** Called after a cook session has been initialized and is about to start ticking and loading/saving packages. */
 	static COREUOBJECT_API FCookInfoEvent CookStarted;
+	/** Called at the end of a cook session, after writing all cook artifacts. */
 	static COREUOBJECT_API FCookInfoEvent CookFinished;
+	/** Called after the cooker has loaded a requested package, before starting to save the package. */
 	static COREUOBJECT_API FValidateSourcePackage ValidateSourcePackage;
 };
 
