Adding FRemoteObjectReference::IsRemote() helper function

#preflught none

diff --git a/Engine/Source/Runtime/CoreUObject/Public/UObject/RemoteObjectTransfer.h b/Engine/Source/Runtime/CoreUObject/Public/UObject/RemoteObjectTransfer.h
--- a/Engine/Source/Runtime/CoreUObject/Public/UObject/RemoteObjectTransfer.h
+++ b/Engine/Source/Runtime/CoreUObject/Public/UObject/RemoteObjectTransfer.h
@@ -242,6 +242,10 @@ public:
 	{
 		return ServerId;
 	}
+	FORCEINLINE bool IsRemote() const
+	{
+		return UE::RemoteObject::Handle::IsRemote(ObjectId);
+	}
 
 	bool Serialize(FArchive& Ar)
 	{
