Deep copy the linker instancing context if we have to cascading abort in the transactionally safe linker. Otherwise we could be referencing dead memory.

diff --git a/Engine/Source/Programs/AutoRTFMTests/Private/UECoreTests.cpp b/Engine/Source/Programs/AutoRTFMTests/Private/UECoreTests.cpp
--- a/Engine/Source/Programs/AutoRTFMTests/Private/UECoreTests.cpp
+++ b/Engine/Source/Programs/AutoRTFMTests/Private/UECoreTests.cpp
@@ -16,6 +16,7 @@
 #include "Logging/StructuredLog.h"
 #include "Serialization/CustomVersion.h"
 #include "UObject/CoreRedirects.h"
+#include "UObject/LinkerInstancingContext.h"
 #include "UObject/NameTypes.h"
 #include "UObject/UObjectArray.h"
 #include "UObject/UObjectGlobals.h"
@@ -3328,6 +3329,22 @@ TEST_CASE("UECore.AsyncLoading")
 			AutoRTFM::AbortTransaction();
 		});
 	}
+
+	SECTION("Stack Local Linker Instancing Context")
+	{
+		int32 RequestId = -1;
+
+		AutoRTFM::Testing::Commit([&]
+		{
+			FLinkerInstancingContext Context;
+			FLoadPackageAsyncOptionalParams Params;
+			Params.InstancingContext = &Context;
+			FString Name(FString::Printf(TEXT("/AutoRTFMTestPackage%d"), __LINE__));
+			RequestId = LoadPackageAsync(Name, MoveTemp(Params));
+		});
+
+		FlushAsyncLoading(RequestId);
+	}
 }
 
 TEST_CASE("UECore.CoreRedirects")
diff --git a/Engine/Source/Runtime/CoreUObject/Private/Serialization/TransactionallySafeAsyncLoading.cpp b/Engine/Source/Runtime/CoreUObject/Private/Serialization/TransactionallySafeAsyncLoading.cpp
--- a/Engine/Source/Runtime/CoreUObject/Private/Serialization/TransactionallySafeAsyncLoading.cpp
+++ b/Engine/Source/Runtime/CoreUObject/Private/Serialization/TransactionallySafeAsyncLoading.cpp
@@ -4,6 +4,7 @@
 #include "Serialization/AsyncPackageLoader.h"
 #include "AutoRTFM/AutoRTFM.h"
 #include "Containers/Map.h"
+#include "UObject/LinkerInstancingContext.h"
 
 class FTransactionallySafeAsyncPackageLoader final : public IAsyncPackageLoader
 {
@@ -63,10 +64,16 @@ public:
 				{
 					FPackagePath PackagePathClone = PackagePath;
 					FName CustomPackageNameClone = FName(CustomPackageName.ToString());
+					TSharedPtr<FLinkerInstancingContext> LinkerInstancingContextClone;
+
+					if (InInstancingContext)
+					{
+						LinkerInstancingContextClone = MakeShared<FLinkerInstancingContext>(FLinkerInstancingContext::DuplicateContext(*InInstancingContext));
+					}
 
 					const AutoRTFM::EContextStatus Status = AutoRTFM::Close([&]
 					{
-						AutoRTFM::CascadingRetryTransaction([this, PackagePathClone, PackagePathAsString, CustomPackageNameClone, InPackageFlags, InPIEInstanceID, InPackagePriority, InInstancingContext, InLoadFlags]
+						AutoRTFM::CascadingRetryTransaction([this, PackagePathClone, PackagePathAsString, CustomPackageNameClone, InPackageFlags, InPIEInstanceID, InPackagePriority, LinkerInstancingContextClone = ::MoveTemp(LinkerInstancingContextClone), InLoadFlags]
 						{
 							FLoadPackageAsyncDelegate CompletionDelegate;
 							CompletionDelegate.BindLambda([this, PackagePathAsString](const FName& PackageName, UPackage* const Package, EAsyncLoadingResult::Type Result)
@@ -81,7 +88,7 @@ public:
 								InPackageFlags,
 								InPIEInstanceID,
 								InPackagePriority,
-								InInstancingContext,
+								LinkerInstancingContextClone.Get(),
 								InLoadFlags);
 							this->WrappedPackageLoader->FlushLoading({ RequestId });
 
