[BugFix] Instanced object in an array property of an actor component is invalidated and marked for garbage collection on undo

The source of this problem occurs when the user clicks the "reset to default" button. unlike other changes, Clicking that button will change the flags of the subobject to include RF_Transactional. This is unexpected behavior. As a result, when the user makes further changes, StaticConstructObject_Internal will toggle the garbage bits when attempting to transact the object. at this point, any undos will kill the object because the transaction buffer thinks it should destroy the object on undo.

The RF_Transactional flag was added during the reset operation because it was copied from the archetype. This change instead only copies RF_Transactional from the archetype if it's set in the parent component (which it isn't). There's a precident for this change all over the engine but one example is in CL 3185535

diff --git a/Engine/Source/Runtime/CoreUObject/Private/UObject/PropertyBaseObject.cpp b/Engine/Source/Runtime/CoreUObject/Private/UObject/PropertyBaseObject.cpp
--- a/Engine/Source/Runtime/CoreUObject/Private/UObject/PropertyBaseObject.cpp
+++ b/Engine/Source/Runtime/CoreUObject/Private/UObject/PropertyBaseObject.cpp
@@ -435,10 +435,11 @@ const TCHAR* FObjectPropertyBase::ImportText_Internal( const TCHAR* InBuffer, vo
 		{
 			ExistingObject->Rename(nullptr, nullptr, REN_DontCreateRedirectors);
 		}
-
-		FObjectDuplicationParameters ObjectDuplicationParams = InitStaticDuplicateObjectParams(Result, Parent, DesiredName);
+		
+		const EObjectFlags MaskedOuterFlags = Parent ? Parent->GetMaskedFlags(RF_PropagateToSubObjects) : RF_AllFlags;
+		FObjectDuplicationParameters ObjectDuplicationParams = InitStaticDuplicateObjectParams(Result, Parent, DesiredName, MaskedOuterFlags);
 		EnumRemoveFlags(ObjectDuplicationParams.FlagMask, RF_ArchetypeObject);
-		if (Parent->IsTemplate())
+		if (Parent && Parent->IsTemplate())
 		{
 			EnumAddFlags(ObjectDuplicationParams.ApplyFlags, RF_ArchetypeObject);
 		}
