Removed deprecated functionality from FSoftObjectPath and related calls.

diff --git a/Engine/Plugins/Editor/ContentBrowser/ContentBrowserAliasDataSource/Source/ContentBrowserAliasDataSource/Private/ContentBrowserAliasDataSource.cpp b/Engine/Plugins/Editor/ContentBrowser/ContentBrowserAliasDataSource/Source/ContentBrowserAliasDataSource/Private/ContentBrowserAliasDataSource.cpp
--- a/Engine/Plugins/Editor/ContentBrowser/ContentBrowserAliasDataSource/Source/ContentBrowserAliasDataSource/Private/ContentBrowserAliasDataSource.cpp
+++ b/Engine/Plugins/Editor/ContentBrowser/ContentBrowserAliasDataSource/Source/ContentBrowserAliasDataSource/Private/ContentBrowserAliasDataSource.cpp
@@ -507,12 +507,10 @@ void UContentBrowserAliasDataSource::RemoveFoldersRecursively(FStringView LeafFo
 	}
 }
 
-PRAGMA_DISABLE_DEPRECATION_WARNINGS
 void UContentBrowserAliasDataSource::RemoveAlias(FName ObjectPath, const FName Alias)
 {
-	RemoveAlias(FSoftObjectPath(ObjectPath), Alias);
+	RemoveAlias(FSoftObjectPath(ObjectPath.ToString()), Alias);
 }
-PRAGMA_ENABLE_DEPRECATION_WARNINGS
 
 void UContentBrowserAliasDataSource::RemoveAlias(const FSoftObjectPath& ObjectPath, const FName Alias)
 {
@@ -592,12 +590,10 @@ void UContentBrowserAliasDataSource::RebuildAliases()
 	OnRebuildAliases().Broadcast();
 }
 
-PRAGMA_DISABLE_DEPRECATION_WARNINGS
 void UContentBrowserAliasDataSource::RemoveAliases(FName ObjectPath)
 {
-	RemoveAliases(FSoftObjectPath(ObjectPath));
+	RemoveAliases(FSoftObjectPath(ObjectPath.ToString()));
 }
-PRAGMA_ENABLE_DEPRECATION_WARNINGS
 
 void UContentBrowserAliasDataSource::RemoveAliases(const FSoftObjectPath& ObjectPath)
 {
@@ -1063,9 +1059,7 @@ FContentBrowserItemData UContentBrowserAliasDataSource::CreateAssetFileItem(cons
 	{
 		LLM_SCOPE(ELLMTag::UI);
 		FName VirtualizedPath;
-		PRAGMA_DISABLE_DEPRECATION_WARNINGS
-		FName InternalPath = AliasData->ObjectPath.ToFName();
-		PRAGMA_ENABLE_DEPRECATION_WARNINGS
+		FName InternalPath = *AliasData->ObjectPath.ToString();
 		TryConvertInternalPathToVirtual(InternalPath, VirtualizedPath);
 
 		// Since AliasID is PackagePath/AssetName, AssetName should also be passed as the ItemName here. This provides the functionality of
diff --git a/Engine/Plugins/FX/Niagara/Source/NiagaraEditor/Private/EdGraphSchema_Niagara.cpp b/Engine/Plugins/FX/Niagara/Source/NiagaraEditor/Private/EdGraphSchema_Niagara.cpp
--- a/Engine/Plugins/FX/Niagara/Source/NiagaraEditor/Private/EdGraphSchema_Niagara.cpp
+++ b/Engine/Plugins/FX/Niagara/Source/NiagaraEditor/Private/EdGraphSchema_Niagara.cpp
@@ -524,9 +524,7 @@ TArray<TSharedPtr<FNiagaraAction_NewNode>> UEdGraphSchema_Niagara::GetGraphActio
 		const ENiagaraMenuSections Section = bSuggested ? ENiagaraMenuSections::Suggested: ENiagaraMenuSections::General;
 
 		UNiagaraNodeFunctionCall* FunctionCallNode = NewObject<UNiagaraNodeFunctionCall>(OwnerOfTemporaries);
-PRAGMA_DISABLE_DEPRECATION_WARNINGS
-		FunctionCallNode->FunctionScriptAssetObjectPath = ScriptAsset.GetSoftObjectPath().ToFName();
-PRAGMA_ENABLE_DEPRECATION_WARNINGS
+		FunctionCallNode->FunctionScriptAssetObjectPath = FName(ScriptAsset.GetSoftObjectPath().ToString());
 
 		TSharedPtr<FNiagaraAction_NewNode> Action = AddNewNodeMenuAction(NewActions, FunctionCallNode, DisplayName, Section, Categories, TooltipDesc, Keywords, SourceData);
 		Action->bIsInLibrary = bIsInLibrary;
diff --git a/Engine/Plugins/FX/Niagara/Source/NiagaraEditor/Private/ViewModels/Stack/NiagaraStackGraphUtilities.cpp b/Engine/Plugins/FX/Niagara/Source/NiagaraEditor/Private/ViewModels/Stack/NiagaraStackGraphUtilities.cpp
--- a/Engine/Plugins/FX/Niagara/Source/NiagaraEditor/Private/ViewModels/Stack/NiagaraStackGraphUtilities.cpp
+++ b/Engine/Plugins/FX/Niagara/Source/NiagaraEditor/Private/ViewModels/Stack/NiagaraStackGraphUtilities.cpp
@@ -2415,13 +2415,12 @@ bool FNiagaraStackGraphUtilities::FindScriptModulesInStack(FAssetData ModuleScri
 
 	OutFunctionCalls.Empty();
 	FString ModuleObjectName = ModuleScriptAsset.GetObjectPathString();
+	FName ModuleScriptAssetSoftObjectPathName = *ModuleScriptAsset.GetSoftObjectPath().ToString();
 	for (UEdGraphNode* Node : Nodes)
 	{
 		if (UNiagaraNodeFunctionCall* FunctionCallNode = Cast<UNiagaraNodeFunctionCall>(Node))
 		{
-PRAGMA_DISABLE_DEPRECATION_WARNINGS
-			if (FunctionCallNode->FunctionScriptAssetObjectPath == ModuleScriptAsset.GetSoftObjectPath().ToFName() ||
-PRAGMA_ENABLE_DEPRECATION_WARNINGS
+			if (FunctionCallNode->FunctionScriptAssetObjectPath == ModuleScriptAssetSoftObjectPathName ||
 				(FunctionCallNode->FunctionScript != nullptr && FunctionCallNode->FunctionScript->GetPathName() == ModuleObjectName))
 			{
 				OutFunctionCalls.Add(FunctionCallNode);
@@ -2527,9 +2526,7 @@ UNiagaraNodeFunctionCall* FNiagaraStackGraphUtilities::AddScriptModuleToStack(FA
 
 	FGraphNodeCreator<UNiagaraNodeFunctionCall> ModuleNodeCreator(*Graph);
 	UNiagaraNodeFunctionCall* NewModuleNode = ModuleNodeCreator.CreateNode();
-PRAGMA_DISABLE_DEPRECATION_WARNINGS
-	NewModuleNode->FunctionScriptAssetObjectPath = ModuleScriptAsset.GetSoftObjectPath().ToFName();
-PRAGMA_ENABLE_DEPRECATION_WARNINGS
+	NewModuleNode->FunctionScriptAssetObjectPath = FName(*ModuleScriptAsset.GetSoftObjectPath().ToString());
 	ModuleNodeCreator.Finalize();
 
 	if (NewModuleNode->HasValidScriptAndGraph() == false)
diff --git a/Engine/Plugins/Runtime/AssetTags/Source/AssetTags/Private/AssetTagsSubsystem.cpp b/Engine/Plugins/Runtime/AssetTags/Source/AssetTags/Private/AssetTagsSubsystem.cpp
--- a/Engine/Plugins/Runtime/AssetTags/Source/AssetTags/Private/AssetTagsSubsystem.cpp
+++ b/Engine/Plugins/Runtime/AssetTags/Source/AssetTags/Private/AssetTagsSubsystem.cpp
@@ -184,9 +184,7 @@ bool UAssetTagsSubsystem::K2_AddAssetToCollection(const FName Name, const FSoftO
 
 bool UAssetTagsSubsystem::AddAssetToCollection(const FName Name, const FName AssetPathName)
 {
-PRAGMA_DISABLE_DEPRECATION_WARNINGS
-	return K2_AddAssetToCollection(Name, FSoftObjectPath(AssetPathName));
-PRAGMA_ENABLE_DEPRECATION_WARNINGS
+	return K2_AddAssetToCollection(Name, FSoftObjectPath(AssetPathName.ToString()));
 }
 
 bool UAssetTagsSubsystem::AddAssetDataToCollection(const FName Name, const FAssetData& AssetData)
@@ -272,9 +270,7 @@ bool UAssetTagsSubsystem::K2_RemoveAssetFromCollection(const FName Name, const F
 
 bool UAssetTagsSubsystem::RemoveAssetFromCollection(const FName Name, const FName AssetPathName)
 {
-PRAGMA_DISABLE_DEPRECATION_WARNINGS
-	return K2_RemoveAssetFromCollection(Name, FSoftObjectPath(AssetPathName));
-PRAGMA_ENABLE_DEPRECATION_WARNINGS
+	return K2_RemoveAssetFromCollection(Name, FSoftObjectPath(AssetPathName.ToString()));
 }
 
 bool UAssetTagsSubsystem::RemoveAssetDataFromCollection(const FName Name, const FAssetData& AssetData)
@@ -492,9 +488,7 @@ TArray<FName> UAssetTagsSubsystem::K2_GetCollectionsContainingAsset(const FSoftO
 
 TArray<FName> UAssetTagsSubsystem::GetCollectionsContainingAsset(const FName AssetPathName)
 {
-PRAGMA_DISABLE_DEPRECATION_WARNINGS
-	return K2_GetCollectionsContainingAsset(FSoftObjectPath(AssetPathName));
-PRAGMA_ENABLE_DEPRECATION_WARNINGS
+	return K2_GetCollectionsContainingAsset(FSoftObjectPath(AssetPathName.ToString()));
 }
 
 TArray<FName> UAssetTagsSubsystem::GetCollectionsContainingAssetData(const FAssetData& AssetData)
diff --git a/Engine/Plugins/Runtime/RigVM/Source/RigVMEditor/Private/EdGraph/NodeSpawners/RigVMEdGraphFunctionRefNodeSpawner.cpp b/Engine/Plugins/Runtime/RigVM/Source/RigVMEditor/Private/EdGraph/NodeSpawners/RigVMEdGraphFunctionRefNodeSpawner.cpp
--- a/Engine/Plugins/Runtime/RigVM/Source/RigVMEditor/Private/EdGraph/NodeSpawners/RigVMEdGraphFunctionRefNodeSpawner.cpp
+++ b/Engine/Plugins/Runtime/RigVM/Source/RigVMEditor/Private/EdGraph/NodeSpawners/RigVMEdGraphFunctionRefNodeSpawner.cpp
@@ -154,9 +154,7 @@ URigVMEdGraphFunctionRefNodeSpawner* URigVMEdGraphFunctionRefNodeSpawner::Create
 		NewArgument.Direction = Arg.Direction;
 		NewArgument.bIsArray = Arg.bIsArray;
 		NewArgument.CPPType = Arg.CPPType;
-		PRAGMA_DISABLE_DEPRECATION_WARNINGS
-		NewArgument.CPPTypeObject = FSoftObjectPath(Arg.CPPTypeObjectPath);
-		PRAGMA_ENABLE_DEPRECATION_WARNINGS
+		NewArgument.CPPTypeObject = FSoftObjectPath(Arg.CPPTypeObjectPath.ToString());
 		Header.Arguments.Add(NewArgument);
 	}
 	NodeSpawner->AssetPath = InAssetData.ToSoftObjectPath();
diff --git a/Engine/Source/Editor/BlueprintGraph/Public/BlueprintActionDatabase.h b/Engine/Source/Editor/BlueprintGraph/Public/BlueprintActionDatabase.h
--- a/Engine/Source/Editor/BlueprintGraph/Public/BlueprintActionDatabase.h
+++ b/Engine/Source/Editor/BlueprintGraph/Public/BlueprintActionDatabase.h
@@ -154,14 +154,6 @@ public:
 	 */
 	void ClearUnloadedAssetActions(const FSoftObjectPath& ObjectPath);
 
-	UE_DEPRECATED(5.1, "FNames containing full asset paths are deprecated. Use FSoftObjectPath instead.")
-	void ClearUnloadedAssetActions(FName ObjectPath)
-	{
-PRAGMA_DISABLE_DEPRECATION_WARNINGS
-		ClearUnloadedAssetActions(FSoftObjectPath(ObjectPath));
-PRAGMA_ENABLE_DEPRECATION_WARNINGS
-	}
-
 	/**
 	 * Moves the unloaded asset actions from one location to another
 	 *
@@ -170,14 +162,6 @@ PRAGMA_ENABLE_DEPRECATION_WARNINGS
 	 */
 	void MoveUnloadedAssetActions(const FSoftObjectPath& SourceObjectPath, const FSoftObjectPath& TargetObjectPath);
 
-	UE_DEPRECATED(5.1, "FNames containing full asset paths are deprecated. Use FSoftObjectPath instead.")
-	void MoveUnloadedAssetActions(FName SourceObjectPath, FName TargetObjectPath) 
-	{
-PRAGMA_DISABLE_DEPRECATION_WARNINGS
-		MoveUnloadedAssetActions(FSoftObjectPath(SourceObjectPath), FSoftObjectPath(TargetObjectPath));
-PRAGMA_ENABLE_DEPRECATION_WARNINGS
-	}
-
 	/** */
 	FOnDatabaseEntryUpdated& OnEntryUpdated() { return EntryRefreshDelegate; }
 	/** */
diff --git a/Engine/Source/Editor/ContentBrowserData/Private/ContentBrowserDataSource.cpp b/Engine/Source/Editor/ContentBrowserData/Private/ContentBrowserDataSource.cpp
--- a/Engine/Source/Editor/ContentBrowserData/Private/ContentBrowserDataSource.cpp
+++ b/Engine/Source/Editor/ContentBrowserData/Private/ContentBrowserDataSource.cpp
@@ -250,9 +250,7 @@ TArray<FContentBrowserItemPath> UContentBrowserDataSource::GetAliasesForPath(con
 
 TArray<FContentBrowserItemPath> UContentBrowserDataSource::GetAliasesForPath(FName InInternalPath) const
 {
-PRAGMA_DISABLE_DEPRECATION_WARNINGS
-	return GetAliasesForPath(FSoftObjectPath(InInternalPath));
-PRAGMA_ENABLE_DEPRECATION_WARNINGS
+	return GetAliasesForPath(FSoftObjectPath(InInternalPath.ToString()));
 }
 
 bool UContentBrowserDataSource::IsDiscoveringItems(FText* OutStatus)
diff --git a/Engine/Source/Runtime/CoreUObject/Private/UObject/SoftObjectPath.cpp b/Engine/Source/Runtime/CoreUObject/Private/UObject/SoftObjectPath.cpp
--- a/Engine/Source/Runtime/CoreUObject/Private/UObject/SoftObjectPath.cpp
+++ b/Engine/Source/Runtime/CoreUObject/Private/UObject/SoftObjectPath.cpp
@@ -32,13 +32,6 @@ static FAutoConsoleVariableRef CVarEnablePathFixupOutsideEditor(TEXT("SoftObject
 FName FSoftObjectPath::NAME_Untracked(TEXT("Untracked"));
 #endif
 
-// Deprecated constructor
-FSoftObjectPath::FSoftObjectPath(FName InAssetPathName, FString InSubPathString)
-{
-	AssetPath = FTopLevelAssetPath(WriteToString<FName::StringBufferSize>(InAssetPathName).ToView());
-	SubPathString = MoveTemp(InSubPathString);
-}
-
 /** Static methods for more meaningful construction sites. */
 FSoftObjectPath FSoftObjectPath::ConstructFromPackageAssetSubpath(FName InPackageName, FName InAssetName, const FString& InSubPathString)
 {
@@ -158,19 +151,6 @@ void FSoftObjectPath::AppendString(FString& Builder) const
 	}
 }
 
-FName FSoftObjectPath::GetAssetPathName() const
-{
-PRAGMA_DISABLE_DEPRECATION_WARNINGS
-	return AssetPath.ToFName();
-PRAGMA_ENABLE_DEPRECATION_WARNINGS
-}
-
-// Deprecated setter
-void FSoftObjectPath::SetAssetPathName(FName InPath)
-{
-	AssetPath = WriteToString<FName::StringBufferSize>(InPath).ToView();
-}
-
 /** Helper function that adds info about the object currently being serialized when triggering an ensure about invalid soft object path */
 static FString GetObjectBeingSerializedForSoftObjectPath()
 {
@@ -267,12 +247,6 @@ void FSoftObjectPath::SetPath(FUtf8StringView Path)
 	SetPath(Wide);
 }
 
-// Deprecated setter
-void FSoftObjectPath::SetPath(FName PathName)
-{
-	SetPath(PathName.ToString());
-}
-
 #if WITH_EDITOR
 	extern bool* GReportSoftObjectPathRedirects;
 #endif
diff --git a/Engine/Source/Runtime/CoreUObject/Public/UObject/SoftObjectPath.h b/Engine/Source/Runtime/CoreUObject/Public/UObject/SoftObjectPath.h
--- a/Engine/Source/Runtime/CoreUObject/Public/UObject/SoftObjectPath.h
+++ b/Engine/Source/Runtime/CoreUObject/Public/UObject/SoftObjectPath.h
@@ -69,18 +69,6 @@ struct FSoftObjectPath
 	explicit FSoftObjectPath(const ANSICHAR* Path)				{ SetPath(FAnsiStringView(Path)); }
 	explicit FSoftObjectPath(TYPE_OF_NULLPTR)					{}
 
-	UE_DEPRECATED(5.1, "Asset path FNames have been deprecated. This constructor should be used only temporarily to fix up old codepaths that produce an FName.")
-	explicit FSoftObjectPath(FName Path) 
-	{
-PRAGMA_DISABLE_DEPRECATION_WARNINGS;
-		SetPath(Path); 
-PRAGMA_ENABLE_DEPRECATION_WARNINGS;
-	}
-
-	/** Construct from an asset FName and subobject pair */
-	UE_DEPRECATED(5.1, "Asset path FNames have been deprecated. This constructor should be used only temporarily to fix up old codepaths that produce an FName.")
-	COREUOBJECT_API FSoftObjectPath(FName InAssetPathName, FString InSubPathString);
-	
 	template <typename T>
 	FSoftObjectPath(const TObjectPtr<T>& InObject)
 	{
@@ -125,21 +113,9 @@ PRAGMA_ENABLE_DEPRECATION_WARNINGS;
 	FSoftObjectPath& operator=(const ANSICHAR* Path)					{ SetPath(FAnsiStringView(Path)); return *this; }
 	FSoftObjectPath& operator=(TYPE_OF_NULLPTR)							{ Reset(); return *this; }
 
-	UE_DEPRECATED(5.1, "Asset path FNames have been deprecated. This assignment operator should be used only temporarily to fix up old codepaths that produce an FName.")
-	FSoftObjectPath& operator=(FName Path) 
-	{
-PRAGMA_DISABLE_DEPRECATION_WARNINGS;
-		SetPath(Path); 
-PRAGMA_ENABLE_DEPRECATION_WARNINGS;
-		return *this;
-	}
-
 	/** Returns string representation of reference, in form /package/path.assetname[:subpath] */
 	COREUOBJECT_API FString ToString() const;
 
-	UE_DEPRECATED(5.1, "Asset path FNames have been deprecated. This function should be used only temporarily to interface with old APIs that require an FName.")
-	FName ToFName() const { return *ToString(); }
-
 	/** Append string representation of reference, in form /package/path.assetname[:subpath] */
 	COREUOBJECT_API void ToString(FStringBuilderBase& Builder) const;
 	COREUOBJECT_API void ToString(FUtf8StringBuilderBase& Builder) const;
@@ -160,13 +136,6 @@ PRAGMA_ENABLE_DEPRECATION_WARNINGS;
 	{
 		return FSoftObjectPath(AssetPath, {});
 	}
-	
-	/** Returns the entire asset path as an FName, including both package and asset but not sub object */
-	UE_DEPRECATED(5.1, "Asset path FNames have been deprecated. Use GetAssetPath instead.")
-	COREUOBJECT_API FName GetAssetPathName() const;
-	
-	UE_DEPRECATED(5.1, "Asset path FNames have been deprecated. Use SetAssetPath instead.")
-	COREUOBJECT_API void SetAssetPathName(FName InAssetPathName);
 
 	/** Returns string version of asset path, including both package and asset but not sub object */
 	FORCEINLINE FString GetAssetPathString() const
@@ -221,8 +190,6 @@ PRAGMA_ENABLE_DEPRECATION_WARNINGS;
 	COREUOBJECT_API void SetPath(FWideStringView Path);
 	COREUOBJECT_API void SetPath(FAnsiStringView Path);
 	COREUOBJECT_API void SetPath(FUtf8StringView Path);
-	UE_DEPRECATED(5.1, "Asset path FNames have been deprecated. This function should be used only temporarily to fix up old codepaths that produce an FName.")
-	COREUOBJECT_API void SetPath(FName Path);
 	void SetPath(const WIDECHAR* Path)			{ SetPath(FWideStringView(Path)); }
 	void SetPath(const ANSICHAR* Path)			{ SetPath(FAnsiStringView(Path)); }
 	void SetPath(const FString& Path)			{ SetPath(FStringView(Path)); }
@@ -351,16 +318,6 @@ PRAGMA_ENABLE_DEPRECATION_WARNINGS;
 		return Hash;
 	}
 
-	UE_DEPRECATED(5.4, "The current object tag is no longer used by TSoftObjectPtr, you can remove all calls")
-	static int32 GetCurrentTag()
-	{
-		return 0;
-	}
-	UE_DEPRECATED(5.4, "The current object tag is no longer used by TSoftObjectPtr, you can remove all calls")
-	static int32 InvalidateTag()
-	{
-		return 0;
-	}
 	static COREUOBJECT_API FSoftObjectPath GetOrCreateIDForObject(FObjectPtr Object);
 	static FORCEINLINE FSoftObjectPath GetOrCreateIDForObject(const UObject* Object)
 	{
@@ -496,12 +453,6 @@ private:
 	static COREUOBJECT_API FSoftObjectPath GetOrCreateIDForObject(const UObject *Object);
 };
 
-UE_DEPRECATED(5.0, "FStringAssetReference was renamed to FSoftObjectPath as it is now not always a string and can also refer to a subobject")
-typedef FSoftObjectPath FStringAssetReference;
-
-UE_DEPRECATED(5.0, "FStringClassReference was renamed to FSoftClassPath")
-typedef FSoftClassPath FStringClassReference;
-
 /** Options for how to set soft object path collection */
 enum class ESoftObjectPathCollectType : uint8
 {
@@ -657,7 +608,10 @@ namespace UE::SoftObjectPath::Private
 		TArray<FName> Out;
 		Out.Reserve(InPaths.Num());
 PRAGMA_DISABLE_DEPRECATION_WARNINGS
-		Algo::Transform(InPaths, Out, [](const FSoftObjectPath& Path) { return Path.ToFName(); });
+		Algo::Transform(InPaths, Out, [](const FSoftObjectPath& Path)
+		{
+			return FName(*Path.ToString());
+		});
 PRAGMA_ENABLE_DEPRECATION_WARNINGS
 		return Out;
 	}
@@ -668,7 +622,12 @@ PRAGMA_ENABLE_DEPRECATION_WARNINGS
 		TArray<FSoftObjectPath> Out;
 		Out.Reserve(InPaths.Num());
 PRAGMA_DISABLE_DEPRECATION_WARNINGS
-		Algo::Transform(InPaths, Out, [](FName Name) { return FSoftObjectPath(Name); });
+		Algo::Transform(InPaths, Out, [](FName Name)
+		{
+				FSoftObjectPath Result;
+				Result.SetPath(Name.ToString());
+				return Result;
+		});
 PRAGMA_ENABLE_DEPRECATION_WARNINGS
 		return Out;
 	}
diff --git a/Engine/Source/Runtime/Engine/Classes/Engine/AssetManager.h b/Engine/Source/Runtime/Engine/Classes/Engine/AssetManager.h
--- a/Engine/Source/Runtime/Engine/Classes/Engine/AssetManager.h
+++ b/Engine/Source/Runtime/Engine/Classes/Engine/AssetManager.h
@@ -206,9 +206,6 @@ public:
 	/** Sees if the passed in object path is a registered primary asset, if so return it. Returns invalid Identifier if not found */
 	ENGINE_API virtual FPrimaryAssetId GetPrimaryAssetIdForPath(const FSoftObjectPath& ObjectPath) const;
 
-	UE_DEPRECATED(5.1, "FName asset paths are deprecated, use FSoftObjectPath instead.")
-	ENGINE_API virtual FPrimaryAssetId GetPrimaryAssetIdForPath(FName ObjectPath) const;
-
 	/** Sees if the package has a primary asset, useful if only the package name is available */
 	ENGINE_API virtual FPrimaryAssetId GetPrimaryAssetIdForPackage(FName PackagePath) const;
 
@@ -459,8 +456,6 @@ public:
 	ENGINE_API virtual FPrimaryAssetId DeterminePrimaryAssetIdForObject(const UObject* Object) const;
 
 	/** Reads AssetManagerSettings for specifically redirected asset paths. This is useful if you need to convert older saved data */
-	UE_DEPRECATED(5.1, "Asset path FNames are deprecated, use FSoftObjectPath instead.")
-	ENGINE_API virtual FName GetRedirectedAssetPath(FName OldPath) const;
 	ENGINE_API virtual FSoftObjectPath GetRedirectedAssetPath(const FSoftObjectPath& OldPath) const;
 
 	/** Extracts all FSoftObjectPaths from a Class/Struct */
diff --git a/Engine/Source/Runtime/Engine/Private/AssetManager.cpp b/Engine/Source/Runtime/Engine/Private/AssetManager.cpp
--- a/Engine/Source/Runtime/Engine/Private/AssetManager.cpp
+++ b/Engine/Source/Runtime/Engine/Private/AssetManager.cpp
@@ -1594,13 +1594,6 @@ FPrimaryAssetId UAssetManager::GetPrimaryAssetIdForPath(const FSoftObjectPath& O
 	return FPrimaryAssetId();
 }
 
-FPrimaryAssetId UAssetManager::GetPrimaryAssetIdForPath(FName ObjectPath) const
-{
-PRAGMA_DISABLE_DEPRECATION_WARNINGS
-	return GetPrimaryAssetIdForPath(FSoftObjectPath(ObjectPath));
-PRAGMA_ENABLE_DEPRECATION_WARNINGS
-}
-
 FPrimaryAssetId UAssetManager::GetPrimaryAssetIdForPackage(FName PackagePath) const
 {
 	FString PackageString = PackagePath.ToString();
@@ -2873,18 +2866,6 @@ void UAssetManager::GetPreviousPrimaryAssetIds(const FPrimaryAssetId& NewId, TAr
 	}
 }
 
-FName UAssetManager::GetRedirectedAssetPath(FName OldPath) const
-{
-PRAGMA_DISABLE_DEPRECATION_WARNINGS
-	const FSoftObjectPath* Redirected = AssetPathRedirects.Find(FSoftObjectPath(OldPath));
-	if (Redirected)
-	{
-		return Redirected->ToFName();
-	}
-	return NAME_None;
-PRAGMA_ENABLE_DEPRECATION_WARNINGS
-}
-
 FSoftObjectPath UAssetManager::GetRedirectedAssetPath(const FSoftObjectPath& ObjectPath) const
 {
 	const FSoftObjectPath* RedirectedName = AssetPathRedirects.Find(ObjectPath.GetWithoutSubPath());
diff --git a/Engine/Source/Runtime/Engine/Private/WorldPartition/WorldPartitionActorDesc.cpp b/Engine/Source/Runtime/Engine/Private/WorldPartition/WorldPartitionActorDesc.cpp
--- a/Engine/Source/Runtime/Engine/Private/WorldPartition/WorldPartitionActorDesc.cpp
+++ b/Engine/Source/Runtime/Engine/Private/WorldPartition/WorldPartitionActorDesc.cpp
@@ -765,9 +765,7 @@ void FWorldPartitionActorDesc::Serialize(FArchive& Ar)
 		{ 
 			FName ActorPathName;
 			Ar << ActorPathName;
-PRAGMA_DISABLE_DEPRECATION_WARNINGS
-			ActorPath = FSoftObjectPath(ActorPathName);
-PRAGMA_ENABLE_DEPRECATION_WARNINGS
+			ActorPath = FSoftObjectPath(ActorPathName.ToString());
 		}
 		else
 		{
