Replaced FString by FUtf8String in FixedTagPrivate::FMarshalledText to save memory (resubmit with fixed check statement during save serialization).

diff --git a/Engine/Source/Runtime/AssetRegistry/Private/AssetRegistryArchive.cpp b/Engine/Source/Runtime/AssetRegistry/Private/AssetRegistryArchive.cpp
--- a/Engine/Source/Runtime/AssetRegistry/Private/AssetRegistryArchive.cpp
+++ b/Engine/Source/Runtime/AssetRegistry/Private/AssetRegistryArchive.cpp
@@ -266,15 +266,16 @@ FAssetRegistryReader::FAssetRegistryReader(FArchive& Inner, int32 NumWorkers, FA
 	SetFilterEditorOnly(Header.bFilterEditorOnlyData);
 	FArchive::SetFilterEditorOnly(Header.bFilterEditorOnlyData); // Workaround for bug in FArchiveProxy
 
-	if (NumWorkers > 0)
+	// Use parallel loading if requested and available; the parallel path is only for latest version
+	if (NumWorkers > 0 && Header.Version == FAssetRegistryVersion::LatestVersion)
 	{
 		TFunction<TArray<FDisplayNameEntryId> ()> GetFutureNames = LoadNameBatchAsync(*this, NumWorkers);
 
 		FixedTagPrivate::FAsyncStoreLoader StoreLoader;
-		Task = StoreLoader.ReadInitialDataAndKickLoad(*this, NumWorkers);
+		Task = StoreLoader.ReadInitialDataAndKickLoad(*this, NumWorkers, Header.Version);
 		
 		Names = GetFutureNames();
-		Tags = StoreLoader.LoadFinalData(*this);
+		Tags = StoreLoader.LoadFinalData(*this, Header.Version);
 	}
 	else
 	{
diff --git a/Engine/Source/Runtime/CoreUObject/Private/AssetRegistry/AssetDataTagMap.cpp b/Engine/Source/Runtime/CoreUObject/Private/AssetRegistry/AssetDataTagMap.cpp
--- a/Engine/Source/Runtime/CoreUObject/Private/AssetRegistry/AssetDataTagMap.cpp
+++ b/Engine/Source/Runtime/CoreUObject/Private/AssetRegistry/AssetDataTagMap.cpp
@@ -167,11 +167,15 @@ static FString LocalizeIfComplexString(const FString& Value)
 
 namespace FixedTagPrivate
 {
+	FMarshalledText::FMarshalledText(const FUtf8String& InComplexString)
+		: String(InComplexString)
+	{
+	}
 	FMarshalledText::FMarshalledText(const FString& InComplexString)
 		: String(InComplexString)
 	{
 	}
-	FMarshalledText::FMarshalledText(FString&& InComplexString)
+	FMarshalledText::FMarshalledText(FUtf8String&& InComplexString)
 		: String(MoveTemp(InComplexString))
 	{
 	}
@@ -184,14 +188,14 @@ namespace FixedTagPrivate
 	{
 	}
 
-	const FString& FMarshalledText::GetAsComplexString() const
+	const FUtf8String& FMarshalledText::GetAsComplexString() const
 	{
 		return String;
 	}
 
 	FText FMarshalledText::GetAsText() const
 	{
-		return FTextStringHelper::CreateFromBuffer(*String);
+		return FTextStringHelper::CreateFromBuffer(StringCast<TCHAR>(*String).Get());
 	}
 
 	int64 FMarshalledText::GetResourceSize() const
@@ -214,6 +218,11 @@ namespace FixedTagPrivate
 		return CityHash32(reinterpret_cast<const char*>(Str), sizeof(TCHAR) * Len);
 	}
 
+	uint32 HashCaseSensitive(const UTF8CHAR* Str, int32 Len)
+	{
+		return CityHash32(reinterpret_cast<const char*>(Str), sizeof(UTF8CHAR) * Len);
+	}
+
 	uint32 HashCombineQuick(uint32 A, uint32 B)
 	{
 		return A ^ (B + 0x9e3779b9 + (A << 6) + (A >> 2));
@@ -431,7 +440,7 @@ namespace FixedTagPrivate
 		case EValueType::NumberlessExportPath:	return Store.NumberlessExportPaths[Index].ToString();
 		case EValueType::ExportPath:			return Store.ExportPaths[Index].ToString();
 		case EValueType::LocalizedText:
-			return bForStorage ? Store.Texts[Index].GetAsComplexString()
+			return bForStorage ? StringCast<TCHAR>(*Store.Texts[Index].GetAsComplexString()).Get()
 				: Store.Texts[Index].GetAsText().ToString();
 		}
 
@@ -1045,7 +1054,7 @@ namespace FixedTagPrivate
 
 		template<> FAssetRegistryExportPath LoadItem()
 		{
-			if (Version >= FAssetRegistryVersion::ClassPaths)
+			if constexpr (Version >= FAssetRegistryVersion::ClassPaths)
 			{
 				FAssetRegistryExportPath ExportPath;
 				Ar << ExportPath;
@@ -1098,14 +1107,23 @@ namespace FixedTagPrivate
 		
 		void SaveItem(const FMarshalledText& Text)
 		{
-			Ar << const_cast<FString&>(Text.GetAsComplexString());
+			Ar << const_cast<FUtf8String&>(Text.GetAsComplexString());
 		}
 
 		template<> FMarshalledText LoadItem()
 		{
-			FString ComplexString;
-			Ar << ComplexString;
-			return FMarshalledText(MoveTemp(ComplexString));
+			if constexpr (Version >= FAssetRegistryVersion::MarshalledTextAsUTF8String)
+			{
+				FUtf8String ComplexString;
+				Ar << ComplexString;
+				return FMarshalledText(MoveTemp(ComplexString));
+			}
+			else
+			{
+				FString ComplexString;
+				Ar << ComplexString;
+				return FMarshalledText(ComplexString);
+			}
 		}
 		
 		template<typename T>
@@ -1292,6 +1310,7 @@ namespace FixedTagPrivate
 
 	void SaveStore(const FStoreData& Store, FArchive& Ar)
 	{
+		// Saving always uses the latest FAssetRegistryVersion
 		TSerializer<FAssetRegistryVersion::LatestVersion>(Ar).Save(Store);
 	}
 
@@ -1308,14 +1327,18 @@ namespace FixedTagPrivate
 		// serialized with FAssetRegistryVersion::ClassPaths - 1 = FAssetRegistryVersion::AddedChunkHashes.
 		// The main purpose of hardcoding version value and passing it as template argument is to eliminate
 		// branches in performance critical serialization code.
-		// Note that we branch in this function only because it's the only fucntion used to load legacy 
+		// Note that we branch in this function only because it's the only function used to load legacy 
 		// asset registries in DiffAssetRegistries commandlet hence it's editor only code too
 		if (Version < FAssetRegistryVersion::ClassPaths)
 		{
 			TSerializer<FAssetRegistryVersion::AddedChunkHashes>(Ar).Load(*Store);
 		}
+		else if (Version <= FAssetRegistryVersion::AssetPackageDataHasExtension)
+		{
+			TSerializer<FAssetRegistryVersion::AssetPackageDataHasExtension>(Ar).Load(*Store);
+		}
 		else // Add more branches should serialization format change
-		{			
+		{
 			TSerializer<FAssetRegistryVersion::LatestVersion>(Ar).Load(*Store);
 		}
 		
@@ -1326,8 +1349,10 @@ namespace FixedTagPrivate
 		: Store(GStores.CreateAndRegister())
 	{}
 
-	TFuture<void> FAsyncStoreLoader::ReadInitialDataAndKickLoad(FArchive& Ar, uint32 MaxWorkerTasks)
+	TFuture<void> FAsyncStoreLoader::ReadInitialDataAndKickLoad(FArchive& Ar, uint32 MaxWorkerTasks, FAssetRegistryVersion::Type HeaderVersion)
 	{
+		// should only load latest version of data
+		check(HeaderVersion == FAssetRegistryVersion::LatestVersion);
 		Order = TSerializer<FAssetRegistryVersion::LatestVersion>(Ar).LoadHeader(*Store);
 
 		if (Order == ELoadOrder::TextFirst)
@@ -1347,10 +1372,12 @@ namespace FixedTagPrivate
 		return TFuture<void>();
 	}
 
-	TRefCountPtr<const FStore> FAsyncStoreLoader::LoadFinalData(FArchive& Ar)
+	TRefCountPtr<const FStore> FAsyncStoreLoader::LoadFinalData(FArchive& Ar, FAssetRegistryVersion::Type HeaderVersion)
 	{
 		if (Order.IsSet())
 		{
+			// should only load latest version of data
+			check(HeaderVersion == FAssetRegistryVersion::LatestVersion);
 			TSerializer<FAssetRegistryVersion::LatestVersion>(Ar).LoadFinalData(/* Out */ *Store, Order.GetValue());
 
 			return TRefCountPtr<const FStore>(Store);
diff --git a/Engine/Source/Runtime/CoreUObject/Public/AssetRegistry/AssetData.h b/Engine/Source/Runtime/CoreUObject/Public/AssetRegistry/AssetData.h
--- a/Engine/Source/Runtime/CoreUObject/Public/AssetRegistry/AssetData.h
+++ b/Engine/Source/Runtime/CoreUObject/Public/AssetRegistry/AssetData.h
@@ -15,6 +15,7 @@
 #include "Containers/StringFwd.h"
 #include "Containers/StringView.h"
 #include "Containers/UnrealString.h"
+#include "Containers/Utf8String.h"
 #include "Containers/VersePathFwd.h"
 #include "HAL/PlatformMath.h"
 #include "HAL/UnrealMemory.h"
@@ -98,6 +99,7 @@ struct FAssetRegistryVersion
 		AddedHeader,						// Added header with bFilterEditorOnlyData flag 
 		AssetPackageDataHasExtension,		// Added Extension to AssetPackageData.
 		AssetPackageDataHasPackageLocation,	// Added PackageLocation to AssetPackageData.
+		MarshalledTextAsUTF8String,			// Replaced 2 byte wide string with UTF8 String
 
 		// -----<new versions can be added above this line>-------------------------------------------------
 		VersionPlusOne,
diff --git a/Engine/Source/Runtime/CoreUObject/Public/AssetRegistry/AssetDataTagMapSerializationDetails.h b/Engine/Source/Runtime/CoreUObject/Public/AssetRegistry/AssetDataTagMapSerializationDetails.h
--- a/Engine/Source/Runtime/CoreUObject/Public/AssetRegistry/AssetDataTagMapSerializationDetails.h
+++ b/Engine/Source/Runtime/CoreUObject/Public/AssetRegistry/AssetDataTagMapSerializationDetails.h
@@ -28,18 +28,19 @@ namespace FixedTagPrivate
 	{
 	public:
 		FMarshalledText() = default;
+		COREUOBJECT_API explicit FMarshalledText(const FUtf8String& InString);
 		COREUOBJECT_API explicit FMarshalledText(const FString& InString);
-		COREUOBJECT_API explicit FMarshalledText(FString&& InString);
+		COREUOBJECT_API explicit FMarshalledText(FUtf8String&& InString);
 		COREUOBJECT_API explicit FMarshalledText(const FText& InText);
 		COREUOBJECT_API explicit FMarshalledText(FText&& InText);
 
-		COREUOBJECT_API const FString& GetAsComplexString() const;
+		COREUOBJECT_API const FUtf8String& GetAsComplexString() const;
 		COREUOBJECT_API FText GetAsText() const;
 		COREUOBJECT_API int32 CompareToCaseIgnored(const FMarshalledText& Other) const;
 		COREUOBJECT_API int64 GetResourceSize() const;
 
 	private:
-		FString String;
+		FUtf8String String;
 	};
 
 	/// Stores a fixed set of values and all the key-values maps used for lookup
@@ -131,6 +132,7 @@ namespace FixedTagPrivate
 	};
 
 	uint32 HashCaseSensitive(const TCHAR* Str, int32 Len);
+	uint32 HashCaseSensitive(const UTF8CHAR* Str, int32 Len);
 	uint32 HashCombineQuick(uint32 A, uint32 B);
 	uint32 HashCombineQuick(uint32 A, uint32 B, uint32 C);
 
@@ -161,10 +163,18 @@ namespace FixedTagPrivate
 			{
 				return A.Equals(B, ESearchCase::CaseSensitive);
 			}
+			static bool Matches(const FUtf8String& A, const FUtf8String& B)
+			{
+				return A.Equals(B, ESearchCase::CaseSensitive);
+			}
 			static uint32 GetKeyHash(const FString& Key)
 			{
 				return HashCaseSensitive(&Key[0], Key.Len());
 			}
+			static uint32 GetKeyHash(const FUtf8String& Key)
+			{
+				return HashCaseSensitive(&Key[0], Key.Len());
+			}
 
 			static bool Matches(FNameEntryId A, FNameEntryId B)
 			{
@@ -263,12 +273,12 @@ namespace FixedTagPrivate
 		/// Won't load FNames to allow concurrent name batch loading
 		/// 
 		/// @return handle to step 3
-		COREUOBJECT_API TFuture<void> ReadInitialDataAndKickLoad(FArchive& Ar, uint32 MaxWorkerTasks);
+		COREUOBJECT_API TFuture<void> ReadInitialDataAndKickLoad(FArchive& Ar, uint32 MaxWorkerTasks, FAssetRegistryVersion::Type HeaderVersion);
 
 		/// 2) Read remaining data, including FNames
 		///
 		/// @return indexed store, usable for FPartialMapHandle::MakeFullHandle()
-		COREUOBJECT_API TRefCountPtr<const FStore> LoadFinalData(FArchive& Ar);
+		COREUOBJECT_API TRefCountPtr<const FStore> LoadFinalData(FArchive& Ar, FAssetRegistryVersion::Type HeaderVersion);
 
 	private:
 		TRefCountPtr<FStore> Store;
