MakeWeakObjectPtr(TObjectPtr<T>) will no longer resolve remote objects

diff --git a/Engine/Source/Runtime/Core/Public/UObject/WeakObjectPtrTemplates.h b/Engine/Source/Runtime/Core/Public/UObject/WeakObjectPtrTemplates.h
--- a/Engine/Source/Runtime/Core/Public/UObject/WeakObjectPtrTemplates.h
+++ b/Engine/Source/Runtime/Core/Public/UObject/WeakObjectPtrTemplates.h
@@ -14,6 +14,7 @@
 struct FFieldPath;
 struct FWeakObjectPtr;
 struct FUObjectItem;
+struct FRemoteObjectId;
 
 /**
  * TWeakObjectPtr is the templated version of the generic FWeakObjectPtr
@@ -77,6 +78,13 @@ public:
 	{
 	}
 
+#if UE_WITH_REMOTE_OBJECT_HANDLE
+	explicit TWeakObjectPtr(const FRemoteObjectId& RemoteId)
+		: WeakPtr(RemoteId)
+	{
+	}
+#endif
+
 	/**
 	 * Reset the weak pointer back to the null state
 	 */
diff --git a/Engine/Source/Runtime/CoreUObject/Public/UObject/ObjectPtr.h b/Engine/Source/Runtime/CoreUObject/Public/UObject/ObjectPtr.h
--- a/Engine/Source/Runtime/CoreUObject/Public/UObject/ObjectPtr.h
+++ b/Engine/Source/Runtime/CoreUObject/Public/UObject/ObjectPtr.h
@@ -1181,7 +1181,16 @@ struct TCallTraits<TObjectPtr<T>> : public TCallTraitsBase<TObjectPtr<T>>
 template <typename T>
 FORCEINLINE TWeakObjectPtr<T> MakeWeakObjectPtr(TObjectPtr<T> Ptr)
 {
-	return TWeakObjectPtr<T>(Ptr);
+#if UE_WITH_REMOTE_OBJECT_HANDLE
+	if (Ptr.IsRemote())
+	{
+		return TWeakObjectPtr<T>(Ptr.GetRemoteId());
+	}
+	else
+#endif
+	{
+		return TWeakObjectPtr<T>(Ptr);
+	}
 }
 
 FORCEINLINE TObjectPtr<UObject>& FObjectPtr::ToTObjectPtr()
