Added a TWeakInterfacePtr constructor that takes a TScriptInterface for convenience

diff --git a/Engine/Source/Runtime/CoreUObject/Public/UObject/WeakInterfacePtr.h b/Engine/Source/Runtime/CoreUObject/Public/UObject/WeakInterfacePtr.h
--- a/Engine/Source/Runtime/CoreUObject/Public/UObject/WeakInterfacePtr.h
+++ b/Engine/Source/Runtime/CoreUObject/Public/UObject/WeakInterfacePtr.h
@@ -56,6 +56,17 @@ struct TWeakInterfacePtr
 		}
 	}
 
+	/**
+	 * Construct from a TScriptInterface of the same interface type
+	 * @param ScriptInterface 	The TScriptInterface to copy from.
+	 * 							No validation is done here; passing an invalid TScriptInterface in will result in an invalid TWeakInterfacePtr.
+	 */
+	TWeakInterfacePtr(const TScriptInterface<T>& ScriptInterface)
+	{
+		ObjectInstance = ScriptInterface.GetObject();
+		InterfaceInstance = ScriptInterface.GetInterface();
+	}
+
 	/**
 	 * Reset the weak pointer back to the null state.
 	 */
