Registering shared object only when its reference gets transferred via RPC

diff --git a/Engine/Source/Runtime/CoreUObject/Private/UObject/RemoteObjectTransfer.cpp b/Engine/Source/Runtime/CoreUObject/Private/UObject/RemoteObjectTransfer.cpp
--- a/Engine/Source/Runtime/CoreUObject/Private/UObject/RemoteObjectTransfer.cpp
+++ b/Engine/Source/Runtime/CoreUObject/Private/UObject/RemoteObjectTransfer.cpp
@@ -815,7 +815,6 @@ FRemoteObjectReference::FRemoteObjectReference(const FObjectPtr& Ptr)
 		if (Ptr.IsResolved())
 		{
 			ServerId = UE::RemoteObject::GetGlobalServerId();
-			UE::RemoteObject::Private::RegisterSharedObject(Ptr.Get());
 		}
 		else if (UE::RemoteObject::Handle::FRemoteObjectStub* Stub = UE::RemoteObject::Private::FindRemoteObjectStub(ObjectId))
 		{
@@ -861,6 +860,7 @@ bool FRemoteObjectReference::NetSerialize(FArchive& Ar, class UPackageMap* Map,
 	{		
 		if (UObject* ExistingObject = StaticFindObjectFastInternal(ObjectId))
 		{
+			UE::RemoteObject::Private::RegisterSharedObject(ExistingObject);
 			PathName = ExistingObject->GetPathName();
 		}
 	}
