[ZenLoader]
- Fix performance regression for the loader at runtime that might check for file existence on disk even for packaged builds

[RN] Fix performance regression for zenloader at runtime that might check for file existence on disk even for packaged builds

diff --git a/Engine/Source/Runtime/CoreUObject/Private/Serialization/AsyncLoading2.cpp b/Engine/Source/Runtime/CoreUObject/Private/Serialization/AsyncLoading2.cpp
--- a/Engine/Source/Runtime/CoreUObject/Private/Serialization/AsyncLoading2.cpp
+++ b/Engine/Source/Runtime/CoreUObject/Private/Serialization/AsyncLoading2.cpp
@@ -274,7 +274,7 @@ void FindAllRuntimeScriptPackages(FRuntimeScriptPackages& OutPackages)
 #endif
 
 #ifndef ALT2_ENABLE_LINKERLOAD_SUPPORT
-#define ALT2_ENABLE_LINKERLOAD_SUPPORT 1
+#define ALT2_ENABLE_LINKERLOAD_SUPPORT WITH_EDITOR
 #endif
 
 #ifndef ALT2_ENABLE_NEW_ARCHIVE_FOR_LINKERLOAD
