Retagging MakeUniqueObjectName to fall under FName LLM category.
#tests Tested with memory traces on cooked build and filtered by callstacks to ensure all fall under correct category.

diff --git a/Engine/Source/Runtime/CoreUObject/Private/UObject/UObjectGlobals.cpp b/Engine/Source/Runtime/CoreUObject/Private/UObject/UObjectGlobals.cpp
--- a/Engine/Source/Runtime/CoreUObject/Private/UObject/UObjectGlobals.cpp
+++ b/Engine/Source/Runtime/CoreUObject/Private/UObject/UObjectGlobals.cpp
@@ -2550,6 +2550,7 @@ namespace NameReuse
 			FEntry(FNameEntryId InBaseName)
 				: BaseName(InBaseName)
 			{
+				LLM_SCOPE(ELLMTag::FName);
 				Names.Reserve(MaxNamesPerEntry);
 			}
 
@@ -2564,6 +2565,8 @@ namespace NameReuse
 
 			void Store(UPTRINT Parent, FName Name)
 			{
+				LLM_SCOPE(ELLMTag::FName);
+
 				int32 Index = Names.IndexOfByPredicate([&](const TTuple<UPTRINT, FName>& Pair) { return Pair.Get<1>() == Name; });
 				if (Index != INDEX_NONE)
 				{
@@ -2609,6 +2612,8 @@ namespace NameReuse
 
 		void Store(UObject* Parent, FNameEntryId BaseId, FName UsedName)
 		{
+			LLM_SCOPE(ELLMTag::FName);
+
 			FEntry* Entry = Entries.FindByPredicate([BaseId](const FEntry& Entry) { return Entry.BaseName == BaseId; });
 			if (!Entry)
 			{
@@ -2660,6 +2665,8 @@ namespace NameReuse
 
 			FName AllocateName(UObject* Parent, FNameEntryId BaseId, FName BaseName)
 			{
+				LLM_SCOPE(ELLMTag::FName);
+
 				checkSlow(BaseId == BaseName.GetComparisonIndex());
 				int32 BaseNumber = BaseName.GetNumber();
 
@@ -2775,6 +2782,7 @@ namespace NameReuse
 
 		FNameRangeEntry* Add(FNameEntryId Id)
 		{
+			LLM_SCOPE(ELLMTag::FName);
 			return Map.Emplace(Id, MakeUnique<FNameRangeEntry>()).Get();
 		}
 
@@ -2795,7 +2803,10 @@ namespace NameReuse
 			return FName();
 		}
 
-		LLM_SCOPE(ELLMTag::UObject);
+		LLM_TAGSET_SCOPE_CLEAR(ELLMTagSet::Assets);
+		LLM_TAGSET_SCOPE_CLEAR(ELLMTagSet::AssetClasses);
+		UE_TRACE_METADATA_CLEAR_SCOPE();
+		LLM_SCOPE(ELLMTag::FName);
 
 		FName ReturnName;
 
