Adding a GC lock when r-constructing an object to simulate Garbage Collector behavior and prevent the object from being accessed by other threads (ALT specifically) while it's being destroyed and reconstructed

diff --git a/Engine/Source/Runtime/CoreUObject/Private/UObject/UObjectGlobals.cpp b/Engine/Source/Runtime/CoreUObject/Private/UObject/UObjectGlobals.cpp
--- a/Engine/Source/Runtime/CoreUObject/Private/UObject/UObjectGlobals.cpp
+++ b/Engine/Source/Runtime/CoreUObject/Private/UObject/UObjectGlobals.cpp
@@ -31,6 +31,7 @@
 #include "UObject/Object.h"
 #include "UObject/ObjectVisibility.h"
 #include "UObject/GarbageCollection.h"
+#include "UObject/GCScopeLock.h"
 #include "UObject/Class.h"
 #include "UObject/CoreRedirects.h"
 #include "UObject/FastReferenceCollector.h"
@@ -3490,6 +3491,56 @@ static thread_local FRestoreForUObjectOverwrite* ObjectRestoreAfterInitProps = n
 extern const FName NAME_UniqueObjectNameForCooking(TEXT("UniqueObjectNameForCooking"));
 COREUOBJECT_API bool GOutputCookingWarnings = false;
 
+// Used by StaticAllocateObject when constructing an object on top of existing object to simulate GC behavior
+class FGCReconstructionGuard
+{
+	enum class EReconstructionLockType
+	{
+		None,
+		GC
+	} LockType = EReconstructionLockType::None;
+
+public:
+
+	FGCReconstructionGuard() = default;
+	~FGCReconstructionGuard()
+	{
+		checkf(LockType == EReconstructionLockType::None, TEXT("FGCReconstructionGuard has not been unlocked"));
+	}
+
+	void Lock()
+	{
+		UE_AUTORTFM_OPEN
+		{			
+			if (!FGCCSyncObject::Get().IsGCLocked()) // GC lock doesn't support recursive locks atm
+			{
+				FGCCSyncObject::Get().GCLock();
+				LockType = EReconstructionLockType::GC;
+			}
+			else
+			{
+				LockType = EReconstructionLockType::None;
+			}
+		};
+		AutoRTFM::PushOnAbortHandler(this, [this]
+			{
+				this->Unlock();
+			});
+	}
+	void Unlock()
+	{
+		UE_AUTORTFM_OPEN
+		{
+			if (LockType == EReconstructionLockType::GC)
+			{
+				FGCCSyncObject::Get().GCUnlock();
+			}
+			LockType = EReconstructionLockType::None;
+		};
+		AutoRTFM::PopOnAbortHandler(this);
+	}
+};
+
 UE_AUTORTFM_ASSUME_SAFE
 UObject* StaticAllocateObject
 (
@@ -3502,7 +3553,8 @@ UObject* StaticAllocateObject
 	bool* bOutRecycledSubobject,
 	UPackage* ExternalPackage,
 	int32 SerialNumber,
-	FRemoteObjectId RemoteId)
+	FRemoteObjectId RemoteId,
+	FGCReconstructionGuard* GCGuard)
 {
 	LLM_SCOPE(ELLMTag::UObject);
 
@@ -3669,11 +3721,18 @@ UObject* StaticAllocateObject
 			OldRemoteId = GUObjectArray.GetRemoteId(OldIndex);
 #endif
 
+			// It's not safe to access the object from other threads while it's being destroyed and reconstucted so simulate GC behavior and lock
+			if (ensureMsgf(IsInGameThread(), TEXT("GC lock can only be acquired on the game thread. If you hit this ensure then an object is being reconstructed on a worker thread which is not thread-safe")) && GCGuard)
+			{
+				GCGuard->Lock();
+			}
+
 			// Destroy the object.
 			SCOPE_CYCLE_COUNTER(STAT_DestroyObject);
 			// Check that the object hasn't been destroyed yet.
 			if(!Obj->HasAnyFlags(RF_FinishDestroyed))
 			{
+#if !UE_WITH_REMOTE_OBJECT_HANDLE
 				if (FPlatformProperties::RequiresCookedData())
 				{
 					ensureAlwaysMsgf(!Obj->HasAnyFlags(RF_NeedLoad|RF_WasLoaded),
@@ -3683,6 +3742,7 @@ UObject* StaticAllocateObject
 						(int32)Obj->GetFlags(),
 						(int32)Obj->GetInternalFlags());
 				}
+#endif // !UE_WITH_REMOTE_OBJECT_HANDLE
 
 				// Get the name before we start the destroy, as destroy renames it
 				FString OldName = Obj->GetFullName();
@@ -4854,8 +4914,9 @@ UObject* StaticConstructObject_Internal(const FStaticConstructObjectParameters&
 	}
 #endif
 
+	FGCReconstructionGuard GCGuard;
 	bool bRecycledSubobject = false;	
-	Result = StaticAllocateObject(InClass, InOuter, InName, InFlags, Params.InternalSetFlags, bCanRecycleSubobjects, &bRecycledSubobject, Params.ExternalPackage, SerialNumber, RemoteId);
+	Result = StaticAllocateObject(InClass, InOuter, InName, InFlags, Params.InternalSetFlags, bCanRecycleSubobjects, &bRecycledSubobject, Params.ExternalPackage, SerialNumber, RemoteId, &GCGuard);
 	check(Result != nullptr);
 	// Don't call the constructor on recycled subobjects, they haven't been destroyed.
 	if (!bRecycledSubobject)
@@ -4863,7 +4924,9 @@ UObject* StaticConstructObject_Internal(const FStaticConstructObjectParameters&
 		STAT(FScopeCycleCounterUObject ConstructorScope(InClass->GetFName().IsNone() ? nullptr : InClass, GET_STATID(STAT_ConstructObject)));
 		(*InClass->ClassConstructor)(FObjectInitializer(Result, Params, RemoteSubObjectOverrides));
 	}
-	
+	// StaticAllocateObject might have locked GCGuard but it can only be unlocked after the object has been fully constructed so unlock it here
+	GCGuard.Unlock();
+
 	if (GIsEditor && 
 		// Do not consider object creation in transaction if the object is marked as async or in being async loaded 
 		!Result->HasAnyInternalFlags(EInternalObjectFlags::Async | EInternalObjectFlags_AsyncLoading) &&
diff --git a/Engine/Source/Runtime/CoreUObject/Public/UObject/UObjectArray.h b/Engine/Source/Runtime/CoreUObject/Public/UObject/UObjectArray.h
--- a/Engine/Source/Runtime/CoreUObject/Public/UObject/UObjectArray.h
+++ b/Engine/Source/Runtime/CoreUObject/Public/UObject/UObjectArray.h
@@ -803,7 +803,7 @@ public:
 class FUObjectArray
 {
 	friend class UObject;
-	friend COREUOBJECT_API UObject* StaticAllocateObject(const UClass*, UObject*, FName, EObjectFlags, EInternalObjectFlags, bool, bool*, UPackage*, int32, FRemoteObjectId);
+	friend COREUOBJECT_API UObject* StaticAllocateObject(const UClass*, UObject*, FName, EObjectFlags, EInternalObjectFlags, bool, bool*, UPackage*, int32, FRemoteObjectId, class FGCReconstructionGuard*);
 
 private:
 	/**
diff --git a/Engine/Source/Runtime/CoreUObject/Public/UObject/UObjectGlobals.h b/Engine/Source/Runtime/CoreUObject/Public/UObject/UObjectGlobals.h
--- a/Engine/Source/Runtime/CoreUObject/Public/UObject/UObjectGlobals.h
+++ b/Engine/Source/Runtime/CoreUObject/Public/UObject/UObjectGlobals.h
@@ -1185,9 +1185,10 @@ bool StaticAllocateObjectErrorTests( const UClass* Class, UObject* InOuter, FNam
  * @param bOutReusedSubobject	flag indicating if the object is a subobject that has already been created (in which case further initialization is not necessary).
  * @param ExternalPackage	External Package assigned to the allocated object, if any	
  * @param RemoteId		Globally unique id for this object
+ * @param GCGuard		Special synchronization object that prevents GC from running when reconstructing an object on top of existing one
  * @return	a pointer to a fully initialized object of the specified class.
  */
-COREUOBJECT_API UObject* StaticAllocateObject(const UClass* Class, UObject* InOuter, FName Name, EObjectFlags SetFlags, EInternalObjectFlags InternalSetFlags = EInternalObjectFlags::None, bool bCanReuseSubobjects = false, bool* bOutReusedSubobject = nullptr, UPackage* ExternalPackage = nullptr, int32 SerialNumber = 0, FRemoteObjectId RemoteId = FRemoteObjectId());
+COREUOBJECT_API UObject* StaticAllocateObject(const UClass* Class, UObject* InOuter, FName Name, EObjectFlags SetFlags, EInternalObjectFlags InternalSetFlags = EInternalObjectFlags::None, bool bCanReuseSubobjects = false, bool* bOutReusedSubobject = nullptr, UPackage* ExternalPackage = nullptr, int32 SerialNumber = 0, FRemoteObjectId RemoteId = FRemoteObjectId(), class FGCReconstructionGuard* GCGuard = nullptr);
 
 /** FObjectInitializer options */
 enum class EObjectInitializerOptions
