Fixing the incorrect usage of GRecentNameCache. The cache was never used if the ReturnName was returning None.

diff --git a/Engine/Source/Runtime/CoreUObject/Private/UObject/UObjectGlobals.cpp b/Engine/Source/Runtime/CoreUObject/Private/UObject/UObjectGlobals.cpp
--- a/Engine/Source/Runtime/CoreUObject/Private/UObject/UObjectGlobals.cpp
+++ b/Engine/Source/Runtime/CoreUObject/Private/UObject/UObjectGlobals.cpp
@@ -2821,7 +2821,7 @@ namespace NameReuse
 				ReturnName = GNameRangeCache.Find(Parent, BaseId, BaseName);
 			}
 
-			if (ReturnName.IsNone())
+			if (!ReturnName.IsNone())
 			{
 				// Store this name for reuse 
 				GRecentNameCache.Store(Parent, BaseId, ReturnName);
