Fix IDO issues found when they are enabled for every object

- Initialize script serialization offsets to -1 because some types bypass UObject::Serialize and do not save offsets.
- Skip serialization of a new IDO if it has no script serialization offsets.
- Skip creation of an IDO struct type for structs that use binary serialization.
- Turn on impersonation when looking up the original type name to fix support for deleted types. This broke when a previous fix disabled impersonation in scopes that it was not meant to be on in.
- Add a CVar IDO.EnableOnEveryObject to allow broad testing of IDOs.

diff --git a/Engine/Source/Runtime/CoreUObject/Private/UObject/InstanceDataObjectUtils.cpp b/Engine/Source/Runtime/CoreUObject/Private/UObject/InstanceDataObjectUtils.cpp
--- a/Engine/Source/Runtime/CoreUObject/Private/UObject/InstanceDataObjectUtils.cpp
+++ b/Engine/Source/Runtime/CoreUObject/Private/UObject/InstanceDataObjectUtils.cpp
@@ -253,7 +253,14 @@ namespace UE
 	FAutoConsoleVariableRef EnableIDOSupportCVar(
 		TEXT("IDO.Enable"),
 		bEnableIDOSupport,
-		TEXT("Allows property bags and IDOs to be created for supported classes.")
+		TEXT("Allows an IDO to be created for an object if its class has support.")
+	);
+
+	bool bEnableIDOSupportOnEveryObject = false;
+	FAutoConsoleVariableRef EnableIDOSupportOnEveryObjectCVar(
+		TEXT("IDO.EnableOnEveryObject"),
+		bEnableIDOSupportOnEveryObject,
+		TEXT("Allows an IDO to be created for every object.")
 	);
 
 	FString ExcludedLoosePropertyTypesVar = TEXT("VerseFunctionProperty");
@@ -294,6 +301,11 @@ namespace UE
 			return false;
 		}
 
+		if (bEnableIDOSupportOnEveryObject)
+		{
+			return true;
+		}
+
 		if (IsInstanceDataObject(InObject))
 		{
 			return true;
@@ -488,11 +500,13 @@ namespace UE
 
 		if (FStructProperty* AsStructProperty = CastField<FStructProperty>(Property))
 		{
-			if (!AsStructProperty->Struct->UseNativeSerialization())
+			// Structs that use native or binary serialization cannot safely generate an IDO.
+			if (!AsStructProperty->Struct->UseNativeSerialization() && !(AsStructProperty->Struct->StructFlags & STRUCT_Immutable))
 			{
 				//@note: Transfer existing metadata over as we build the InstanceDataObject from the struct or it owners, if any, this is useful for testing purposes
 				FString OriginalName;
-				if (const FString* OriginalType = FindOriginalTypeName(AsStructProperty))
+				if (TGuardValue<bool> ImpersonateScope(FUObjectThreadContext::Get().GetSerializeContext()->bImpersonateProperties, true);
+					const FString* OriginalType = FindOriginalTypeName(AsStructProperty))
 				{
 					OriginalName = *OriginalType;
 				}
diff --git a/Engine/Source/Runtime/CoreUObject/Private/UObject/PropertyBagRepository.cpp b/Engine/Source/Runtime/CoreUObject/Private/UObject/PropertyBagRepository.cpp
--- a/Engine/Source/Runtime/CoreUObject/Private/UObject/PropertyBagRepository.cpp
+++ b/Engine/Source/Runtime/CoreUObject/Private/UObject/PropertyBagRepository.cpp
@@ -811,17 +811,20 @@ UObject* FPropertyBagRepository::CreateInstanceDataObject(UObject* Owner, FArchi
 	UObject* InstanceDataObject = StaticConstructObject_Internal(Params);
 
 	// Load the IDO.
-	const int64 OffsetToRestore = Archive.Tell();
-	FScopedObjectSerializeContext ObjectSerializeContext(InstanceDataObject, Archive);
-	Archive.Seek(StartOffset);
+	if (StartOffset != EndOffset)
 	{
-		FGuardValue_Bitfield(Archive.ArMergeOverrides, true);
-		InstanceDataObject->SerializeScriptProperties(Archive);
+		const int64 OffsetToRestore = Archive.Tell();
+		FScopedObjectSerializeContext ObjectSerializeContext(InstanceDataObject, Archive);
+		Archive.Seek(StartOffset);
+		{
+			FGuardValue_Bitfield(Archive.ArMergeOverrides, true);
+			InstanceDataObject->SerializeScriptProperties(Archive);
+		}
+		ensureMsgf(Archive.Tell() == EndOffset,
+			TEXT("Serializing %s into its IDO consumed %" INT64_FMT " bytes when %" INT64_FMT " bytes were expected."),
+			*Owner->GetPathName(), Archive.Tell() - StartOffset, EndOffset - StartOffset);
+		Archive.Seek(OffsetToRestore);
 	}
-	ensureMsgf(Archive.Tell() == EndOffset,
-		TEXT("Serializing %s into its IDO consumed %" INT64_FMT " bytes when %" INT64_FMT " bytes were expected."),
-		*Owner->GetPathName(), Archive.Tell() - StartOffset, EndOffset - StartOffset);
-	Archive.Seek(OffsetToRestore);
 
 	InstanceDataObjectToOwner.Add(InstanceDataObject, Owner);
 	BagData.InstanceDataObject = InstanceDataObject;
diff --git a/Engine/Source/Runtime/CoreUObject/Private/UObject/UObjectThreadContext.cpp b/Engine/Source/Runtime/CoreUObject/Private/UObject/UObjectThreadContext.cpp
--- a/Engine/Source/Runtime/CoreUObject/Private/UObject/UObjectThreadContext.cpp
+++ b/Engine/Source/Runtime/CoreUObject/Private/UObject/UObjectThreadContext.cpp
@@ -147,6 +147,9 @@ FScopedObjectSerializeContext::FScopedObjectSerializeContext(UObject* InObject,
 	bSavedTrackUnknownEnumNames = SerializeContext->bTrackUnknownEnumNames;
 	bSavedImpersonateProperties = SerializeContext->bImpersonateProperties;
 
+	SerializeContext->SerializedObjectScriptStartOffset = -1;
+	SerializeContext->SerializedObjectScriptEndOffset = -1;
+
 	SerializeContext->SerializedObject = Object;
 
 	const bool bIsLoading = Archive.IsLoading();
