Rename OwnerServerId to ResidentServerId to be more indicative of what it means.

diff --git a/Engine/Source/Runtime/CoreUObject/Private/UObject/RemoteObject.cpp b/Engine/Source/Runtime/CoreUObject/Private/UObject/RemoteObject.cpp
--- a/Engine/Source/Runtime/CoreUObject/Private/UObject/RemoteObject.cpp
+++ b/Engine/Source/Runtime/CoreUObject/Private/UObject/RemoteObject.cpp
@@ -46,7 +46,7 @@ public:
 		return nullptr;
 	}
 
-	UE::RemoteObject::Handle::FRemoteObjectStub* FindOrAddRemoteObjectStub(FRemoteObjectId Id, FRemoteServerId OwnerServerId = FRemoteServerId())
+	UE::RemoteObject::Handle::FRemoteObjectStub* FindOrAddRemoteObjectStub(FRemoteObjectId Id, FRemoteServerId ResidentServerId = FRemoteServerId())
 	{
 		using namespace UE::RemoteObject::Handle;
 
@@ -56,7 +56,7 @@ public:
 		{
 			Stub = new FRemoteObjectStub();
 			Stub->Id = Id;
-			Stub->RemoteServerId = OwnerServerId.IsValid() ? OwnerServerId : Id.GetServerId();
+			Stub->ResidentServerId = ResidentServerId.IsValid() ? ResidentServerId : Id.GetServerId();
 		}		
 		return Stub;
 	}
@@ -71,7 +71,7 @@ public:
 		{
 			Stub = new FRemoteObjectStub(Object);
 		}
-		Stub->RemoteServerId = DestinationServerId;
+		Stub->ResidentServerId = DestinationServerId;
 		FUObjectItem* ObjectItem = GUObjectArray.ObjectToObjectItem(Object);
 		checkf(ObjectItem, TEXT("Attempting to get a serial number for an object that does not exist in the global UObject array (it's possible GUObjectArray is not initialized yet, ObjectIndex=%d)"), GUObjectArray.ObjectToIndex(Object));
 		Stub->SerialNumber = ObjectItem->SerialNumber;
@@ -126,9 +126,9 @@ void InitRemoteObjects()
 	FCoreDelegates::OnEndFrame.AddStatic(&UE::RemoteObject::Transfer::WaitForServersToSync);
 }
 
-void RegisterRemoteObjectId(FRemoteObjectId ObjectId, FRemoteServerId OwnerServerId)
+void RegisterRemoteObjectId(FRemoteObjectId ObjectId, FRemoteServerId ResidentServerId)
 {
-	ObjectMaps->FindOrAddRemoteObjectStub(ObjectId, OwnerServerId);
+	ObjectMaps->FindOrAddRemoteObjectStub(ObjectId, ResidentServerId);
 }
 
 void RegisterSharedObject(UObject* Object)
@@ -263,7 +263,7 @@ UObject* ResolveObject(const FRemoteObjectStub* Stub, ERemoteReferenceType RefTy
 
 		UObject* Outer = Object ? Object->GetOuter() : StaticFindObjectFastInternal(Stub->OuterId);
 
-		MigrateObjectFromRemoteServer(Stub->Id, Stub->RemoteServerId, Outer);
+		MigrateObjectFromRemoteServer(Stub->Id, Stub->ResidentServerId, Outer);
 
 		// if running transactionally, we will have aborted and not reached here
 
@@ -294,7 +294,7 @@ bool RequestObject(const FRemoteObjectStub* Stub)
 
 	if (bRemoteObject)
 	{
-		RequestObjectData(Stub->Id, Stub->RemoteServerId, GetGlobalServerId());
+		RequestObjectData(Stub->Id, Stub->ResidentServerId, GetGlobalServerId());
 	}
 
 	return !bRemoteObject;
diff --git a/Engine/Source/Runtime/CoreUObject/Private/UObject/RemoteObjectTransfer.cpp b/Engine/Source/Runtime/CoreUObject/Private/UObject/RemoteObjectTransfer.cpp
--- a/Engine/Source/Runtime/CoreUObject/Private/UObject/RemoteObjectTransfer.cpp
+++ b/Engine/Source/Runtime/CoreUObject/Private/UObject/RemoteObjectTransfer.cpp
@@ -434,18 +434,18 @@ void MigrateObjectToRemoteServerWithExplicitPriority(uint32 RequestToken, FRemot
 			UE_LOG(LogRemoteObject, Verbose, TEXT("    : MigrateObjectToRemoteServer(%u) did not find obj %s (also no stub)"), RequestToken,*ObjectId.ToString());
 		}
 
-		FRemoteServerId CurrentOwnerServerId = RemoteObjectStub ? RemoteObjectStub->RemoteServerId : ObjectId.GetServerId();
-		checkf(CurrentOwnerServerId != GetGlobalServerId(), TEXT("PeerId:%s attempting to forward request for object %s to PeerId:%s but object location is unknown"), 
+		FRemoteServerId CurrentResidentServerId = RemoteObjectStub ? RemoteObjectStub->ResidentServerId : ObjectId.GetServerId();
+		checkf(CurrentResidentServerId != GetGlobalServerId(), TEXT("PeerId:%s attempting to forward request for object %s to PeerId:%s but object location is unknown"), 
 			*GetGlobalServerId().ToString(), *ObjectId.ToString(), *DestinationServerId.ToString());
 
-		if (CurrentOwnerServerId != DestinationServerId)
+		if (CurrentResidentServerId != DestinationServerId)
 		{
-			UE_LOG(LogRemoteObject, Verbose, TEXT("    : MigrateObjectToRemoteServer(%u) forwarding request (we believe %s owns the Object)"), RequestToken, *CurrentOwnerServerId.ToString());
-			GTransferQueue.RequestRemoteObject(RequestToken, RequestPriority, ObjectId, CurrentOwnerServerId, DestinationServerId);
+			UE_LOG(LogRemoteObject, Verbose, TEXT("    : MigrateObjectToRemoteServer(%u) forwarding request (we believe %s owns the Object)"), RequestToken, *CurrentResidentServerId.ToString());
+			GTransferQueue.RequestRemoteObject(RequestToken, RequestPriority, ObjectId, CurrentResidentServerId, DestinationServerId);
 		}
 		else
 		{
-			UE_LOG(LogRemoteObject, Verbose, TEXT("    : MigrateObjectToRemoteServer(%u) denying request %s (we believe %s already has the Object)"), RequestToken, *ObjectId.ToString(), *CurrentOwnerServerId.ToString());
+			UE_LOG(LogRemoteObject, Verbose, TEXT("    : MigrateObjectToRemoteServer(%u) denying request %s (we believe %s already has the Object)"), RequestToken, *ObjectId.ToString(), *CurrentResidentServerId.ToString());
 			RemoteObjectDeniedTransferDelegate.Execute(RequestToken, ObjectId, DestinationServerId);
 		}
 	}
@@ -818,7 +818,7 @@ FRemoteObjectReference::FRemoteObjectReference(const FObjectPtr& Ptr)
 		}
 		else if (UE::RemoteObject::Handle::FRemoteObjectStub* Stub = UE::RemoteObject::Private::FindRemoteObjectStub(ObjectId))
 		{
-			ServerId = Stub->RemoteServerId;
+			ServerId = Stub->ResidentServerId;
 		}
 	}
 }
diff --git a/Engine/Source/Runtime/CoreUObject/Public/UObject/RemoteObject.h b/Engine/Source/Runtime/CoreUObject/Public/UObject/RemoteObject.h
--- a/Engine/Source/Runtime/CoreUObject/Public/UObject/RemoteObject.h
+++ b/Engine/Source/Runtime/CoreUObject/Public/UObject/RemoteObject.h
@@ -38,8 +38,8 @@ namespace UE::RemoteObject::Handle
 		/** SerialNumber this object had on this server */
 		int32 SerialNumber = 0;
 
-		/** Server id that currently owns the object */
-		FRemoteServerId RemoteServerId;
+		/** Server id that where the object currently resides */
+		FRemoteServerId ResidentServerId;
 
 		FRemoteObjectStub() = default;
 		COREUOBJECT_API explicit FRemoteObjectStub(UObject* Object);
