[AutoRTFM] Add EInternalObjectFlags::AutoRTFMConstructionAborted

Set this flag if an object was freshly constructed within a transaction that is aborted.
Filter out objects that have this flag when calling ForEachObjectOfClass(), and other FUObjectHashTables functions.

This fixes a crash in UCurieGlobals::IsCurieParticipant() where a reverted object was cast, which gave an invalid interface.

Also:
* Update the EInternalObjectFlags enum in UnrealEngineTypes.cs to mirror the enum in ObjectMacros.h.
* Skip the body of UObject::DestroyNonNativeProperties() if AutoRTFMConstructionAborted is flagged - this is called from GC and can also attempt to use data that has been reverted by a transaction abort.

diff --git a/Engine/Source/Programs/Shared/EpicGames.Core/UnrealEngineTypes.cs b/Engine/Source/Programs/Shared/EpicGames.Core/UnrealEngineTypes.cs
--- a/Engine/Source/Programs/Shared/EpicGames.Core/UnrealEngineTypes.cs
+++ b/Engine/Source/Programs/Shared/EpicGames.Core/UnrealEngineTypes.cs
@@ -684,6 +684,21 @@ namespace EpicGames.Core
 		/// </summary>
 		ReachabilityFlag2 = 1 << 2,
 
+		/// <summary>
+		/// Object was constructed in an AutoRTFM transaction that has been aborted
+		/// </summary>
+		AutoRTFMConstructionAborted = 1 << 17,
+
+		/// <summary>
+		/// Object is no longer owned by this process
+		/// </summary>
+		Remote = 1 << 18,
+
+		/// <summary>
+		/// Object referenced by remote process
+		/// </summary>
+		RemoteReference = 1 << 19,
+		
 		/// <summary>
 		/// Object is ready to be imported by another package during loading
 		/// </summary>
@@ -694,6 +709,11 @@ namespace EpicGames.Core
 		/// </summary>
 		Garbage = 1 << 21,
 
+		/// <summary>
+		/// Object is being asynchronously loaded
+		/// </summary>
+		AsyncLoadingPhase1 = 1 << 22,
+
 		/// <summary>
 		/// External reference to object in cluster exists
 		/// </summary>
diff --git a/Engine/Source/Runtime/CoreUObject/Private/Blueprint/BlueprintSupport.cpp b/Engine/Source/Runtime/CoreUObject/Private/Blueprint/BlueprintSupport.cpp
--- a/Engine/Source/Runtime/CoreUObject/Private/Blueprint/BlueprintSupport.cpp
+++ b/Engine/Source/Runtime/CoreUObject/Private/Blueprint/BlueprintSupport.cpp
@@ -2639,6 +2639,13 @@ bool UObject::IsInBlueprint() const
 
 void UObject::DestroyNonNativeProperties()
 {
+	if (HasAnyInternalFlags(EInternalObjectFlags::AutoRTFMConstructionAborted))
+	{
+		// Object was constructed in a transaction that was aborted.
+		// Return as the object is unlikely to be in a valid state.
+		return;
+	}
+
 	// Destroy properties that won't be destroyed by the native destructor
 	GetClass()->DestroyPersistentUberGraphFrame(this);
 
diff --git a/Engine/Source/Runtime/CoreUObject/Private/UObject/UObjectGlobals.cpp b/Engine/Source/Runtime/CoreUObject/Private/UObject/UObjectGlobals.cpp
--- a/Engine/Source/Runtime/CoreUObject/Private/UObject/UObjectGlobals.cpp
+++ b/Engine/Source/Runtime/CoreUObject/Private/UObject/UObjectGlobals.cpp
@@ -3834,6 +3834,17 @@ UObject* StaticAllocateObject
 #endif
 				RemoteId);
 		};
+		if (!bWasConstructedOnOldObject)
+		{
+			// If the object was freshly constructed inside a transaction, then
+			// mark the AutoRTFMConstructionAborted bit if the transaction is
+			// aborted. This is used to prevent aborted objects from being
+			// returned from global lists, e.g. ForEachObjectOfClass().
+			UE_AUTORTFM_ONABORT(Obj)
+			{
+				Obj->SetInternalFlags(EInternalObjectFlags::AutoRTFMConstructionAborted);
+			};
+		}
 	}
 	else
 	{
diff --git a/Engine/Source/Runtime/CoreUObject/Private/UObject/UObjectHash.cpp b/Engine/Source/Runtime/CoreUObject/Private/UObject/UObjectHash.cpp
--- a/Engine/Source/Runtime/CoreUObject/Private/UObject/UObjectHash.cpp
+++ b/Engine/Source/Runtime/CoreUObject/Private/UObject/UObjectHash.cpp
@@ -45,6 +45,10 @@ LLM_DEFINE_TAG(UObjectHash, TEXT("UObject hashtables"));
 // Global UObject array instance
 FUObjectArray GUObjectArray;
 
+// Common EInternalObjectFlags used in this file to filter out objects.
+static constexpr EInternalObjectFlags DefaultInternalExclusionFlags =
+	EInternalObjectFlags::Unreachable | EInternalObjectFlags::AutoRTFMConstructionAborted;
+
 static FORCENOINLINE void OnHashFailure(UObjectBaseUtility* Object, const TCHAR* HashName, const TCHAR* FailureKind)
 {
 	UE_LOG(LogUObjectHash, Error, TEXT("UObject %s consistency failure (%s). Checking for memory corruption"), HashName, FailureKind);
@@ -586,7 +590,7 @@ static int32 GetObjectOuterHash(FName ObjName, PTRINT Outer)
 
 UObject* StaticFindObjectFastExplicitThreadSafe(FUObjectHashTables& ThreadHash, const UClass* ObjectClass, FName ObjectName, const FString& ObjectPathName, bool bExactClass, EObjectFlags ExcludeFlags/*=0*/)
 {
-	const EInternalObjectFlags ExclusiveInternalFlags = EInternalObjectFlags::Unreachable;
+	const EInternalObjectFlags ExclusiveInternalFlags = DefaultInternalExclusionFlags;
 
 	// Find an object with the specified name and (optional) class, in any package; if bAnyPackage is false, only matches top-level packages
 	int32 Hash = GetObjectHash(ObjectName);
@@ -716,7 +720,7 @@ struct FObjectSearchPath
 
 UObject* StaticFindObjectInPackageInternal(FUObjectHashTables& ThreadHash, const UClass* ObjectClass, const UPackage* ObjectPackage, FName ObjectName, bool bExactClass, EObjectFlags ExcludeFlags, EInternalObjectFlags ExclusiveInternalFlags)
 {
-	ExclusiveInternalFlags |= EInternalObjectFlags::Unreachable;
+	ExclusiveInternalFlags |= DefaultInternalExclusionFlags;
 	UObject* Result = nullptr;
 	if (FHashBucket* Inners = ThreadHash.PackageToObjectListMap.Find(ObjectPackage))
 	{
@@ -750,7 +754,7 @@ UObject* StaticFindObjectInPackageInternal(FUObjectHashTables& ThreadHash, const
 
 UObject* StaticFindObjectFastInternalThreadSafe(FUObjectHashTables& ThreadHash, const UClass* ObjectClass, const UObject* ObjectPackage, FName ObjectName, bool bExactClass, bool bAnyPackage, EObjectFlags ExcludeFlags, EInternalObjectFlags ExclusiveInternalFlags)
 {
-	ExclusiveInternalFlags |= EInternalObjectFlags::Unreachable;
+	ExclusiveInternalFlags |= DefaultInternalExclusionFlags;
 
 	// If they specified an outer use that during the hashing
 	UObject* Result = nullptr;
@@ -912,7 +916,7 @@ bool StaticFindAllObjectsFastInternal(TArray<UObject*>& OutFoundObjects, const U
 {
 	INC_DWORD_STAT(STAT_FindObjectFast);
 
-	ExclusiveInternalFlags |= EInternalObjectFlags::Unreachable;
+	ExclusiveInternalFlags |= DefaultInternalExclusionFlags;
 
 	FObjectSearchPath SearchPath(ObjectName);
 	const int32 Hash = GetObjectHash(SearchPath.Inner);
@@ -956,7 +960,7 @@ UObject* StaticFindFirstObjectFastInternal(const UClass* ObjectClass, FName Obje
 {
 	INC_DWORD_STAT(STAT_FindObjectFast);
 
-	ExclusiveInternalFlags |= EInternalObjectFlags::Unreachable;
+	ExclusiveInternalFlags |= DefaultInternalExclusionFlags;
 
 	UObject* Result = nullptr;
 	FObjectSearchPath SearchPath(ObjectName);
@@ -1000,7 +1004,7 @@ UObject* StaticFindFirstObjectFastInternal(const UClass* ObjectClass, FName Obje
 UObject* StaticFindObjectFastInternalThreadSafe(FUObjectHashTables& ThreadHash, FRemoteObjectId RemoteId, EObjectFlags InExclusiveFlags, EInternalObjectFlags InExlusiveInternalFlags)
 {
 #if UE_WITH_REMOTE_OBJECT_HANDLE
-	const EInternalObjectFlags ExclusiveInternalFlags = EInternalObjectFlags::Unreachable | InExlusiveInternalFlags;
+	const EInternalObjectFlags ExclusiveInternalFlags = DefaultInternalExclusionFlags | InExlusiveInternalFlags;
 	uint32 Hash = GetTypeHash(RemoteId);
 
 	FHashTableLock HashLock(ThreadHash);
@@ -1283,7 +1287,7 @@ void GetObjectsWithOuter(const class UObjectBase* Outer, TArray<UObject *>& Resu
 #endif
 
 	// We don't want to return any objects that are currently being background loaded unless we're using the object iterator during async loading.
-	ExclusionInternalFlags |= EInternalObjectFlags::Unreachable | UE::GetAsyncLoadingInternalFlagsExclusion();
+	ExclusionInternalFlags |= DefaultInternalExclusionFlags | UE::GetAsyncLoadingInternalFlagsExclusion();
 
 	int32 StartNum = Results.Num();
 	FUObjectHashTables& ThreadHash = FUObjectHashTables::Get();
@@ -1347,7 +1351,7 @@ void ForEachObjectWithOuterBreakable(const class UObjectBase* Outer, TFunctionRe
 #endif
 
 	// We don't want to return any objects that are currently being background loaded unless we're using the object iterator during async loading.
-	ExclusionInternalFlags |= EInternalObjectFlags::Unreachable | UE::GetAsyncLoadingInternalFlagsExclusion();
+	ExclusionInternalFlags |= DefaultInternalExclusionFlags | UE::GetAsyncLoadingInternalFlagsExclusion();
 
 	FUObjectHashTables& ThreadHash = FUObjectHashTables::Get();
 	FHashTableLock HashLock(ThreadHash);
@@ -1392,7 +1396,7 @@ UObjectBase* FindObjectWithOuter(const class UObjectBase* Outer, const class UCl
 	UObject* Result = nullptr;
 	check( Outer );
 	// We don't want to return any objects that are currently being background loaded unless we're using the object iterator during async loading.
-	EInternalObjectFlags ExclusionInternalFlags = EInternalObjectFlags::Unreachable | UE::GetAsyncLoadingInternalFlagsExclusion();
+	EInternalObjectFlags ExclusionInternalFlags = DefaultInternalExclusionFlags | UE::GetAsyncLoadingInternalFlagsExclusion();
 
 	if( NameToLookFor != NAME_None )
 	{
@@ -1442,7 +1446,7 @@ void ForEachObjectWithPackage(const class UPackage* Package, TFunctionRef<bool(U
 	check(Package != nullptr);
 
 	// We don't want to return any objects that are currently being background loaded unless we're using the object iterator during async loading.
-	ExclusionInternalFlags |= EInternalObjectFlags::Unreachable | UE::GetAsyncLoadingInternalFlagsExclusion();
+	ExclusionInternalFlags |= DefaultInternalExclusionFlags | UE::GetAsyncLoadingInternalFlagsExclusion();
 
 	FUObjectHashTables& ThreadHash = FUObjectHashTables::Get();
 	FHashTableLock HashLock(ThreadHash);
@@ -1549,7 +1553,7 @@ FORCEINLINE void ForEachObjectOfClasses_Implementation(FUObjectHashTables& Threa
 	TRACE_CPUPROFILER_EVENT_SCOPE(ForEachObjectOfClasses_Implementation);
 
 	// We don't want to return any objects that are currently being background loaded unless we're using the object iterator during async loading.
-	ExclusionInternalFlags |= EInternalObjectFlags::Unreachable | UE::GetAsyncLoadingInternalFlagsExclusion();
+	ExclusionInternalFlags |= DefaultInternalExclusionFlags | UE::GetAsyncLoadingInternalFlagsExclusion();
 
 	TBucketMapLock ClassToObjectListMapLock(ThreadHash.ClassToObjectListMap);
 
diff --git a/Engine/Source/Runtime/CoreUObject/Public/UObject/ObjectMacros.h b/Engine/Source/Runtime/CoreUObject/Public/UObject/ObjectMacros.h
--- a/Engine/Source/Runtime/CoreUObject/Public/UObject/ObjectMacros.h
+++ b/Engine/Source/Runtime/CoreUObject/Public/UObject/ObjectMacros.h
@@ -623,6 +623,7 @@ enum class EInternalObjectFlags : int32
 	ReachabilityFlag1 = 1 << 1, ///< One of the flags used by Garbage Collector to determine UObject's reachability state
 	ReachabilityFlag2 = 1 << 2, ///< One of the flags used by Garbage Collector to determine UObject's reachability state
 
+	AutoRTFMConstructionAborted = 1 << 17, //< Object was constructed in an AutoRTFM transaction that has been aborted
 	Remote = 1 << 18, //< Object is no longer owned by this process
 	RemoteReference = 1 << 19, //< Object referenced by remote process
 	LoaderImport = 1 << 20, ///< Object is ready to be imported by another package during loading
@@ -647,7 +648,7 @@ ENUM_CLASS_FLAGS(EInternalObjectFlags);
 
 //~ Make sure these macros are up to date!
 #define EInternalObjectFlags_GarbageCollectionKeepFlags (EInternalObjectFlags::Native | EInternalObjectFlags::Async | EInternalObjectFlags::AsyncLoadingPhase1 | EInternalObjectFlags::AsyncLoadingPhase2 | EInternalObjectFlags::LoaderImport | EInternalObjectFlags::RefCounted)
-#define EInternalObjectFlags_AllFlags (EInternalObjectFlags::ReachabilityFlag0 | EInternalObjectFlags::ReachabilityFlag1 | EInternalObjectFlags::ReachabilityFlag2 | EInternalObjectFlags::Remote | EInternalObjectFlags::RemoteReference | EInternalObjectFlags::LoaderImport | EInternalObjectFlags::Garbage | EInternalObjectFlags::ReachableInCluster | EInternalObjectFlags::ClusterRoot | EInternalObjectFlags::Native | EInternalObjectFlags::RefCounted | EInternalObjectFlags::Async | EInternalObjectFlags::AsyncLoadingPhase1 | EInternalObjectFlags::AsyncLoadingPhase2 | EInternalObjectFlags::RootSet | EInternalObjectFlags::PendingConstruction | (EInternalObjectFlags::Unreachable))
+#define EInternalObjectFlags_AllFlags (EInternalObjectFlags::ReachabilityFlag0 | EInternalObjectFlags::ReachabilityFlag1 | EInternalObjectFlags::ReachabilityFlag2 | EInternalObjectFlags::AutoRTFMConstructionAborted | EInternalObjectFlags::Remote | EInternalObjectFlags::RemoteReference | EInternalObjectFlags::LoaderImport | EInternalObjectFlags::Garbage | EInternalObjectFlags::ReachableInCluster | EInternalObjectFlags::ClusterRoot | EInternalObjectFlags::Native | EInternalObjectFlags::RefCounted | EInternalObjectFlags::Async | EInternalObjectFlags::AsyncLoadingPhase1 | EInternalObjectFlags::AsyncLoadingPhase2 | EInternalObjectFlags::RootSet | EInternalObjectFlags::PendingConstruction | (EInternalObjectFlags::Unreachable))
 #define EInternalObjectFlags_RootFlags (EInternalObjectFlags::RootSet | EInternalObjectFlags_GarbageCollectionKeepFlags)
 #define EInternalObjectFlags_ReachabilityFlags (EInternalObjectFlags::ReachabilityFlag0 | EInternalObjectFlags::ReachabilityFlag1 | EInternalObjectFlags::ReachabilityFlag2 | EInternalObjectFlags::Unreachable)
 #define EInternalObjectFlags_AsyncLoading (EInternalObjectFlags::AsyncLoadingPhase1 | EInternalObjectFlags::AsyncLoadingPhase2)
diff --git a/Engine/Source/Runtime/CoreUObject/Public/UObject/UObjectIterator.h b/Engine/Source/Runtime/CoreUObject/Public/UObject/UObjectIterator.h
--- a/Engine/Source/Runtime/CoreUObject/Public/UObject/UObjectIterator.h
+++ b/Engine/Source/Runtime/CoreUObject/Public/UObject/UObjectIterator.h
@@ -58,7 +58,7 @@ enum class EObjectIteratorThreadSafetyOptions : uint8
 
 inline EInternalObjectFlags GetObjectIteratorDefaultInternalExclusionFlags(EInternalObjectFlags InternalExclusionFlags)
 {
-	return InternalExclusionFlags | EInternalObjectFlags::Unreachable | EInternalObjectFlags::PendingConstruction | EInternalObjectFlags::Remote | UE::GetAsyncLoadingInternalFlagsExclusion();
+	return InternalExclusionFlags | EInternalObjectFlags::Unreachable | EInternalObjectFlags::AutoRTFMConstructionAborted | EInternalObjectFlags::PendingConstruction | EInternalObjectFlags::Remote | UE::GetAsyncLoadingInternalFlagsExclusion();
 }
 
 /**
diff --git a/Engine/Source/Runtime/Engine/Private/Tests/AutoRTFM/AutoRTFMActorComponentTests.cpp b/Engine/Source/Runtime/Engine/Private/Tests/AutoRTFM/AutoRTFMActorComponentTests.cpp
--- a/Engine/Source/Runtime/Engine/Private/Tests/AutoRTFM/AutoRTFMActorComponentTests.cpp
+++ b/Engine/Source/Runtime/Engine/Private/Tests/AutoRTFM/AutoRTFMActorComponentTests.cpp
@@ -143,6 +143,17 @@ private:
 	const int32 TestLine;
 };
 
+// Calls ForEachObjectOfClass() to count the number of AAutoRTFMTestActor instances currently alive.
+size_t CountAAutoRTFMTestActors()
+{
+	size_t Count = 0;
+	ForEachObjectOfClass(AAutoRTFMTestActor::StaticClass(), [&](UObject* const Obj)
+	{
+		Count++;
+	});
+	return Count;
+};
+
 // General tests for calling RegisterComponent() and UnregisterComponent() in transactions.
 // See: SOL-6709
 AUTORTFM_ACTOR_COMPONENT_TEST(RegisterComponent_UnregisterComponent)
@@ -536,6 +547,49 @@ AUTORTFM_ACTOR_COMPONENT_TEST(ReconstructCameraShake) // SOL-7529
 	TestTrueExpr(123 == Old->ShakeScale);
 }
 
+AUTORTFM_ACTOR_COMPONENT_TEST(ForEachObjectOfClass_Fresh)
+{
+	const size_t InitialCount = CountAAutoRTFMTestActors();
+
+	TestTrueExpr(CountAAutoRTFMTestActors() == InitialCount);
+
+	AutoRTFM::Testing::Abort([&]
+	{
+		AAutoRTFMTestActor* Object = NewObject<AAutoRTFMTestActor>();
+		TestTrueExpr(CountAAutoRTFMTestActors() == InitialCount + 1);
+		AutoRTFM::AbortTransaction();
+	});
+
+	TestTrueExpr(CountAAutoRTFMTestActors() == InitialCount);
+}
+
+AUTORTFM_ACTOR_COMPONENT_TEST(ForEachObjectOfClass_Reconstructed)
+{
+	const size_t InitialCount = CountAAutoRTFMTestActors();
+	
+	TestTrueExpr(CountAAutoRTFMTestActors() == InitialCount);
+
+	AutoRTFM::Testing::Commit([&]
+	{
+		FName Name = TEXT("MyObject");
+		UObject* Outer = static_cast<UObject*>(GetTransientPackage());
+		AAutoRTFMTestActor* Old = NewObject<AAutoRTFMTestActor>(Outer, Name);
+		TestTrueExpr(CountAAutoRTFMTestActors() == InitialCount + 1u);
+
+		AutoRTFM::Testing::Abort([&]
+		{
+			AAutoRTFMTestActor* New = NewObject<AAutoRTFMTestActor>(Outer, Name);
+			TestTrueExpr(Old == New);
+			TestTrueExpr(CountAAutoRTFMTestActors() == InitialCount + 1u);
+			AutoRTFM::AbortTransaction();
+		});
+
+		TestTrueExpr(CountAAutoRTFMTestActors() == InitialCount + 1u);
+	});
+
+	TestTrueExpr(CountAAutoRTFMTestActors() == InitialCount + 1u);
+}
+
 }  // anonymous namespace
 
 #undef AUTORTFM_ACTOR_COMPONENT_TEST
