Re-enable garbage reference warnings and modify the value of the cvar for the cooker instead.

diff --git a/Engine/Source/Editor/UnrealEd/Private/Commandlets/CookCommandlet.cpp b/Engine/Source/Editor/UnrealEd/Private/Commandlets/CookCommandlet.cpp
--- a/Engine/Source/Editor/UnrealEd/Private/Commandlets/CookCommandlet.cpp
+++ b/Engine/Source/Editor/UnrealEd/Private/Commandlets/CookCommandlet.cpp
@@ -205,6 +205,9 @@ int32 UCookCommandlet::Main(const FString& CmdLineParams)
 	COOK_STAT(FParse::Value(*Params, TEXT("CookCultures="), DetailedCookStats::CookCultures));
 	COOK_STAT(FParse::Value(*Params, TEXT("CookLabel="), DetailedCookStats::CookLabel));
 	
+	// Temporarily disable garbage ref reporting pending a fix to FWorldPartitionLevelHelper::MoveExternalActorsToLevel
+	TGuardConsoleVariable<int32> GarbageRefGuard(TEXT("gc.GarbageReferenceTrackingEnabled"), 0);
+
 	ITargetPlatformManagerModule& TPM = GetTargetPlatformManagerRef();
 	if ( bCookOnTheFly )
 	{
diff --git a/Engine/Source/Runtime/Core/Public/HAL/IConsoleManager.h b/Engine/Source/Runtime/Core/Public/HAL/IConsoleManager.h
--- a/Engine/Source/Runtime/Core/Public/HAL/IConsoleManager.h
+++ b/Engine/Source/Runtime/Core/Public/HAL/IConsoleManager.h
@@ -1257,6 +1257,34 @@ private:
 };
 
 
+/**
+ * Utility class to modify a console variable's value for a certain scope and restore it to its previous value afterwards.
+ * Type parameter may be any supported by IConsoleVariable::GetValue, i.e. bool, int32, float, FString.
+ */
+template <typename T>
+struct TGuardConsoleVariable
+{
+	T PreviousValue;
+	IConsoleVariable* ConsoleVariable;
+
+	TGuardConsoleVariable(const TCHAR* ConsoleVariableName, const T& NewValue)
+		: ConsoleVariable(IConsoleManager::Get().FindConsoleVariable(ConsoleVariableName))
+	{
+		if (ensureMsgf(ConsoleVariable, TEXT("TGuardConsoleVariable: Trying to set console variable \"%s\" but it doesn't exist"), ConsoleVariableName))
+		{
+			ConsoleVariable->GetValue(PreviousValue);
+			ConsoleVariable->Set(NewValue);
+		}
+	}
+	~TGuardConsoleVariable()
+	{
+		if (ConsoleVariable)
+		{
+			ConsoleVariable->Set(PreviousValue);
+		}
+	}
+};
+
 /**
  * auto registering console variable sinks (register a callback function that is called when ever a cvar is changes by the user, changes are grouped and happen in specific engine spots during the frame/main loop)
  */
diff --git a/Engine/Source/Runtime/CoreUObject/Private/UObject/GarbageCollection.cpp b/Engine/Source/Runtime/CoreUObject/Private/UObject/GarbageCollection.cpp
--- a/Engine/Source/Runtime/CoreUObject/Private/UObject/GarbageCollection.cpp
+++ b/Engine/Source/Runtime/CoreUObject/Private/UObject/GarbageCollection.cpp
@@ -5335,10 +5335,9 @@ static void DumpGarbageReferencers(TConstArrayView<TUniquePtr<FWorkerContext>> C
 			for (FGarbageReferenceInfo& GarbageReference : Context->GarbageReferences)
 			{
 				// Note: EngineTest tests rely on the format of this message, see TestGarbageElimination
-				// TODO: Restore severity of this message to Warning once we have fixed the spurious occurrence of the warning during cook.
-				UE_LOG(LogGarbage, Display, TEXT("Reachable garbage object: %s"), *GarbageReference.GarbageObject->GetFullName());
-				UE_LOG(LogGarbage, Display, TEXT("Referenced by:            %s"), *GarbageReference.GetReferencingObjectInfo());
-				UE_LOG(LogGarbage, Display, TEXT(""));
+				UE_LOG(LogGarbage, Warning, TEXT("Reachable garbage object: %s"), *GarbageReference.GarbageObject->GetFullName());
+				UE_LOG(LogGarbage, Warning, TEXT("Referenced by:            %s"), *GarbageReference.GetReferencingObjectInfo());
+				UE_LOG(LogGarbage, Warning, TEXT(""));
 				NumGarbageReferences++;
 			}
 			Context->GarbageReferences.Reset();
@@ -5361,10 +5360,9 @@ static void DumpGarbageReferencers(TConstArrayView<TUniquePtr<FWorkerContext>> C
 				if (Seen.Contains(Key) == false)
 				{
 					// Note: EngineTest tests rely on the format of this message, see TestGarbageElimination
-					// TODO: Restore severity of this message to Warning once we have fixed the spurious occurrence of the warning during cook.
-					UE_LOG(LogGarbage, Display, TEXT("Reachable garbage object: %s"), *GarbageReference.GarbageObject->GetFullName());
-					UE_LOG(LogGarbage, Display, TEXT("Referenced by:            %s"), *GarbageReference.GetReferencingObjectInfo());
-					UE_LOG(LogGarbage, Display, TEXT(""));
+					UE_LOG(LogGarbage, Warning, TEXT("Reachable garbage object: %s"), *GarbageReference.GarbageObject->GetFullName());
+					UE_LOG(LogGarbage, Warning, TEXT("Referenced by:            %s"), *GarbageReference.GetReferencingObjectInfo());
+					UE_LOG(LogGarbage, Warning, TEXT(""));
 					Seen.Add(Key);
 					ReportedGarbageReferences++;
 				}
