- Added initial support for "hybrid cooked editor" to load cooked-for-editor assets from IoStore/ZenStore without staging a full cooked build.
1 Run "RunUAT MakeCookedEditor -project=<project> -build -cook -HybridCookedEditor"
- This will generate WindowsEditor, not WindowsCookedEditor
2 Run "UnrealEditor <project> <map> -HybridCookedEditor" (or similar) to activate the hybrid mode.
- For now, opening the default map may crash in a ULevel function, but loading a random map should work better. (If it still fails, run the editor without -hybridce, make a new test map, save it, then open that one with -hybridce)

diff --git a/Engine/Source/Developer/CookedEditor/Public/CookedEditorTargetPlatformControls.h b/Engine/Source/Developer/CookedEditor/Public/CookedEditorTargetPlatformControls.h
--- a/Engine/Source/Developer/CookedEditor/Public/CookedEditorTargetPlatformControls.h
+++ b/Engine/Source/Developer/CookedEditor/Public/CookedEditorTargetPlatformControls.h
@@ -39,7 +39,8 @@ public:
 
 	virtual FString PlatformName() const override
 	{
-		static FString CachedPlatformName = this->GetTargetPlatformSettings()->IniPlatformName() + TEXT("CookedEditor");
+		static bool bIsCookingForHybridCE = FParse::Param(FCommandLine::Get(), TEXT("hybridcookededitor"));
+		static FString CachedPlatformName = this->GetTargetPlatformSettings()->IniPlatformName() + (bIsCookingForHybridCE  ? TEXT("Editor") : TEXT("CookedEditor"));
 		return CachedPlatformName;
 	}
 
diff --git a/Engine/Source/Programs/AutomationTool/CookedEditor/CookedEditor.Automation.cs b/Engine/Source/Programs/AutomationTool/CookedEditor/CookedEditor.Automation.cs
--- a/Engine/Source/Programs/AutomationTool/CookedEditor/CookedEditor.Automation.cs
+++ b/Engine/Source/Programs/AutomationTool/CookedEditor/CookedEditor.Automation.cs
@@ -359,6 +359,7 @@ public class ModifyStageContext
 public class MakeCookedEditor : BuildCommand
 {
 	protected bool bIsCookedCooker;
+	protected bool bIsHybridCE;
 	protected FileReference ProjectFile;
 
 	protected ConfigHelper ConfigHelper;
@@ -376,6 +377,7 @@ public class MakeCookedEditor : BuildCommand
 		Logger.LogInformation("************************* MakeCookedEditor");
 
 		bIsCookedCooker = ParseParam("cookedcooker");
+		bIsHybridCE = ParseParam("hybridcookededitor");
 		ProjectFile = ParseProjectParam();
 
 		// set up config sections and the like
@@ -1161,7 +1163,7 @@ public class MakeCookedEditor : BuildCommand
 		// control the server/client taregts
 		Params.ServerCookedTargets.Clear();
 		Params.ClientCookedTargets.Clear();
-		List<TargetPlatformDescriptor> TargetPlatformList = new List<TargetPlatformDescriptor>() { new TargetPlatformDescriptor(Platform, TargetPlatformType) };
+		List<TargetPlatformDescriptor> TargetPlatformList = new List<TargetPlatformDescriptor>() { new TargetPlatformDescriptor(Platform, bIsHybridCE ? "Editor" : TargetPlatformType) };
 		if (bIsCookedCooker)
 		{
 			Params.EditorTargets.Add(TargetName);
@@ -1191,6 +1193,11 @@ public class MakeCookedEditor : BuildCommand
 		// and choose to cook them or not (generally for editoronly assets)
 		Params.AdditionalCookerOptions += " -DlcReevaluateUncookedAssets";
 
+		if (bIsHybridCE)
+		{
+			Params.AdditionalCookerOptions += " -hybridcookededitor";
+		}
+
 		string AssetRegistryCacheRootFolder = ParseParamValue("AssetRegistryCacheRootFolder", "");
 		if (!string.IsNullOrEmpty(AssetRegistryCacheRootFolder))
 		{
diff --git a/Engine/Source/Programs/AutomationTool/CookedEditor/PlatformWrappers.cs b/Engine/Source/Programs/AutomationTool/CookedEditor/PlatformWrappers.cs
--- a/Engine/Source/Programs/AutomationTool/CookedEditor/PlatformWrappers.cs
+++ b/Engine/Source/Programs/AutomationTool/CookedEditor/PlatformWrappers.cs
@@ -79,4 +79,42 @@ public class MacCookedCooker : MacPlatform
 		return new TargetPlatformDescriptor(TargetPlatformType, "CookedCooker");
 	}
 
-}
\ No newline at end of file
+}
+
+
+// Placeholders while we figure out hybrid cooked editor solution for cooking
+public class WindowsEditor : Win64Platform
+{
+	public override string GetCookPlatform(bool bDedicatedServer, bool bIsClientOnly)
+	{
+		return "WindowsEditor";
+	}
+	public override TargetPlatformDescriptor GetTargetPlatformDescriptor()
+	{
+		return new TargetPlatformDescriptor(TargetPlatformType, "Editor");
+	}
+}
+
+public class LinuxEditor : GenericLinuxPlatform
+{
+	public override string GetCookPlatform(bool bDedicatedServer, bool bIsClientOnly)
+	{
+		return "LinuxEditor";
+	}
+	public override TargetPlatformDescriptor GetTargetPlatformDescriptor()
+	{
+		return new TargetPlatformDescriptor(TargetPlatformType, "Editor");
+	}
+}
+
+public class MacEditor : MacPlatform
+{
+	public override string GetCookPlatform(bool bDedicatedServer, bool bIsClientOnly)
+	{
+		return "MacEditor";
+	}
+	public override TargetPlatformDescriptor GetTargetPlatformDescriptor()
+	{
+		return new TargetPlatformDescriptor(TargetPlatformType, "Editor");
+	}
+}
diff --git a/Engine/Source/Runtime/AssetRegistry/Private/AssetRegistry.cpp b/Engine/Source/Runtime/AssetRegistry/Private/AssetRegistry.cpp
--- a/Engine/Source/Runtime/AssetRegistry/Private/AssetRegistry.cpp
+++ b/Engine/Source/Runtime/AssetRegistry/Private/AssetRegistry.cpp
@@ -6916,7 +6916,9 @@ void FAssetRegistryImpl::DependencyDataGathered(TMultiMap<FName, FPackageDepende
 		{
 			// Update package data
 			FAssetPackageData* PackageData = State.CreateOrGetAssetPackageData(Result.PackageName);
+			FPackageName::EPackageLocationFilter OriginalLocation = PackageData->GetPackageLocation();
 			*PackageData = Result.PackageData;
+			PackageData->SetPackageLocation(FPackageName::EPackageLocationFilter(uint8(PackageData->GetPackageLocation()) | uint8(OriginalLocation)));
 		}
 
 		if (Result.bHasDependencyData)
diff --git a/Engine/Source/Runtime/Core/Private/Misc/CoreGlobals.cpp b/Engine/Source/Runtime/Core/Private/Misc/CoreGlobals.cpp
--- a/Engine/Source/Runtime/Core/Private/Misc/CoreGlobals.cpp
+++ b/Engine/Source/Runtime/Core/Private/Misc/CoreGlobals.cpp
@@ -176,6 +176,7 @@ bool GIsRunningUnattendedScript = false;
 #if WITH_EDITOR
 bool					PRIVATE_GIsRunningCookCommandlet	= false;				/** Whether this executable is running the cook commandlet */
 bool					PRIVATE_GIsRunningDLCCookCommandlet = false;				/** Whether this executable is running the cook commandlet on a DLC plugin */
+bool					PRIVATE_GIsRunningHybridCookedEditor= false;				/** Whether the editor is running in hybrid cooked editor mode */
 namespace UE::Private
 {
 int32					GMultiprocessId = 0;
diff --git a/Engine/Source/Runtime/Core/Public/CoreGlobals.h b/Engine/Source/Runtime/Core/Public/CoreGlobals.h
--- a/Engine/Source/Runtime/Core/Public/CoreGlobals.h
+++ b/Engine/Source/Runtime/Core/Public/CoreGlobals.h
@@ -239,8 +239,21 @@ extern CORE_API bool GFirstFrameIntraFrameDebugging;
 #if WITH_EDITOR
 extern CORE_API bool PRIVATE_GIsRunningCookCommandlet;
 extern CORE_API bool PRIVATE_GIsRunningDLCCookCommandlet;
+extern CORE_API bool PRIVATE_GIsRunningHybridCookedEditor;
 #endif
 
+/**
+ * Check if the editor is running in Hybrid Cooked Editor mode (where it can dynamically choose between cooked assets and uncooked files
+ */
+FORCEINLINE bool IsRunningHybridCookedEditor()
+{
+#if WITH_EDITOR
+	return PRIVATE_GIsRunningHybridCookedEditor;
+#else
+	return false;
+#endif
+}
+
 /**
 * Check to see if this executable is running a commandlet (custom command-line processing code in an editor-like environment)
 */
diff --git a/Engine/Source/Runtime/Core/Public/Misc/Build.h b/Engine/Source/Runtime/Core/Public/Misc/Build.h
--- a/Engine/Source/Runtime/Core/Public/Misc/Build.h
+++ b/Engine/Source/Runtime/Core/Public/Misc/Build.h
@@ -472,10 +472,10 @@
 #ifndef UE_IS_COOKED_EDITOR
 #define UE_IS_COOKED_EDITOR 0
 #endif
-// Controls whether to enable loading cooked packages from I/O store in editor builds. Defaults to UE_IS_COOKED_EDITOR
-// but can be defined separately if needed
+// Controls whether to enable loading cooked packages from I/O store in editor builds. Defaults to WITH_EDITOR as we can now
+// load from IoStore without needing to stage an entire cooked build
 #ifndef WITH_IOSTORE_IN_EDITOR
-#define WITH_IOSTORE_IN_EDITOR UE_IS_COOKED_EDITOR
+#define WITH_IOSTORE_IN_EDITOR WITH_EDITOR
 #endif
 
 // Controls if iostore will be forced on
diff --git a/Engine/Source/Runtime/CoreUObject/Private/Serialization/AsyncLoading2.cpp b/Engine/Source/Runtime/CoreUObject/Private/Serialization/AsyncLoading2.cpp
--- a/Engine/Source/Runtime/CoreUObject/Private/Serialization/AsyncLoading2.cpp
+++ b/Engine/Source/Runtime/CoreUObject/Private/Serialization/AsyncLoading2.cpp
@@ -441,6 +441,17 @@ static FAutoConsoleVariableRef CVarAllowLoadingPackagesFromFileSystem(
 	ECVF_Default
 	);
 
+// In editor builds, for packages that are both local and in zenstore, do we prefer uncooked (local) or cooked (zenstore) packages
+static bool GPreferCookedPackagesInEditor = false;
+static FAutoConsoleVariableRef CVarPreferCookedPackagesInEditor(
+	TEXT("s.PreferCookedPackagesInEditor"),
+	GPreferCookedPackagesInEditor,
+	TEXT("In editor builds, for packages that are both local and in zenstore, do we prefer uncooked (local) or cooked (zenstore) packages"),
+	ECVF_Default
+);
+
+
+
 #if WITH_EDITOR
 // This is important for the editor because the linker can end up recreating missing exports from its CreateImport function.
 // If we don't put the package in the zenloader loading queue, we can be left with objects in a RF_NeedLoad state that will never be actually loaded.
@@ -3988,7 +3999,7 @@ public:
 
 private:
 #if ALT2_ENABLE_LINKERLOAD_SUPPORT || WITH_EDITOR
-	bool TryGetPackagePathFromFileSystem(FName& PackageNameToLoad, FName& UPackageName, FPackagePath& OutPackagePath)
+	bool GetLinkerLoadInformationAndReturnIfRequired(FName& InOutPackageNameToLoad, FName& InOutUPackageName, FPackagePath& InOutPackagePath)
 	{
 		// If disabling loading from the file system in a shipping/test build and we are not the Editor avoid doing this look-up
 		// As this can be expensive to call and can cause contention on the GameThread and ASyncLoadingThread
@@ -3997,6 +4008,14 @@ private:
 			return false;
 		}
 
+		bool bExistsInIoStore = FIoDispatcher::Get().DoesChunkExist(CreatePackageDataChunkId(FPackageId::FromName(InOutPackagePath.GetPackageFName())));
+		// if we prefer cooked packages in the editor, then we don't want to use the linkerload version
+		// NOTE that we won't have fixed up the FNames, but if we aren't using linkerloads, it shouldn't matter
+		if (bExistsInIoStore && GIsEditor && GPreferCookedPackagesInEditor)
+		{
+			return false;
+		}
+
 #if WITH_EDITORONLY_DATA
 		// In editor, set MatchCaseOnDisk=true so that we set the capitalization of the Package's FName to match the
 		// capitalization on disk. Different capitalizations can arise from imports of the package that were somehow
@@ -4008,19 +4027,24 @@ private:
 		constexpr bool bMatchCaseOnDisk = false;
 #endif
 
-		if (!PackageNameToLoad.IsNone() &&
-			FPackageName::DoesPackageExistEx(OutPackagePath, FPackageName::EPackageLocationFilter::FileSystem,
-				bMatchCaseOnDisk, &OutPackagePath) != FPackageName::EPackageLocationFilter::None)
+		// make sure the package is available in filesystem
+		bool bExistsInFileSystem = false;
+		if (!InOutPackageNameToLoad.IsNone() &&
+			FPackageName::DoesPackageExistEx(InOutPackagePath, FPackageName::EPackageLocationFilter::FileSystem,
+				bMatchCaseOnDisk, &InOutPackagePath) != FPackageName::EPackageLocationFilter::None)
 		{
-			FName CaseCorrectedPackageName = OutPackagePath.GetPackageFName();
-			if (PackageNameToLoad == UPackageName)
+			bExistsInFileSystem = true;
+
+			FName CaseCorrectedPackageName = InOutPackagePath.GetPackageFName();
+			if (InOutPackageNameToLoad == InOutUPackageName)
 			{
-				UPackageName = CaseCorrectedPackageName;
+				InOutUPackageName = CaseCorrectedPackageName;
 			}
-			PackageNameToLoad = CaseCorrectedPackageName;
-			return true;
+			InOutPackageNameToLoad = CaseCorrectedPackageName;
 		}
-		return false;
+
+		// if we prefer uncooked in the editor, then load from filesystem if it's there
+		return bExistsInFileSystem;
 	}
 #endif
 
@@ -4410,6 +4434,12 @@ void FAsyncLoadingThread2::InitializeLoading()
 		ParsePackageNames(VerbosePackageNamesString, GAsyncLoading2_VerbosePackageIds);
 		ParsePackageNames(DebugPackageNamesString, GAsyncLoading2_VerbosePackageIds);
 		GAsyncLoading2_VerboseLogFilter = GAsyncLoading2_VerbosePackageIds.Num() > 0 ? 1 : 2;
+
+		// force allowing cooked packages in hybrid cooked editor mode
+		if (IsRunningHybridCookedEditor())
+		{
+			GPreferCookedPackagesInEditor = true;
+		}
 	}
 
 	FileOpenLogWrapper = (FPlatformFileOpenLog*)(FPlatformFileManager::Get().FindPlatformFile(FPlatformFileOpenLog::GetTypeName()));
@@ -4845,7 +4875,7 @@ bool FAsyncLoadingThread2::CreateAsyncPackagesFromQueue(FAsyncLoadingThreadState
 #if ALT2_ENABLE_LINKERLOAD_SUPPORT
 			FPackagePath::TryFromPackageName(PackageNameToLoad, Request.PackagePath);
 
-			// Lookup if we have the package in the importstore already to avoid expensive calls in TryGetPackagePathFromFileSystem
+			// Lookup if we have the package in the importstore already to avoid expensive calls in GetLinkerLoadInformationAndReturnIfRequired
 			if (FLoadedPackageRef* PackageRef = GlobalImportStore.FindPackageRef(PackageIdToLoad))
 			{
 				if (!PackageRef->HasErrors())
@@ -4855,10 +4885,14 @@ bool FAsyncLoadingThread2::CreateAsyncPackagesFromQueue(FAsyncLoadingThreadState
 					PackageStatus = EPackageStoreEntryStatus::Ok;
 				}
 			}
-			else if (TryGetPackagePathFromFileSystem(PackageNameToLoad, UPackageName, Request.PackagePath))
+			else
 			{
-				bIsZenPackage = false;
-				PackageStatus = EPackageStoreEntryStatus::Ok;
+				// update some package names, and if linker loads are required, disable zen
+				if (GetLinkerLoadInformationAndReturnIfRequired(PackageNameToLoad, UPackageName, Request.PackagePath))
+				{
+					bIsZenPackage = false;
+					PackageStatus = EPackageStoreEntryStatus::Ok;
+				}
 			}
 #endif
 
@@ -5500,10 +5534,11 @@ void FAsyncPackage2::ImportPackagesRecursiveInner(FAsyncLoadingThreadState2& Thr
 #if ALT2_ENABLE_LINKERLOAD_SUPPORT || WITH_EDITOR
 		bool bIsZenPackage = true;
 		bool bIsZenPackageImport = true;
+		bool bIsOnlyOnFileSystem = false;
 		FPackagePath::TryFromPackageName(ImportedPackageNameToLoad, ImportedPackagePath);
 		bIsZenPackage = !LinkerLoadState.IsSet();
 
-		// Lookup if we have the package in the importstore already to avoid expensive calls in TryGetPackagePathFromFileSystem
+		// Lookup if we have the package in the importstore already to avoid expensive calls in GetLinkerLoadInformationAndReturnIfRequired
 		if (FLoadedPackageRef* PackageRef = AsyncLoadingThread.GlobalImportStore.FindPackageRef(ImportedPackageId))
 		{
 			if (!PackageRef->HasErrors())
@@ -5513,10 +5548,14 @@ void FAsyncPackage2::ImportPackagesRecursiveInner(FAsyncLoadingThreadState2& Thr
 				ImportedPackageStatus = EPackageStoreEntryStatus::Ok;
 			}
 		}
-		else if (AsyncLoadingThread.TryGetPackagePathFromFileSystem(ImportedPackageNameToLoad, ImportedPackageUPackageName, ImportedPackagePath))
+		else
 		{
-			bIsZenPackageImport = false;
-			ImportedPackageStatus = EPackageStoreEntryStatus::Ok;
+			// update some package names, and if linker loads are required, disable zen
+			if (AsyncLoadingThread.GetLinkerLoadInformationAndReturnIfRequired(ImportedPackageNameToLoad, ImportedPackageUPackageName, ImportedPackagePath))
+			{
+				bIsZenPackageImport = false;
+				ImportedPackageStatus = EPackageStoreEntryStatus::Ok;
+			}
 		}
 #else
 		constexpr bool bIsZenPackage = true;
diff --git a/Engine/Source/Runtime/CoreUObject/Private/UObject/LinkerLoad.cpp b/Engine/Source/Runtime/CoreUObject/Private/UObject/LinkerLoad.cpp
--- a/Engine/Source/Runtime/CoreUObject/Private/UObject/LinkerLoad.cpp
+++ b/Engine/Source/Runtime/CoreUObject/Private/UObject/LinkerLoad.cpp
@@ -172,6 +172,17 @@ static FAutoConsoleVariableRef CVarAllowCookedDataInEditorBuilds(
 	ECVF_Default
 );
 
+#if WITH_EDITOR
+// force the cvar value in HybridCookerEditor mode
+FDelayedAutoRegisterHelper GInitDevice(EDelayedRegisterRunPhase::IniSystemReady, []
+	{
+		if (IsRunningHybridCookedEditor())
+		{
+			GAllowCookedDataInEditorBuilds = 1;
+		}
+	});
+#endif
+
 int32 GSkipAsyncLoaderForCookedData = 0;
 static FAutoConsoleVariableRef CVarSkipAsyncLoaderForCookedData(
 	TEXT("cook.SkipAsyncLoaderForCookedData"),
diff --git a/Engine/Source/Runtime/Launch/Private/LaunchEngineLoop.cpp b/Engine/Source/Runtime/Launch/Private/LaunchEngineLoop.cpp
--- a/Engine/Source/Runtime/Launch/Private/LaunchEngineLoop.cpp
+++ b/Engine/Source/Runtime/Launch/Private/LaunchEngineLoop.cpp
@@ -1950,6 +1950,18 @@ int32 FEngineLoop::PreInitPreStartupScreen(const TCHAR* CmdLine)
 		GShouldEmitVerboseNamedEvents = true;
 	}
 
+#if WITH_EDITOR
+	FString RunToken;
+	// we are running in HybridCookedEditor if it's on the commandline, and we aren't cooking in hybridce mode
+	// (if we are cooking then we don't want to load cooked data, but -hybridcookededitor is still on the commandline 
+	// because we need to modify the platform name from <Plat>CookedEditor to <Plat>Editor based on the param)
+	if (FParse::Param(CmdLine, TEXT("hybridcookededitor")) && 
+		(!FParse::Value(CmdLine, TEXT("-run="), RunToken) || RunToken != TEXT("Cook")))
+	{
+		PRIVATE_GIsRunningHybridCookedEditor = true;
+	}
+#endif
+
 	// Set the flag for whether we've build DebugGame instead of Development. The engine does not know this (whereas the launch module does) because it is always built in development.
 #if UE_BUILD_DEVELOPMENT && defined(UE_BUILD_DEVELOPMENT_WITH_DEBUGGAME) && UE_BUILD_DEVELOPMENT_WITH_DEBUGGAME
 	FApp::SetDebugGame(true);
diff --git a/Engine/Source/Runtime/StorageServerClient/Private/StorageServerPlatformFile.cpp b/Engine/Source/Runtime/StorageServerClient/Private/StorageServerPlatformFile.cpp
--- a/Engine/Source/Runtime/StorageServerClient/Private/StorageServerPlatformFile.cpp
+++ b/Engine/Source/Runtime/StorageServerClient/Private/StorageServerPlatformFile.cpp
@@ -63,6 +63,17 @@ DEFINE_LOG_CATEGORY_STATIC(LogStorageServerPlatformFile, Log, All);
 
 static FDateTime GAssumedImmutableTimeStamp = FDateTime::Now();
 
+// If this is set, then StorageServer will not be used for non-assets (OpenRead, etc is used for non-assets; assets are handled via the IoDispatcher)
+static bool GPreferLocalForNonAssets = false;
+static FAutoConsoleVariableRef CVarPreferCookedPackagesInEditor(
+	TEXT("s.PreferLocalForNonAssets"),
+	GPreferLocalForNonAssets,
+	TEXT("Set to true to look at the local file sytem for files before loading from StorageServer"),
+	ECVF_Default
+);
+
+
+
 FStorageServerFileSystemTOC::~FStorageServerFileSystemTOC()
 {
 	FWriteScopeLock _(TocLock);
@@ -961,6 +972,14 @@ FAnsiString FStorageServerPlatformFile::MakeBaseURI()
 
 bool FStorageServerPlatformFile::ShouldBeUsed(IPlatformFile* Inner, const TCHAR* CmdLine) const
 {
+	// in the non-fully-cooked-editor, we only want to use zen store when we are using HybridCookedEditor mode
+#if WITH_EDITOR && !UE_IS_COOKED_EDITOR
+	if (!IsRunningHybridCookedEditor())
+	{
+		return false;
+	}
+#endif
+
 #if WITH_COTF
 	UE::Cook::ICookOnTheFlyModule& CookOnTheFlyModule = FModuleManager::LoadModuleChecked<UE::Cook::ICookOnTheFlyModule>(TEXT("CookOnTheFly"));
 	TSharedPtr<UE::Cook::ICookOnTheFlyServerConnection> DefaultConnection = CookOnTheFlyModule.GetDefaultServerConnection();
@@ -1033,6 +1052,12 @@ bool FStorageServerPlatformFile::ShouldBeUsed(IPlatformFile* Inner, const TCHAR*
 
 bool FStorageServerPlatformFile::Initialize(IPlatformFile* Inner, const TCHAR* CmdLine)
 {
+	// hybrid cooked editor wants to load any local files when possible, instead of any potential non-assets that were imported with a zen store
+	if (IsRunningHybridCookedEditor())
+	{
+		GPreferLocalForNonAssets = 1;
+	}
+
 	LowerLevel = Inner;
 	if (HostAddrs.Num() > 0)
 	{
@@ -1315,6 +1340,16 @@ IFileHandle* FStorageServerPlatformFile::OpenRead(const TCHAR* Filename, bool bA
 {
 	TStringBuilder<1024> StorageServerFilename;
 
+	// if we prefer local files, look local before checking if the file is in the ServerToc
+	if (GPreferLocalForNonAssets && LowerLevel && IsNonServerFilenameAllowed(Filename))
+	{
+		IFileHandle* Handle = LowerLevel->OpenRead(Filename, bAllowWrite);
+		if (Handle != nullptr)
+		{
+			return Handle;
+		}
+	}	
+
 	if (MakeStorageServerPath(Filename, StorageServerFilename))
 	{
 		FIoChunkId FileChunkId;
@@ -1324,7 +1359,8 @@ IFileHandle* FStorageServerPlatformFile::OpenRead(const TCHAR* Filename, bool bA
 			return InternalOpenFile(FileChunkId, RawSize, Filename);
 		}
 	}
-	return (LowerLevel && IsNonServerFilenameAllowed(Filename)) ? LowerLevel->OpenRead(Filename, bAllowWrite) : nullptr;
+	// if we preferred Server over Local, look in local if Server failed
+	return (!GPreferLocalForNonAssets && LowerLevel && IsNonServerFilenameAllowed(Filename)) ? LowerLevel->OpenRead(Filename, bAllowWrite) : nullptr;
 }
 
 bool FStorageServerPlatformFile::IterateDirectory(const TCHAR* Directory, IPlatformFile::FDirectoryVisitor& Visitor)
