RemoteObjectTransfer fixes:

- fixed some issues with the work priority. because we inherit priority for dependent work, we could end up with an actor tick and a dependent physics query at the same priority on different servers both fighting over the same object causing a live lock. we should never have a situation where there isn't a deterministic 'winner' between two tasks. added FRemoteWorkPriority to capture some additional information for deciding the winner. we now use the server id of the root work (which was the previous priority), plus the 'depth' of the work (so that 'more leaf' work always takes priority over 'more root' work, since root work can't complete until dependent work is complete). If we still have equal priority (two works from the same root server at the same depth) then we use the transaction id (transactions started on the root server take priority over later transactions) and as a final arbitrator we take the server ids of the servers actually executing the work to pick a winner.

- fixed object migration arbitration being confused by GRemoteObjectsMigrateFullHierarchy. Since we always serialize the full hierarchy now, we now find the 'canonical root' object when doing arbitration so that we always talk about the root object instead of some particular subobject.

diff --git a/Engine/Source/Runtime/CoreUObject/Private/UObject/RemoteObject.cpp b/Engine/Source/Runtime/CoreUObject/Private/UObject/RemoteObject.cpp
--- a/Engine/Source/Runtime/CoreUObject/Private/UObject/RemoteObject.cpp
+++ b/Engine/Source/Runtime/CoreUObject/Private/UObject/RemoteObject.cpp
@@ -116,7 +116,7 @@ void InitRemoteObjects()
 	}
 	if (!UE::RemoteObject::Transfer::FetchDeferredRPCDelegate.IsBound())
 	{
-		UE::RemoteObject::Transfer::FetchDeferredRPCDelegate.BindStatic([](FRemoteServerId) { return TTuple<FRemoteServerId, TFunction<void(void)>>(); });
+		UE::RemoteObject::Transfer::FetchDeferredRPCDelegate.BindStatic([](FRemoteWorkPriority) { return TTuple<FRemoteWorkPriority, TFunction<void(void)>>(); });
 	}
 	if (!UE::RemoteObject::Transfer::PhysicsQueryTickDelegate.IsBound())
 	{
diff --git a/Engine/Source/Runtime/CoreUObject/Private/UObject/RemoteObjectSerialization.cpp b/Engine/Source/Runtime/CoreUObject/Private/UObject/RemoteObjectSerialization.cpp
--- a/Engine/Source/Runtime/CoreUObject/Private/UObject/RemoteObjectSerialization.cpp
+++ b/Engine/Source/Runtime/CoreUObject/Private/UObject/RemoteObjectSerialization.cpp
@@ -70,7 +70,7 @@ FString GenerateRemoteObjectFilename(FRemoteObjectId ObjectId, FRemoteServerId D
 	return FPaths::Combine(*FPaths::ProjectSavedDir(), *(FString::Printf(TEXT("%s_%s.remote"), *ObjectId.ToString(), *DestinationServerId.ToString())));
 }
 
-void LoadObjectFromDisk(uint32 RequestToken, FRemoteServerId RequestPriority, FRemoteObjectId ObjectId, FRemoteServerId CurrentOwnerServerId, FRemoteServerId DestinationServerId)
+void LoadObjectFromDisk(uint32 RequestToken, FRemoteWorkPriority RequestPriority, FRemoteObjectId ObjectId, FRemoteServerId CurrentOwnerServerId, FRemoteServerId DestinationServerId)
 {
 	const FString Filename(GenerateRemoteObjectFilename(ObjectId, CurrentOwnerServerId));
 	TUniquePtr<FArchive> FileReader{ IFileManager::Get().CreateFileReader(*Filename) };
@@ -737,14 +737,7 @@ FRemoteObjectData* SerializeObjectData(UObject* InObject, TSet<UObject*>& OutObj
 
 	if (GRemoteObjectsMigrateFullHierarchy)
 	{
-		// If the object we want to migrate is a default subobject, make sure we migrate its parent too 
-		// so that it can safely be reconstructed on the other end (otherwise its outer may be missing)
-		while (Object && Object->IsDefaultSubobject())
-		{
-			Object = Object->GetOuter();
-		}
-		// Sanity check because by definition default subobjects should always have an outer
-		checkf(Object, TEXT("Attempted to migrate default subobject (%s) hierarchy but it doesn't have an outer or all outers are also default subobjects"), *GetPathNameSafe(InObject));
+		Object = FindCanonicalRootObjectForSerialization(Object);
 	}
 	
 	{
@@ -875,4 +868,17 @@ int32 DeserializeObjectData(FRemoteObjectData& ObjectData, TArray<UObject*>& Out
 	return RequestedObjectIndex;
 }
 
+UObject* FindCanonicalRootObjectForSerialization(UObject* Object)
+{
+	while (Object && Object->IsDefaultSubobject())
+	{
+		Object = Object->GetOuter();
+	}
+
+	// Sanity check because by definition default subobjects should always have an outer
+	checkf(Object, TEXT("Attempted to migrate default subobject (%s) hierarchy but it doesn't have an outer or all outers are also default subobjects"), *GetPathNameSafe(Object));
+
+	return Object;
+}
+
 } // namespace UE::RemoteObject::Serialization
\ No newline at end of file
diff --git a/Engine/Source/Runtime/CoreUObject/Private/UObject/RemoteObjectSerialization.h b/Engine/Source/Runtime/CoreUObject/Private/UObject/RemoteObjectSerialization.h
--- a/Engine/Source/Runtime/CoreUObject/Private/UObject/RemoteObjectSerialization.h
+++ b/Engine/Source/Runtime/CoreUObject/Private/UObject/RemoteObjectSerialization.h
@@ -71,11 +71,17 @@ namespace UE::RemoteObject::Serialization
 	*/
 	int32 DeserializeObjectData(FRemoteObjectData& ObjectData, TArray<UObject*>& OutReceivedObjects);
 
+	/**
+	* Finds the canonical 'root' object that is used for remote object serialization
+	* - we trace up the chain of Outer pointers until we reach the first non default subobject
+	*/
+	UObject* FindCanonicalRootObjectForSerialization(UObject* Object);
+
 } // namespace UE::RemoveObject::Serialization
 
 
 namespace UE::RemoteObject::Serialization::Disk
 {
-	void LoadObjectFromDisk(uint32 RequestToken, FRemoteServerId RequestPriority, FRemoteObjectId ObjectId, FRemoteServerId CurrentOwnerServerId, FRemoteServerId DestinationServerId);
+	void LoadObjectFromDisk(uint32 RequestToken, FRemoteWorkPriority RequestPriority, FRemoteObjectId ObjectId, FRemoteServerId CurrentOwnerServerId, FRemoteServerId DestinationServerId);
 	void SaveObjectToDisk(uint32 RequestToken, FRemoteObjectId ObjectId, FRemoteServerId ServerId, FRemoteObjectData& ObjectData);
 }
\ No newline at end of file
diff --git a/Engine/Source/Runtime/CoreUObject/Private/UObject/RemoteObjectTransfer.cpp b/Engine/Source/Runtime/CoreUObject/Private/UObject/RemoteObjectTransfer.cpp
--- a/Engine/Source/Runtime/CoreUObject/Private/UObject/RemoteObjectTransfer.cpp
+++ b/Engine/Source/Runtime/CoreUObject/Private/UObject/RemoteObjectTransfer.cpp
@@ -21,24 +21,24 @@ namespace UE::RemoteObject::Transfer
 {
 const FRemoteServerId DatabaseId(0x3ff);
 
-TDelegate<void(uint32 RequestToken, FRemoteServerId RequestPriority, FRemoteObjectId ObjectId, FRemoteServerId CurrentOwnerServerId, FRemoteServerId DestinationServerId)> RequestRemoteObjectDelegate;
-TDelegate<void(uint32 RequestToken, FRemoteObjectId ObjectId, FRemoteServerId DestinationServerId, FRemoteObjectData& ObjectData)> RemoteObjectTransferDelegate;
-TDelegate<void(uint32 RequestToken, FRemoteObjectId ObjectId, FRemoteServerId DestinationServerId)> RemoteObjectDeniedTransferDelegate;
+TDelegate<void(uint32 /*RequestToken*/, FRemoteWorkPriority /*RequestPriority*/, FRemoteObjectId /*ObjectId*/, FRemoteServerId /*CurrentOwnerServerId*/, FRemoteServerId /*DestinationServerId*/)> RequestRemoteObjectDelegate;
+TDelegate<void(uint32 /*RequestToken*/, FRemoteObjectId /*ObjectId*/, FRemoteServerId /*DestinationServerId*/, FRemoteObjectData& /*ObjectData*/)> RemoteObjectTransferDelegate;
+TDelegate<void(uint32 /*RequestToken*/, FRemoteObjectId /*ObjectId*/, FRemoteServerId /*DestinationServerId*/)> RemoteObjectDeniedTransferDelegate;
 TDelegate<void()> TickNetworkDelegate;
-TDelegate<TTuple<FRemoteServerId, TFunction<void(void)>>(FRemoteServerId)> FetchDeferredRPCDelegate;
-TDelegate<TTuple<FRemoteServerId, TFunction<void(void)>>(void)> FetchNextDeferredRPCDelegate;
+TDelegate<TTuple<FRemoteWorkPriority, TFunction<void(void)>>(FRemoteWorkPriority)> FetchDeferredRPCDelegate;
+TDelegate<TTuple<FRemoteWorkPriority, TFunction<void(void)>>(void)> FetchNextDeferredRPCDelegate;
 TDelegate<void()> WaitForServersToSyncDelegate;
 
-TMulticastDelegate<void(const FRemoteObjectData& ObjectData, FRemoteServerId RemoteServerId)> OnObjectDataSentDelegate;
-TMulticastDelegate<void(const TSet<UObject*>& MigratedObjects, FRemoteServerId RemoteServerId)> OnObjectsSentDelegate;
-TMulticastDelegate<void(const FRemoteObjectData& ObjectData, FRemoteServerId RemoteServerId)> OnObjectDataReceivedDelegate;
-TMulticastDelegate<void(FRemoteObjectId ObjectId, FRemoteServerId RemoteServerId)> OnObjectDataDeniedDelegate;
-TMulticastDelegate<void(const TArray<UObject*>& MigratedObjects, FRemoteServerId RemoteServerId)> OnObjectsReceivedDelegate;
+TMulticastDelegate<void(const FRemoteObjectData& /*ObjectData*/, FRemoteServerId /*RemoteServerId*/)> OnObjectDataSentDelegate;
+TMulticastDelegate<void(const TSet<UObject*>& /*MigratedObjects*/, FRemoteServerId /*RemoteServerId*/)> OnObjectsSentDelegate;
+TMulticastDelegate<void(const FRemoteObjectData& /*ObjectData*/, FRemoteServerId /*RemoteServerId*/)> OnObjectDataReceivedDelegate;
+TMulticastDelegate<void(FRemoteObjectId /*ObjectId*/, FRemoteServerId /*RemoteServerId*/)> OnObjectDataDeniedDelegate;
+TMulticastDelegate<void(const TArray<UObject*>& /*MigratedObjects*/, FRemoteServerId /*RemoteServerId*/)> OnObjectsReceivedDelegate;
 
-TDelegate<void(uint32 RequestToken, FRemoteObjectId ObjectId, FRemoteServerId DatabaseId, FRemoteObjectData& ObjectData)> StoreRemoteObjectDataDelegate;
-TDelegate<void(uint32 RequestToken, FRemoteServerId RequestPriority, FRemoteObjectId ObjectId, FRemoteServerId DatabaseId, FRemoteServerId DestinationServerId)> RestoreRemoteObjectDataDelegate;
+TDelegate<void(uint32 /*RequestToken*/, FRemoteObjectId /*ObjectId*/, FRemoteServerId /*DatabaseId*/, FRemoteObjectData& /*ObjectData*/)> StoreRemoteObjectDataDelegate;
+TDelegate<void(uint32 /*RequestToken*/, FRemoteWorkPriority /*RequestPriority*/, FRemoteObjectId /*ObjectId*/, FRemoteServerId /*DatabaseId*/, FRemoteServerId /*DestinationServerId*/)> RestoreRemoteObjectDataDelegate;
 
-TDelegate<void()> PhysicsQueryBeginRequestDelegate;
+TDelegate<void(FRemoteWorkPriority)> PhysicsQueryBeginRequestDelegate;
 TDelegate<bool()> PhysicsQueryTickDelegate;
 TDelegate<void()> PhysicsQueryEndRequestDelegate;
 
@@ -49,6 +49,15 @@ FString DebugLastAbortReason;
 
 UObject* MigrateObjectFromObjectData(FRemoteObjectData& ObjectData, FRemoteServerId RemoteServerId, EObjectMigrationFlags ObjectFlags);
 
+void ExecuteTransactionalInternal(FRemoteWorkPriority WorkPriority, const TFunctionRef<void(void)>& Work);
+
+TAtomic<uint32> NextTransactionId;
+
+uint32 GenerateNextTransactionId()
+{
+	return (NextTransactionId++ % 0x800000u) + 1000u;
+}
+
 /**
  * Remote object transfer queue. 
  * Queued requests (both send and receive) are processed on the game thread since some systems (like RPCs) that are used for transfering object data are GT-only
@@ -56,14 +65,14 @@ UObject* MigrateObjectFromObjectData(FRemoteObjectData& ObjectData, FRemoteServe
 class FRemoteObjectTransferQueue
 {
 public:
-	TArray<FRemoteObjectId> LockedObjects;
-	TArray<FRemoteObjectId> RequestedObjects;
+	TArray<FRemoteObjectId> PendingObjectRequests;
+	TArray<FRemoteObjectId> RequiredObjects;
+	TArray<FRemoteObjectId> RequiredObjectsCanonicalRoot;
 	TAtomic<uint32> ActiveRequestToken;
 	TAtomic<uint32> NextActiveRequestToken;
-	FRemoteServerId ActiveRequestPriority;
-	int32 PendingRequests = 0;
+	FRemoteWorkPriority ActiveRequestPriority;
 
-	void BeginRequest(FRemoteServerId Priority)
+	void BeginRequest(FRemoteWorkPriority Priority)
 	{
 		check(!AutoRTFM::IsClosed());
 
@@ -71,11 +80,9 @@ public:
 		check(ActiveRequestToken.Load() == 0);
 
 		// we expect the object lists to be empty
-		check(LockedObjects.Num() == 0);
-		check(RequestedObjects.Num() == 0);
-
-		// we expect to not be currently waiting for any requests
-		check(PendingRequests == 0);
+		check(PendingObjectRequests.Num() == 0);
+		check(RequiredObjects.Num() == 0);
+		check(RequiredObjectsCanonicalRoot.Num() == 0);
 
 		uint32 ExpectedToken = 0u;
 		bool TokenExchangeSuccess = ActiveRequestToken.CompareExchange(ExpectedToken, (NextActiveRequestToken.Load() % 32000) + 1000);
@@ -93,18 +100,16 @@ public:
 		uint32 OldToken = ActiveRequestToken.Exchange(0u);
 		check(OldToken != 0);
 
-		// we expect to have no pending network requests
-		check(PendingRequests == 0);
-
 		// clear the object lists
-		LockedObjects.Reset();
-		RequestedObjects.Reset();
+		PendingObjectRequests.Reset();
+		RequiredObjects.Reset();
+		RequiredObjectsCanonicalRoot.Reset();
 
 		// clear the priority
-		ActiveRequestPriority = FRemoteServerId();
+		ActiveRequestPriority = FRemoteWorkPriority();
 	}
 
-	void RequestRemoteObject(uint32 RequestToken, FRemoteServerId RequestPriority, FRemoteObjectId ObjectId, FRemoteServerId CurrentOwnerServerId, FRemoteServerId DestinationServerId)
+	void RequestRemoteObject(uint32 RequestToken, FRemoteWorkPriority RequestPriority, FRemoteObjectId ObjectId, FRemoteServerId CurrentOwnerServerId, FRemoteServerId DestinationServerId)
 	{
 		check(!AutoRTFM::IsClosed());
 
@@ -146,17 +151,13 @@ public:
 	*/
 	void FulfillReceiveRequest(uint32 RequestToken, FRemoteObjectId ObjectId, FRemoteServerId RemoteServerId, FRemoteObjectData* ObjectData)
 	{
+		using namespace UE::RemoteObject::Serialization;
 		check(!AutoRTFM::IsClosed());
 
-		bool bMatchesActiveRequest = (RequestToken == ActiveRequestToken.Load());
-		bool bFoundRequest = (RequestedObjects.Find(ObjectId) != INDEX_NONE);
-		bool bAlreadyLocked = (LockedObjects.Find(ObjectId) != INDEX_NONE);
-
-		// locked objects are ones that we believe to be local - it should be impossible for
-		// another server to send us an object that we already have
-		check(!bAlreadyLocked);
+		const int32 PendingRequestIndex = PendingObjectRequests.Find(ObjectId);
+		const bool bFoundRequest = (PendingRequestIndex != INDEX_NONE);
 
-		UE_LOG(LogRemoteObject, Verbose, TEXT("     : FulfillReceiveRequest(%u) %s (was requested: %d, token matches: %d)"), RequestToken, *ObjectId.ToString(), bFoundRequest, bMatchesActiveRequest);
+		UE_LOG(LogRemoteObject, Verbose, TEXT("     : FulfillReceiveRequest(%u) %s (was requested: %d)"), RequestToken, *ObjectId.ToString(), bFoundRequest);
 
 		const bool bAssignedOwnership = (RequestToken == PushMigrationRequestToken);
 		EObjectMigrationFlags MigrationFlags = bAssignedOwnership ? EObjectMigrationFlags::Acquire : EObjectMigrationFlags::Receive;
@@ -168,14 +169,14 @@ public:
 
 		if (bFoundRequest)
 		{
-			// we received an object that we requested (was either push migrated or directly responded to),
-			// so now consider it locked
-			LockedObjects.Add(ObjectId);
-		}
+			PendingObjectRequests.RemoveAt(PendingRequestIndex);
 
-		if (bMatchesActiveRequest && bFoundRequest)
-		{
-			PendingRequests--;
+			const int32 RequiredObjectIndex = RequiredObjects.Find(ObjectId);
+			check(RequiredObjectIndex != INDEX_NONE);
+
+			// now that we have the object locally we can discover and update the canonical root object id
+			UObject* RootObject = FindCanonicalRootObjectForSerialization(Result);
+			RequiredObjectsCanonicalRoot[RequiredObjectIndex] = UE::CoreUObject::Private::FObjectHandleUtils::GetRemoteId(RootObject);
 		}
 	}
 
@@ -183,14 +184,14 @@ public:
 	{
 		check(!AutoRTFM::IsClosed());
 
-		bool bMatchesActiveRequest = (RequestToken == ActiveRequestToken.Load());
-		bool bFoundRequest = (RequestedObjects.Find(ObjectId) != INDEX_NONE);
+		const int32 PendingRequestIndex = PendingObjectRequests.Find(ObjectId);
+		const bool bFoundRequest = (PendingRequestIndex != INDEX_NONE);
 
-		UE_LOG(LogRemoteObject, Verbose, TEXT("     : DenyReceiveRequest(%u) %s, (was requested: %d, token matches: %d)"), RequestToken, *ObjectId.ToString(), bFoundRequest, bMatchesActiveRequest);
+		UE_LOG(LogRemoteObject, Verbose, TEXT("     : DenyReceiveRequest(%u) %s, (was requested: %d)"), RequestToken, *ObjectId.ToString(), bFoundRequest);
 
-		if (bMatchesActiveRequest && bFoundRequest)
+		if (bFoundRequest)
 		{
-			PendingRequests--;
+			PendingObjectRequests.RemoveAt(PendingRequestIndex);
 		}
 	}
 };
@@ -202,7 +203,7 @@ void RequestObjectData(FRemoteObjectId ObjectId, FRemoteServerId CurrentOwnerSer
 	check(!AutoRTFM::IsClosed());
 
 	check(GTransferQueue.ActiveRequestToken.Load() != 0);
-	GTransferQueue.RequestRemoteObject(GTransferQueue.ActiveRequestToken.Load(), GTransferQueue.ActiveRequestPriority, ObjectId, CurrentOwnerServerId, DestinationServerId);
+	GTransferQueue.RequestRemoteObject(GTransferQueue.ActiveRequestToken.Load(), GTransferQueue.ActiveRequestPriority.CreateDependentWorkPriority(), ObjectId, CurrentOwnerServerId, DestinationServerId);
 }
 
 UObject* MigrateObjectFromObjectData(FRemoteObjectData& ObjectData, FRemoteServerId RemoteServerId, EObjectMigrationFlags MigrationFlags)
@@ -220,12 +221,15 @@ UObject* MigrateObjectFromObjectData(FRemoteObjectData& ObjectData, FRemoteServe
 
 	if (ReceivedObjects.Num() > 0)
 	{
-		for(int32 SubObjectIndex = 1; SubObjectIndex < ReceivedObjects.Num(); SubObjectIndex++)
+		if (UE_LOG_ACTIVE(LogRemoteObjectD, Verbose))
 		{
-			FRemoteObjectId SubObjectId = UE::CoreUObject::Private::FObjectHandleUtils::GetRemoteId(ReceivedObjects[SubObjectIndex]);
-			UE_LOG(LogRemoteObjectD, Verbose, TEXT("         subobject[%2d] : %s"), SubObjectIndex, *SubObjectId.ToString());
+			for(int32 SubObjectIndex = 0; SubObjectIndex < ReceivedObjects.Num(); SubObjectIndex++)
+			{
+				FRemoteObjectId SubObjectId = UE::CoreUObject::Private::FObjectHandleUtils::GetRemoteId(ReceivedObjects[SubObjectIndex]);
+				UE_LOG(LogRemoteObjectD, Verbose, TEXT("         subobject[%2d] : %s %s"), SubObjectIndex, *SubObjectId.ToString(), (SubObjectIndex == RequestedObjectIndex) ? TEXT("<requested>") : TEXT(""));
+			}
+			UE_LOG(LogRemoteObjectD, Verbose, TEXT("         <done listing subobjects>"));
 		}
-		UE_LOG(LogRemoteObjectD, Verbose, TEXT("         <done listing subobjects>"));
 	}
 
 #if UE_WITH_REMOTE_OBJECT_HANDLE
@@ -254,17 +258,18 @@ void MigrateObjectFromRemoteServer(FRemoteObjectId ObjectId, FRemoteServerId Cur
 	{
 		// if we are not in a transaction (dev/test workflow), begin a
 		// temp request
-		GTransferQueue.BeginRequest(GetGlobalServerId());
+		GTransferQueue.BeginRequest(FRemoteWorkPriority());
 	}
 
 	UE_AUTORTFM_OPEN
 	{
-		// it's possible a prefetch already added this to the RequestedObjects list
-		bAdded = (GTransferQueue.RequestedObjects.Find(ObjectId) == INDEX_NONE);
+		// it's possible a prefetch already added this to the RequiredObjects list
+		bAdded = (GTransferQueue.RequiredObjects.Find(ObjectId) == INDEX_NONE);
 
 		if (bAdded)
 		{
-			GTransferQueue.RequestedObjects.Add(ObjectId);
+			GTransferQueue.RequiredObjects.Add(ObjectId);
+			GTransferQueue.RequiredObjectsCanonicalRoot.Add(ObjectId);
 		}
 
 		UE_LOG(LogRemoteObject, Verbose, TEXT("Aborting transaction, need obj id %s on server %s"),
@@ -292,7 +297,6 @@ void MigrateObjectFromRemoteServer(FRemoteObjectId ObjectId, FRemoteServerId Cur
 	{
 		// if we are not in a transaction, request the object 
 		// (which should immediately get it) and end the request
-		GTransferQueue.PendingRequests++;
 		UE::RemoteObject::Handle::RequestObject(ObjectId);
 		GTransferQueue.EndRequest();
 	}
@@ -300,10 +304,11 @@ void MigrateObjectFromRemoteServer(FRemoteObjectId ObjectId, FRemoteServerId Cur
 
 void MigrateObjectToRemoteServer(uint32 RequestToken, UObject* Object, FRemoteServerId DestinationServerId)
 {
-	MigrateObjectToRemoteServerWithExplicitPriority(RequestToken, GetGlobalServerId(), Object, DestinationServerId);
+	FRemoteWorkPriority RootWorkPriority = FRemoteWorkPriority::CreateRootWorkPriority(GetGlobalServerId(), GenerateNextTransactionId());
+	MigrateObjectToRemoteServerWithExplicitPriority(RequestToken, RootWorkPriority, Object, DestinationServerId);
 }
 
-void MigrateObjectToRemoteServerWithExplicitPriority(uint32 RequestToken, FRemoteServerId RequestPriority, UObject* Object, FRemoteServerId DestinationServerId)
+void MigrateObjectToRemoteServerWithExplicitPriority(uint32 RequestToken, FRemoteWorkPriority RequestPriority, UObject* Object, FRemoteServerId DestinationServerId)
 {
 	using namespace UE::RemoteObject::Serialization;
 	using namespace UE::RemoteObject::Private;
@@ -312,42 +317,59 @@ void MigrateObjectToRemoteServerWithExplicitPriority(uint32 RequestToken, FRemot
 
 	FRemoteObjectId ObjectId = UE::CoreUObject::Private::FObjectHandleUtils::GetRemoteId(Object);
 
-	UE_LOG(LogRemoteObject, Verbose, TEXT("     : MigrateObjectToRemoteServer(%u) sending obj %s [%s] (isremote %d) (%d locked objs)"), RequestToken, *ObjectId.ToString(), *Object->GetName(), Handle::IsRemote(Object), GTransferQueue.LockedObjects.Num());
+	UE_LOG(LogRemoteObject, Verbose, TEXT("     : MigrateObjectToRemoteServer(%u %s) %s is requesting obj %s (%s)"), RequestToken, *RequestPriority.ToString(), *DestinationServerId.ToString(), *ObjectId.ToString(), *Object->GetName());
+
+	// find the canonical root object id to use for arbitration
+	UObject* const RootObject = FindCanonicalRootObjectForSerialization(Object);
+	const FRemoteObjectId CanonicalRootObjectId = UE::CoreUObject::Private::FObjectHandleUtils::GetRemoteId(RootObject);
 
 	// is the object locked?
 	bool bIsObjectLocked = false;
 
-	for (int32 LockedObjectIndex = 0; LockedObjectIndex < GTransferQueue.LockedObjects.Num(); LockedObjectIndex++)
+	for (int32 RequiredObjectIndex = 0; RequiredObjectIndex < GTransferQueue.RequiredObjectsCanonicalRoot.Num(); RequiredObjectIndex++)
 	{
-		FRemoteObjectId LockedObjectId = GTransferQueue.LockedObjects[LockedObjectIndex];
+		FRemoteObjectId RequiredObjectId = GTransferQueue.RequiredObjectsCanonicalRoot[RequiredObjectIndex];
 
-		if (LockedObjectId == ObjectId)
+		if (RequiredObjectId == CanonicalRootObjectId)
 		{
-			// do we give it up, or deny their request?
-			if(GTransferQueue.ActiveRequestPriority < RequestPriority)
+			bool bIsHigherPriority = IsHigherPriority(GTransferQueue.ActiveRequestPriority, RequestPriority);
+			bool bIsLowerPriority = IsHigherPriority(RequestPriority, GTransferQueue.ActiveRequestPriority);
+			bool bIsEqualPriority = IsEqualPriority(GTransferQueue.ActiveRequestPriority, RequestPriority);
+
+			// we keep our object locked either if the work priority is explicitly higher, or otherwise
+			// as a last measure of arbitration we use the server ids when the priorities are equal
+			if (bIsHigherPriority)
 			{
-				// we are higher priority, send a deny response
+				// we keep our object locked if our priority is higher
+				UE_LOG(LogRemoteObject, Verbose, TEXT("     : MigrateObjectToRemoteServer(%u) DENYING request for our locked %s because our priority is higher: %s vs %s"), RequestToken, *ObjectId.ToString(), *GTransferQueue.ActiveRequestPriority.ToString(), *RequestPriority.ToString());
 				bIsObjectLocked = true;
-
-				UE_LOG(LogRemoteObject, Verbose, TEXT("     : MigrateObjectToRemoteServer(%u) DENYING request for %s because it's locked"), RequestToken, *LockedObjectId.ToString());
-				RemoteObjectDeniedTransferDelegate.Execute(RequestToken, LockedObjectId, DestinationServerId);
+			}
+			else if(bIsEqualPriority && (GetGlobalServerId() < DestinationServerId))
+			{
+				// if the priorities are equal, we arbitrate based on server id
+				UE_LOG(LogRemoteObject, Verbose, TEXT("     : MigrateObjectToRemoteServer(%u) DENYING request for our locked %s because priority %s is equal, and our server index is lower: %s vs %s"), RequestToken, *ObjectId.ToString(), *GTransferQueue.ActiveRequestPriority.ToString(), *GetGlobalServerId().ToString(), *DestinationServerId.ToString());
+				bIsObjectLocked = true;
+			}
+			else if (bIsLowerPriority)
+			{
+				UE_LOG(LogRemoteObject, Verbose, TEXT("     : MigrateObjectToRemoteServer(%u) GRANTING request for our locked %s because our priority is lower: %s vs %s"), RequestToken, *ObjectId.ToString(), *GTransferQueue.ActiveRequestPriority.ToString(), *RequestPriority.ToString());
 			}
 			else
 			{
-				UE_LOG(LogRemoteObject, Verbose, TEXT("     : MigrateObjectToRemoteServer(%u) GIVING UP our locked %s"), RequestToken, *LockedObjectId.ToString());
-
-				// unlock the object from ourselves
-				GTransferQueue.LockedObjects.RemoveAt(LockedObjectIndex, EAllowShrinking::No);
+				UE_LOG(LogRemoteObject, Verbose, TEXT("     : MigrateObjectToRemoteServer(%u) GRANTING request for our locked %s because priority %s is equal, and our server index is higher: %s vs %s"), RequestToken, *ObjectId.ToString(), *GTransferQueue.ActiveRequestPriority.ToString(), *GetGlobalServerId().ToString(), *DestinationServerId.ToString());
 			}
 
 			break;
 		}
 	}
 
-	if (!bIsObjectLocked)
+	if (bIsObjectLocked)
+	{
+		// we are higher priority, send a deny response
+		RemoteObjectDeniedTransferDelegate.Execute(RequestToken, ObjectId, DestinationServerId);
+	}
+	else
 	{
-		UE_LOG(LogRemoteObject, Verbose, TEXT("     : MigrateObjectToRemoteServer(%u) serializing obj %s (isremote %d)"), RequestToken, *ObjectId.ToString(), Handle::IsRemote(Object));
-
 		// send the object
 		TSet<UObject*> SentObjects;
 		FRemoteObjectData* ObjectData = SerializeObjectData(Object, SentObjects);
@@ -378,7 +400,7 @@ void MigrateObjectToRemoteServerWithExplicitPriority(uint32 RequestToken, FRemot
 			// Are we loaning it out or changing ownership?
 			const bool bOwnershipChange = (RequestToken == PushMigrationRequestToken);
 			FPostMigrateParams PostMigrateParams;
-			PostMigrateParams.RequestingServerId = RequestPriority;
+			PostMigrateParams.RequestingServerId = RequestPriority.GetRootServerId();
 			PostMigrateParams.DestinationServerId = DestinationServerId;
 			PostMigrateParams.MigrationFlags = bOwnershipChange ? EObjectMigrationFlags::Reassign : EObjectMigrationFlags::Loan;
 			for (UObject* SentObject : SentObjects)
@@ -404,10 +426,11 @@ void RegisterRemoteObjectId(FRemoteObjectId Id, FRemoteServerId OwnerServerId)
 
 void MigrateObjectToRemoteServer(uint32 RequestToken, FRemoteObjectId ObjectId, FRemoteServerId DestinationServerId)
 {
-	MigrateObjectToRemoteServerWithExplicitPriority(RequestToken, GetGlobalServerId(), ObjectId, DestinationServerId);
+	FRemoteWorkPriority RootWorkPriority = FRemoteWorkPriority::CreateRootWorkPriority(GetGlobalServerId(), GenerateNextTransactionId());
+	MigrateObjectToRemoteServerWithExplicitPriority(RequestToken, RootWorkPriority, ObjectId, DestinationServerId);
 }
 
-void MigrateObjectToRemoteServerWithExplicitPriority(uint32 RequestToken, FRemoteServerId RequestPriority, FRemoteObjectId ObjectId, FRemoteServerId DestinationServerId)
+void MigrateObjectToRemoteServerWithExplicitPriority(uint32 RequestToken, FRemoteWorkPriority RequestPriority, FRemoteObjectId ObjectId, FRemoteServerId DestinationServerId)
 {
 	using namespace UE::RemoteObject::Handle;
 	using namespace UE::RemoteObject::Serialization;
@@ -476,35 +499,44 @@ void PrefetchRemoteObject(FRemoteObjectId ObjectId, UObject* DestinationOuter)
 	
 	UE_AUTORTFM_OPEN
 	{
-		// it's possible a previous prefetch or strong fetch already added this to the RequestedObjects list
-		bAdded = (GTransferQueue.RequestedObjects.Find(ObjectId) == INDEX_NONE);
+		// it's possible a previous prefetch or strong fetch already added this to the RequiredObjects list
+		bAdded = (GTransferQueue.RequiredObjects.Find(ObjectId) == INDEX_NONE);
 
 		if (bAdded)
 		{
-			GTransferQueue.RequestedObjects.Add(ObjectId);
+			GTransferQueue.RequiredObjects.Add(ObjectId);
+			GTransferQueue.RequiredObjectsCanonicalRoot.Add(ObjectId);
 			UE_LOG(LogRemoteObject, Verbose, TEXT("PrefetchRemoteObject(%u) : obj %s"), GTransferQueue.ActiveRequestToken.Load(), *ObjectId.ToString());
 		}
 	};
 }
 
-void ExecuteTransactionalWithExplicitPriority(FRemoteServerId WorkPriority, const TFunctionRef<void(void)>& Work)
+void ExecuteTransactionalWithExplicitPriority(FRemoteWorkPriority WorkPriority, const TFunctionRef<void(void)>& Work)
 {
+#if UE_WITH_REMOTE_OBJECT_HANDLE
 	ExecuteTransactionalInternal(WorkPriority, Work);
+#else
+	Work();
+#endif
 }
 
 void ExecuteTransactional(const TFunctionRef<void(void)>& Work)
 {
-	ExecuteTransactionalWithExplicitPriority(GetGlobalServerId(), Work);
+	FRemoteWorkPriority RootWorkPriority = FRemoteWorkPriority::CreateRootWorkPriority(GetGlobalServerId(), GenerateNextTransactionId());
+	ExecuteTransactionalWithExplicitPriority(RootWorkPriority, Work);
 }
 
-void ExecuteTransactionalInternal(FRemoteServerId InWorkPriority, const TFunctionRef<void(void)>& InWork)
+void ExecuteTransactionalInternal(FRemoteWorkPriority InWorkPriority, const TFunctionRef<void(void)>& InWork)
 {
+	using namespace UE::RemoteObject::Serialization;
 	check(!AutoRTFM::IsClosed());
 
-	TArray<TTuple<FRemoteServerId, TFunction<void(void)>>> PendingWorks;
+	double LastWaitingPrintTime = FPlatformTime::Seconds();
+
+	TArray<TTuple<FRemoteWorkPriority, TFunction<void(void)>>> PendingWorks;
 	// it's not possible to construct a TFunction from a TFunctionRef, but it's 
 	// safe to do in this context so we just wrap the TFunctionRef in a lambda
-	PendingWorks.Add(TTuple<FRemoteServerId, TFunction<void(void)>>(InWorkPriority, [&InWork](){ InWork(); }));
+	PendingWorks.Add(TTuple<FRemoteWorkPriority, TFunction<void(void)>>(InWorkPriority, [&InWork](){ InWork(); }));
 
 	uint32 LocalIterationNumber = 0;
 
@@ -517,9 +549,10 @@ void ExecuteTransactionalInternal(FRemoteServerId InWorkPriority, const TFunctio
 		// earlier requests could be dependent on later ones completing
 		for(;;)
 		{
-			TTuple<FRemoteServerId, TFunction<void(void)>> NewWork = FetchDeferredRPCDelegate.Execute(PendingWorks[0].Get<0>());
+			TTuple<FRemoteWorkPriority, TFunction<void(void)>> NewWork = FetchDeferredRPCDelegate.Execute(PendingWorks[0].Get<0>());
 			if(NewWork.Get<0>().IsValid())
 			{
+				UE_LOG(LogRemoteObjectD, VeryVerbose, TEXT("-- ExecuteTransactional : iteration %u popped a higher priority RPC --"), LocalIterationNumber);
 				PendingWorks.Insert(NewWork, 0);
 			}
 			else
@@ -528,12 +561,12 @@ void ExecuteTransactionalInternal(FRemoteServerId InWorkPriority, const TFunctio
 			}
 		}
 
-		TTuple<FRemoteServerId, TFunction<void(void)>> ExecutingWork = PendingWorks[0];
-		FRemoteServerId ExecutingWorkPriority = ExecutingWork.Get<0>();
+		TTuple<FRemoteWorkPriority, TFunction<void(void)>> ExecutingWork = PendingWorks[0];
+		FRemoteWorkPriority ExecutingWorkPriority = ExecutingWork.Get<0>();
 
 		GTransferQueue.BeginRequest(ExecutingWorkPriority);
 
-		PhysicsQueryBeginRequestDelegate.ExecuteIfBound();
+		PhysicsQueryBeginRequestDelegate.ExecuteIfBound(ExecutingWorkPriority);
 
 		UE_LOG(LogRemoteObjectD, VeryVerbose, TEXT("-- ExecuteTransactional : iteration %u request %u pri %s --"), LocalIterationNumber, GTransferQueue.ActiveRequestToken.Load(), *ExecutingWorkPriority.ToString());
 
@@ -548,7 +581,7 @@ void ExecuteTransactionalInternal(FRemoteServerId InWorkPriority, const TFunctio
 			bool bRequiresContextSwitch = false;
 			for (;;)
 			{
-				TTuple<FRemoteServerId, TFunction<void(void)>> NewWork = FetchDeferredRPCDelegate.Execute(ExecutingWorkPriority);
+				TTuple<FRemoteWorkPriority, TFunction<void(void)>> NewWork = FetchDeferredRPCDelegate.Execute(ExecutingWorkPriority);
 				if (NewWork.Get<0>().IsValid())
 				{
 					PendingWorks.Insert(NewWork, 0);
@@ -565,94 +598,88 @@ void ExecuteTransactionalInternal(FRemoteServerId InWorkPriority, const TFunctio
 				break;
 			}
 
-			UE_LOG(LogRemoteObjectD, Verbose, TEXT("   : %d locked objs :"), GTransferQueue.LockedObjects.Num());
-			for (FRemoteObjectId ObjectId : GTransferQueue.LockedObjects)
-			{
-				UE_LOG(LogRemoteObjectD, Verbose, TEXT("       %s"), *ObjectId.ToString());
-			}
-
-			UE_LOG(LogRemoteObjectD, Verbose, TEXT("   : %d requested objs :"), GTransferQueue.RequestedObjects.Num());
-			for (FRemoteObjectId ObjectId : GTransferQueue.RequestedObjects)
-			{
-				UE_LOG(LogRemoteObjectD, Verbose, TEXT("       %s"), *ObjectId.ToString());
-			}
-
-			TArray<FRemoteObjectId> ObjectsToRequest;
+			check(GTransferQueue.RequiredObjects.Num() == GTransferQueue.RequiredObjectsCanonicalRoot.Num());
 
-			// try to migrate all RequestedObjects -> LockedObjects
-			for (FRemoteObjectId RequestedObjectId : GTransferQueue.RequestedObjects)
+			UE_LOG(LogRemoteObjectD, Verbose, TEXT("   : %d required objs (%d pending requests):"), GTransferQueue.RequiredObjects.Num(), GTransferQueue.PendingObjectRequests.Num());
+			for(int32 RequiredObjectIndex = 0; RequiredObjectIndex < GTransferQueue.RequiredObjects.Num(); RequiredObjectIndex++)
 			{
-				if (GTransferQueue.LockedObjects.Find(RequestedObjectId) == INDEX_NONE)
-				{
-					ObjectsToRequest.Add(RequestedObjectId);
-				}
+				FRemoteObjectId RequestedObjectId = GTransferQueue.RequiredObjects[RequiredObjectIndex];
+				FRemoteObjectId CanonicalObjectId = GTransferQueue.RequiredObjectsCanonicalRoot[RequiredObjectIndex];
+				UE_LOG(LogRemoteObjectD, Verbose, TEXT("       %s [root %s]"), *RequestedObjectId.ToString(), *CanonicalObjectId.ToString());
 			}
 
-			UE_LOG(LogRemoteObjectD, Verbose, TEXT("   : ObjectsToRequest: %d entries"), ObjectsToRequest.Num());
+			// pump the network
+			TickNetworkDelegate.ExecuteIfBound();
 
 			// send the requests... slowly for now (to avoid straining the network subsystem)
 			int32 RequestsSent = 0;
-			double LastWaitingPrintTime = FPlatformTime::Seconds();
-			for(;;)
+
+			for (int32 RequiredObjectIndex = 0; RequiredObjectIndex < GTransferQueue.RequiredObjects.Num(); RequiredObjectIndex++)
 			{
-				// do we need to, and can we send more requests?
-				while ((RequestsSent < ObjectsToRequest.Num()) && (GTransferQueue.PendingRequests < 5))
+				if (GTransferQueue.PendingObjectRequests.Num() >= 5)
 				{
-					// yes, send one request
-					FRemoteObjectId RequestedObjectId = ObjectsToRequest[RequestsSent++];
-
-					if (UE::RemoteObject::Handle::RequestObject(RequestedObjectId))
-					{
-						// true means we already have the object locally and didn't actually send a request for it
-						GTransferQueue.LockedObjects.Add(RequestedObjectId);
-						UE_LOG(LogRemoteObjectD, Verbose, TEXT("   : marking %s as locked, it is already local"), *RequestedObjectId.ToString());
-					}
-					else
-					{
-						GTransferQueue.PendingRequests++;
-						UE_LOG(LogRemoteObjectD, Verbose, TEXT("   : sent request for %s (%d pending requests)"), *RequestedObjectId.ToString(), GTransferQueue.PendingRequests);
-					}
+					break;
 				}
 
-				// pump the network
-				TickNetworkDelegate.ExecuteIfBound();
+				FRemoteObjectId RequiredObjectId = GTransferQueue.RequiredObjects[RequiredObjectIndex];
 
-				// are all requests sent?
-				if (RequestsSent == ObjectsToRequest.Num())
+				if (Handle::IsRemote(RequiredObjectId) && (GTransferQueue.PendingObjectRequests.Find(RequiredObjectId) == INDEX_NONE))
 				{
-					// yes, all requests are sent
-					// are they all complete?
-					if (GTransferQueue.PendingRequests == 0)
-					{
-						// yes, all done
-						break;
-					}
+					// send one request
+					bool bAlreadyHaveObject = UE::RemoteObject::Handle::RequestObject(RequiredObjectId);
+
+					// bAlreadyHaveObject should be false since we already tested IsRemote()
+					check(!bAlreadyHaveObject);
+
+					UE_LOG(LogRemoteObjectD, Verbose, TEXT("   : sent request for %s (%s) (%d pending requests)"), *RequiredObjectId.ToString(), *GTransferQueue.ActiveRequestPriority.ToString(), GTransferQueue.PendingObjectRequests.Num());
+					GTransferQueue.PendingObjectRequests.Add(RequiredObjectId);
 				}
 
-				double TimeNow = FPlatformTime::Seconds();
-				if (TimeNow > LastWaitingPrintTime + 1.0)
+				if (!Handle::IsRemote(RequiredObjectId))
 				{
-					UE_LOG(LogRemoteObjectD, VeryVerbose, TEXT("   : waiting for %d responses..."), GTransferQueue.PendingRequests);
-					LastWaitingPrintTime = TimeNow;
+					// discover the canonical root object id
+					UObject* const RootObject = FindCanonicalRootObjectForSerialization(StaticFindObjectFastInternal(RequiredObjectId));
+					GTransferQueue.RequiredObjectsCanonicalRoot[RequiredObjectIndex] = UE::CoreUObject::Private::FObjectHandleUtils::GetRemoteId(RootObject);
 				}
 			}
 
+			const double TimeNow = FPlatformTime::Seconds();
+
 			if (!PhysicsQueryTickDelegate.Execute())
 			{
+				if (TimeNow > LastWaitingPrintTime + 1.0)
+				{
+					UE_LOG(LogRemoteObjectD, VeryVerbose, TEXT("   : waiting for physics responses..."));
+					LastWaitingPrintTime = TimeNow;
+				}
+
 				// physics doesn't have all responses yet, try again
 				continue;
 			}
 
-			// now we also might have migrated objects away, so check to see if LockedObjects == RequestedObjects
-			if(GTransferQueue.LockedObjects.Num() != GTransferQueue.RequestedObjects.Num())
+			// do we have all required objects local?
+			int32 RemoteObjectCount = 0;
+			for (int32 RequiredObjectIndex = 0; RequiredObjectIndex < GTransferQueue.RequiredObjects.Num(); RequiredObjectIndex++)
+			{
+				if (Handle::IsRemote(GTransferQueue.RequiredObjects[RequiredObjectIndex]))
+				{
+					RemoteObjectCount++;
+				}
+			}
+
+			if(RemoteObjectCount > 0)
 			{
-				UE_LOG(LogRemoteObjectD, VeryVerbose, TEXT("   : locked count %d mismatches requests %d, did we give up locked objects? trying again..."), GTransferQueue.LockedObjects.Num(), GTransferQueue.RequestedObjects.Num());
+				if (TimeNow > LastWaitingPrintTime + 1.0)
+				{
+					UE_LOG(LogRemoteObjectD, VeryVerbose, TEXT("   : waiting for %d object responses..."), RemoteObjectCount);
+					LastWaitingPrintTime = TimeNow;
+				}
 
 				// try again!
 				continue;
 			}
 
-			UE_LOG(LogRemoteObjectD, VeryVerbose, TEXT("   : READY: all %d requested objects are locked"), GTransferQueue.RequestedObjects.Num());
+			UE_LOG(LogRemoteObjectD, VeryVerbose, TEXT("   : READY: all %d required objects are locked"), GTransferQueue.RequiredObjects.Num());
 
 			// all of the requested objects are locked, try to perform the work
 			if (!bWorkSucceeded)
@@ -673,18 +700,7 @@ void ExecuteTransactionalInternal(FRemoteServerId InWorkPriority, const TFunctio
 
 					bWorkSucceeded = true;
 					OnTransactionCompletedDelegate.Broadcast();
-
-					if (GTransferQueue.LockedObjects.Num() == GTransferQueue.RequestedObjects.Num())
-					{
-						// the work succeeded and no more object requests were made, we're all done!
-						break;
-					}
-					else
-					{
-						// the work is complete, but it added object requests (prefetches), so we
-						// run through the loop again in order to fetch those objects
-						continue;
-					}
+					break;
 				}
 			}
 		}
@@ -718,7 +734,7 @@ void ExecuteTransactionalInternal(FRemoteServerId InWorkPriority, const TFunctio
 			// we came here to execute. but before we start going back to find more
 			// work for ourselves, try to process any pending RPCs by searching
 			// for any priority available to help unblock other servers
-			TTuple<FRemoteServerId, TFunction<void(void)>> NewWork = FetchDeferredRPCDelegate.Execute(FRemoteServerId());
+			TTuple<FRemoteWorkPriority, TFunction<void(void)>> NewWork = FetchDeferredRPCDelegate.Execute(FRemoteWorkPriority());
 			if (NewWork.Get<0>().IsValid())
 			{
 				PendingWorks.Add(NewWork);
@@ -886,6 +902,61 @@ bool FRemoteObjectReference::NetSerialize(FArchive& Ar, class UPackageMap* Map,
 	return bOutSuccess;
 }
 
+FString FRemoteWorkPriority::ToString() const
+{
+	return FString::Format(TEXT("[pri: rsi {0} depth {1} id {2}]"),
+		{*GetRootServerId().ToString(),
+		GetWorkDepth(),
+		GetRootWorkTransactionId()});
+}
+
+FRemoteWorkPriority FRemoteWorkPriority::CreateRootWorkPriority(FRemoteServerId ServerId, uint32 TransactionId)
+{
+	FRemoteWorkPriority Result{};
+	uint32 RemoteServerId = 0;
+	FMemory::Memcpy(&RemoteServerId, &ServerId, sizeof(FRemoteServerId));
+
+	Result.PackedData |= (static_cast<uint64>(RemoteServerId) << 32);
+	Result.PackedData |= (0xFFull << 24);
+	Result.PackedData |= (TransactionId & 0xFFFFFFu);
+	return Result;
+}
+
+FRemoteWorkPriority FRemoteWorkPriority::CreateDependentWorkPriority() const
+{
+	FRemoteWorkPriority Result = *this;
+	Result.PackedData &= 0xFFFFFFFF00FFFFFFull;
+	
+	// work depth must be non-zero else we would wrap back around
+	check(GetWorkDepth() > 0);
+
+	Result.PackedData |= static_cast<uint64>(GetWorkDepth() - 1) << 24;
+	return Result;
+}
+
+// returns if Lhs is higher priority than Rhs
+bool IsHigherPriority(FRemoteWorkPriority Lhs, FRemoteWorkPriority Rhs)
+{
+	return Lhs.PackedData < Rhs.PackedData;
+}
+
+bool IsEqualPriority(FRemoteWorkPriority Lhs, FRemoteWorkPriority Rhs)
+{
+	// they're equal if neither is higher than the other
+	return !IsHigherPriority(Lhs, Rhs) && !IsHigherPriority(Rhs, Lhs);
+}
+
+bool IsHigherOrEqualPriority(FRemoteWorkPriority Lhs, FRemoteWorkPriority Rhs)
+{
+	return IsHigherPriority(Lhs, Rhs) || IsEqualPriority(Lhs, Rhs);
+}
+
+FArchive& operator<<(FArchive& Ar, FRemoteWorkPriority& Priority)
+{
+	Ar << Priority.PackedData;
+	return Ar;
+}
+
 FArchive& operator<<(FArchive& Ar, FRemoteObjectReference& Ref)
 {
 	Ar << Ref.ObjectId;
@@ -897,4 +968,4 @@ FArchive& operator<<(FArchive& Ar, FRemoteObjectReference& Ref)
 	}
 
 	return Ar;
-}
\ No newline at end of file
+}
diff --git a/Engine/Source/Runtime/CoreUObject/Public/UObject/NoExportTypes.h b/Engine/Source/Runtime/CoreUObject/Public/UObject/NoExportTypes.h
--- a/Engine/Source/Runtime/CoreUObject/Public/UObject/NoExportTypes.h
+++ b/Engine/Source/Runtime/CoreUObject/Public/UObject/NoExportTypes.h
@@ -3441,6 +3441,15 @@ struct FRemoteServerId
 	uint32 Id;
 };
 
+USTRUCT(noexport, BlueprintType)
+struct FRemoteWorkPriority
+{
+	GENERATED_BODY()
+
+	UPROPERTY()
+	uint64 PackedData;
+};
+
 USTRUCT(noexport, BlueprintType)
 struct FRemoteObjectReference
 {
diff --git a/Engine/Source/Runtime/CoreUObject/Public/UObject/RemoteObjectTransfer.h b/Engine/Source/Runtime/CoreUObject/Public/UObject/RemoteObjectTransfer.h
--- a/Engine/Source/Runtime/CoreUObject/Public/UObject/RemoteObjectTransfer.h
+++ b/Engine/Source/Runtime/CoreUObject/Public/UObject/RemoteObjectTransfer.h
@@ -41,6 +41,41 @@ struct FPostMigrateParams
 	EObjectMigrationFlags MigrationFlags;
 };
 
+struct FRemoteWorkPriority
+{
+	// the data is packed in a uint64 to allow passing by value in a register
+	// [63:32] : FRemoteServerId RootServerId
+	// [31:24] : uint8 WorkDepth
+	// [23:0]  : uint32 RootWorkTransactionId
+	//
+	// this allows comparison of priority to be a single uint64 compare-less-than
+	// 
+	// the WorkDepth is stored subtracted from 0xFF so that less-than means higher priority
+	//   0xFF is the root depth
+	//   0xFE is one deeper
+	//   
+	uint64 PackedData = 0;
+	static_assert(sizeof(FRemoteServerId) == 4);
+
+	FRemoteServerId GetRootServerId() const { return FRemoteServerId(static_cast<uint32>(PackedData >> 32)); }
+	uint32 GetWorkDepth() const { return (PackedData >> 24) & 0xFFu; }
+	uint32 GetRootWorkTransactionId() const { return PackedData & 0xFFFFFFu; }
+
+	bool IsValid() const { return GetRootServerId().IsValid(); }
+
+	COREUOBJECT_API static FRemoteWorkPriority CreateRootWorkPriority(FRemoteServerId ServerId, uint32 TransactionId);
+
+    // creates a copy of this priority with the WorkDepth decreased by 1
+	COREUOBJECT_API FRemoteWorkPriority CreateDependentWorkPriority() const;
+
+	COREUOBJECT_API FString ToString() const;
+};
+
+COREUOBJECT_API FArchive& operator<<(FArchive& Ar, FRemoteWorkPriority& Priority);
+COREUOBJECT_API bool IsHigherPriority(FRemoteWorkPriority Lhs, FRemoteWorkPriority Rhs);
+COREUOBJECT_API  bool IsEqualPriority(FRemoteWorkPriority Lhs, FRemoteWorkPriority Rhs);
+COREUOBJECT_API  bool IsHigherOrEqualPriority(FRemoteWorkPriority Lhs, FRemoteWorkPriority Rhs);
+
 namespace UE::RemoteObject::Transfer
 {
 	constexpr uint32 PushMigrationRequestToken = 0xFFFFFFFFu;
@@ -50,7 +85,7 @@ namespace UE::RemoteObject::Transfer
 	* as well as negotiating remote object transfers, retrying the transaction until it succeeds.
 	*/
 	COREUOBJECT_API void ExecuteTransactional(const TFunctionRef<void(void)>& Work);
-	COREUOBJECT_API void ExecuteTransactionalWithExplicitPriority(FRemoteServerId WorkPriority, const TFunctionRef<void(void)>& Work);
+	COREUOBJECT_API void ExecuteTransactionalWithExplicitPriority(FRemoteWorkPriority WorkPriority, const TFunctionRef<void(void)>& Work);
 
 	/**
 	* (Dev API) Sets a friendly string indicating why we are aborting the current transaction
@@ -92,7 +127,7 @@ namespace UE::RemoteObject::Transfer
 	* @param DestinationServerId Id of a remote server to receive the object
 	*/
 	COREUOBJECT_API void MigrateObjectToRemoteServer(uint32 RequestToken, UObject* Object, FRemoteServerId DestinationServerId);
-	COREUOBJECT_API void MigrateObjectToRemoteServerWithExplicitPriority(uint32 RequestToken, FRemoteServerId RequestPriority, UObject* Object, FRemoteServerId DestinationServerId);
+	COREUOBJECT_API void MigrateObjectToRemoteServerWithExplicitPriority(uint32 RequestToken, FRemoteWorkPriority RequestPriority, UObject* Object, FRemoteServerId DestinationServerId);
 
 	/**
 	* Migrates an object to a remote server
@@ -100,7 +135,7 @@ namespace UE::RemoteObject::Transfer
 	* @param DestinationServerId Id of a remote server to receive the object
 	*/
 	COREUOBJECT_API void MigrateObjectToRemoteServer(uint32 RequestToken, FRemoteObjectId Id, FRemoteServerId DestinationServerId);
-	COREUOBJECT_API void MigrateObjectToRemoteServerWithExplicitPriority(uint32 RequestToken, FRemoteServerId RequestPriority, FRemoteObjectId Id, FRemoteServerId DestinationServerId);
+	COREUOBJECT_API void MigrateObjectToRemoteServerWithExplicitPriority(uint32 RequestToken, FRemoteWorkPriority RequestPriority, FRemoteObjectId Id, FRemoteServerId DestinationServerId);
 	
 	/**
 	* Migrates an object from a remote server (temp function)
@@ -151,14 +186,14 @@ namespace UE::RemoteObject::Transfer
 	extern COREUOBJECT_API TDelegate<void(uint32 /*RequestToken*/, FRemoteObjectId /*ObjectId*/, FRemoteServerId /*DestinationServerId*/)> RemoteObjectDeniedTransferDelegate;
 
 	/** Delegate that requests remote object data from CurrentOwnerServerId to be transferred to DestinationServerId */
-	extern COREUOBJECT_API TDelegate<void(uint32 /*RequestToken*/, FRemoteServerId /*RequestPriority*/, FRemoteObjectId /*ObjectId*/, FRemoteServerId /*CurrentOwnerServerId*/, FRemoteServerId /*DestinationServerId*/)> RequestRemoteObjectDelegate;
+	extern COREUOBJECT_API TDelegate<void(uint32 /*RequestToken*/, FRemoteWorkPriority /*RequestPriority*/, FRemoteObjectId /*ObjectId*/, FRemoteServerId /*CurrentOwnerServerId*/, FRemoteServerId /*DestinationServerId*/)> RequestRemoteObjectDelegate;
 
 	/** Delegate that pumps the network and executes RPC handlers */
 	extern COREUOBJECT_API TDelegate<void()> TickNetworkDelegate;
 	
 	/** Fetch the next deferred RPC request **/
-	extern COREUOBJECT_API TDelegate<TTuple<FRemoteServerId, TFunction<void(void)>>(FRemoteServerId)> FetchDeferredRPCDelegate;
-	extern COREUOBJECT_API TDelegate<TTuple<FRemoteServerId, TFunction<void(void)>>(void)> FetchNextDeferredRPCDelegate;
+	extern COREUOBJECT_API TDelegate<TTuple<FRemoteWorkPriority, TFunction<void(void)>>(FRemoteWorkPriority)> FetchDeferredRPCDelegate;
+	extern COREUOBJECT_API TDelegate<TTuple<FRemoteWorkPriority, TFunction<void(void)>>(void)> FetchNextDeferredRPCDelegate;
 
 	/** Delegate that syncs all servers */
 	extern COREUOBJECT_API TDelegate<void()> WaitForServersToSyncDelegate;
@@ -186,10 +221,10 @@ namespace UE::RemoteObject::Transfer
 	extern COREUOBJECT_API TDelegate<void(uint32 /*RequestToken*/, FRemoteObjectId /*ObjectId*/, FRemoteServerId /*DatabaseId*/, FRemoteObjectData& /*ObjectData*/)> StoreRemoteObjectDataDelegate;
 
 	/** Delegate that restores object data from a database. */
-	extern COREUOBJECT_API TDelegate<void(uint32 /*RequestToken*/, FRemoteServerId /*RequestPriority*/, FRemoteObjectId /*ObjectId*/, FRemoteServerId /*DatabaseId*/, FRemoteServerId /*DestinationServerId*/)> RestoreRemoteObjectDataDelegate;
+	extern COREUOBJECT_API TDelegate<void(uint32 /*RequestToken*/, FRemoteWorkPriority /*RequestPriority*/, FRemoteObjectId /*ObjectId*/, FRemoteServerId /*DatabaseId*/, FRemoteServerId /*DestinationServerId*/)> RestoreRemoteObjectDataDelegate;
 
 	/** Physics query delegates **/
-	extern COREUOBJECT_API TDelegate<void()> PhysicsQueryBeginRequestDelegate;
+	extern COREUOBJECT_API TDelegate<void(FRemoteWorkPriority /*RequestPriority*/)> PhysicsQueryBeginRequestDelegate;
 	extern COREUOBJECT_API TDelegate<bool()> PhysicsQueryTickDelegate;
 	extern COREUOBJECT_API TDelegate<void()> PhysicsQueryEndRequestDelegate;
 
@@ -198,9 +233,6 @@ namespace UE::RemoteObject::Transfer
 
 	/** Delegate executed when a ExecuteTransactional observes an aborted AutoRTFM transaction */
 	extern COREUOBJECT_API TMulticastDelegate<void()> OnTransactionAbortedDelegate;
-
-	extern COREUOBJECT_API void ExecuteTransactionalInternal(FRemoteServerId WorkPriority, const TFunctionRef<void(void)>& Work);
-
 } // namespace UE::RemoveObject::Transfer
 
 struct FRemoteObjectReference
