PropertyBag: Prevents crashes when the user requests an unsupported type.

diff --git a/Engine/Source/Runtime/CoreUObject/Private/StructUtils/PropertyBag.cpp b/Engine/Source/Runtime/CoreUObject/Private/StructUtils/PropertyBag.cpp
--- a/Engine/Source/Runtime/CoreUObject/Private/StructUtils/PropertyBag.cpp
+++ b/Engine/Source/Runtime/CoreUObject/Private/StructUtils/PropertyBag.cpp
@@ -1686,7 +1686,7 @@ void FInstancedPropertyBag::AddProperties(const TConstArrayView<FPropertyBagProp
 			ExistingProperty->ValueType = NewDesc.ValueType;
 			ExistingProperty->ValueTypeObject = NewDesc.ValueTypeObject;
 		}
-		else
+		else if (NewDesc.ValueType != EPropertyBagPropertyType::None)
 		{
 			Descs.Add(NewDesc);
 		}
@@ -3568,6 +3568,9 @@ const UPropertyBag* UPropertyBag::GetOrCreateFromDescs(const TConstArrayView<FPr
 	{
 		UE::StructUtils::Private::RemovePropertyByName(NewBag->PropertyDescs, NewBag->PropertyDescs[Index].Name, Index + 1);
 	}
+
+	// Remove unsupported types
+	NewBag->PropertyDescs.RemoveAll([](const FPropertyBagPropertyDesc& Desc){ return Desc.ValueType == EPropertyBagPropertyType::None; });
 	
 	// Add properties (AddCppProperty() adds them backwards in the linked list)
 	for (int32 DescIndex = NewBag->PropertyDescs.Num() - 1; DescIndex >= 0; DescIndex--)
