Added s.UsePostloadGroupsWithIOStore (enabled by default).

#rnx

diff --git a/Engine/Source/Runtime/CoreUObject/Private/Serialization/AsyncLoading2.cpp b/Engine/Source/Runtime/CoreUObject/Private/Serialization/AsyncLoading2.cpp
--- a/Engine/Source/Runtime/CoreUObject/Private/Serialization/AsyncLoading2.cpp
+++ b/Engine/Source/Runtime/CoreUObject/Private/Serialization/AsyncLoading2.cpp
@@ -424,6 +424,14 @@ static FAutoConsoleVariableRef CVarGVerifyObjectLoadFlags(
 	ECVF_Default
 	);
 
+static bool GUsePostloadGroupsWithIOStore = true;
+static FAutoConsoleVariableRef CVarGUsePostloadGroupsWithIOStore(
+	TEXT("s.UsePostloadGroupsWithIOStore"),
+	GUsePostloadGroupsWithIOStore,
+	TEXT("Postload groups reduces the likelihood of race conditions between the loading thread and the game thread but are slower."),
+	ECVF_Default
+);
+
 static bool GOnlyProcessRequiredPackagesWhenSyncLoading = true;
 static FAutoConsoleVariableRef CVarGOnlyProcessRequiredPackagesWhenSyncLoading(
 	TEXT("s.OnlyProcessRequiredPackagesWhenSyncLoading"),
@@ -4638,7 +4646,7 @@ FAsyncPackage2* FAsyncLoadingThread2::FindOrInsertPackage(FAsyncLoadingThreadSta
 	// highly recursive InitDefaultMaterials function on some platforms.
 	// Since postload groups are there to protect against race conditions between postloads and serialize,
 	// no such race can exists until the loading thread is started.
-	const bool bIsPostLoadGroupFeatureActive = IsMultithreaded();
+	const bool bIsPostLoadGroupFeatureActive = IsMultithreaded() && GUsePostloadGroupsWithIOStore;
 #endif
 
 	if (bInserted)
