GitHub #12729 : Add operators to TStructView & TConstStructView to make them easier to use

diff --git a/Engine/Source/Runtime/CoreUObject/Public/StructUtils/StructView.h b/Engine/Source/Runtime/CoreUObject/Public/StructUtils/StructView.h
--- a/Engine/Source/Runtime/CoreUObject/Public/StructUtils/StructView.h
+++ b/Engine/Source/Runtime/CoreUObject/Public/StructUtils/StructView.h
@@ -183,6 +183,13 @@ public:
 		}
 	}
 
+	/** Simplified accessor for members of the guaranteed minimum structure type */
+	BaseStructT* operator->() const
+	{
+		check(IsValid());
+		return GetPtr();
+	}
+
 	/** Comparison operators. Note: it does not compare the internal structure itself */
 	template<typename T = BaseStructT, typename = std::enable_if_t<std::is_base_of_v<BaseStructT, std::decay_t<T>>>>
 	bool operator==(const TStructView<T>& Other) const
@@ -395,6 +402,13 @@ struct TConstStructView : FConstStructView
 		}
 	}
 
+	/** Simplified accessor for members of the guaranteed minimum structure type */
+	BaseStructT* operator->() const
+	{
+		check(IsValid());
+		return GetPtr();
+	}
+
 	/** Comparison operators. Note: it does not compare the internal structure itself */
 	template<typename T = BaseStructT, typename = std::enable_if_t<std::is_base_of_v<BaseStructT, std::decay_t<T>>>>
 	bool operator==(const TConstStructView<T>& Other) const
