Removed deprecated functions from FStructView.

diff --git a/Engine/Source/Runtime/CoreUObject/Public/StructUtils/StructView.h b/Engine/Source/Runtime/CoreUObject/Public/StructUtils/StructView.h
--- a/Engine/Source/Runtime/CoreUObject/Public/StructUtils/StructView.h
+++ b/Engine/Source/Runtime/CoreUObject/Public/StructUtils/StructView.h
@@ -60,22 +60,6 @@ public:
 		return UE::StructUtils::GetStructPtr<T>(ScriptStruct, StructMemory);
 	}
 
-	/** Returns mutable reference to the struct, this getter assumes that all data is valid. */
-	template<typename T>
-	UE_DEPRECATED(5.3, "Use Get() instead")
-	T& GetMutable() const
-	{
-		return Get<T>();
-	}
-
-	/** Returns mutable pointer to the struct, or nullptr if cast is not valid. */
-	template<typename T>
-	UE_DEPRECATED(5.3, "Use GetPtr() instead")
-	T* GetMutablePtr() const
-	{
-		return GetPtr<T>();
-	}
-
 	/** Returns struct type. */
 	const UScriptStruct* GetScriptStruct() const
 	{
@@ -88,13 +72,6 @@ public:
 		return StructMemory;
 	}
 
-	/** Returns a mutable pointer to struct memory. */
-	UE_DEPRECATED(5.3, "Use GetMemory() instead")
-	uint8* GetMutableMemory()
-	{
-		return GetMemory();
-	}
-
 	/** Reset to empty. */
 	void Reset()
 	{
@@ -272,22 +249,6 @@ public:
 		return UE::StructUtils::GetStructPtr<T>(ScriptStruct, StructMemory);
 	}
 
-	/** Returns const reference to the struct, this getter assumes that all data is valid. */
-	template<typename T>
-	UE_DEPRECATED(5.3, "Use the version of this function which takes const in the template type")
-	constexpr typename TEnableIf<!TIsConst<T>::Value, const T&>::Type Get() const
-	{
-		return UE::StructUtils::GetStructRef<T>(ScriptStruct, StructMemory);
-	}
-
-	/** Returns const pointer to the struct, or nullptr if cast is not valid. */
-	template<typename T>
-	UE_DEPRECATED(5.3, "Use the version of this function which takes const in the template type")
-	constexpr typename TEnableIf<!TIsConst<T>::Value, const T*>::Type GetPtr() const
-	{
-		return UE::StructUtils::GetStructPtr<T>(ScriptStruct, StructMemory);
-	}
-
 	/** Returns struct type. */
 	const UScriptStruct* GetScriptStruct() const
 	{
