Fixed IDO GC when IDOs are enabled for every object

diff --git a/Engine/Source/Runtime/CoreUObject/Internal/UObject/InstanceDataObjectUtils.h b/Engine/Source/Runtime/CoreUObject/Internal/UObject/InstanceDataObjectUtils.h
--- a/Engine/Source/Runtime/CoreUObject/Internal/UObject/InstanceDataObjectUtils.h
+++ b/Engine/Source/Runtime/CoreUObject/Internal/UObject/InstanceDataObjectUtils.h
@@ -27,6 +27,8 @@ namespace UE
 UE_API bool IsInstanceDataObjectSupportEnabled();
 /** Query if InstanceDataObject support is enabled for a specific object. */
 UE_API bool IsInstanceDataObjectSupportEnabled(const UObject* Object);
+/** Query if InstanceDataObject support is enabled for a specific class. */
+UE_API bool IsInstanceDataObjectSupportEnabledForClass(const UClass* Class);
 /** Query if uninitialized alert UI is enabled **/
 UE_API bool IsUninitializedAlertUIEnabled();
 
diff --git a/Engine/Source/Runtime/CoreUObject/Private/UObject/GarbageCollection.cpp b/Engine/Source/Runtime/CoreUObject/Private/UObject/GarbageCollection.cpp
--- a/Engine/Source/Runtime/CoreUObject/Private/UObject/GarbageCollection.cpp
+++ b/Engine/Source/Runtime/CoreUObject/Private/UObject/GarbageCollection.cpp
@@ -40,6 +40,7 @@
 #include "UObject/GarbageCollectionHistory.h"
 #include "UObject/GarbageCollectionInternalFlags.h"
 #include "UObject/GarbageCollectionTesting.h"
+#include "UObject/InstanceDataObjectUtils.h"
 #include "UObject/PropertyBagRepository.h"
 #include "UObject/PropertyOptional.h"
 #include "UObject/ExpandingChunkedList.h"
@@ -6729,7 +6730,7 @@ static UE::GC::ObjectAROFn GetARO(UClass* Class)
 	UE::GC::ObjectAROFn ARO = Class->CppClassStaticFunctions.GetAddReferencedObjects();
 	check(ARO != nullptr);
 #if WITH_EDITORONLY_DATA
-	return ARO != &UObject::AddReferencedObjects || Class->CanCreateInstanceDataObject() ? ARO : nullptr;
+	return ARO != &UObject::AddReferencedObjects || UE::IsInstanceDataObjectSupportEnabledForClass(Class) ? ARO : nullptr;
 #else
 	return ARO != &UObject::AddReferencedObjects ? ARO : nullptr;
 #endif
diff --git a/Engine/Source/Runtime/CoreUObject/Private/UObject/InstanceDataObjectUtils.cpp b/Engine/Source/Runtime/CoreUObject/Private/UObject/InstanceDataObjectUtils.cpp
--- a/Engine/Source/Runtime/CoreUObject/Private/UObject/InstanceDataObjectUtils.cpp
+++ b/Engine/Source/Runtime/CoreUObject/Private/UObject/InstanceDataObjectUtils.cpp
@@ -249,18 +249,22 @@ namespace UE
 	static const FName NAME_IDOMapKey(ANSITEXTVIEW("Key"));
 	static const FName NAME_IDOMapValue(ANSITEXTVIEW("Value"));
 
+	static void OnInstanceDataObjectSupportChanged(IConsoleVariable*);
+
 	bool bEnableIDOSupport = true;
 	FAutoConsoleVariableRef EnableIDOSupportCVar(
 		TEXT("IDO.Enable"),
 		bEnableIDOSupport,
-		TEXT("Allows an IDO to be created for an object if its class has support.")
+		TEXT("Allows an IDO to be created for an object if its class has support."),
+		FConsoleVariableDelegate::CreateStatic(OnInstanceDataObjectSupportChanged)
 	);
 
 	bool bEnableIDOSupportOnEveryObject = false;
 	FAutoConsoleVariableRef EnableIDOSupportOnEveryObjectCVar(
 		TEXT("IDO.EnableOnEveryObject"),
 		bEnableIDOSupportOnEveryObject,
-		TEXT("Allows an IDO to be created for every object.")
+		TEXT("Allows an IDO to be created for every object."),
+		FConsoleVariableDelegate::CreateStatic(OnInstanceDataObjectSupportChanged)
 	);
 
 	FString ExcludedLoosePropertyTypesVar = TEXT("VerseFunctionProperty");
@@ -301,11 +305,6 @@ namespace UE
 			return false;
 		}
 
-		if (bEnableIDOSupportOnEveryObject)
-		{
-			return true;
-		}
-
 		if (IsInstanceDataObject(InObject))
 		{
 			return true;
@@ -326,7 +325,7 @@ namespace UE
 
 		//@todo FH: change to check trait when available or use config object
 		const UClass* ObjClass = InObject->GetClass();
-		if (!ObjClass->CanCreateInstanceDataObject())
+		if (!IsInstanceDataObjectSupportEnabledForClass(ObjClass))
 		{
 			return false;
 		}
@@ -362,6 +361,23 @@ namespace UE
 		return true;
 	}
 
+	bool IsInstanceDataObjectSupportEnabledForClass(const UClass* Class)
+	{
+		return bEnableIDOSupport && (bEnableIDOSupportOnEveryObject || Class->CanCreateInstanceDataObject());
+	}
+
+	static void OnInstanceDataObjectSupportChanged(IConsoleVariable*)
+	{
+		// The reference token stream is dependent on the return value of IsInstanceDataObjectSupportEnabledForClass.
+		TArray<UClass*> AllClasses;
+		AllClasses.Add(UObject::StaticClass());
+		GetDerivedClasses(UObject::StaticClass(), AllClasses);
+		for (UClass* Class : AllClasses)
+		{
+			Class->AssembleReferenceTokenStream(/*bForce*/ true);
+		}
+	}
+
 	bool CanCreatePropertyBagPlaceholderTypeForImportClass(const UClass* ImportClass)
 	{
 		// @todo - Expand to other import types (e.g. prefab BPs) later; for now restricted to Verse class objects only.
