[cook] add a multicast delegate that object types can register to log additional information when package cooks are blocked from saving.

diff --git a/Engine/Source/Editor/UnrealEd/Private/CookOnTheFlyServer.cpp b/Engine/Source/Editor/UnrealEd/Private/CookOnTheFlyServer.cpp
--- a/Engine/Source/Editor/UnrealEd/Private/CookOnTheFlyServer.cpp
+++ b/Engine/Source/Editor/UnrealEd/Private/CookOnTheFlyServer.cpp
@@ -1748,6 +1748,16 @@ void UCookOnTheFlyServer::SetSaveBusy(bool bInBusy)
 						break;
 					}
 					UE_LOG(LogCook, Display, TEXT("    %s"), *Object->GetFullName());
+
+					TStringBuilder<2048> AdditionalDebugInfo;
+					FDelegates::PackageBlocked.Broadcast(Object, AdditionalDebugInfo);
+					FStringView AdditionalDebugInfoStr = AdditionalDebugInfo.ToView();
+					UE::String::ParseTokens(AdditionalDebugInfo.ToView(), TEXT('\n'),
+						[](FStringView Line)
+						{
+							UE_LOG(LogCook, Display, TEXT("        %.*s"), Line.Len(), Line.GetData());
+						}, UE::String::EParseTokensOptions::SkipEmpty);
+
 					++DisplayCount;
 				}
 			}
diff --git a/Engine/Source/Runtime/CoreUObject/Private/UObject/ICookInfo.cpp b/Engine/Source/Runtime/CoreUObject/Private/UObject/ICookInfo.cpp
--- a/Engine/Source/Runtime/CoreUObject/Private/UObject/ICookInfo.cpp
+++ b/Engine/Source/Runtime/CoreUObject/Private/UObject/ICookInfo.cpp
@@ -87,6 +87,7 @@ PRAGMA_ENABLE_DEPRECATION_WARNINGS;
 FCookInfoEvent FDelegates::CookStarted;
 FCookInfoEvent FDelegates::CookFinished;
 FValidateSourcePackage FDelegates::ValidateSourcePackage;
+FPackageBlockedEvent FDelegates::PackageBlocked;
 
 const TCHAR* GetReferencedSetFilename()
 {
diff --git a/Engine/Source/Runtime/CoreUObject/Public/UObject/ICookInfo.h b/Engine/Source/Runtime/CoreUObject/Public/UObject/ICookInfo.h
--- a/Engine/Source/Runtime/CoreUObject/Public/UObject/ICookInfo.h
+++ b/Engine/Source/Runtime/CoreUObject/Public/UObject/ICookInfo.h
@@ -208,6 +208,7 @@ public:
 
 DECLARE_MULTICAST_DELEGATE_OneParam(FCookInfoEvent, ICookInfo&);
 DECLARE_DELEGATE_RetVal_TwoParams(EDataValidationResult, FValidateSourcePackage, UPackage* /*Package*/, FDataValidationContext& /*ValidationContext*/);
+DECLARE_MULTICAST_DELEGATE_TwoParams(FPackageBlockedEvent, const UObject* /* Object */, FStringBuilderBase& OutDebugInfo)
 
 /** UE::Cook::FDelegates: callbacks for cook events. */
 struct FDelegates
@@ -223,6 +224,8 @@ public:
 	static COREUOBJECT_API FCookInfoEvent CookFinished;
 	/** Called after the cooker has loaded a requested package, before starting to save the package. */
 	static COREUOBJECT_API FValidateSourcePackage ValidateSourcePackage;
+	/** Called when the given package has been blocked from saving for beyond the configured time threshold; can be registered to append additional debug info in this case. */
+	static COREUOBJECT_API FPackageBlockedEvent PackageBlocked;
 };
 
 
