Fix for assert in AsyncLoading2 caused by copying more flags than neccessary in FObjectPropertyBase::ConstructDefaultObjectValueIfNecessary

#rnx

diff --git a/Engine/Source/Runtime/CoreUObject/Private/UObject/PropertyBaseObject.cpp b/Engine/Source/Runtime/CoreUObject/Private/UObject/PropertyBaseObject.cpp
--- a/Engine/Source/Runtime/CoreUObject/Private/UObject/PropertyBaseObject.cpp
+++ b/Engine/Source/Runtime/CoreUObject/Private/UObject/PropertyBaseObject.cpp
@@ -786,6 +786,10 @@ UObject* FObjectPropertyBase::ConstructDefaultObjectValueIfNecessary(UObject* Ex
 		Outer = GetTransientPackage();
 	}
 
+	// Filter out any flags that do not belong to the RF_Load or RF_PropagateToSubObjects groups. We mostly want to exclude
+	// UObject life-time and internal GC flags.
+	const EObjectFlags FlagsToKeepMask = (RF_Load | RF_PropagateToSubObjects);
+
 	if (ExistingValue)
 	{
 		UClass* ExistingValueClass = ExistingValue->GetClass();
@@ -795,8 +799,10 @@ UObject* FObjectPropertyBase::ConstructDefaultObjectValueIfNecessary(UObject* Ex
 			if (ExistingValue->IsTemplate() && 	// Existing value is a template so we can construct a new value with it as the archetype
 				ExistingValue->GetOuter() != Outer) // Unless the template's Outer is the same as the new Outer in which case the template (ExistingValue) IS the object we can reuse
 			{
+				EObjectFlags NewFlags = ExistingValue->GetFlags() & FlagsToKeepMask;
+
 				// We probably got here because an object value failed to load (missing import class) and the property is left with a template of default subobject
-				NewDefaultObjectValue = NewObject<UObject>(Outer, ExistingValue->GetClass(), NAME_None, ExistingValue->GetFlags(), ExistingValue);
+				NewDefaultObjectValue = NewObject<UObject>(Outer, ExistingValue->GetClass(), NAME_None, NewFlags, ExistingValue);
 			}
 			else
 			{
@@ -836,7 +842,7 @@ UObject* FObjectPropertyBase::ConstructDefaultObjectValueIfNecessary(UObject* Ex
 			bool bCanReplaceObjectWithExistingName = !ObjectWithExistingName || ObjectWithExistingName->GetClass()->IsChildOf(PropertyClass);
 
 			FName NewName = (bCanReplaceObjectWithExistingName) ? ExistingName : NAME_None;
-			EObjectFlags NewFlags = ExistingValue ? ExistingValue->GetFlags() : RF_NoFlags;
+			EObjectFlags NewFlags = ExistingValue ? (ExistingValue->GetFlags() & FlagsToKeepMask) : RF_NoFlags;
 			UObject* Template = Defaults ? GetObjectPropertyValue(Defaults) : nullptr;
 
 			// Existing value did not exist or it could not be used as a template
