RemoteExecutor:

- refactor the ExecuteTransactional system out of RemoteObjectTransfer and into its own thing
- remote object transfer and other subsystems decoupled from the core executor
- subsystem requests are no longer singleton, this lays the groundwork for being able to make progress on multiple requests concurrently

diff --git a/Engine/Source/Runtime/CoreUObject/Private/UObject/RemoteExecutor.cpp b/Engine/Source/Runtime/CoreUObject/Private/UObject/RemoteExecutor.cpp
new file mode 100644
--- /dev/null
+++ b/Engine/Source/Runtime/CoreUObject/Private/UObject/RemoteExecutor.cpp
@@ -0,0 +1,404 @@
+// Copyright Epic Games, Inc. All Rights Reserved.
+
+#include "UObject/RemoteExecutor.h"
+#include "AutoRTFM/AutoRTFM.h"
+#include "UObject/RemoteObject.h"
+
+DEFINE_LOG_CATEGORY(LogRemoteExec);
+
+namespace UE::RemoteExecutor
+{
+TDelegate<void()> TickNetworkDelegate;
+TDelegate<TTuple<FRemoteWorkPriority, TFunction<void(void)>>(FRemoteWorkPriority)> FetchDeferredRPCDelegate;
+TDelegate<TTuple<FRemoteWorkPriority, TFunction<void(void)>>(void)> FetchNextDeferredRPCDelegate;
+TDelegate<void()> WaitForServersToSyncDelegate;
+
+TMulticastDelegate<void()> OnTransactionCompletedDelegate;
+TMulticastDelegate<void()> OnTransactionAbortedDelegate;
+}
+
+FArchive& operator<<(FArchive& Ar, FRemoteTransactionId& Id)
+{
+	Ar << Id.Id;
+	return Ar;
+}
+
+class FRemoteExecutor
+{
+	TArray<FRemoteSubsystemBase*> Subsystems;
+
+	uint32 NextTransactionRequestId = 0;
+
+	TArray<TTuple<FRemoteWorkPriority, TFunction<void(void)>>> PendingWorks;
+	TArray<FRemoteTransactionId> PendingWorksRequestId;
+	TArray<TArray<void*>> PendingWorksRequests;
+
+public:
+	FString DebugLastAbortReason;
+
+	void RegisterSubsystem(FRemoteSubsystemBase* Subsystem)
+	{
+		Subsystems.Add(Subsystem);
+	}
+
+	FRemoteTransactionId GenerateNextTransactionId()
+	{
+		uint32 RawRequestId = NextTransactionRequestId++;
+		return FRemoteTransactionId((RawRequestId++ % 0x800000u) + 1000u);
+	}
+
+	void EnqueueWork(FRemoteWorkPriority InWorkPriority, const TFunction<void(void)>& InWork)
+	{
+		PendingWorks.Add(TTuple<FRemoteWorkPriority, TFunction<void(void)>>(InWorkPriority, InWork));
+		PendingWorksRequestId.Add(FRemoteTransactionId::Invalid());
+		PendingWorksRequests.Add(TArray<void*>());
+	}
+
+	void ExecutePendingWork()
+	{
+		uint32 LocalIterationNumber = 0;
+		double LastStallPrintTime = FPlatformTime::Seconds();
+
+		while (PendingWorks.Num() > 0)
+		{
+			LocalIterationNumber++;
+
+			// first check for higher priority work and insert it at the front
+			// this causes requests received later to be put at the start, because
+			// earlier requests could be dependent on later ones completing
+			for (;;)
+			{
+				TTuple<FRemoteWorkPriority, TFunction<void(void)>> NewWork = UE::RemoteExecutor::FetchDeferredRPCDelegate.Execute(PendingWorks[0].Get<0>());
+				if (NewWork.Get<0>().IsValid())
+				{
+					PendingWorks.Insert(NewWork, 0);
+					PendingWorksRequestId.Insert(FRemoteTransactionId::Invalid(), 0);
+					PendingWorksRequests.Insert(TArray<void*>(), 0);
+
+					UE_LOG(LogRemoteExec, VeryVerbose, TEXT("-- ExecuteTransactional : iteration %u popped a higher priority RPC --"), LocalIterationNumber);
+				}
+				else
+				{
+					break;
+				}
+			}
+
+			TTuple<FRemoteWorkPriority, TFunction<void(void)>> ExecutingWork = PendingWorks[0];
+			FRemoteWorkPriority ExecutingWorkPriority = ExecutingWork.Get<0>();
+			FRemoteTransactionId ExecutingWorkRequestId = PendingWorksRequestId[0];
+
+			UE_LOG(LogRemoteExec, VeryVerbose, TEXT("-- ExecuteTransactional : iteration %u request %s pri %s (%d PendingWorks) --"), LocalIterationNumber, *ExecutingWorkRequestId.ToString(), *ExecutingWorkPriority.ToString(), PendingWorks.Num());
+
+			if (PendingWorks.Num() > 1)
+			{
+				for (int32 PWI = 0; PWI < PendingWorks.Num(); PWI++)
+				{
+					UE_LOG(LogRemoteExec, VeryVerbose, TEXT("\t PendingWork[%d] Request (%s) Pri (%s)"), PWI, *PendingWorksRequestId[PWI].ToString(), *PendingWorks[PWI].Get<0>().ToString());
+				}
+			}
+
+			const bool bIsNewRequest = !ExecutingWorkRequestId.IsValid();
+
+			if (bIsNewRequest)
+			{
+				ExecutingWorkRequestId = GenerateNextTransactionId();
+				PendingWorksRequestId[0] = ExecutingWorkRequestId;
+			}
+
+			// begin requests
+			UE_LOG(LogRemoteExec, VeryVerbose, TEXT("\tactivate request %s"), *ExecutingWorkRequestId.ToString());
+
+			for (int32 SubsystemIndex = 0; SubsystemIndex < Subsystems.Num(); SubsystemIndex++)
+			{
+				FRemoteSubsystemBase* Subsystem = Subsystems[SubsystemIndex];
+
+				if (bIsNewRequest)
+				{
+					UE_LOG(LogRemoteExec, VeryVerbose, TEXT("\tactivate request[%s] %s"), Subsystem->NameForDebug(), *ExecutingWorkRequestId.ToString());
+					Subsystem->CreateRequest(ExecutingWorkRequestId);
+				}
+
+				Subsystem->SetActiveRequest(ExecutingWorkRequestId);
+
+				if (bIsNewRequest)
+				{
+					UE_LOG(LogRemoteExec, VeryVerbose, TEXT("\tactivate request[%s] %s"), Subsystem->NameForDebug(), *ExecutingWorkRequestId.ToString());
+					Subsystem->BeginRequest(ExecutingWorkPriority);
+				}
+			}
+
+			// we might run through this loop after the work succeeds in the case
+			// that the work prefetches objects that it doesn't access
+			bool bWorkSucceeded = false;
+
+			for (;;)
+			{
+				double Now = FPlatformTime::Seconds();
+				bool bPrintStall = false;
+
+				if ((Now - LastStallPrintTime) > 1.0)
+				{
+					LastStallPrintTime = Now;
+					bPrintStall = true;
+				}
+
+				// we need to check to see if there is any higher priority work
+				// that we need to context switch to
+				bool bRequiresContextSwitch = false;
+				for (;;)
+				{
+					TTuple<FRemoteWorkPriority, TFunction<void(void)>> NewWork = UE::RemoteExecutor::FetchDeferredRPCDelegate.Execute(ExecutingWorkPriority);
+					if (NewWork.Get<0>().IsValid())
+					{
+						PendingWorks.Insert(NewWork, 0);
+						PendingWorksRequestId.Insert(FRemoteTransactionId::Invalid(), 0);
+						PendingWorksRequests.Add(TArray<void*>());
+
+						UE_LOG(LogRemoteExec, Verbose, TEXT("\tGOT higher priority work (%s), requires context switch"), *NewWork.Get<0>().ToString());
+
+						bRequiresContextSwitch = true;
+					}
+					else
+					{
+						break;
+					}
+				}
+
+				if (bRequiresContextSwitch)
+				{
+					break;
+				}
+
+				// pump the network
+				UE::RemoteExecutor::TickNetworkDelegate.ExecuteIfBound();
+
+				// tick the subsystems
+				for (int32 SubsystemIndex = 0; SubsystemIndex < Subsystems.Num(); SubsystemIndex++)
+				{
+					FRemoteSubsystemBase* Subsystem = Subsystems[SubsystemIndex];
+					Subsystem->Tick();
+				}
+
+				int32 SubsystemReadyCount = 0;
+				for (int32 SubsystemIndex = 0; SubsystemIndex < Subsystems.Num(); SubsystemIndex++)
+				{
+					FRemoteSubsystemBase* Subsystem = Subsystems[SubsystemIndex];
+					if (Subsystem->AreDependenciesSatisfied())
+					{
+						SubsystemReadyCount++;
+					}
+					else
+					{
+						if (bPrintStall)
+						{
+							UE_LOG(LogRemoteExec, VeryVerbose, TEXT("\tsubsystem %s not ready..."), Subsystem->NameForDebug());
+						}
+					}
+				}
+
+				if (SubsystemReadyCount != Subsystems.Num())
+				{
+					if (bPrintStall)
+					{
+						UE_LOG(LogRemoteExec, VeryVerbose, TEXT("\t%d of %d subsystems not ready..."), Subsystems.Num() - SubsystemReadyCount, Subsystems.Num());
+					}
+					continue;
+				}
+
+				// all of the subsystems are ready, try to perform the work
+				if (!bWorkSucceeded)
+				{
+					DebugLastAbortReason = TEXT("");
+
+					AutoRTFM::ETransactionResult WorkTransactResult = AutoRTFM::Transact([&ExecutingWork]() { ExecutingWork.Get<1>()(); });
+
+					if (WorkTransactResult == AutoRTFM::ETransactionResult::AbortedByRequest)
+					{
+						// the work aborted, it should have inserted a new requested object, retry
+						UE_LOG(LogRemoteExec, Verbose, TEXT("\twork ABORTED (%s)"), *DebugLastAbortReason);
+						UE::RemoteExecutor::OnTransactionAbortedDelegate.Broadcast();
+					}
+					else
+					{
+						UE_LOG(LogRemoteExec, VeryVerbose, TEXT("\twork COMPLETED!"));
+
+						bWorkSucceeded = true;
+						UE::RemoteExecutor::OnTransactionCompletedDelegate.Broadcast();
+
+						// the work succeeded, we're all done!
+						break;
+					}
+				}
+			}
+
+			// if we get here and the work was successfully executed, we're done
+			// otherwise, we broke out of the request because there are higher
+			// priority pending requests that need to be serviced so we need
+			// to loop back to the top
+			if (bWorkSucceeded)
+			{
+				UE_LOG(LogRemoteExec, VeryVerbose, TEXT("-- ExecuteTransactional %s DONE --"), *ExecutingWorkRequestId.ToString());
+			}
+			else
+			{
+				UE_LOG(LogRemoteExec, Verbose, TEXT("-- ExecuteTransactional %s ABANDONED, context switching to higher priority work --"), *ExecutingWorkRequestId.ToString());
+			}
+
+			if (bWorkSucceeded)
+			{
+				UE_LOG(LogRemoteExec, VeryVerbose, TEXT("\tending request %s"), *ExecutingWorkRequestId.ToString());
+
+				for (int32 SubsystemIndex = 0; SubsystemIndex < Subsystems.Num(); SubsystemIndex++)
+				{
+					FRemoteSubsystemBase* Subsystem = Subsystems[SubsystemIndex];
+					Subsystem->EndRequest();
+				}
+
+				for (int32 SubsystemIndex = 0; SubsystemIndex < Subsystems.Num(); SubsystemIndex++)
+				{
+					FRemoteSubsystemBase* Subsystem = Subsystems[SubsystemIndex];
+					Subsystem->DestroyRequest(ExecutingWorkRequestId);
+				}
+
+				// remove this work from the list
+				PendingWorks.RemoveAt(0, EAllowShrinking::No);
+				PendingWorksRequestId.RemoveAt(0, EAllowShrinking::No);
+				PendingWorksRequests.RemoveAt(0, EAllowShrinking::No);
+			}
+
+			if (PendingWorks.Num() == 0)
+			{
+				// if we have no more pending work, then we are done executing what
+				// we came here to execute. but before we start going back to find more
+				// work for ourselves, try to process any pending RPCs by searching
+				// for any priority available to help unblock other servers
+				TTuple<FRemoteWorkPriority, TFunction<void(void)>> NewWork = UE::RemoteExecutor::FetchDeferredRPCDelegate.Execute(FRemoteWorkPriority());
+				if (NewWork.Get<0>().IsValid())
+				{
+					PendingWorks.Add(NewWork);
+					PendingWorksRequestId.Add(FRemoteTransactionId::Invalid());
+					PendingWorksRequests.Add(TArray<void*>());
+
+					UE_LOG(LogRemoteExec, VeryVerbose, TEXT("\tdone with original work, enqueueing lower priority deferred RPC"));
+				}
+			}
+		}
+	}
+};
+
+FRemoteExecutor GRemoteExecutor;
+
+void ExecuteTransactionalInternal(FRemoteWorkPriority InWorkPriority, const TFunctionRef<void(void)>& InWork)
+{
+#if UE_WITH_REMOTE_OBJECT_HANDLE
+	check(!AutoRTFM::IsClosed());
+
+	GRemoteExecutor.EnqueueWork(InWorkPriority, [&InWork]() { InWork(); });
+	GRemoteExecutor.ExecutePendingWork();
+#else
+	InWork();
+#endif
+}
+
+namespace UE::RemoteExecutor
+{
+void RegisterRemoteSubsystem(FRemoteSubsystemBase* Subsystem)
+{
+	GRemoteExecutor.RegisterSubsystem(Subsystem);
+}
+
+void WaitForServersToSync()
+{
+	UE_LOG(LogRemoteExec, VeryVerbose, TEXT("-- WaitForServersToSync --"));
+	WaitForServersToSyncDelegate.ExecuteIfBound();
+}
+
+void SetAbortReason(FStringView Reason)
+{
+	UE_AUTORTFM_OPEN
+	{
+		GRemoteExecutor.DebugLastAbortReason = Reason;
+	};
+}
+
+FRemoteWorkPriority CreateRootWorkPriority()
+{
+	FRemoteWorkPriority RootWorkPriority = FRemoteWorkPriority::CreateRootWorkPriority(UE::RemoteObject::GetGlobalServerId(), GRemoteExecutor.GenerateNextTransactionId());
+	return RootWorkPriority;
+}
+
+void ExecuteTransactionalWithExplicitPriority(FRemoteWorkPriority WorkPriority, const TFunctionRef<void(void)>& Work)
+{
+#if UE_WITH_REMOTE_OBJECT_HANDLE
+	ExecuteTransactionalInternal(WorkPriority, Work);
+#else
+	Work();
+#endif
+}
+
+void ExecuteTransactional(const TFunctionRef<void(void)>& Work)
+{
+	FRemoteWorkPriority RootWorkPriority = FRemoteWorkPriority::CreateRootWorkPriority(UE::RemoteObject::GetGlobalServerId(), GRemoteExecutor.GenerateNextTransactionId());
+	ExecuteTransactionalWithExplicitPriority(RootWorkPriority, Work);
+}
+
+} // namespace UE::RemoteExecutor
+
+FString FRemoteWorkPriority::ToString() const
+{
+	return FString::Format(TEXT("[pri: rsi {0} depth {1} id {2}]"),
+		{ *GetRootServerId().ToString(),
+		GetWorkDepth(),
+		*GetRootWorkTransactionId().ToString() });
+}
+
+FRemoteWorkPriority FRemoteWorkPriority::CreateRootWorkPriority(FRemoteServerId ServerId, FRemoteTransactionId TransactionId)
+{
+	FRemoteWorkPriority Result{};
+
+	uint32 RawServerId = 0;
+	FMemory::Memcpy(&RawServerId, &ServerId, sizeof(FRemoteServerId));
+
+	uint32 RawTransactionId = 0;
+	FMemory::Memcpy(&RawTransactionId, &TransactionId, sizeof(FRemoteTransactionId));
+
+	Result.PackedData |= (static_cast<uint64>(RawServerId) << 32);
+	Result.PackedData |= (0xFFull << 24);
+	Result.PackedData |= (RawTransactionId & 0xFFFFFFu);
+	return Result;
+}
+
+FRemoteWorkPriority FRemoteWorkPriority::CreateDependentWorkPriority() const
+{
+	FRemoteWorkPriority Result = *this;
+	Result.PackedData &= 0xFFFFFFFF00FFFFFFull;
+
+	// work depth must be non-zero else we would wrap back around
+	check(GetWorkDepth() > 0);
+
+	Result.PackedData |= static_cast<uint64>(GetWorkDepth() - 1) << 24;
+	return Result;
+}
+
+// returns if Lhs is higher priority than Rhs
+bool IsHigherPriority(FRemoteWorkPriority Lhs, FRemoteWorkPriority Rhs)
+{
+	return Lhs.PackedData < Rhs.PackedData;
+}
+
+bool IsEqualPriority(FRemoteWorkPriority Lhs, FRemoteWorkPriority Rhs)
+{
+	// they're equal if neither is higher than the other
+	return !IsHigherPriority(Lhs, Rhs) && !IsHigherPriority(Rhs, Lhs);
+}
+
+bool IsHigherOrEqualPriority(FRemoteWorkPriority Lhs, FRemoteWorkPriority Rhs)
+{
+	return IsHigherPriority(Lhs, Rhs) || IsEqualPriority(Lhs, Rhs);
+}
+
+FArchive& operator<<(FArchive& Ar, FRemoteWorkPriority& Priority)
+{
+	Ar << Priority.PackedData;
+	return Ar;
+}
diff --git a/Engine/Source/Runtime/CoreUObject/Private/UObject/RemoteObject.cpp b/Engine/Source/Runtime/CoreUObject/Private/UObject/RemoteObject.cpp
--- a/Engine/Source/Runtime/CoreUObject/Private/UObject/RemoteObject.cpp
+++ b/Engine/Source/Runtime/CoreUObject/Private/UObject/RemoteObject.cpp
@@ -1,6 +1,7 @@
 // Copyright Epic Games, Inc. All Rights Reserved.
 
 #include "UObject/RemoteObject.h"
+#include "UObject/RemoteExecutor.h"
 #include "UObject/ObjectPtr.h"
 #include "UObject/Class.h"
 #include "UObject/GCObject.h"
@@ -18,7 +19,6 @@
 #include <atomic>
 
 DEFINE_LOG_CATEGORY(LogRemoteObject);
-DEFINE_LOG_CATEGORY(LogRemoteObjectD);
 
 namespace UE::RemoteObject::Private
 {
@@ -98,6 +98,8 @@ void InitRemoteObjects()
 	RemoteServerId = FRemoteServerId(ServerIdNumber);
 	checkf(RemoteServerId != UE::RemoteObject::Transfer::DatabaseId, TEXT("Remote Server Id=%s is reserved for database"), *RemoteServerId.ToString());
 
+	UE::RemoteObject::Transfer::InitRemoteObjectTransfer();
+
 	if (!UE::RemoteObject::Transfer::RemoteObjectTransferDelegate.IsBound())
 	{
 		UE::RemoteObject::Transfer::RemoteObjectTransferDelegate.BindStatic(&UE::RemoteObject::Serialization::Disk::SaveObjectToDisk);
@@ -114,16 +116,13 @@ void InitRemoteObjects()
 	{
 		UE::RemoteObject::Transfer::RestoreRemoteObjectDataDelegate.BindStatic(&UE::RemoteObject::Serialization::Disk::LoadObjectFromDisk);
 	}
-	if (!UE::RemoteObject::Transfer::FetchDeferredRPCDelegate.IsBound())
-	{
-		UE::RemoteObject::Transfer::FetchDeferredRPCDelegate.BindStatic([](FRemoteWorkPriority) { return TTuple<FRemoteWorkPriority, TFunction<void(void)>>(); });
-	}
-	if (!UE::RemoteObject::Transfer::PhysicsQueryTickDelegate.IsBound())
+
+	if (!UE::RemoteExecutor::FetchDeferredRPCDelegate.IsBound())
 	{
-		UE::RemoteObject::Transfer::PhysicsQueryTickDelegate.BindStatic([](){ return true; });
+		UE::RemoteExecutor::FetchDeferredRPCDelegate.BindStatic([](FRemoteWorkPriority) { return TTuple<FRemoteWorkPriority, TFunction<void(void)>>(); });
 	}
 
-	FCoreDelegates::OnEndFrame.AddStatic(&UE::RemoteObject::Transfer::WaitForServersToSync);
+	FCoreDelegates::OnEndFrame.AddStatic(&UE::RemoteExecutor::WaitForServersToSync);
 }
 
 void RegisterRemoteObjectId(FRemoteObjectId ObjectId, FRemoteServerId ResidentServerId)
diff --git a/Engine/Source/Runtime/CoreUObject/Private/UObject/RemoteObjectPrivate.h b/Engine/Source/Runtime/CoreUObject/Private/UObject/RemoteObjectPrivate.h
--- a/Engine/Source/Runtime/CoreUObject/Private/UObject/RemoteObjectPrivate.h
+++ b/Engine/Source/Runtime/CoreUObject/Private/UObject/RemoteObjectPrivate.h
@@ -9,7 +9,6 @@
 #include "UObject/Object.h"
 
 DECLARE_LOG_CATEGORY_EXTERN(LogRemoteObject, Display, All);
-DECLARE_LOG_CATEGORY_EXTERN(LogRemoteObjectD, Display, All);
 
 namespace UE::RemoteObject::Private
 {
diff --git a/Engine/Source/Runtime/CoreUObject/Private/UObject/RemoteObjectSerialization.cpp b/Engine/Source/Runtime/CoreUObject/Private/UObject/RemoteObjectSerialization.cpp
--- a/Engine/Source/Runtime/CoreUObject/Private/UObject/RemoteObjectSerialization.cpp
+++ b/Engine/Source/Runtime/CoreUObject/Private/UObject/RemoteObjectSerialization.cpp
@@ -70,7 +70,7 @@ FString GenerateRemoteObjectFilename(FRemoteObjectId ObjectId, FRemoteServerId D
 	return FPaths::Combine(*FPaths::ProjectSavedDir(), *(FString::Printf(TEXT("%s_%s.remote"), *ObjectId.ToString(), *DestinationServerId.ToString())));
 }
 
-void LoadObjectFromDisk(uint32 RequestToken, FRemoteWorkPriority RequestPriority, FRemoteObjectId ObjectId, FRemoteServerId CurrentOwnerServerId, FRemoteServerId DestinationServerId)
+void LoadObjectFromDisk(FRemoteTransactionId RequestId, FRemoteWorkPriority RequestPriority, FRemoteObjectId ObjectId, FRemoteServerId CurrentOwnerServerId, FRemoteServerId DestinationServerId)
 {
 	const FString Filename(GenerateRemoteObjectFilename(ObjectId, CurrentOwnerServerId));
 	TUniquePtr<FArchive> FileReader{ IFileManager::Get().CreateFileReader(*Filename) };
@@ -82,10 +82,10 @@ void LoadObjectFromDisk(uint32 RequestToken, FRemoteWorkPriority RequestPriority
 	FileReader->Close();
 	IFileManager::Get().Delete(*Filename, false, true, true);
 
-	UE::RemoteObject::Transfer::OnObjectDataReceived(RequestToken, ObjectId, CurrentOwnerServerId, ObjectData);
+	UE::RemoteObject::Transfer::OnObjectDataReceived(RequestId, ObjectId, CurrentOwnerServerId, ObjectData);
 }
 
-void SaveObjectToDisk(uint32 RequestToken, FRemoteObjectId ObjectId, FRemoteServerId DestinationServerId, FRemoteObjectData& ObjectData)
+void SaveObjectToDisk(FRemoteTransactionId RequestId, FRemoteObjectId ObjectId, FRemoteServerId DestinationServerId, FRemoteObjectData& ObjectData)
 {
 	TUniquePtr<FArchive> FileWriter{ IFileManager::Get().CreateFileWriter(*GenerateRemoteObjectFilename(ObjectId, DestinationServerId)) };
 	checkf(FileWriter.IsValid(), TEXT("Unable to create file writer for remote object %s"), *ObjectId.ToString());
diff --git a/Engine/Source/Runtime/CoreUObject/Private/UObject/RemoteObjectSerialization.h b/Engine/Source/Runtime/CoreUObject/Private/UObject/RemoteObjectSerialization.h
--- a/Engine/Source/Runtime/CoreUObject/Private/UObject/RemoteObjectSerialization.h
+++ b/Engine/Source/Runtime/CoreUObject/Private/UObject/RemoteObjectSerialization.h
@@ -82,6 +82,6 @@ namespace UE::RemoteObject::Serialization
 
 namespace UE::RemoteObject::Serialization::Disk
 {
-	void LoadObjectFromDisk(uint32 RequestToken, FRemoteWorkPriority RequestPriority, FRemoteObjectId ObjectId, FRemoteServerId CurrentOwnerServerId, FRemoteServerId DestinationServerId);
-	void SaveObjectToDisk(uint32 RequestToken, FRemoteObjectId ObjectId, FRemoteServerId ServerId, FRemoteObjectData& ObjectData);
+	void LoadObjectFromDisk(FRemoteTransactionId RequestId, FRemoteWorkPriority RequestPriority, FRemoteObjectId ObjectId, FRemoteServerId CurrentOwnerServerId, FRemoteServerId DestinationServerId);
+	void SaveObjectToDisk(FRemoteTransactionId RequestId, FRemoteObjectId ObjectId, FRemoteServerId ServerId, FRemoteObjectData& ObjectData);
 }
\ No newline at end of file
diff --git a/Engine/Source/Runtime/CoreUObject/Private/UObject/RemoteObjectTransfer.cpp b/Engine/Source/Runtime/CoreUObject/Private/UObject/RemoteObjectTransfer.cpp
--- a/Engine/Source/Runtime/CoreUObject/Private/UObject/RemoteObjectTransfer.cpp
+++ b/Engine/Source/Runtime/CoreUObject/Private/UObject/RemoteObjectTransfer.cpp
@@ -21,13 +21,9 @@ namespace UE::RemoteObject::Transfer
 {
 const FRemoteServerId DatabaseId(0x3ff);
 
-TDelegate<void(uint32 /*RequestToken*/, FRemoteWorkPriority /*RequestPriority*/, FRemoteObjectId /*ObjectId*/, FRemoteServerId /*CurrentOwnerServerId*/, FRemoteServerId /*DestinationServerId*/)> RequestRemoteObjectDelegate;
-TDelegate<void(uint32 /*RequestToken*/, FRemoteObjectId /*ObjectId*/, FRemoteServerId /*DestinationServerId*/, FRemoteObjectData& /*ObjectData*/)> RemoteObjectTransferDelegate;
-TDelegate<void(uint32 /*RequestToken*/, FRemoteObjectId /*ObjectId*/, FRemoteServerId /*DestinationServerId*/)> RemoteObjectDeniedTransferDelegate;
-TDelegate<void()> TickNetworkDelegate;
-TDelegate<TTuple<FRemoteWorkPriority, TFunction<void(void)>>(FRemoteWorkPriority)> FetchDeferredRPCDelegate;
-TDelegate<TTuple<FRemoteWorkPriority, TFunction<void(void)>>(void)> FetchNextDeferredRPCDelegate;
-TDelegate<void()> WaitForServersToSyncDelegate;
+TDelegate<void(FRemoteTransactionId /*RequestId*/, FRemoteWorkPriority /*RequestPriority*/, FRemoteObjectId /*ObjectId*/, FRemoteServerId /*CurrentOwnerServerId*/, FRemoteServerId /*DestinationServerId*/)> RequestRemoteObjectDelegate;
+TDelegate<void(FRemoteTransactionId /*RequestId*/, FRemoteObjectId /*ObjectId*/, FRemoteServerId /*DestinationServerId*/, FRemoteObjectData& /*ObjectData*/)> RemoteObjectTransferDelegate;
+TDelegate<void(FRemoteTransactionId /*RequestId*/, FRemoteObjectId /*ObjectId*/, FRemoteServerId /*DestinationServerId*/)> RemoteObjectDeniedTransferDelegate;
 
 TMulticastDelegate<void(const FRemoteObjectData& /*ObjectData*/, FRemoteServerId /*RemoteServerId*/)> OnObjectDataSentDelegate;
 TMulticastDelegate<void(const TSet<UObject*>& /*MigratedObjects*/, FRemoteServerId /*RemoteServerId*/)> OnObjectsSentDelegate;
@@ -35,92 +31,135 @@ TMulticastDelegate<void(const FRemoteObjectData& /*ObjectData*/, FRemoteServerId
 TMulticastDelegate<void(FRemoteObjectId /*ObjectId*/, FRemoteServerId /*RemoteServerId*/)> OnObjectDataDeniedDelegate;
 TMulticastDelegate<void(const TArray<UObject*>& /*MigratedObjects*/, FRemoteServerId /*RemoteServerId*/)> OnObjectsReceivedDelegate;
 
-TDelegate<void(uint32 /*RequestToken*/, FRemoteObjectId /*ObjectId*/, FRemoteServerId /*DatabaseId*/, FRemoteObjectData& /*ObjectData*/)> StoreRemoteObjectDataDelegate;
-TDelegate<void(uint32 /*RequestToken*/, FRemoteWorkPriority /*RequestPriority*/, FRemoteObjectId /*ObjectId*/, FRemoteServerId /*DatabaseId*/, FRemoteServerId /*DestinationServerId*/)> RestoreRemoteObjectDataDelegate;
-
-TDelegate<void(FRemoteWorkPriority)> PhysicsQueryBeginRequestDelegate;
-TDelegate<bool()> PhysicsQueryTickDelegate;
-TDelegate<void()> PhysicsQueryEndRequestDelegate;
-
-TMulticastDelegate<void()> OnTransactionCompletedDelegate;
-TMulticastDelegate<void()> OnTransactionAbortedDelegate;
-
-FString DebugLastAbortReason;
+TDelegate<void(FRemoteTransactionId /*RequestId*/, FRemoteObjectId /*ObjectId*/, FRemoteServerId /*DatabaseId*/, FRemoteObjectData& /*ObjectData*/)> StoreRemoteObjectDataDelegate;
+TDelegate<void(FRemoteTransactionId /*RequestId*/, FRemoteWorkPriority /*RequestPriority*/, FRemoteObjectId /*ObjectId*/, FRemoteServerId /*DatabaseId*/, FRemoteServerId /*DestinationServerId*/)> RestoreRemoteObjectDataDelegate;
 
 UObject* MigrateObjectFromObjectData(FRemoteObjectData& ObjectData, FRemoteServerId RemoteServerId, EObjectMigrationFlags ObjectFlags);
 
-void ExecuteTransactionalInternal(FRemoteWorkPriority WorkPriority, const TFunctionRef<void(void)>& Work);
-
-TAtomic<uint32> NextTransactionId;
-
-uint32 GenerateNextTransactionId()
-{
-	return (NextTransactionId++ % 0x800000u) + 1000u;
-}
-
 /**
  * Remote object transfer queue. 
  * Queued requests (both send and receive) are processed on the game thread since some systems (like RPCs) that are used for transfering object data are GT-only
  */
-class FRemoteObjectTransferQueue
+
+class FRemoteObjectRequest
 {
 public:
-	TArray<FRemoteObjectId> PendingObjectRequests;
+	FRemoteTransactionId RequestId;
+	FRemoteWorkPriority Priority;
+
 	TArray<FRemoteObjectId> RequiredObjects;
 	TArray<FRemoteObjectId> RequiredObjectsCanonicalRoot;
-	TAtomic<uint32> ActiveRequestToken;
-	TAtomic<uint32> NextActiveRequestToken;
-	FRemoteWorkPriority ActiveRequestPriority;
 
-	void BeginRequest(FRemoteWorkPriority Priority)
+	FRemoteTransactionId GetRequestId() const { return RequestId; }
+	FRemoteWorkPriority GetPriority() const { return Priority; }
+
+	void BeginRequest(FRemoteWorkPriority InPriority)
 	{
 		check(!AutoRTFM::IsClosed());
+		Priority = InPriority;
+	}
 
-		// we expect to not have an active request token
-		check(ActiveRequestToken.Load() == 0);
+	void EndRequest()
+	{
+		check(!AutoRTFM::IsClosed());
+	}
+};
 
-		// we expect the object lists to be empty
-		check(PendingObjectRequests.Num() == 0);
-		check(RequiredObjects.Num() == 0);
-		check(RequiredObjectsCanonicalRoot.Num() == 0);
+class FRemoteObjectTransferQueue : public FRemoteSubsystem<FRemoteObjectRequest>
+{
+	// these are objects that we have outstanding requests for (across all requests)
+	TArray<FRemoteObjectId> PendingObjectRequests;
 
-		uint32 ExpectedToken = 0u;
-		bool TokenExchangeSuccess = ActiveRequestToken.CompareExchange(ExpectedToken, (NextActiveRequestToken.Load() % 32000) + 1000);
-		check(TokenExchangeSuccess);
+public:
+	const TCHAR* NameForDebug() final { return TEXT("RemoteObjectTransferQueue"); }
 
-		ActiveRequestPriority = Priority;
-		NextActiveRequestToken++;
+	void BeginRequest(FRemoteWorkPriority Priority) final
+	{
+		ActiveRequest->BeginRequest(Priority);
 	}
 
-	void EndRequest()
+	void Tick() final
 	{
-		check(!AutoRTFM::IsClosed());
+		using namespace UE::RemoteObject::Serialization;
+
+		check(ActiveRequest->RequiredObjects.Num() == ActiveRequest->RequiredObjectsCanonicalRoot.Num());
 
-		// we expect to have an active request token
-		uint32 OldToken = ActiveRequestToken.Exchange(0u);
-		check(OldToken != 0);
+		UE_LOG(LogRemoteObject, VeryVerbose, TEXT("\t%d required objs (%d pending requests):"), ActiveRequest->RequiredObjects.Num(), PendingObjectRequests.Num());
+		for (int32 RequiredObjectIndex = 0; RequiredObjectIndex < ActiveRequest->RequiredObjects.Num(); RequiredObjectIndex++)
+		{
+			FRemoteObjectId RequestedObjectId = ActiveRequest->RequiredObjects[RequiredObjectIndex];
+			FRemoteObjectId CanonicalObjectId = ActiveRequest->RequiredObjectsCanonicalRoot[RequiredObjectIndex];
+			UE_LOG(LogRemoteObject, VeryVerbose, TEXT("\t\t%s [root %s]"), *RequestedObjectId.ToString(), *CanonicalObjectId.ToString());
+		}
 
-		// clear the object lists
-		PendingObjectRequests.Reset();
-		RequiredObjects.Reset();
-		RequiredObjectsCanonicalRoot.Reset();
+		// send the requests... slowly for now (to avoid straining the network subsystem)
+		for (int32 RequiredObjectIndex = 0; RequiredObjectIndex < ActiveRequest->RequiredObjects.Num(); RequiredObjectIndex++)
+		{
+			if (PendingObjectRequests.Num() >= 5)
+			{
+				break;
+			}
 
-		// clear the priority
-		ActiveRequestPriority = FRemoteWorkPriority();
+			FRemoteObjectId RequiredObjectId = ActiveRequest->RequiredObjects[RequiredObjectIndex];
+
+			if (Handle::IsRemote(RequiredObjectId) && (PendingObjectRequests.Find(RequiredObjectId) == INDEX_NONE))
+			{
+				// send one request
+				bool bAlreadyHaveObject = UE::RemoteObject::Handle::RequestObject(RequiredObjectId);
+				
+				// bAlreadyHaveObject should be false since we already tested IsRemote()
+				check (!bAlreadyHaveObject);
+
+				UE_LOG(LogRemoteObject, Verbose, TEXT("\tsent request for %s %s (%d pending requests)"), *RequiredObjectId.ToString(), *ActiveRequest->Priority.ToString(), PendingObjectRequests.Num());
+
+				PendingObjectRequests.Add(RequiredObjectId);
+			}
+
+			if (!Handle::IsRemote(RequiredObjectId))
+			{
+				// discover the canonical root object id
+				// todo what about the other requests that might have this?
+				UObject* const RootObject = FindCanonicalRootObjectForSerialization(StaticFindObjectFastInternal(RequiredObjectId));
+				ActiveRequest->RequiredObjectsCanonicalRoot[RequiredObjectIndex] = UE::CoreUObject::Private::FObjectHandleUtils::GetRemoteId(RootObject);
+			}
+		}
 	}
 
-	void RequestRemoteObject(uint32 RequestToken, FRemoteWorkPriority RequestPriority, FRemoteObjectId ObjectId, FRemoteServerId CurrentOwnerServerId, FRemoteServerId DestinationServerId)
+	bool AreDependenciesSatisfied() const final
+	{
+		for (int32 RequiredObjectIndex = 0; RequiredObjectIndex < ActiveRequest->RequiredObjects.Num(); RequiredObjectIndex++)
+		{
+			FRemoteObjectId RequiredObjectId = ActiveRequest->RequiredObjects[RequiredObjectIndex];
+
+			if (Handle::IsRemote(RequiredObjectId))
+			{
+				return false;
+			}
+		}
+
+		UE_LOG(LogRemoteObject, VeryVerbose, TEXT("\tREADY: all %d required objects are local"), ActiveRequest->RequiredObjects.Num());
+
+		return true;
+	}
+
+	void EndRequest() final
+	{
+		ActiveRequest->EndRequest();
+		ActiveRequest = nullptr;
+	}
+
+	void RequestRemoteObject(FRemoteTransactionId RequestId, FRemoteWorkPriority RequestPriority, FRemoteObjectId ObjectId, FRemoteServerId CurrentOwnerServerId, FRemoteServerId DestinationServerId)
 	{
 		check(!AutoRTFM::IsClosed());
 
 		// Request an object from remote server
 		if (CurrentOwnerServerId != DatabaseId)
 		{
-			RequestRemoteObjectDelegate.Execute(RequestToken, RequestPriority, ObjectId, CurrentOwnerServerId, DestinationServerId);
+			RequestRemoteObjectDelegate.Execute(RequestId, RequestPriority, ObjectId, CurrentOwnerServerId, DestinationServerId);
 		}
 		else
 		{
-			RestoreRemoteObjectDataDelegate.Execute(RequestToken, RequestPriority, ObjectId, DatabaseId, GetGlobalServerId());
+			RestoreRemoteObjectDataDelegate.Execute(RequestId, RequestPriority, ObjectId, DatabaseId, GetGlobalServerId());
 		}
 	}
 
@@ -131,16 +170,16 @@ public:
 	* @param ObjectData Object data to send. Send request takes ownership of the data
 	* @return Future that's going to receive the object data
 	*/
-	void SendRemoteObject(uint32 RequestToken, FRemoteObjectId ObjectId, FRemoteServerId DestinationServerId, FRemoteObjectData* ObjectData)
+	void SendRemoteObject(FRemoteTransactionId RequestId, FRemoteObjectId ObjectId, FRemoteServerId DestinationServerId, FRemoteObjectData* ObjectData)
 	{
 		check(!AutoRTFM::IsClosed());
-		RemoteObjectTransferDelegate.Execute(RequestToken, ObjectId, DestinationServerId, *ObjectData);
+		RemoteObjectTransferDelegate.Execute(RequestId, ObjectId, DestinationServerId, *ObjectData);
 	}
 
-	void SendDeniedRemoteObject(uint32 RequestToken, FRemoteObjectId ObjectId, FRemoteServerId DestinationServerId)
+	void SendDeniedRemoteObject(FRemoteTransactionId RequestId, FRemoteObjectId ObjectId, FRemoteServerId DestinationServerId)
 	{
 		check(!AutoRTFM::IsClosed());
-		RemoteObjectDeniedTransferDelegate.Execute(RequestToken, ObjectId, DestinationServerId);
+		RemoteObjectDeniedTransferDelegate.Execute(RequestId, ObjectId, DestinationServerId);
 	}
 
 	/**
@@ -149,45 +188,52 @@ public:
 	* @param Data Object data. Data ownership is transferred to the receive request (if it exists)
 	* @return True if a receive request for the specified object exists (and object data ownership was transferred), false otherwise
 	*/
-	void FulfillReceiveRequest(uint32 RequestToken, FRemoteObjectId ObjectId, FRemoteServerId RemoteServerId, FRemoteObjectData* ObjectData)
+	void FulfillReceiveRequest(FRemoteTransactionId RequestId, FRemoteObjectId ObjectId, FRemoteServerId RemoteServerId, FRemoteObjectData* ObjectData)
 	{
 		using namespace UE::RemoteObject::Serialization;
+
 		check(!AutoRTFM::IsClosed());
 
 		const int32 PendingRequestIndex = PendingObjectRequests.Find(ObjectId);
 		const bool bFoundRequest = (PendingRequestIndex != INDEX_NONE);
 
-		UE_LOG(LogRemoteObject, Verbose, TEXT("     : FulfillReceiveRequest(%u) %s (was requested: %d)"), RequestToken, *ObjectId.ToString(), bFoundRequest);
+		UE_LOG(LogRemoteObject, Verbose, TEXT("     : FulfillReceiveRequest(%s) %s (was requested: %d)"), *RequestId.ToString(), *ObjectId.ToString(), bFoundRequest);
 
-		const bool bAssignedOwnership = (RequestToken == PushMigrationRequestToken);
+		const bool bAssignedOwnership = !RequestId.IsValid();
 		EObjectMigrationFlags MigrationFlags = bAssignedOwnership ? EObjectMigrationFlags::Acquire : EObjectMigrationFlags::Receive;
 		UObject* Result = MigrateObjectFromObjectData(*ObjectData, RemoteServerId, MigrationFlags);
 		checkf(UE::CoreUObject::Private::FObjectHandleUtils::GetRemoteId(Result) == ObjectId, TEXT("PeerId:%s requested object %s migration from PeerId:%s but received %s"),
 			*GetGlobalServerId().ToString(), *ObjectId.ToString(), *RemoteServerId.ToString(), *UE::CoreUObject::Private::FObjectHandleUtils::GetRemoteId(Result).ToString());
 
-		UE_LOG(LogRemoteObject, Verbose, TEXT("     : FulfillReceiveRequest(%u) %s DONE"), RequestToken, *ObjectId.ToString());
-
 		if (bFoundRequest)
 		{
 			PendingObjectRequests.RemoveAt(PendingRequestIndex);
 
-			const int32 RequiredObjectIndex = RequiredObjects.Find(ObjectId);
-			check(RequiredObjectIndex != INDEX_NONE);
-
 			// now that we have the object locally we can discover and update the canonical root object id
 			UObject* RootObject = FindCanonicalRootObjectForSerialization(Result);
-			RequiredObjectsCanonicalRoot[RequiredObjectIndex] = UE::CoreUObject::Private::FObjectHandleUtils::GetRemoteId(RootObject);
+			FRemoteObjectId RootObjectRemoteId = UE::CoreUObject::Private::FObjectHandleUtils::GetRemoteId(RootObject);
+
+			for (int32 RequestIndex = 0; RequestIndex < GetRequestCount(); RequestIndex++)
+			{
+				FRemoteObjectRequest* ExistingRequest = GetRequestByIndex(RequestIndex);
+
+				const int32 RequiredObjectIndex = ExistingRequest->RequiredObjects.Find(ObjectId);
+				if (RequiredObjectIndex != INDEX_NONE)
+				{
+					ExistingRequest->RequiredObjectsCanonicalRoot[RequiredObjectIndex] = RootObjectRemoteId;
+				}
+			}
 		}
 	}
 
-	void DenyReceiveRequest(uint32 RequestToken, FRemoteObjectId ObjectId)
+	void DenyReceiveRequest(FRemoteTransactionId RequestId, FRemoteObjectId ObjectId)
 	{
 		check(!AutoRTFM::IsClosed());
 
 		const int32 PendingRequestIndex = PendingObjectRequests.Find(ObjectId);
 		const bool bFoundRequest = (PendingRequestIndex != INDEX_NONE);
 
-		UE_LOG(LogRemoteObject, Verbose, TEXT("     : DenyReceiveRequest(%u) %s, (was requested: %d)"), RequestToken, *ObjectId.ToString(), bFoundRequest);
+		UE_LOG(LogRemoteObject, Verbose, TEXT("     : DenyReceiveRequest(%s) %s, (was requested: %d)"), *RequestId.ToString(), *ObjectId.ToString(), bFoundRequest);
 
 		if (bFoundRequest)
 		{
@@ -202,8 +248,13 @@ void RequestObjectData(FRemoteObjectId ObjectId, FRemoteServerId CurrentOwnerSer
 {
 	check(!AutoRTFM::IsClosed());
 
-	check(GTransferQueue.ActiveRequestToken.Load() != 0);
-	GTransferQueue.RequestRemoteObject(GTransferQueue.ActiveRequestToken.Load(), GTransferQueue.ActiveRequestPriority.CreateDependentWorkPriority(), ObjectId, CurrentOwnerServerId, DestinationServerId);
+	check(GTransferQueue.ActiveRequest != nullptr);
+	GTransferQueue.RequestRemoteObject(
+		GTransferQueue.ActiveRequest->GetRequestId(),
+		GTransferQueue.ActiveRequest->GetPriority().CreateDependentWorkPriority(),
+		ObjectId,
+		CurrentOwnerServerId,
+		DestinationServerId);
 }
 
 UObject* MigrateObjectFromObjectData(FRemoteObjectData& ObjectData, FRemoteServerId RemoteServerId, EObjectMigrationFlags MigrationFlags)
@@ -221,14 +272,18 @@ UObject* MigrateObjectFromObjectData(FRemoteObjectData& ObjectData, FRemoteServe
 
 	if (ReceivedObjects.Num() > 0)
 	{
-		if (UE_LOG_ACTIVE(LogRemoteObjectD, Verbose))
+		if (UE_LOG_ACTIVE(LogRemoteObject, VeryVerbose))
 		{
 			for(int32 SubObjectIndex = 0; SubObjectIndex < ReceivedObjects.Num(); SubObjectIndex++)
 			{
 				FRemoteObjectId SubObjectId = UE::CoreUObject::Private::FObjectHandleUtils::GetRemoteId(ReceivedObjects[SubObjectIndex]);
-				UE_LOG(LogRemoteObjectD, Verbose, TEXT("         subobject[%2d] : %s %s"), SubObjectIndex, *SubObjectId.ToString(), (SubObjectIndex == RequestedObjectIndex) ? TEXT("<requested>") : TEXT(""));
+				UE_LOG(LogRemoteObject, VeryVerbose, TEXT("         subobject[%2d] : %s %s"),
+					SubObjectIndex,
+					*SubObjectId.ToString(),
+					(SubObjectIndex == RequestedObjectIndex) ? TEXT("<requested>") : TEXT(""));
 			}
-			UE_LOG(LogRemoteObjectD, Verbose, TEXT("         <done listing subobjects>"));
+
+			UE_LOG(LogRemoteObject, VeryVerbose, TEXT("         <done listing subobjects>"));
 		}
 	}
 
@@ -254,22 +309,17 @@ void MigrateObjectFromRemoteServer(FRemoteObjectId ObjectId, FRemoteServerId Cur
 	// ensure the request is added in the Open so after we abort it is preserved
 	bool bAdded = false;
 
-	if (!AutoRTFM::IsClosed())
-	{
-		// if we are not in a transaction (dev/test workflow), begin a
-		// temp request
-		GTransferQueue.BeginRequest(FRemoteWorkPriority());
-	}
+	check(GTransferQueue.ActiveRequest);
 
 	UE_AUTORTFM_OPEN
 	{
-		// it's possible a prefetch already added this to the RequiredObjects list
-		bAdded = (GTransferQueue.RequiredObjects.Find(ObjectId) == INDEX_NONE);
+		// it's possible a prefetch already added this to the RequestedObjects list
+		bAdded = (GTransferQueue.ActiveRequest->RequiredObjects.Find(ObjectId) == INDEX_NONE);
 
 		if (bAdded)
 		{
-			GTransferQueue.RequiredObjects.Add(ObjectId);
-			GTransferQueue.RequiredObjectsCanonicalRoot.Add(ObjectId);
+			GTransferQueue.ActiveRequest->RequiredObjects.Add(ObjectId);
+			GTransferQueue.ActiveRequest->RequiredObjectsCanonicalRoot.Add(ObjectId);
 		}
 
 		UE_LOG(LogRemoteObject, Verbose, TEXT("Aborting transaction, need obj id %s on server %s"),
@@ -285,30 +335,20 @@ void MigrateObjectFromRemoteServer(FRemoteObjectId ObjectId, FRemoteServerId Cur
 		}
 	};
 
-	if (AutoRTFM::IsClosed())
-	{
-		// TODO: DO we want to log when we're pulling over an object we've NEVER known about?!
-		// 
-		// if we are in a transaction, abort so the outer handler can renegotiate object transfers and retry
-		SetAbortReason(TEXT("object migration"));
-		AutoRTFM::AbortTransaction();
-	}
-	else
-	{
-		// if we are not in a transaction, request the object 
-		// (which should immediately get it) and end the request
-		UE::RemoteObject::Handle::RequestObject(ObjectId);
-		GTransferQueue.EndRequest();
-	}
+	// TODO: DO we want to log when we're pulling over an object we've NEVER known about?!
+	// 
+	// abort so the outer handler can renegotiate object transfers and retry
+	UE::RemoteExecutor::SetAbortReason(TEXT("object migration"));
+	AutoRTFM::AbortTransaction();
 }
 
-void MigrateObjectToRemoteServer(uint32 RequestToken, UObject* Object, FRemoteServerId DestinationServerId)
+void MigrateObjectToRemoteServer(FRemoteTransactionId RequestId, UObject* Object, FRemoteServerId DestinationServerId)
 {
-	FRemoteWorkPriority RootWorkPriority = FRemoteWorkPriority::CreateRootWorkPriority(GetGlobalServerId(), GenerateNextTransactionId());
-	MigrateObjectToRemoteServerWithExplicitPriority(RequestToken, RootWorkPriority, Object, DestinationServerId);
+	FRemoteWorkPriority RootWorkPriority = UE::RemoteExecutor::CreateRootWorkPriority();
+	MigrateObjectToRemoteServerWithExplicitPriority(RequestId, RootWorkPriority, Object, DestinationServerId);
 }
 
-void MigrateObjectToRemoteServerWithExplicitPriority(uint32 RequestToken, FRemoteWorkPriority RequestPriority, UObject* Object, FRemoteServerId DestinationServerId)
+void MigrateObjectToRemoteServerWithExplicitPriority(FRemoteTransactionId RequestId, FRemoteWorkPriority RequestPriority, UObject* Object, FRemoteServerId DestinationServerId)
 {
 	using namespace UE::RemoteObject::Serialization;
 	using namespace UE::RemoteObject::Private;
@@ -317,75 +357,111 @@ void MigrateObjectToRemoteServerWithExplicitPriority(uint32 RequestToken, FRemot
 
 	FRemoteObjectId ObjectId = UE::CoreUObject::Private::FObjectHandleUtils::GetRemoteId(Object);
 
-	UE_LOG(LogRemoteObject, Verbose, TEXT("     : MigrateObjectToRemoteServer(%u %s) %s is requesting obj %s (%s)"), RequestToken, *RequestPriority.ToString(), *DestinationServerId.ToString(), *ObjectId.ToString(), *Object->GetName());
+	UE_LOG(LogRemoteObject, Verbose, TEXT("     : MigrateObjectToRemoteServer(%s %s) %s is requesting obj %s (%s)"),
+		*RequestId.ToString(), 
+		*RequestPriority.ToString(), 
+		*DestinationServerId.ToString(), 
+		*ObjectId.ToString(),
+		*Object->GetName());
 
 	// find the canonical root object id to use for arbitration
 	UObject* const RootObject = FindCanonicalRootObjectForSerialization(Object);
 	const FRemoteObjectId CanonicalRootObjectId = UE::CoreUObject::Private::FObjectHandleUtils::GetRemoteId(RootObject);
 
 	// is the object locked?
-	bool bIsObjectLocked = false;
+	const FRemoteObjectRequest* HigherPriorityRequest = nullptr;
+	const FRemoteObjectRequest* EqualPriorityRequest = nullptr;
+	const FRemoteObjectRequest* LowerPriorityRequest = nullptr;
+
+	UE_LOG(LogRemoteObject, Verbose, TEXT("     : MigrateObjectToRemoteServer(%s) sending obj %s [%s] (isremote %d)"), *RequestId.ToString(), *ObjectId.ToString(), *Object->GetName(), Handle::IsRemote(Object));
 
-	for (int32 RequiredObjectIndex = 0; RequiredObjectIndex < GTransferQueue.RequiredObjectsCanonicalRoot.Num(); RequiredObjectIndex++)
+	for (int32 RequestIndex = 0; RequestIndex < GTransferQueue.GetRequestCount(); RequestIndex++)
 	{
-		FRemoteObjectId RequiredObjectId = GTransferQueue.RequiredObjectsCanonicalRoot[RequiredObjectIndex];
+		FRemoteObjectRequest* ExistingRequest = GTransferQueue.GetRequestByIndex(RequestIndex);
 
-		if (RequiredObjectId == CanonicalRootObjectId)
+		for (int32 RequiredObjectIndex = 0; RequiredObjectIndex < ExistingRequest->RequiredObjects.Num(); RequiredObjectIndex++)
 		{
-			bool bIsHigherPriority = IsHigherPriority(GTransferQueue.ActiveRequestPriority, RequestPriority);
-			bool bIsLowerPriority = IsHigherPriority(RequestPriority, GTransferQueue.ActiveRequestPriority);
-			bool bIsEqualPriority = IsEqualPriority(GTransferQueue.ActiveRequestPriority, RequestPriority);
+			FRemoteObjectId RequiredObjectId = ExistingRequest->RequiredObjects[RequiredObjectIndex];
 
-			// we keep our object locked either if the work priority is explicitly higher, or otherwise
-			// as a last measure of arbitration we use the server ids when the priorities are equal
-			if (bIsHigherPriority)
-			{
-				// we keep our object locked if our priority is higher
-				UE_LOG(LogRemoteObject, Verbose, TEXT("     : MigrateObjectToRemoteServer(%u) DENYING request for our locked %s because our priority is higher: %s vs %s"), RequestToken, *ObjectId.ToString(), *GTransferQueue.ActiveRequestPriority.ToString(), *RequestPriority.ToString());
-				bIsObjectLocked = true;
-			}
-			else if(bIsEqualPriority && (GetGlobalServerId() < DestinationServerId))
-			{
-				// if the priorities are equal, we arbitrate based on server id
-				UE_LOG(LogRemoteObject, Verbose, TEXT("     : MigrateObjectToRemoteServer(%u) DENYING request for our locked %s because priority %s is equal, and our server index is lower: %s vs %s"), RequestToken, *ObjectId.ToString(), *GTransferQueue.ActiveRequestPriority.ToString(), *GetGlobalServerId().ToString(), *DestinationServerId.ToString());
-				bIsObjectLocked = true;
-			}
-			else if (bIsLowerPriority)
+			if (RequiredObjectId == ObjectId)
 			{
-				UE_LOG(LogRemoteObject, Verbose, TEXT("     : MigrateObjectToRemoteServer(%u) GRANTING request for our locked %s because our priority is lower: %s vs %s"), RequestToken, *ObjectId.ToString(), *GTransferQueue.ActiveRequestPriority.ToString(), *RequestPriority.ToString());
-			}
-			else
-			{
-				UE_LOG(LogRemoteObject, Verbose, TEXT("     : MigrateObjectToRemoteServer(%u) GRANTING request for our locked %s because priority %s is equal, and our server index is higher: %s vs %s"), RequestToken, *ObjectId.ToString(), *GTransferQueue.ActiveRequestPriority.ToString(), *GetGlobalServerId().ToString(), *DestinationServerId.ToString());
+				// do we give it up, or deny their request?
+				bool bIsHigherPriority = IsHigherPriority(ExistingRequest->Priority, RequestPriority);
+				bool bIsEqualPriority = IsEqualPriority(ExistingRequest->Priority, RequestPriority);
+				bool bIsLowerPriority = IsHigherPriority(RequestPriority, ExistingRequest->Priority);
+
+				if (bIsHigherPriority)
+				{
+					HigherPriorityRequest = ExistingRequest;
+					
+					// if we ever find a higher priority request, we can stop looking
+					break;
+				}
+				else if (bIsEqualPriority)
+				{
+					EqualPriorityRequest = ExistingRequest;
+				}
+				else if (bIsLowerPriority)
+				{
+					LowerPriorityRequest = ExistingRequest;
+				}
 			}
+		}
+	}
 
-			break;
+	bool bIsObjectLocked = false;
+
+	if (HigherPriorityRequest != nullptr)
+	{
+		// we keep our object locked if our priority is higher
+		UE_LOG(LogRemoteObject, Verbose, TEXT("     : MigrateObjectToRemoteServer(%s) DENYING request for our locked %s because our priority is higher: %s vs %s"), *RequestId.ToString(), *ObjectId.ToString(), *HigherPriorityRequest->Priority.ToString(), *RequestPriority.ToString());
+		bIsObjectLocked = true;
+	}
+	else if (EqualPriorityRequest)
+	{
+		if (GetGlobalServerId() < DestinationServerId)
+		{
+			UE_LOG(LogRemoteObject, Verbose, TEXT("     : MigrateObjectToRemoteServer(%s) DENYING request for our locked %s because priority %s is equal, and our server index is lower: %s vs %s"), *RequestId.ToString(), *ObjectId.ToString(), *EqualPriorityRequest->Priority.ToString(), *GetGlobalServerId().ToString(), *DestinationServerId.ToString());
+			bIsObjectLocked = true;
+		}
+		else
+		{
+			UE_LOG(LogRemoteObject, Verbose, TEXT("     : MigrateObjectToRemoteServer(%s) GRANTING request for our locked %s because priority %s is equal, and our server index is higher: %s vs %s"), *RequestId.ToString(), *ObjectId.ToString(), *EqualPriorityRequest->Priority.ToString(), *GetGlobalServerId().ToString(), *DestinationServerId.ToString());
 		}
 	}
+	else if (LowerPriorityRequest)
+	{
+		UE_LOG(LogRemoteObject, Verbose, TEXT("     : MigrateObjectToRemoteServer(%s) GRANTING request for our locked %s because our priority is lower: %s vs %s"), *RequestId.ToString(), *ObjectId.ToString(), *LowerPriorityRequest->Priority.ToString(), *RequestPriority.ToString());
+	}
+	else
+	{
+		UE_LOG(LogRemoteObject, Verbose, TEXT("     : MigrateObjectToRemoteServer(%s) GRANTING request for %s because none of our active requests need it"), *RequestId.ToString(), *ObjectId.ToString());
+	}
 
 	if (bIsObjectLocked)
 	{
-		// we are higher priority, send a deny response
-		RemoteObjectDeniedTransferDelegate.Execute(RequestToken, ObjectId, DestinationServerId);
+		RemoteObjectDeniedTransferDelegate.Execute(RequestId, ObjectId, DestinationServerId);
 	}
 	else
 	{
+		UE_LOG(LogRemoteObject, VeryVerbose, TEXT("     : MigrateObjectToRemoteServer(%s) serializing obj %s (isremote %d)"), *RequestId.ToString(), *ObjectId.ToString(), Handle::IsRemote(Object));
+
 		// send the object
 		TSet<UObject*> SentObjects;
 		FRemoteObjectData* ObjectData = SerializeObjectData(Object, SentObjects);
 
-		UE_LOG(LogRemoteObject, VeryVerbose, TEXT("     : MigrateObjectToRemoteServer(%u) serialization complete obj %s (isremote %d)"), RequestToken, *ObjectId.ToString(), Handle::IsRemote(Object));
+		UE_LOG(LogRemoteObject, VeryVerbose, TEXT("     : MigrateObjectToRemoteServer(%s) serialization complete obj %s (isremote %d)"), *RequestId.ToString(), *ObjectId.ToString(), Handle::IsRemote(Object));
 
 		for(UObject* SO : SentObjects)
 		{
 			FRemoteObjectId SubObjectId = UE::CoreUObject::Private::FObjectHandleUtils::GetRemoteId(SO);
-			UE_LOG(LogRemoteObject, Verbose, TEXT("         sent subobject : %s [%s]"), *SubObjectId.ToString(), *GetNameSafe(SO));
+			UE_LOG(LogRemoteObject, VeryVerbose, TEXT("         sent subobject : %s [%s]"), *SubObjectId.ToString(), *GetNameSafe(SO));
 		}
 
 		// Notify that object data has been migrated (actually will be but SendObjectData does a MoveTemp on the data)
 		OnObjectDataSentDelegate.Broadcast(*ObjectData, DestinationServerId);
 
-		GTransferQueue.SendRemoteObject(RequestToken, ObjectId, DestinationServerId, ObjectData);
+		GTransferQueue.SendRemoteObject(RequestId, ObjectId, DestinationServerId, ObjectData);
 
 		checkf(SentObjects.Num() > 0, TEXT("PeerId:%s unable to migrate object %s to PeerId:%s"), *GetGlobalServerId().ToString(), *UE::CoreUObject::Private::FObjectHandleUtils::GetRemoteId(Object).ToString(), *DestinationServerId.ToString());
 
@@ -398,7 +474,7 @@ void MigrateObjectToRemoteServerWithExplicitPriority(uint32 RequestToken, FRemot
 			OnObjectsSentDelegate.Broadcast(SentObjects, DestinationServerId);
 
 			// Are we loaning it out or changing ownership?
-			const bool bOwnershipChange = (RequestToken == PushMigrationRequestToken);
+			const bool bOwnershipChange = !RequestId.IsValid();
 			FPostMigrateParams PostMigrateParams;
 			PostMigrateParams.RequestingServerId = RequestPriority.GetRootServerId();
 			PostMigrateParams.DestinationServerId = DestinationServerId;
@@ -424,25 +500,25 @@ void RegisterRemoteObjectId(FRemoteObjectId Id, FRemoteServerId OwnerServerId)
 	UE::RemoteObject::Private::RegisterRemoteObjectId(Id, OwnerServerId);
 }
 
-void MigrateObjectToRemoteServer(uint32 RequestToken, FRemoteObjectId ObjectId, FRemoteServerId DestinationServerId)
+void MigrateObjectToRemoteServer(FRemoteTransactionId RequestId, FRemoteObjectId ObjectId, FRemoteServerId DestinationServerId)
 {
-	FRemoteWorkPriority RootWorkPriority = FRemoteWorkPriority::CreateRootWorkPriority(GetGlobalServerId(), GenerateNextTransactionId());
-	MigrateObjectToRemoteServerWithExplicitPriority(RequestToken, RootWorkPriority, ObjectId, DestinationServerId);
+	FRemoteWorkPriority RootWorkPriority = UE::RemoteExecutor::CreateRootWorkPriority();
+	MigrateObjectToRemoteServerWithExplicitPriority(RequestId, RootWorkPriority, ObjectId, DestinationServerId);
 }
 
-void MigrateObjectToRemoteServerWithExplicitPriority(uint32 RequestToken, FRemoteWorkPriority RequestPriority, FRemoteObjectId ObjectId, FRemoteServerId DestinationServerId)
+void MigrateObjectToRemoteServerWithExplicitPriority(FRemoteTransactionId RequestId, FRemoteWorkPriority RequestPriority, FRemoteObjectId ObjectId, FRemoteServerId DestinationServerId)
 {
 	using namespace UE::RemoteObject::Handle;
 	using namespace UE::RemoteObject::Serialization;
 	using namespace UE::RemoteObject::Private;
 
-	UE_LOG(LogRemoteObject, Verbose, TEXT("MigrateObjectToRemoteServer(%u) obj %s to %s"), RequestToken, *ObjectId.ToString(), *DestinationServerId.ToString());
+	UE_LOG(LogRemoteObject, Verbose, TEXT("MigrateObjectToRemoteServer(%s) obj %s to %s"), *RequestId.ToString(), *ObjectId.ToString(), *DestinationServerId.ToString());
 
 	UObject* Object = StaticFindObjectFastInternal(ObjectId);
 	if (Object && !IsRemote(Object))
 	{
-		UE_LOG(LogRemoteObject, VeryVerbose, TEXT("    : MigrateObjectToRemoteServer(%u) found obj %s"), RequestToken, *ObjectId.ToString());
-		MigrateObjectToRemoteServerWithExplicitPriority(RequestToken, RequestPriority, Object, DestinationServerId);
+		UE_LOG(LogRemoteObject, VeryVerbose, TEXT("    : MigrateObjectToRemoteServer(%s) found obj %s"), *RequestId.ToString(), *ObjectId.ToString());
+		MigrateObjectToRemoteServerWithExplicitPriority(RequestId, RequestPriority, Object, DestinationServerId);
 	}
 	else
 	{
@@ -450,11 +526,11 @@ void MigrateObjectToRemoteServerWithExplicitPriority(uint32 RequestToken, FRemot
 
 		if (RemoteObjectStub)
 		{
-			UE_LOG(LogRemoteObject, Verbose, TEXT("    : MigrateObjectToRemoteServer(%u) did not find obj %s (we do have a stub)"), RequestToken,*ObjectId.ToString());
+			UE_LOG(LogRemoteObject, Verbose, TEXT("    : MigrateObjectToRemoteServer(%s) did not find obj %s (we do have a stub)"), *RequestId.ToString(),*ObjectId.ToString());
 		}
 		else
 		{
-			UE_LOG(LogRemoteObject, Verbose, TEXT("    : MigrateObjectToRemoteServer(%u) did not find obj %s (also no stub)"), RequestToken,*ObjectId.ToString());
+			UE_LOG(LogRemoteObject, Verbose, TEXT("    : MigrateObjectToRemoteServer(%s) did not find obj %s (also no stub)"), *RequestId.ToString(),*ObjectId.ToString());
 		}
 
 		FRemoteServerId CurrentResidentServerId = RemoteObjectStub ? RemoteObjectStub->ResidentServerId : ObjectId.GetServerId();
@@ -463,27 +539,27 @@ void MigrateObjectToRemoteServerWithExplicitPriority(uint32 RequestToken, FRemot
 
 		if (CurrentResidentServerId != DestinationServerId)
 		{
-			UE_LOG(LogRemoteObject, Verbose, TEXT("    : MigrateObjectToRemoteServer(%u) forwarding request (we believe %s owns the Object)"), RequestToken, *CurrentResidentServerId.ToString());
-			GTransferQueue.RequestRemoteObject(RequestToken, RequestPriority, ObjectId, CurrentResidentServerId, DestinationServerId);
+			UE_LOG(LogRemoteObject, Verbose, TEXT("    : MigrateObjectToRemoteServer(%s) forwarding request (we believe %s owns the Object)"), *RequestId.ToString(), *CurrentResidentServerId.ToString());
+			GTransferQueue.RequestRemoteObject(RequestId, RequestPriority, ObjectId, CurrentResidentServerId, DestinationServerId);
 		}
 		else
 		{
-			UE_LOG(LogRemoteObject, Verbose, TEXT("    : MigrateObjectToRemoteServer(%u) denying request %s (we believe %s already has the Object)"), RequestToken, *ObjectId.ToString(), *CurrentResidentServerId.ToString());
-			RemoteObjectDeniedTransferDelegate.Execute(RequestToken, ObjectId, DestinationServerId);
+			UE_LOG(LogRemoteObject, Verbose, TEXT("    : MigrateObjectToRemoteServer(%s) denying request %s (we believe %s already has the Object)"), *RequestId.ToString(), *ObjectId.ToString(), *CurrentResidentServerId.ToString());
+			RemoteObjectDeniedTransferDelegate.Execute(RequestId, ObjectId, DestinationServerId);
 		}
 	}
 }
 
-void OnObjectDataReceived(uint32 RequestToken, FRemoteObjectId ObjectId, FRemoteServerId RemoteServerId, FRemoteObjectData* Data)
+void OnObjectDataReceived(FRemoteTransactionId RequestId, FRemoteObjectId ObjectId, FRemoteServerId RemoteServerId, FRemoteObjectData* Data)
 {
 	check(!AutoRTFM::IsClosed());
-	GTransferQueue.FulfillReceiveRequest(RequestToken, ObjectId, RemoteServerId, Data);
+	GTransferQueue.FulfillReceiveRequest(RequestId, ObjectId, RemoteServerId, Data);
 }
 
-void OnObjectDataDenied(uint32 RequestToken, FRemoteObjectId ObjectId, FRemoteServerId RemoteServerId)
+void OnObjectDataDenied(FRemoteTransactionId RequestId, FRemoteObjectId ObjectId, FRemoteServerId RemoteServerId)
 {
 	check(!AutoRTFM::IsClosed());
-	GTransferQueue.DenyReceiveRequest(RequestToken, ObjectId);
+	GTransferQueue.DenyReceiveRequest(RequestId, ObjectId);
 }
 
 void PrefetchRemoteObject(FRemoteObjectId ObjectId, UObject* DestinationOuter)
@@ -500,261 +576,20 @@ void PrefetchRemoteObject(FRemoteObjectId ObjectId, UObject* DestinationOuter)
 	UE_AUTORTFM_OPEN
 	{
 		// it's possible a previous prefetch or strong fetch already added this to the RequiredObjects list
-		bAdded = (GTransferQueue.RequiredObjects.Find(ObjectId) == INDEX_NONE);
+		bAdded = (GTransferQueue.ActiveRequest->RequiredObjects.Find(ObjectId) == INDEX_NONE);
 
 		if (bAdded)
 		{
-			GTransferQueue.RequiredObjects.Add(ObjectId);
-			GTransferQueue.RequiredObjectsCanonicalRoot.Add(ObjectId);
-			UE_LOG(LogRemoteObject, Verbose, TEXT("PrefetchRemoteObject(%u) : obj %s"), GTransferQueue.ActiveRequestToken.Load(), *ObjectId.ToString());
+			GTransferQueue.ActiveRequest->RequiredObjects.Add(ObjectId);
+			GTransferQueue.ActiveRequest->RequiredObjectsCanonicalRoot.Add(ObjectId);
+			UE_LOG(LogRemoteObject, VeryVerbose, TEXT("PrefetchRemoteObject(%s) : obj %s"), *GTransferQueue.ActiveRequest->GetRequestId().ToString(), *ObjectId.ToString());
 		}
 	};
 }
 
-void ExecuteTransactionalWithExplicitPriority(FRemoteWorkPriority WorkPriority, const TFunctionRef<void(void)>& Work)
-{
-#if UE_WITH_REMOTE_OBJECT_HANDLE
-	ExecuteTransactionalInternal(WorkPriority, Work);
-#else
-	Work();
-#endif
-}
-
-void ExecuteTransactional(const TFunctionRef<void(void)>& Work)
-{
-	FRemoteWorkPriority RootWorkPriority = FRemoteWorkPriority::CreateRootWorkPriority(GetGlobalServerId(), GenerateNextTransactionId());
-	ExecuteTransactionalWithExplicitPriority(RootWorkPriority, Work);
-}
-
-void ExecuteTransactionalInternal(FRemoteWorkPriority InWorkPriority, const TFunctionRef<void(void)>& InWork)
-{
-	using namespace UE::RemoteObject::Serialization;
-	check(!AutoRTFM::IsClosed());
-
-	double LastWaitingPrintTime = FPlatformTime::Seconds();
-
-	TArray<TTuple<FRemoteWorkPriority, TFunction<void(void)>>> PendingWorks;
-	// it's not possible to construct a TFunction from a TFunctionRef, but it's 
-	// safe to do in this context so we just wrap the TFunctionRef in a lambda
-	PendingWorks.Add(TTuple<FRemoteWorkPriority, TFunction<void(void)>>(InWorkPriority, [&InWork](){ InWork(); }));
-
-	uint32 LocalIterationNumber = 0;
-
-	while (PendingWorks.Num() > 0)
-	{
-		LocalIterationNumber++;
-
-		// first check for higher priority work and insert it at the front
-		// this causes requests received later to be put at the start, because
-		// earlier requests could be dependent on later ones completing
-		for(;;)
-		{
-			TTuple<FRemoteWorkPriority, TFunction<void(void)>> NewWork = FetchDeferredRPCDelegate.Execute(PendingWorks[0].Get<0>());
-			if(NewWork.Get<0>().IsValid())
-			{
-				UE_LOG(LogRemoteObjectD, VeryVerbose, TEXT("-- ExecuteTransactional : iteration %u popped a higher priority RPC --"), LocalIterationNumber);
-				PendingWorks.Insert(NewWork, 0);
-			}
-			else
-			{
-				break;
-			}
-		}
-
-		TTuple<FRemoteWorkPriority, TFunction<void(void)>> ExecutingWork = PendingWorks[0];
-		FRemoteWorkPriority ExecutingWorkPriority = ExecutingWork.Get<0>();
-
-		GTransferQueue.BeginRequest(ExecutingWorkPriority);
-
-		PhysicsQueryBeginRequestDelegate.ExecuteIfBound(ExecutingWorkPriority);
-
-		UE_LOG(LogRemoteObjectD, VeryVerbose, TEXT("-- ExecuteTransactional : iteration %u request %u pri %s --"), LocalIterationNumber, GTransferQueue.ActiveRequestToken.Load(), *ExecutingWorkPriority.ToString());
-
-		// we might run through this loop after the work succeeds in the case
-		// that the work prefetches objects that it doesn't access
-		bool bWorkSucceeded = false;
-
-		for (;;)
-		{
-			// we need to check to see if there is any higher priority work
-			// that we need to context switch to
-			bool bRequiresContextSwitch = false;
-			for (;;)
-			{
-				TTuple<FRemoteWorkPriority, TFunction<void(void)>> NewWork = FetchDeferredRPCDelegate.Execute(ExecutingWorkPriority);
-				if (NewWork.Get<0>().IsValid())
-				{
-					PendingWorks.Insert(NewWork, 0);
-					bRequiresContextSwitch = true;
-				}
-				else
-				{
-					break;
-				}
-			}
-
-			if (bRequiresContextSwitch)
-			{
-				break;
-			}
-
-			check(GTransferQueue.RequiredObjects.Num() == GTransferQueue.RequiredObjectsCanonicalRoot.Num());
-
-			UE_LOG(LogRemoteObjectD, Verbose, TEXT("   : %d required objs (%d pending requests):"), GTransferQueue.RequiredObjects.Num(), GTransferQueue.PendingObjectRequests.Num());
-			for(int32 RequiredObjectIndex = 0; RequiredObjectIndex < GTransferQueue.RequiredObjects.Num(); RequiredObjectIndex++)
-			{
-				FRemoteObjectId RequestedObjectId = GTransferQueue.RequiredObjects[RequiredObjectIndex];
-				FRemoteObjectId CanonicalObjectId = GTransferQueue.RequiredObjectsCanonicalRoot[RequiredObjectIndex];
-				UE_LOG(LogRemoteObjectD, Verbose, TEXT("       %s [root %s]"), *RequestedObjectId.ToString(), *CanonicalObjectId.ToString());
-			}
-
-			// pump the network
-			TickNetworkDelegate.ExecuteIfBound();
-
-			// send the requests... slowly for now (to avoid straining the network subsystem)
-			int32 RequestsSent = 0;
-
-			for (int32 RequiredObjectIndex = 0; RequiredObjectIndex < GTransferQueue.RequiredObjects.Num(); RequiredObjectIndex++)
-			{
-				if (GTransferQueue.PendingObjectRequests.Num() >= 5)
-				{
-					break;
-				}
-
-				FRemoteObjectId RequiredObjectId = GTransferQueue.RequiredObjects[RequiredObjectIndex];
-
-				if (Handle::IsRemote(RequiredObjectId) && (GTransferQueue.PendingObjectRequests.Find(RequiredObjectId) == INDEX_NONE))
-				{
-					// send one request
-					bool bAlreadyHaveObject = UE::RemoteObject::Handle::RequestObject(RequiredObjectId);
-
-					// bAlreadyHaveObject should be false since we already tested IsRemote()
-					check(!bAlreadyHaveObject);
-
-					UE_LOG(LogRemoteObjectD, Verbose, TEXT("   : sent request for %s (%s) (%d pending requests)"), *RequiredObjectId.ToString(), *GTransferQueue.ActiveRequestPriority.ToString(), GTransferQueue.PendingObjectRequests.Num());
-					GTransferQueue.PendingObjectRequests.Add(RequiredObjectId);
-				}
-
-				if (!Handle::IsRemote(RequiredObjectId))
-				{
-					// discover the canonical root object id
-					UObject* const RootObject = FindCanonicalRootObjectForSerialization(StaticFindObjectFastInternal(RequiredObjectId));
-					GTransferQueue.RequiredObjectsCanonicalRoot[RequiredObjectIndex] = UE::CoreUObject::Private::FObjectHandleUtils::GetRemoteId(RootObject);
-				}
-			}
-
-			const double TimeNow = FPlatformTime::Seconds();
-
-			if (!PhysicsQueryTickDelegate.Execute())
-			{
-				if (TimeNow > LastWaitingPrintTime + 1.0)
-				{
-					UE_LOG(LogRemoteObjectD, VeryVerbose, TEXT("   : waiting for physics responses..."));
-					LastWaitingPrintTime = TimeNow;
-				}
-
-				// physics doesn't have all responses yet, try again
-				continue;
-			}
-
-			// do we have all required objects local?
-			int32 RemoteObjectCount = 0;
-			for (int32 RequiredObjectIndex = 0; RequiredObjectIndex < GTransferQueue.RequiredObjects.Num(); RequiredObjectIndex++)
-			{
-				if (Handle::IsRemote(GTransferQueue.RequiredObjects[RequiredObjectIndex]))
-				{
-					RemoteObjectCount++;
-				}
-			}
-
-			if(RemoteObjectCount > 0)
-			{
-				if (TimeNow > LastWaitingPrintTime + 1.0)
-				{
-					UE_LOG(LogRemoteObjectD, VeryVerbose, TEXT("   : waiting for %d object responses..."), RemoteObjectCount);
-					LastWaitingPrintTime = TimeNow;
-				}
-
-				// try again!
-				continue;
-			}
-
-			UE_LOG(LogRemoteObjectD, VeryVerbose, TEXT("   : READY: all %d required objects are locked"), GTransferQueue.RequiredObjects.Num());
-
-			// all of the requested objects are locked, try to perform the work
-			if (!bWorkSucceeded)
-			{
-				DebugLastAbortReason = TEXT("");
-
-				AutoRTFM::ETransactionResult WorkTransactResult = AutoRTFM::Transact([&ExecutingWork](){ ExecutingWork.Get<1>()(); });
-
-				if (WorkTransactResult == AutoRTFM::ETransactionResult::AbortedByRequest)
-				{
-					// the work aborted, it should have inserted a new requested object, retry
-					UE_LOG(LogRemoteObjectD, VeryVerbose, TEXT("   : work ABORTED (%s)"), *DebugLastAbortReason);
-					OnTransactionAbortedDelegate.Broadcast();
-				}
-				else
-				{
-					UE_LOG(LogRemoteObjectD, VeryVerbose, TEXT("   : work COMPLETED!"));
-
-					bWorkSucceeded = true;
-					OnTransactionCompletedDelegate.Broadcast();
-					break;
-				}
-			}
-		}
-
-		// if we get here and the work was successfully executed, we're done
-		// otherwise, we broke out of the request because there are higher
-		// priority pending requests that need to be serviced so we need
-		// to loop back to the top
-		if (bWorkSucceeded)
-		{
-			UE_LOG(LogRemoteObjectD, VeryVerbose, TEXT("-- ExecuteTransactional %u DONE --"), GTransferQueue.ActiveRequestToken.Load());
-		}
-		else
-		{
-			UE_LOG(LogRemoteObjectD, VeryVerbose, TEXT("-- ExecuteTransactional %u ABANDONED, context switching to higher priority work --"), GTransferQueue.ActiveRequestToken.Load());
-		}
-
-		GTransferQueue.EndRequest();
-
-		PhysicsQueryEndRequestDelegate.ExecuteIfBound();
-
-		if (bWorkSucceeded)
-		{
-			// remove this work from the list
-			PendingWorks.RemoveAt(0, EAllowShrinking::No);
-		}
-
-		if (PendingWorks.Num() == 0)
-		{
-			// if we have no more pending work, then we are done executing what
-			// we came here to execute. but before we start going back to find more
-			// work for ourselves, try to process any pending RPCs by searching
-			// for any priority available to help unblock other servers
-			TTuple<FRemoteWorkPriority, TFunction<void(void)>> NewWork = FetchDeferredRPCDelegate.Execute(FRemoteWorkPriority());
-			if (NewWork.Get<0>().IsValid())
-			{
-				PendingWorks.Add(NewWork);
-			}
-		}
-	}
-}
-
-void WaitForServersToSync()
+void InitRemoteObjectTransfer()
 {
-	UE_LOG(LogRemoteObjectD, VeryVerbose, TEXT("-- WaitForServersToSync --"));
-	WaitForServersToSyncDelegate.ExecuteIfBound();
-}
-
-void SetAbortReason(FStringView Reason)
-{
-	UE_AUTORTFM_OPEN
-	{
-		DebugLastAbortReason = Reason;
-	};
+	UE::RemoteExecutor::RegisterRemoteSubsystem(&GTransferQueue);
 }
 
 } // namespace UE::RemoteObject::Transfer
@@ -798,7 +633,7 @@ void StoreObjectToDatabase(UObject* Object)
 		TSet<UObject*> SentObjects;
 		{
 			TUniquePtr<FRemoteObjectData> ObjectData(SerializeObjectData(OutermostUnreachableRemoteReference, SentObjects));
-			StoreRemoteObjectDataDelegate.Execute(PushMigrationRequestToken, FRemoteObjectId(OutermostUnreachableRemoteReference), DatabaseId, *ObjectData);
+			StoreRemoteObjectDataDelegate.Execute(PushMigrationRequestId, FRemoteObjectId(OutermostUnreachableRemoteReference), DatabaseId, *ObjectData);
 		}
 
 		checkf(SentObjects.Num() > 0, TEXT("PeerId:%s unable to store object %s to DatabaseId:%s"), *GetGlobalServerId().ToString(), *FRemoteObjectId(Object).ToString(), *DatabaseId.ToString());
@@ -902,61 +737,6 @@ bool FRemoteObjectReference::NetSerialize(FArchive& Ar, class UPackageMap* Map,
 	return bOutSuccess;
 }
 
-FString FRemoteWorkPriority::ToString() const
-{
-	return FString::Format(TEXT("[pri: rsi {0} depth {1} id {2}]"),
-		{*GetRootServerId().ToString(),
-		GetWorkDepth(),
-		GetRootWorkTransactionId()});
-}
-
-FRemoteWorkPriority FRemoteWorkPriority::CreateRootWorkPriority(FRemoteServerId ServerId, uint32 TransactionId)
-{
-	FRemoteWorkPriority Result{};
-	uint32 RemoteServerId = 0;
-	FMemory::Memcpy(&RemoteServerId, &ServerId, sizeof(FRemoteServerId));
-
-	Result.PackedData |= (static_cast<uint64>(RemoteServerId) << 32);
-	Result.PackedData |= (0xFFull << 24);
-	Result.PackedData |= (TransactionId & 0xFFFFFFu);
-	return Result;
-}
-
-FRemoteWorkPriority FRemoteWorkPriority::CreateDependentWorkPriority() const
-{
-	FRemoteWorkPriority Result = *this;
-	Result.PackedData &= 0xFFFFFFFF00FFFFFFull;
-	
-	// work depth must be non-zero else we would wrap back around
-	check(GetWorkDepth() > 0);
-
-	Result.PackedData |= static_cast<uint64>(GetWorkDepth() - 1) << 24;
-	return Result;
-}
-
-// returns if Lhs is higher priority than Rhs
-bool IsHigherPriority(FRemoteWorkPriority Lhs, FRemoteWorkPriority Rhs)
-{
-	return Lhs.PackedData < Rhs.PackedData;
-}
-
-bool IsEqualPriority(FRemoteWorkPriority Lhs, FRemoteWorkPriority Rhs)
-{
-	// they're equal if neither is higher than the other
-	return !IsHigherPriority(Lhs, Rhs) && !IsHigherPriority(Rhs, Lhs);
-}
-
-bool IsHigherOrEqualPriority(FRemoteWorkPriority Lhs, FRemoteWorkPriority Rhs)
-{
-	return IsHigherPriority(Lhs, Rhs) || IsEqualPriority(Lhs, Rhs);
-}
-
-FArchive& operator<<(FArchive& Ar, FRemoteWorkPriority& Priority)
-{
-	Ar << Priority.PackedData;
-	return Ar;
-}
-
 FArchive& operator<<(FArchive& Ar, FRemoteObjectReference& Ref)
 {
 	Ar << Ref.ObjectId;
@@ -968,4 +748,4 @@ FArchive& operator<<(FArchive& Ar, FRemoteObjectReference& Ref)
 	}
 
 	return Ar;
-}
+}
\ No newline at end of file
diff --git a/Engine/Source/Runtime/CoreUObject/Public/UObject/NoExportTypes.h b/Engine/Source/Runtime/CoreUObject/Public/UObject/NoExportTypes.h
--- a/Engine/Source/Runtime/CoreUObject/Public/UObject/NoExportTypes.h
+++ b/Engine/Source/Runtime/CoreUObject/Public/UObject/NoExportTypes.h
@@ -3462,6 +3462,15 @@ struct FRemoteObjectReference
 	FRemoteServerId ServerId;
 };
 
+USTRUCT(noexport, BlueprintType)
+struct FRemoteTransactionId
+{
+	GENERATED_BODY()
+
+	UPROPERTY()
+	uint32 Id;
+};
+
 /// @endcond
 
 #endif
diff --git a/Engine/Source/Runtime/CoreUObject/Public/UObject/RemoteExecutor.h b/Engine/Source/Runtime/CoreUObject/Public/UObject/RemoteExecutor.h
new file mode 100644
--- /dev/null
+++ b/Engine/Source/Runtime/CoreUObject/Public/UObject/RemoteExecutor.h
@@ -0,0 +1,185 @@
+// Copyright Epic Games, Inc. All Rights Reserved.
+
+#pragma once
+
+#include "CoreTypes.h"
+#include "Misc/AssertionMacros.h"
+#include "UObject/RemoteObjectTypes.h"
+#include "UObject/Class.h"
+#include "Async/Async.h"
+
+DECLARE_LOG_CATEGORY_EXTERN(LogRemoteExec, Display, All);
+
+struct FRemoteTransactionId
+{
+	FRemoteTransactionId() = default;
+	explicit constexpr FRemoteTransactionId(uint32 InId)
+		: Id(InId)
+	{
+	}
+
+	static FRemoteTransactionId Invalid()
+	{
+		return FRemoteTransactionId();
+	}
+
+	bool operator==(const FRemoteTransactionId& Rhs) const = default;
+	bool operator!=(const FRemoteTransactionId& Rhs) const = default;
+
+	bool IsValid() const
+	{
+		return Id != 0;
+	}
+	FString ToString() const
+	{
+		return FString::FromInt(Id);
+	}
+	void operator = (FRemoteTransactionId Other)
+	{
+		Id = Other.Id;
+	}
+
+	friend COREUOBJECT_API FArchive& operator<<(FArchive& Ar, FRemoteTransactionId& Id);
+
+private:
+	uint32 Id = 0;
+};
+
+struct FRemoteWorkPriority
+{
+	// the data is packed in a uint64 to allow passing by value in a register
+	// [63:32] : FRemoteServerId RootServerId
+	// [31:24] : uint8 WorkDepth
+	// [23:0]  : FRemoteTransactionId RootWorkTransactionId
+	//
+	// this allows comparison of priority to be a single uint64 compare-less-than
+	// 
+	// the WorkDepth is stored subtracted from 0xFF so that less-than means higher priority
+	//   0xFF is the root depth
+	//   0xFE is one deeper
+	//   
+	uint64 PackedData = 0;
+	static_assert(sizeof(FRemoteServerId) == 4);
+	static_assert(sizeof(FRemoteTransactionId) == 4);
+
+	FRemoteServerId GetRootServerId() const { return FRemoteServerId(static_cast<uint32>(PackedData >> 32)); }
+	uint32 GetWorkDepth() const { return (PackedData >> 24) & 0xFFu; }
+	FRemoteTransactionId GetRootWorkTransactionId() const { return FRemoteTransactionId(PackedData & 0xFFFFFFu); }
+
+	bool IsValid() const { return GetRootServerId().IsValid(); }
+
+	COREUOBJECT_API static FRemoteWorkPriority CreateRootWorkPriority(FRemoteServerId ServerId, FRemoteTransactionId TransactionId);
+
+	// creates a copy of this priority with the WorkDepth decreased by 1
+	COREUOBJECT_API FRemoteWorkPriority CreateDependentWorkPriority() const;
+
+	COREUOBJECT_API FString ToString() const;
+};
+
+COREUOBJECT_API FArchive& operator<<(FArchive& Ar, FRemoteWorkPriority& Priority);
+COREUOBJECT_API bool IsHigherPriority(FRemoteWorkPriority Lhs, FRemoteWorkPriority Rhs);
+COREUOBJECT_API  bool IsEqualPriority(FRemoteWorkPriority Lhs, FRemoteWorkPriority Rhs);
+COREUOBJECT_API  bool IsHigherOrEqualPriority(FRemoteWorkPriority Lhs, FRemoteWorkPriority Rhs);
+
+class FRemoteSubsystemBase
+{
+protected:
+	TArray<TTuple<FRemoteTransactionId, void*>> Requests;
+
+public:
+
+	FRemoteSubsystemBase() = default;
+	virtual ~FRemoteSubsystemBase() = default;
+
+	int32 GetRequestCount() const { return Requests.Num(); }
+
+	// implemented in FRemoteSubsystem<>
+	virtual void* CreateRequest(FRemoteTransactionId RequestId) = 0;
+	virtual void DestroyRequest(FRemoteTransactionId RequestId) = 0;
+	virtual void SetActiveRequest(FRemoteTransactionId RequestId) = 0;
+
+	// for subsystem implementor to override:
+	virtual const TCHAR* NameForDebug() = 0;
+	virtual void BeginRequest(FRemoteWorkPriority Priority) = 0;
+	virtual void Tick() = 0;
+	virtual bool AreDependenciesSatisfied() const = 0;
+	virtual void EndRequest() = 0;
+};
+
+template<typename RequestType>
+class FRemoteSubsystem : public FRemoteSubsystemBase
+{
+public:
+	RequestType* ActiveRequest = nullptr;
+
+	RequestType* GetRequestByIndex(int32 RequestIndex)
+	{
+		TTuple<FRemoteTransactionId, void*> RequestData = Requests[RequestIndex];
+		return static_cast<RequestType*>(RequestData.Get<1>());
+	}
+
+	RequestType* TryGetRequest(FRemoteTransactionId RequestId)
+	{
+		for (int32 RequestIndex = 0; RequestIndex < Requests.Num(); RequestIndex++)
+		{
+			TTuple<FRemoteTransactionId, void*>& RequestData = Requests[RequestIndex];
+			if (RequestData.Get<0>() == RequestId)
+				return static_cast<RequestType*>(RequestData.Get<1>());
+		}
+
+		return nullptr;
+	}
+
+	void* CreateRequest(FRemoteTransactionId RequestId) final
+	{
+		RequestType* Request = new RequestType{};
+		Request->RequestId = RequestId;
+		Requests.Add(TTuple<FRemoteTransactionId, void*>(RequestId, Request));
+		return Request;
+	}
+
+	void DestroyRequest(FRemoteTransactionId RequestId) final
+	{
+		RequestType* Request = TryGetRequest(RequestId);
+		check(Request);
+		Requests.Remove(TTuple<FRemoteTransactionId, void*>(RequestId, Request));
+		delete Request;
+	}
+
+	void SetActiveRequest(FRemoteTransactionId RequestId) final
+	{
+		ActiveRequest = TryGetRequest(RequestId);
+		check(ActiveRequest);
+	}
+};
+
+namespace UE::RemoteExecutor
+{
+	/** Delegate that pumps the network and executes RPC handlers */
+	extern COREUOBJECT_API TDelegate<void()> TickNetworkDelegate;
+
+	/** Fetch the next deferred RPC request **/
+	extern COREUOBJECT_API TDelegate<TTuple<FRemoteWorkPriority, TFunction<void(void)>>(FRemoteWorkPriority)> FetchDeferredRPCDelegate;
+	extern COREUOBJECT_API TDelegate<TTuple<FRemoteWorkPriority, TFunction<void(void)>>(void)> FetchNextDeferredRPCDelegate;
+
+	/** Delegate that syncs all servers */
+	extern COREUOBJECT_API TDelegate<void()> WaitForServersToSyncDelegate;
+
+	/** Delegate executed when a ExecuteTransactional observes a completed AutoRTFM transaction */
+	extern COREUOBJECT_API TMulticastDelegate<void()> OnTransactionCompletedDelegate;
+
+	/** Delegate executed when a ExecuteTransactional observes an aborted AutoRTFM transaction */
+	extern COREUOBJECT_API TMulticastDelegate<void()> OnTransactionAbortedDelegate;
+
+	COREUOBJECT_API void RegisterRemoteSubsystem(FRemoteSubsystemBase* Subsystem);
+
+	COREUOBJECT_API void ExecuteTransactional(const TFunctionRef<void(void)>& Work);
+	COREUOBJECT_API void ExecuteTransactionalWithExplicitPriority(FRemoteWorkPriority WorkPriority, const TFunctionRef<void(void)>& Work);
+
+	COREUOBJECT_API void SetAbortReason(FStringView Reason);
+	
+	/** Establishes a synchronization barrier across all servers */
+	COREUOBJECT_API void WaitForServersToSync();
+	
+	COREUOBJECT_API FRemoteWorkPriority CreateRootWorkPriority();
+}
diff --git a/Engine/Source/Runtime/CoreUObject/Public/UObject/RemoteObjectTransfer.h b/Engine/Source/Runtime/CoreUObject/Public/UObject/RemoteObjectTransfer.h
--- a/Engine/Source/Runtime/CoreUObject/Public/UObject/RemoteObjectTransfer.h
+++ b/Engine/Source/Runtime/CoreUObject/Public/UObject/RemoteObjectTransfer.h
@@ -7,6 +7,7 @@
 #include "UObject/RemoteObjectTypes.h"
 #include "UObject/Class.h"
 #include "Async/Async.h"
+#include "UObject/RemoteExecutor.h"
 
 enum class EObjectMigrationFlags : uint8
 {
@@ -41,56 +42,9 @@ struct FPostMigrateParams
 	EObjectMigrationFlags MigrationFlags;
 };
 
-struct FRemoteWorkPriority
-{
-	// the data is packed in a uint64 to allow passing by value in a register
-	// [63:32] : FRemoteServerId RootServerId
-	// [31:24] : uint8 WorkDepth
-	// [23:0]  : uint32 RootWorkTransactionId
-	//
-	// this allows comparison of priority to be a single uint64 compare-less-than
-	// 
-	// the WorkDepth is stored subtracted from 0xFF so that less-than means higher priority
-	//   0xFF is the root depth
-	//   0xFE is one deeper
-	//   
-	uint64 PackedData = 0;
-	static_assert(sizeof(FRemoteServerId) == 4);
-
-	FRemoteServerId GetRootServerId() const { return FRemoteServerId(static_cast<uint32>(PackedData >> 32)); }
-	uint32 GetWorkDepth() const { return (PackedData >> 24) & 0xFFu; }
-	uint32 GetRootWorkTransactionId() const { return PackedData & 0xFFFFFFu; }
-
-	bool IsValid() const { return GetRootServerId().IsValid(); }
-
-	COREUOBJECT_API static FRemoteWorkPriority CreateRootWorkPriority(FRemoteServerId ServerId, uint32 TransactionId);
-
-    // creates a copy of this priority with the WorkDepth decreased by 1
-	COREUOBJECT_API FRemoteWorkPriority CreateDependentWorkPriority() const;
-
-	COREUOBJECT_API FString ToString() const;
-};
-
-COREUOBJECT_API FArchive& operator<<(FArchive& Ar, FRemoteWorkPriority& Priority);
-COREUOBJECT_API bool IsHigherPriority(FRemoteWorkPriority Lhs, FRemoteWorkPriority Rhs);
-COREUOBJECT_API  bool IsEqualPriority(FRemoteWorkPriority Lhs, FRemoteWorkPriority Rhs);
-COREUOBJECT_API  bool IsHigherOrEqualPriority(FRemoteWorkPriority Lhs, FRemoteWorkPriority Rhs);
-
 namespace UE::RemoteObject::Transfer
 {
-	constexpr uint32 PushMigrationRequestToken = 0xFFFFFFFFu;
-
-	/**
-	* Execute the Work under an transaction that can catch remote object requests
-	* as well as negotiating remote object transfers, retrying the transaction until it succeeds.
-	*/
-	COREUOBJECT_API void ExecuteTransactional(const TFunctionRef<void(void)>& Work);
-	COREUOBJECT_API void ExecuteTransactionalWithExplicitPriority(FRemoteWorkPriority WorkPriority, const TFunctionRef<void(void)>& Work);
-
-	/**
-	* (Dev API) Sets a friendly string indicating why we are aborting the current transaction
-	*/
-	COREUOBJECT_API void SetAbortReason(FStringView Reason);
+	constexpr FRemoteTransactionId PushMigrationRequestId = FRemoteTransactionId(0xFFFFFFFFu);
 
 	/**
 	* (Advanced API) Adds the remote object to the current transaction's requested object list, causing
@@ -104,14 +58,14 @@ namespace UE::RemoteObject::Transfer
 	* @param RemoteServerId Server that sent object data
 	* @param Data Remote object data. Data ownership is transferred to the remote object queue
 	*/
-	COREUOBJECT_API void OnObjectDataReceived(uint32 RequestToken, FRemoteObjectId ObjectId, FRemoteServerId RemoteServerId, FRemoteObjectData* Data);
+	COREUOBJECT_API void OnObjectDataReceived(FRemoteTransactionId RequestId, FRemoteObjectId ObjectId, FRemoteServerId RemoteServerId, FRemoteObjectData* Data);
 
 	/**
 	* Called when a remote object request was denied by a remote server
 	* @param ObjectId Remote object id
 	* @param RemoteServerId Server that owns the object data
 	*/
-	COREUOBJECT_API void OnObjectDataDenied(uint32 RequestToken, FRemoteObjectId ObjectId, FRemoteServerId RemoteServerId);
+	COREUOBJECT_API void OnObjectDataDenied(FRemoteTransactionId RequestId, FRemoteObjectId ObjectId, FRemoteServerId RemoteServerId);
 
 	/**
 	* Requests remote object data to be transferred from Current Owner Server to Destination Server
@@ -126,16 +80,16 @@ namespace UE::RemoteObject::Transfer
 	* @param Object Local object to migrate
 	* @param DestinationServerId Id of a remote server to receive the object
 	*/
-	COREUOBJECT_API void MigrateObjectToRemoteServer(uint32 RequestToken, UObject* Object, FRemoteServerId DestinationServerId);
-	COREUOBJECT_API void MigrateObjectToRemoteServerWithExplicitPriority(uint32 RequestToken, FRemoteWorkPriority RequestPriority, UObject* Object, FRemoteServerId DestinationServerId);
+	COREUOBJECT_API void MigrateObjectToRemoteServer(FRemoteTransactionId RequestId, UObject* Object, FRemoteServerId DestinationServerId);
+	COREUOBJECT_API void MigrateObjectToRemoteServerWithExplicitPriority(FRemoteTransactionId RequestId, FRemoteWorkPriority RequestPriority, UObject* Object, FRemoteServerId DestinationServerId);
 
 	/**
 	* Migrates an object to a remote server
 	* @param ObjectId Remote object id
 	* @param DestinationServerId Id of a remote server to receive the object
 	*/
-	COREUOBJECT_API void MigrateObjectToRemoteServer(uint32 RequestToken, FRemoteObjectId Id, FRemoteServerId DestinationServerId);
-	COREUOBJECT_API void MigrateObjectToRemoteServerWithExplicitPriority(uint32 RequestToken, FRemoteWorkPriority RequestPriority, FRemoteObjectId Id, FRemoteServerId DestinationServerId);
+	COREUOBJECT_API void MigrateObjectToRemoteServer(FRemoteTransactionId RequestId, FRemoteObjectId Id, FRemoteServerId DestinationServerId);
+	COREUOBJECT_API void MigrateObjectToRemoteServerWithExplicitPriority(FRemoteTransactionId RequestId, FRemoteWorkPriority RequestPriority, FRemoteObjectId Id, FRemoteServerId DestinationServerId);
 	
 	/**
 	* Migrates an object from a remote server (temp function)
@@ -172,31 +126,18 @@ namespace UE::RemoteObject::Transfer
 	*/
 	COREUOBJECT_API void RegisterSharedObject(UObject* Object);
 
-	/**
-	* Establishes a synchronization barrier across all servers
-	*/
-	COREUOBJECT_API void WaitForServersToSync();
-
 	extern COREUOBJECT_API const FRemoteServerId DatabaseId;
 
+	extern COREUOBJECT_API void InitRemoteObjectTransfer();
+
 	/** Delegate that transfers object data to another server */
-	extern COREUOBJECT_API TDelegate<void(uint32 /*RequestToken*/, FRemoteObjectId /*ObjectId*/, FRemoteServerId /*DestinationServerId*/, FRemoteObjectData& /*ObjectData*/)> RemoteObjectTransferDelegate;
+	extern COREUOBJECT_API TDelegate<void(FRemoteTransactionId /*RequestId*/, FRemoteObjectId /*ObjectId*/, FRemoteServerId /*DestinationServerId*/, FRemoteObjectData& /*ObjectData*/)> RemoteObjectTransferDelegate;
 
 	/** Delegate that handles an object request being denied */
-	extern COREUOBJECT_API TDelegate<void(uint32 /*RequestToken*/, FRemoteObjectId /*ObjectId*/, FRemoteServerId /*DestinationServerId*/)> RemoteObjectDeniedTransferDelegate;
+	extern COREUOBJECT_API TDelegate<void(FRemoteTransactionId /*RequestId*/, FRemoteObjectId /*ObjectId*/, FRemoteServerId /*DestinationServerId*/)> RemoteObjectDeniedTransferDelegate;
 
 	/** Delegate that requests remote object data from CurrentOwnerServerId to be transferred to DestinationServerId */
-	extern COREUOBJECT_API TDelegate<void(uint32 /*RequestToken*/, FRemoteWorkPriority /*RequestPriority*/, FRemoteObjectId /*ObjectId*/, FRemoteServerId /*CurrentOwnerServerId*/, FRemoteServerId /*DestinationServerId*/)> RequestRemoteObjectDelegate;
-
-	/** Delegate that pumps the network and executes RPC handlers */
-	extern COREUOBJECT_API TDelegate<void()> TickNetworkDelegate;
-	
-	/** Fetch the next deferred RPC request **/
-	extern COREUOBJECT_API TDelegate<TTuple<FRemoteWorkPriority, TFunction<void(void)>>(FRemoteWorkPriority)> FetchDeferredRPCDelegate;
-	extern COREUOBJECT_API TDelegate<TTuple<FRemoteWorkPriority, TFunction<void(void)>>(void)> FetchNextDeferredRPCDelegate;
-
-	/** Delegate that syncs all servers */
-	extern COREUOBJECT_API TDelegate<void()> WaitForServersToSyncDelegate;
+	extern COREUOBJECT_API TDelegate<void(FRemoteTransactionId /*RequestId*/, FRemoteWorkPriority /*RequestPriority*/, FRemoteObjectId /*ObjectId*/, FRemoteServerId /*CurrentOwnerServerId*/, FRemoteServerId /*DestinationServerId*/)> RequestRemoteObjectDelegate;
 
 	/** 
 	 * Delegate executed when objects have been migrated from another server. 
@@ -218,21 +159,11 @@ namespace UE::RemoteObject::Transfer
 	extern COREUOBJECT_API TMulticastDelegate<void(const TSet<UObject*>& /*MigratedObjects*/, FRemoteServerId /*RemoteServerId*/)> OnObjectsSentDelegate;
 
 	/** Delegate that stores locally unrachable object data into a database. */
-	extern COREUOBJECT_API TDelegate<void(uint32 /*RequestToken*/, FRemoteObjectId /*ObjectId*/, FRemoteServerId /*DatabaseId*/, FRemoteObjectData& /*ObjectData*/)> StoreRemoteObjectDataDelegate;
+	extern COREUOBJECT_API TDelegate<void(FRemoteTransactionId /*RequestId*/, FRemoteObjectId /*ObjectId*/, FRemoteServerId /*DatabaseId*/, FRemoteObjectData& /*ObjectData*/)> StoreRemoteObjectDataDelegate;
 
 	/** Delegate that restores object data from a database. */
-	extern COREUOBJECT_API TDelegate<void(uint32 /*RequestToken*/, FRemoteWorkPriority /*RequestPriority*/, FRemoteObjectId /*ObjectId*/, FRemoteServerId /*DatabaseId*/, FRemoteServerId /*DestinationServerId*/)> RestoreRemoteObjectDataDelegate;
-
-	/** Physics query delegates **/
-	extern COREUOBJECT_API TDelegate<void(FRemoteWorkPriority /*RequestPriority*/)> PhysicsQueryBeginRequestDelegate;
-	extern COREUOBJECT_API TDelegate<bool()> PhysicsQueryTickDelegate;
-	extern COREUOBJECT_API TDelegate<void()> PhysicsQueryEndRequestDelegate;
+	extern COREUOBJECT_API TDelegate<void(FRemoteTransactionId /*RequestId*/, FRemoteWorkPriority /*RequestPriority*/, FRemoteObjectId /*ObjectId*/, FRemoteServerId /*DatabaseId*/, FRemoteServerId /*DestinationServerId*/)> RestoreRemoteObjectDataDelegate;
 
-	/** Delegate executed when a ExecuteTransactional observes a completed AutoRTFM transaction */
-	extern COREUOBJECT_API TMulticastDelegate<void()> OnTransactionCompletedDelegate;
-
-	/** Delegate executed when a ExecuteTransactional observes an aborted AutoRTFM transaction */
-	extern COREUOBJECT_API TMulticastDelegate<void()> OnTransactionAbortedDelegate;
 } // namespace UE::RemoveObject::Transfer
 
 struct FRemoteObjectReference
@@ -298,4 +229,4 @@ struct TStructOpsTypeTraits<FRemoteObjectReference> : public TStructOpsTypeTrait
 		WithNetSerializer = true,
 		WithSerializer = true,
 	};
-};
\ No newline at end of file
+};
diff --git a/Engine/Source/Runtime/CoreUObject/Public/UObject/RemoteObjectTypes.h b/Engine/Source/Runtime/CoreUObject/Public/UObject/RemoteObjectTypes.h
--- a/Engine/Source/Runtime/CoreUObject/Public/UObject/RemoteObjectTypes.h
+++ b/Engine/Source/Runtime/CoreUObject/Public/UObject/RemoteObjectTypes.h
@@ -44,6 +44,10 @@ struct FRemoteServerId
 	{
 		return Id <= Other.Id;
 	}
+	void operator = (FRemoteServerId Other)
+	{
+		Id = Other.Id;
+	}
 
 	friend COREUOBJECT_API FArchive& operator<<(FArchive& Ar, FRemoteServerId& Id);
 
diff --git a/Engine/Source/Runtime/Engine/Private/Actor.cpp b/Engine/Source/Runtime/Engine/Private/Actor.cpp
--- a/Engine/Source/Runtime/Engine/Private/Actor.cpp
+++ b/Engine/Source/Runtime/Engine/Private/Actor.cpp
@@ -63,6 +63,7 @@
 #include "Engine/SCS_Node.h"
 
 #if UE_WITH_REMOTE_OBJECT_HANDLE
+#include "UObject/RemoteExecutor.h"
 #include "UObject/RemoteObjectTransfer.h"
 #endif
 
@@ -6270,7 +6271,7 @@ void AActor::LifeSpanExpired()
 {
 #if UE_WITH_REMOTE_OBJECT_HANDLE
 	// World Timer Manager is not yet transactional, so let's make this call transactional to make up for it
-	UE::RemoteObject::Transfer::ExecuteTransactional([this] { Destroy(); });
+	UE::RemoteExecutor::ExecuteTransactional([this] { Destroy(); });
 #else
 	Destroy();
 #endif
diff --git a/Engine/Source/Runtime/Engine/Private/Collision/WorldCollisionAsync.cpp b/Engine/Source/Runtime/Engine/Private/Collision/WorldCollisionAsync.cpp
--- a/Engine/Source/Runtime/Engine/Private/Collision/WorldCollisionAsync.cpp
+++ b/Engine/Source/Runtime/Engine/Private/Collision/WorldCollisionAsync.cpp
@@ -14,7 +14,7 @@
 #include "Physics/Experimental/PhysInterface_Chaos.h"
 #include "ProfilingDebugging/CsvProfiler.h"
 #include "ProfilingDebugging/ScopedTimers.h"
-#include "UObject/RemoteObjectTransfer.h"
+#include "UObject/RemoteExecutor.h"
 
 CSV_DEFINE_CATEGORY(WorldCollision, true);
 DEFINE_LOG_CATEGORY_STATIC(LogWorldCollision, Log, All);
@@ -808,7 +808,7 @@ void UWorld::FinishAsyncTrace()
 	};
 	if (PhysicsQueryHandler != nullptr)
 	{
-		UE::RemoteObject::Transfer::ExecuteTransactional(DoWorkCallback);
+		UE::RemoteExecutor::ExecuteTransactional(DoWorkCallback);
 	}
 #endif
 
diff --git a/Engine/Source/Runtime/Engine/Private/DataReplication.cpp b/Engine/Source/Runtime/Engine/Private/DataReplication.cpp
--- a/Engine/Source/Runtime/Engine/Private/DataReplication.cpp
+++ b/Engine/Source/Runtime/Engine/Private/DataReplication.cpp
@@ -27,7 +27,7 @@
 #include "Net/Core/PushModel/Types/PushModelPerNetDriverState.h"
 #include "Net/RPCDoSDetection.h"
 #include "UObject/CoreNetContext.h"
-#include "UObject/RemoteObjectTransfer.h"
+#include "UObject/RemoteExecutor.h"
 
 DECLARE_LLM_MEMORY_STAT(TEXT("NetObjReplicator"), STAT_NetObjReplicatorLLM, STATGROUP_LLMFULL);
 LLM_DEFINE_TAG(NetObjReplicator, NAME_None, TEXT("Networking"), GET_STATFNAME(STAT_NetObjReplicatorLLM), GET_STATFNAME(STAT_NetworkingSummaryLLM));
@@ -1390,7 +1390,7 @@ bool FObjectReplicator::ReceivedRPC(FNetBitReader& Reader, const FReplicationFla
 #if UE_WITH_REMOTE_OBJECT_HANDLE
 					if (Connection->Driver->bReplicateTransactionally)
 					{
-						UE::RemoteObject::Transfer::ExecuteTransactional([Object, Function, Parms]()
+						UE::RemoteExecutor::ExecuteTransactional([Object, Function, Parms]()
 							{
 								CallProcessEventForReceivedRPC(Object, Function, Parms);
 							});
diff --git a/Engine/Source/Runtime/Engine/Private/NetDriver.cpp b/Engine/Source/Runtime/Engine/Private/NetDriver.cpp
--- a/Engine/Source/Runtime/Engine/Private/NetDriver.cpp
+++ b/Engine/Source/Runtime/Engine/Private/NetDriver.cpp
@@ -50,7 +50,7 @@
 #include "Net/NetworkMetricsConfig.h"
 #include "Net/NetSubObjectRegistryGetter.h"
 #include "Net/NetworkGranularMemoryLogging.h"
-#include "UObject/RemoteObjectTransfer.h"
+#include "UObject/RemoteExecutor.h"
 #include "UObject/Stack.h"
 #if UE_WITH_IRIS
 #include "Iris/IrisConfig.h"
@@ -1036,7 +1036,7 @@ void UNetDriver::TickFlush(float DeltaSeconds)
 #if UE_WITH_REMOTE_OBJECT_HANDLE
 			if (bReplicateTransactionally)
 			{
-				UE::RemoteObject::Transfer::ExecuteTransactional([this, DeltaSeconds]()
+				UE::RemoteExecutor::ExecuteTransactional([this, DeltaSeconds]()
 				{
 					ServerReplicateActors(DeltaSeconds);
 				});
diff --git a/Engine/Source/Runtime/Engine/Private/TickTaskManager.cpp b/Engine/Source/Runtime/Engine/Private/TickTaskManager.cpp
--- a/Engine/Source/Runtime/Engine/Private/TickTaskManager.cpp
+++ b/Engine/Source/Runtime/Engine/Private/TickTaskManager.cpp
@@ -13,7 +13,7 @@
 #include "Misc/Fork.h"
 #include "Misc/TimeGuard.h"
 #include "Templates/UniquePtr.h"
-#include "UObject/RemoteObjectTransfer.h"
+#include "UObject/RemoteExecutor.h"
 
 DEFINE_LOG_CATEGORY_STATIC(LogTick, Log, All);
 
@@ -320,7 +320,7 @@ public:
 
 			if (Target->bRunTransactionally)
 			{
-				UE::RemoteObject::Transfer::ExecuteTransactional(ExecuteTickWork);
+				UE::RemoteExecutor::ExecuteTransactional(ExecuteTickWork);
 			}
 			else
 			{
