Fixed an issue with IDO GC when the CVars are changed very early

#rnx

diff --git a/Engine/Source/Runtime/CoreUObject/Private/UObject/InstanceDataObjectUtils.cpp b/Engine/Source/Runtime/CoreUObject/Private/UObject/InstanceDataObjectUtils.cpp
--- a/Engine/Source/Runtime/CoreUObject/Private/UObject/InstanceDataObjectUtils.cpp
+++ b/Engine/Source/Runtime/CoreUObject/Private/UObject/InstanceDataObjectUtils.cpp
@@ -374,7 +374,11 @@ namespace UE
 		GetDerivedClasses(UObject::StaticClass(), AllClasses);
 		for (UClass* Class : AllClasses)
 		{
-			Class->AssembleReferenceTokenStream(/*bForce*/ true);
+			// Only re-assemble if it has been assembled because this can run before intrinsic schemas are declared.
+			if (Class->HasAnyClassFlags(CLASS_TokenStreamAssembled))
+			{
+				Class->AssembleReferenceTokenStream(/*bForce*/ true);
+			}
 		}
 	}
 
