Adding accessor ArrayIndicesPerObject to be able to use this information to find a specific property on UObject
https://p4-swarm.epicgames.net/reviews/39190288

diff --git a/Engine/Source/Runtime/CoreUObject/Public/UObject/UnrealType.h b/Engine/Source/Runtime/CoreUObject/Public/UObject/UnrealType.h
--- a/Engine/Source/Runtime/CoreUObject/Public/UObject/UnrealType.h
+++ b/Engine/Source/Runtime/CoreUObject/Public/UObject/UnrealType.h
@@ -6468,6 +6468,17 @@ struct FPropertyChangedEvent
 		InstancesChanged = Forward<T>(InInstancesChanged);
 	}
 
+	bool GetArrayIndicesPerObject(int32 InObjectIteratorIndex, TMap<FString, int32>& OutArrayIndicesPerObject) const
+	{
+		if (!(ArrayIndicesPerObject.IsValidIndex(InObjectIteratorIndex)))
+		{
+			return false;
+		}
+
+		OutArrayIndicesPerObject = ArrayIndicesPerObject[InObjectIteratorIndex];
+		return true;
+	}
+	
 	/**
 	 * Gets the Array Index of the "current object" based on a particular name
 	 * InName - Name of the property to find the array index for
