RemoteExecutor: step two towards being able to service multiple requests concurrently

- removed all of the priority sorting
- now we just collect all pending work and round-robin through it continuously until it's all done. the priority system will sort out preventing thrashing of dependencies so forward progress can still be made even if we execute lower priority work

diff --git a/Engine/Source/Runtime/CoreUObject/Private/UObject/RemoteExecutor.cpp b/Engine/Source/Runtime/CoreUObject/Private/UObject/RemoteExecutor.cpp
--- a/Engine/Source/Runtime/CoreUObject/Private/UObject/RemoteExecutor.cpp
+++ b/Engine/Source/Runtime/CoreUObject/Private/UObject/RemoteExecutor.cpp
@@ -9,7 +9,6 @@ DEFINE_LOG_CATEGORY(LogRemoteExec);
 namespace UE::RemoteExecutor
 {
 TDelegate<void()> TickNetworkDelegate;
-TDelegate<TTuple<FRemoteWorkPriority, TFunction<void(void)>>(FRemoteWorkPriority)> FetchDeferredRPCDelegate;
 TDelegate<TTuple<FRemoteWorkPriority, TFunction<void(void)>>(void)> FetchNextDeferredRPCDelegate;
 TDelegate<void()> WaitForServersToSyncDelegate;
 
@@ -23,15 +22,21 @@ FArchive& operator<<(FArchive& Ar, FRemoteTransactionId& Id)
 	return Ar;
 }
 
+struct FRemoteExecutorWork
+{
+	FRemoteTransactionId RequestId;
+	FRemoteWorkPriority Priority;
+	TFunction<void(void)> Work;
+	TArray<void*> SubsystemRequests;
+};
+
 class FRemoteExecutor
 {
 	TArray<FRemoteSubsystemBase*> Subsystems;
 
 	uint32 NextTransactionRequestId = 0;
 
-	TArray<TTuple<FRemoteWorkPriority, TFunction<void(void)>>> PendingWorks;
-	TArray<FRemoteTransactionId> PendingWorksRequestId;
-	TArray<TArray<void*>> PendingWorksRequests;
+	TArray<FRemoteExecutorWork> PendingWorks;
 
 public:
 	FString DebugLastAbortReason;
@@ -49,9 +54,9 @@ public:
 
 	void EnqueueWork(FRemoteWorkPriority InWorkPriority, const TFunction<void(void)>& InWork)
 	{
-		PendingWorks.Add(TTuple<FRemoteWorkPriority, TFunction<void(void)>>(InWorkPriority, InWork));
-		PendingWorksRequestId.Add(FRemoteTransactionId::Invalid());
-		PendingWorksRequests.Add(TArray<void*>());
+		FRemoteExecutorWork& NewWork = PendingWorks.Emplace_GetRef();
+		NewWork.Priority = InWorkPriority;
+		NewWork.Work = InWork;
 	}
 
 	void ExecutePendingWork()
@@ -63,19 +68,18 @@ public:
 		{
 			LocalIterationNumber++;
 
-			// first check for higher priority work and insert it at the front
-			// this causes requests received later to be put at the start, because
-			// earlier requests could be dependent on later ones completing
+			UE_LOG(LogRemoteExec, VeryVerbose, TEXT("ExecutePendingWork iteration %u"), LocalIterationNumber);
+
+			// pump the network
+			UE::RemoteExecutor::TickNetworkDelegate.ExecuteIfBound();
+
+			// fetch pending RPC work
 			for (;;)
 			{
-				TTuple<FRemoteWorkPriority, TFunction<void(void)>> NewWork = UE::RemoteExecutor::FetchDeferredRPCDelegate.Execute(PendingWorks[0].Get<0>());
+				TTuple<FRemoteWorkPriority, TFunction<void(void)>> NewWork = UE::RemoteExecutor::FetchNextDeferredRPCDelegate.Execute();
 				if (NewWork.Get<0>().IsValid())
 				{
-					PendingWorks.Insert(NewWork, 0);
-					PendingWorksRequestId.Insert(FRemoteTransactionId::Invalid(), 0);
-					PendingWorksRequests.Insert(TArray<void*>(), 0);
-
-					UE_LOG(LogRemoteExec, VeryVerbose, TEXT("-- ExecuteTransactional : iteration %u popped a higher priority RPC --"), LocalIterationNumber);
+					EnqueueWork(NewWork.Get<0>(), NewWork.Get<1>());
 				}
 				else
 				{
@@ -83,56 +87,50 @@ public:
 				}
 			}
 
-			TTuple<FRemoteWorkPriority, TFunction<void(void)>> ExecutingWork = PendingWorks[0];
-			FRemoteWorkPriority ExecutingWorkPriority = ExecutingWork.Get<0>();
-			FRemoteTransactionId ExecutingWorkRequestId = PendingWorksRequestId[0];
-
-			UE_LOG(LogRemoteExec, VeryVerbose, TEXT("-- ExecuteTransactional : iteration %u request %s pri %s (%d PendingWorks) --"), LocalIterationNumber, *ExecutingWorkRequestId.ToString(), *ExecutingWorkPriority.ToString(), PendingWorks.Num());
-
-			if (PendingWorks.Num() > 1)
+			for (int32 PendingWorkIndex = 0; PendingWorkIndex < PendingWorks.Num(); PendingWorkIndex++)
 			{
-				for (int32 PWI = 0; PWI < PendingWorks.Num(); PWI++)
-				{
-					UE_LOG(LogRemoteExec, VeryVerbose, TEXT("\t PendingWork[%d] Request (%s) Pri (%s)"), PWI, *PendingWorksRequestId[PWI].ToString(), *PendingWorks[PWI].Get<0>().ToString());
-				}
-			}
-
-			const bool bIsNewRequest = !ExecutingWorkRequestId.IsValid();
+				FRemoteExecutorWork& PendingWork = PendingWorks[PendingWorkIndex];
 
-			if (bIsNewRequest)
-			{
-				ExecutingWorkRequestId = GenerateNextTransactionId();
-				PendingWorksRequestId[0] = ExecutingWorkRequestId;
+				UE_LOG(LogRemoteExec, VeryVerbose, TEXT("PendingWork[%d] RequestId %s %s"),
+					PendingWorkIndex,
+					*PendingWork.RequestId.ToString(),
+					*PendingWork.Priority.ToString());
 			}
 
-			// begin requests
-			UE_LOG(LogRemoteExec, VeryVerbose, TEXT("\tactivate request %s"), *ExecutingWorkRequestId.ToString());
-
-			for (int32 SubsystemIndex = 0; SubsystemIndex < Subsystems.Num(); SubsystemIndex++)
+			for(int32 PendingWorkIndex = 0; PendingWorkIndex < PendingWorks.Num(); )
 			{
-				FRemoteSubsystemBase* Subsystem = Subsystems[SubsystemIndex];
+				FRemoteExecutorWork& ExecutingWork = PendingWorks[PendingWorkIndex];
+
+				UE_LOG(LogRemoteExec, VeryVerbose, TEXT("-- ExecuteTransactional : iteration %u request %s %s --"),
+					LocalIterationNumber,
+					*ExecutingWork.RequestId.ToString(),
+					*ExecutingWork.Priority.ToString());
 
-				if (bIsNewRequest)
+				if (!ExecutingWork.RequestId.IsValid())
 				{
-					UE_LOG(LogRemoteExec, VeryVerbose, TEXT("\tactivate request[%s] %s"), Subsystem->NameForDebug(), *ExecutingWorkRequestId.ToString());
-					Subsystem->CreateRequest(ExecutingWorkRequestId);
-				}
+					// new request
+					ExecutingWork.RequestId = GenerateNextTransactionId();
+
+					for (int32 SubsystemIndex = 0; SubsystemIndex < Subsystems.Num(); SubsystemIndex++)
+					{
+						FRemoteSubsystemBase* Subsystem = Subsystems[SubsystemIndex];
+						UE_LOG(LogRemoteExec, VeryVerbose, TEXT("\tcreate request[%s]"), Subsystem->NameForDebug());
+						Subsystem->CreateRequest(ExecutingWork.RequestId, ExecutingWork.Priority);
 
-				Subsystem->SetActiveRequest(ExecutingWorkRequestId);
+						Subsystem->SetActiveRequest(ExecutingWork.RequestId);
 
-				if (bIsNewRequest)
+						Subsystem->BeginRequest();
+					}
+				}
+				else
 				{
-					UE_LOG(LogRemoteExec, VeryVerbose, TEXT("\tactivate request[%s] %s"), Subsystem->NameForDebug(), *ExecutingWorkRequestId.ToString());
-					Subsystem->BeginRequest(ExecutingWorkPriority);
+					for (int32 SubsystemIndex = 0; SubsystemIndex < Subsystems.Num(); SubsystemIndex++)
+					{
+						FRemoteSubsystemBase* Subsystem = Subsystems[SubsystemIndex];
+						Subsystem->SetActiveRequest(ExecutingWork.RequestId);
+					}
 				}
-			}
 
-			// we might run through this loop after the work succeeds in the case
-			// that the work prefetches objects that it doesn't access
-			bool bWorkSucceeded = false;
-
-			for (;;)
-			{
 				double Now = FPlatformTime::Seconds();
 				bool bPrintStall = false;
 
@@ -142,36 +140,6 @@ public:
 					bPrintStall = true;
 				}
 
-				// we need to check to see if there is any higher priority work
-				// that we need to context switch to
-				bool bRequiresContextSwitch = false;
-				for (;;)
-				{
-					TTuple<FRemoteWorkPriority, TFunction<void(void)>> NewWork = UE::RemoteExecutor::FetchDeferredRPCDelegate.Execute(ExecutingWorkPriority);
-					if (NewWork.Get<0>().IsValid())
-					{
-						PendingWorks.Insert(NewWork, 0);
-						PendingWorksRequestId.Insert(FRemoteTransactionId::Invalid(), 0);
-						PendingWorksRequests.Add(TArray<void*>());
-
-						UE_LOG(LogRemoteExec, Verbose, TEXT("\tGOT higher priority work (%s), requires context switch"), *NewWork.Get<0>().ToString());
-
-						bRequiresContextSwitch = true;
-					}
-					else
-					{
-						break;
-					}
-				}
-
-				if (bRequiresContextSwitch)
-				{
-					break;
-				}
-
-				// pump the network
-				UE::RemoteExecutor::TickNetworkDelegate.ExecuteIfBound();
-
 				// tick the subsystems
 				for (int32 SubsystemIndex = 0; SubsystemIndex < Subsystems.Num(); SubsystemIndex++)
 				{
@@ -202,84 +170,50 @@ public:
 					{
 						UE_LOG(LogRemoteExec, VeryVerbose, TEXT("\t%d of %d subsystems not ready..."), Subsystems.Num() - SubsystemReadyCount, Subsystems.Num());
 					}
+
+					PendingWorkIndex++;
 					continue;
 				}
 
 				// all of the subsystems are ready, try to perform the work
-				if (!bWorkSucceeded)
-				{
-					DebugLastAbortReason = TEXT("");
+				DebugLastAbortReason = TEXT("");
 
-					AutoRTFM::ETransactionResult WorkTransactResult = AutoRTFM::Transact([&ExecutingWork]() { ExecutingWork.Get<1>()(); });
+				AutoRTFM::ETransactionResult WorkTransactResult = AutoRTFM::Transact([&ExecutingWork]() { ExecutingWork.Work(); });
 
-					if (WorkTransactResult == AutoRTFM::ETransactionResult::AbortedByRequest)
-					{
-						// the work aborted, it should have inserted a new requested object, retry
-						UE_LOG(LogRemoteExec, Verbose, TEXT("\twork ABORTED (%s)"), *DebugLastAbortReason);
-						UE::RemoteExecutor::OnTransactionAbortedDelegate.Broadcast();
-					}
-					else
-					{
-						UE_LOG(LogRemoteExec, VeryVerbose, TEXT("\twork COMPLETED!"));
-
-						bWorkSucceeded = true;
-						UE::RemoteExecutor::OnTransactionCompletedDelegate.Broadcast();
+				if (WorkTransactResult == AutoRTFM::ETransactionResult::AbortedByRequest)
+				{
+					// the work aborted, it should have inserted a new requested object, retry
+					UE_LOG(LogRemoteExec, Verbose, TEXT("\twork ABORTED (%s)"), *DebugLastAbortReason);
+					UE::RemoteExecutor::OnTransactionAbortedDelegate.Broadcast();
 
-						// the work succeeded, we're all done!
-						break;
-					}
+					PendingWorkIndex++;
 				}
-			}
-
-			// if we get here and the work was successfully executed, we're done
-			// otherwise, we broke out of the request because there are higher
-			// priority pending requests that need to be serviced so we need
-			// to loop back to the top
-			if (bWorkSucceeded)
-			{
-				UE_LOG(LogRemoteExec, VeryVerbose, TEXT("-- ExecuteTransactional %s DONE --"), *ExecutingWorkRequestId.ToString());
-			}
-			else
-			{
-				UE_LOG(LogRemoteExec, Verbose, TEXT("-- ExecuteTransactional %s ABANDONED, context switching to higher priority work --"), *ExecutingWorkRequestId.ToString());
-			}
+				else
+				{
+					UE_LOG(LogRemoteExec, VeryVerbose, TEXT("\twork COMPLETED!"));
 
-			if (bWorkSucceeded)
-			{
-				UE_LOG(LogRemoteExec, VeryVerbose, TEXT("\tending request %s"), *ExecutingWorkRequestId.ToString());
+					UE::RemoteExecutor::OnTransactionCompletedDelegate.Broadcast();
 
-				for (int32 SubsystemIndex = 0; SubsystemIndex < Subsystems.Num(); SubsystemIndex++)
-				{
-					FRemoteSubsystemBase* Subsystem = Subsystems[SubsystemIndex];
-					Subsystem->EndRequest();
-				}
+					for (int32 SubsystemIndex = 0; SubsystemIndex < Subsystems.Num(); SubsystemIndex++)
+					{
+						FRemoteSubsystemBase* Subsystem = Subsystems[SubsystemIndex];
+						UE_LOG(LogRemoteExec, VeryVerbose, TEXT("\tend request[%s]"), Subsystem->NameForDebug());
+						Subsystem->EndRequest();
+						Subsystem->ClearActiveRequest();
+					}
 
-				for (int32 SubsystemIndex = 0; SubsystemIndex < Subsystems.Num(); SubsystemIndex++)
-				{
-					FRemoteSubsystemBase* Subsystem = Subsystems[SubsystemIndex];
-					Subsystem->DestroyRequest(ExecutingWorkRequestId);
-				}
+					for (int32 SubsystemIndex = 0; SubsystemIndex < Subsystems.Num(); SubsystemIndex++)
+					{
+						FRemoteSubsystemBase* Subsystem = Subsystems[SubsystemIndex];
+						Subsystem->DestroyRequest(ExecutingWork.RequestId);
+					}
 
-				// remove this work from the list
-				PendingWorks.RemoveAt(0, EAllowShrinking::No);
-				PendingWorksRequestId.RemoveAt(0, EAllowShrinking::No);
-				PendingWorksRequests.RemoveAt(0, EAllowShrinking::No);
-			}
+					UE_LOG(LogRemoteExec, VeryVerbose, TEXT("-- ExecuteTransactional %s DONE --"), *ExecutingWork.RequestId.ToString());
 
-			if (PendingWorks.Num() == 0)
-			{
-				// if we have no more pending work, then we are done executing what
-				// we came here to execute. but before we start going back to find more
-				// work for ourselves, try to process any pending RPCs by searching
-				// for any priority available to help unblock other servers
-				TTuple<FRemoteWorkPriority, TFunction<void(void)>> NewWork = UE::RemoteExecutor::FetchDeferredRPCDelegate.Execute(FRemoteWorkPriority());
-				if (NewWork.Get<0>().IsValid())
-				{
-					PendingWorks.Add(NewWork);
-					PendingWorksRequestId.Add(FRemoteTransactionId::Invalid());
-					PendingWorksRequests.Add(TArray<void*>());
+					// remove this work from the list
+					PendingWorks.RemoveAt(PendingWorkIndex, EAllowShrinking::No);
 
-					UE_LOG(LogRemoteExec, VeryVerbose, TEXT("\tdone with original work, enqueueing lower priority deferred RPC"));
+					// do not increment PendingWorkIndex
 				}
 			}
 		}
diff --git a/Engine/Source/Runtime/CoreUObject/Private/UObject/RemoteObject.cpp b/Engine/Source/Runtime/CoreUObject/Private/UObject/RemoteObject.cpp
--- a/Engine/Source/Runtime/CoreUObject/Private/UObject/RemoteObject.cpp
+++ b/Engine/Source/Runtime/CoreUObject/Private/UObject/RemoteObject.cpp
@@ -117,9 +117,9 @@ void InitRemoteObjects()
 		UE::RemoteObject::Transfer::RestoreRemoteObjectDataDelegate.BindStatic(&UE::RemoteObject::Serialization::Disk::LoadObjectFromDisk);
 	}
 
-	if (!UE::RemoteExecutor::FetchDeferredRPCDelegate.IsBound())
+	if (!UE::RemoteExecutor::FetchNextDeferredRPCDelegate.IsBound())
 	{
-		UE::RemoteExecutor::FetchDeferredRPCDelegate.BindStatic([](FRemoteWorkPriority) { return TTuple<FRemoteWorkPriority, TFunction<void(void)>>(); });
+		UE::RemoteExecutor::FetchNextDeferredRPCDelegate.BindStatic([]() { return TTuple<FRemoteWorkPriority, TFunction<void(void)>>(); });
 	}
 
 	FCoreDelegates::OnEndFrame.AddStatic(&UE::RemoteExecutor::WaitForServersToSync);
diff --git a/Engine/Source/Runtime/CoreUObject/Private/UObject/RemoteObjectTransfer.cpp b/Engine/Source/Runtime/CoreUObject/Private/UObject/RemoteObjectTransfer.cpp
--- a/Engine/Source/Runtime/CoreUObject/Private/UObject/RemoteObjectTransfer.cpp
+++ b/Engine/Source/Runtime/CoreUObject/Private/UObject/RemoteObjectTransfer.cpp
@@ -52,17 +52,6 @@ public:
 
 	FRemoteTransactionId GetRequestId() const { return RequestId; }
 	FRemoteWorkPriority GetPriority() const { return Priority; }
-
-	void BeginRequest(FRemoteWorkPriority InPriority)
-	{
-		check(!AutoRTFM::IsClosed());
-		Priority = InPriority;
-	}
-
-	void EndRequest()
-	{
-		check(!AutoRTFM::IsClosed());
-	}
 };
 
 class FRemoteObjectTransferQueue : public FRemoteSubsystem<FRemoteObjectRequest>
@@ -73,9 +62,8 @@ class FRemoteObjectTransferQueue : public FRemoteSubsystem<FRemoteObjectRequest>
 public:
 	const TCHAR* NameForDebug() final { return TEXT("RemoteObjectTransferQueue"); }
 
-	void BeginRequest(FRemoteWorkPriority Priority) final
+	void BeginRequest() final
 	{
-		ActiveRequest->BeginRequest(Priority);
 	}
 
 	void Tick() final
@@ -144,8 +132,6 @@ public:
 
 	void EndRequest() final
 	{
-		ActiveRequest->EndRequest();
-		ActiveRequest = nullptr;
 	}
 
 	void RequestRemoteObject(FRemoteTransactionId RequestId, FRemoteWorkPriority RequestPriority, FRemoteObjectId ObjectId, FRemoteServerId CurrentOwnerServerId, FRemoteServerId DestinationServerId)
diff --git a/Engine/Source/Runtime/CoreUObject/Public/UObject/RemoteExecutor.h b/Engine/Source/Runtime/CoreUObject/Public/UObject/RemoteExecutor.h
--- a/Engine/Source/Runtime/CoreUObject/Public/UObject/RemoteExecutor.h
+++ b/Engine/Source/Runtime/CoreUObject/Public/UObject/RemoteExecutor.h
@@ -94,13 +94,14 @@ public:
 	int32 GetRequestCount() const { return Requests.Num(); }
 
 	// implemented in FRemoteSubsystem<>
-	virtual void* CreateRequest(FRemoteTransactionId RequestId) = 0;
+	virtual void* CreateRequest(FRemoteTransactionId RequestId, FRemoteWorkPriority Priority) = 0;
 	virtual void DestroyRequest(FRemoteTransactionId RequestId) = 0;
 	virtual void SetActiveRequest(FRemoteTransactionId RequestId) = 0;
+	virtual void ClearActiveRequest() = 0;
 
 	// for subsystem implementor to override:
 	virtual const TCHAR* NameForDebug() = 0;
-	virtual void BeginRequest(FRemoteWorkPriority Priority) = 0;
+	virtual void BeginRequest() = 0;
 	virtual void Tick() = 0;
 	virtual bool AreDependenciesSatisfied() const = 0;
 	virtual void EndRequest() = 0;
@@ -130,19 +131,38 @@ public:
 		return nullptr;
 	}
 
-	void* CreateRequest(FRemoteTransactionId RequestId) final
+	int32 TryGetRequestIndex(FRemoteTransactionId RequestId)
+	{
+		int32 Result = INDEX_NONE;
+
+		for (int32 RequestIndex = 0; RequestIndex < Requests.Num(); RequestIndex++)
+		{
+			TTuple<FRemoteTransactionId, void*>& RequestData = Requests[RequestIndex];
+			if (RequestData.Get<0>() == RequestId)
+			{
+				Result = RequestIndex;
+				break;
+			}
+		}
+
+		return Result;
+	}
+
+	void* CreateRequest(FRemoteTransactionId RequestId, FRemoteWorkPriority Priority) final
 	{
 		RequestType* Request = new RequestType{};
 		Request->RequestId = RequestId;
+		Request->Priority = Priority;
 		Requests.Add(TTuple<FRemoteTransactionId, void*>(RequestId, Request));
 		return Request;
 	}
 
 	void DestroyRequest(FRemoteTransactionId RequestId) final
 	{
-		RequestType* Request = TryGetRequest(RequestId);
-		check(Request);
-		Requests.Remove(TTuple<FRemoteTransactionId, void*>(RequestId, Request));
+		int32 RequestIndex = TryGetRequestIndex(RequestId);
+		check(RequestIndex != INDEX_NONE);
+		RequestType* Request = GetRequestByIndex(RequestIndex);
+		Requests.RemoveAt(RequestIndex, EAllowShrinking::No);
 		delete Request;
 	}
 
@@ -151,6 +171,11 @@ public:
 		ActiveRequest = TryGetRequest(RequestId);
 		check(ActiveRequest);
 	}
+
+	void ClearActiveRequest() final
+	{
+		ActiveRequest = nullptr;
+	}
 };
 
 namespace UE::RemoteExecutor
@@ -159,7 +184,6 @@ namespace UE::RemoteExecutor
 	extern COREUOBJECT_API TDelegate<void()> TickNetworkDelegate;
 
 	/** Fetch the next deferred RPC request **/
-	extern COREUOBJECT_API TDelegate<TTuple<FRemoteWorkPriority, TFunction<void(void)>>(FRemoteWorkPriority)> FetchDeferredRPCDelegate;
 	extern COREUOBJECT_API TDelegate<TTuple<FRemoteWorkPriority, TFunction<void(void)>>(void)> FetchNextDeferredRPCDelegate;
 
 	/** Delegate that syncs all servers */
