[ZenLoader]
- Add missing support for localized packages when running editor with -game

diff --git a/Engine/Source/Runtime/CoreUObject/Private/Serialization/AsyncLoading2.cpp b/Engine/Source/Runtime/CoreUObject/Private/Serialization/AsyncLoading2.cpp
--- a/Engine/Source/Runtime/CoreUObject/Private/Serialization/AsyncLoading2.cpp
+++ b/Engine/Source/Runtime/CoreUObject/Private/Serialization/AsyncLoading2.cpp
@@ -55,6 +55,7 @@
 #include "ProfilingDebugging/CountersTrace.h"
 #include "Serialization/AsyncPackage.h"
 #include "Serialization/UnversionedPropertySerialization.h"
+#include "Internationalization/PackageLocalizationManager.h"
 #include "Serialization/Zenaphore.h"
 #include "UObject/GCObject.h"
 #include "UObject/ObjectRedirector.h"
@@ -513,6 +514,36 @@ FString FormatPackageId(FPackageId PackageId)
 #endif
 }
 
+#if WITH_EDITOR
+
+FName ApplyPackageNameRedirections(FName PackageName)
+{
+	if (!PackageName.IsNone())
+	{
+		FName NewPackageName = FCoreRedirects::GetRedirectedName(ECoreRedirectFlags::Type_Package, FCoreRedirectObjectName(NAME_None, NAME_None, PackageName)).PackageName;
+		if (NewPackageName != PackageName)
+		{
+			UE_LOG(LogStreaming, Verbose, TEXT("Applying core redirection for package %s to package %s"), *PackageName.ToString(), *NewPackageName.ToString());
+			PackageName = NewPackageName;
+		}
+
+		// Handle localization when running in -game
+		if (!GIsEditor)
+		{
+			NewPackageName = FPackageLocalizationManager::Get().FindLocalizedPackageName(PackageName);
+			if (!NewPackageName.IsNone())
+			{
+				UE_LOG(LogStreaming, Verbose, TEXT("Applying localization redirection for package %s to package %s"), *PackageName.ToString(), *NewPackageName.ToString());
+				PackageName = NewPackageName;
+			}
+		}
+	}
+
+	return PackageName;
+}
+
+#endif
+
 struct FAsyncPackage2;
 class FAsyncLoadingThread2;
 
@@ -4813,11 +4844,7 @@ bool FAsyncLoadingThread2::CreateAsyncPackagesFromQueue(FAsyncLoadingThreadState
 			}
 
 #if WITH_EDITOR
-			const FCoreRedirectObjectName RedirectedPackageName = FCoreRedirects::GetRedirectedName(ECoreRedirectFlags::Type_Package, FCoreRedirectObjectName(NAME_None, NAME_None, PackageNameToLoad));
-			if (RedirectedPackageName.PackageName != PackageNameToLoad)
-			{
-				PackageNameToLoad = RedirectedPackageName.PackageName;
-			}
+			PackageNameToLoad = ApplyPackageNameRedirections(PackageNameToLoad);
 #endif
 
 			FPackageId PackageIdToLoad = FPackageId::FromName(PackageNameToLoad);
@@ -5451,10 +5478,10 @@ void FAsyncPackage2::ImportPackagesRecursiveInner(FAsyncLoadingThreadState2& Thr
 #if WITH_EDITOR
 		if (!ImportedPackageNameToLoad.IsNone())
 		{
-			const FCoreRedirectObjectName RedirectedPackageName = FCoreRedirects::GetRedirectedName(ECoreRedirectFlags::Type_Package, FCoreRedirectObjectName(NAME_None, NAME_None, ImportedPackageNameToLoad));
-			if (RedirectedPackageName.PackageName != ImportedPackageNameToLoad)
+			const FName NewPackageNameToLoad = ApplyPackageNameRedirections(ImportedPackageNameToLoad);
+			if (ImportedPackageNameToLoad != NewPackageNameToLoad)
 			{
-				ImportedPackageNameToLoad = RedirectedPackageName.PackageName;
+				ImportedPackageNameToLoad = NewPackageNameToLoad;
 				ImportedPackageIdToLoad = FPackageId::FromName(ImportedPackageNameToLoad);
 				ImportedPackageUPackageName = ImportedPackageNameToLoad;
 				ImportedPackageId = ImportedPackageIdToLoad;
