Fixed WasAdded state not being preserved on instances of prefab when saving it.

#rnx

diff --git a/Engine/Source/Editor/UnrealEd/Private/Kismet2/KismetReinstanceUtilities.cpp b/Engine/Source/Editor/UnrealEd/Private/Kismet2/KismetReinstanceUtilities.cpp
--- a/Engine/Source/Editor/UnrealEd/Private/Kismet2/KismetReinstanceUtilities.cpp
+++ b/Engine/Source/Editor/UnrealEd/Private/Kismet2/KismetReinstanceUtilities.cpp
@@ -2309,6 +2309,7 @@ static void ReplaceObjectHelper(UObject*& OldObject, UClass* OldClass, UObject*&
 	{
 		UEditorEngine::CopyPropertiesForUnrelatedObjects(Pair.Key, Pair.Value, Options);
 	}
+	FOverridableManager::Get().RestoreOverrideState(OldObject, NewUObject);
 
 	UWorld* RegisteredWorld = nullptr;
 	bool bWasRegistered = false;
diff --git a/Engine/Source/Runtime/CoreUObject/Private/UObject/OverridableManager.cpp b/Engine/Source/Runtime/CoreUObject/Private/UObject/OverridableManager.cpp
--- a/Engine/Source/Runtime/CoreUObject/Private/UObject/OverridableManager.cpp
+++ b/Engine/Source/Runtime/CoreUObject/Private/UObject/OverridableManager.cpp
@@ -137,6 +137,17 @@ FOverriddenPropertySet* FOverridableManager::RestoreOverrideOperation(TNotNull<U
 
 }
 
+void FOverridableManager::RestoreOverrideState(TNotNull<const UObject*> OldObject, TNotNull<UObject*> NewObject)
+{
+	if (const FOverriddenPropertySet* OldOverriddenProperties = GetOverriddenProperties(*OldObject))
+	{
+		if (FOverriddenPropertySet* NewOverriddenProperties = GetOverriddenProperties(*NewObject))
+		{
+			NewOverriddenProperties->RestoreOverriddenState(*OldOverriddenProperties);
+		}
+	}
+}
+
 EOverriddenState FOverridableManager::GetOverriddenState(UObject& Object)
 {
 #if WITH_EDITORONLY_DATA
diff --git a/Engine/Source/Runtime/CoreUObject/Private/UObject/OverriddenPropertySet.cpp b/Engine/Source/Runtime/CoreUObject/Private/UObject/OverriddenPropertySet.cpp
--- a/Engine/Source/Runtime/CoreUObject/Private/UObject/OverriddenPropertySet.cpp
+++ b/Engine/Source/Runtime/CoreUObject/Private/UObject/OverriddenPropertySet.cpp
@@ -352,6 +352,11 @@ void FOverriddenPropertyNodeID::HandleObjectsReInstantiated(const TMap<UObject*,
 }
 
 
+void FOverriddenPropertySet::RestoreOverriddenState(const FOverriddenPropertySet& FromOverriddenProperties)
+{
+	bWasAdded = FromOverriddenProperties.bWasAdded;
+}
+
 //----------------------------------------------------------------------//
 // FOverriddenPropertySet
 //----------------------------------------------------------------------//
diff --git a/Engine/Source/Runtime/CoreUObject/Public/UObject/OverridableManager.h b/Engine/Source/Runtime/CoreUObject/Public/UObject/OverridableManager.h
--- a/Engine/Source/Runtime/CoreUObject/Public/UObject/OverridableManager.h
+++ b/Engine/Source/Runtime/CoreUObject/Public/UObject/OverridableManager.h
@@ -167,6 +167,13 @@ public:
 	 * @return the overridden properties of the object */
 	COREUOBJECT_API FOverriddenPropertySet* RestoreOverrideOperation(TNotNull<UObject*> Object, EOverriddenPropertyOperation Operation, const bool bNeedsSubobjectTemplateInstantiation);
 
+	/**
+	 * Restore some of the overridden state that is not necessarily restored by the CPFUO
+	 * (ex: bWasAdded come from the owner of the object and reinstantiating the object does not preserve it)
+	 * @param OldObject to old object to restore from the overridden state
+	 * @param NewObject the new object to restore to the overridden state */
+	COREUOBJECT_API void RestoreOverrideState(TNotNull<const UObject*> OldObject, TNotNull<UObject*> NewObject);
+
 	/**
 	 * Retrieve the overridden state for the specified object
 	 * @param Object to fetch the overridden properties
diff --git a/Engine/Source/Runtime/CoreUObject/Public/UObject/OverriddenPropertySet.h b/Engine/Source/Runtime/CoreUObject/Public/UObject/OverriddenPropertySet.h
--- a/Engine/Source/Runtime/CoreUObject/Public/UObject/OverriddenPropertySet.h
+++ b/Engine/Source/Runtime/CoreUObject/Public/UObject/OverriddenPropertySet.h
@@ -422,6 +422,13 @@ public:
 	{
 		return Owner;
 	}
+
+	/**
+	 * Restore some of the overridden state that is not necessarily restored by the CPFUO
+	 * (ex: bWasAdded come from the owner of the object and reinstantiating the object does not preserve it)
+	 * @param FromOverriddenProperties overridable properties to restore from*/
+	void RestoreOverriddenState(const FOverriddenPropertySet& FromOverriddenProperties);
+
 protected:
 
 	FOverriddenPropertyNode& FindOrAddNode(FOverriddenPropertyNode& ParentPropertyNode, FOverriddenPropertyNodeID NodeID);
