Replace an IDO GetStructTypeHash ensure with a warning

This scenario arises when a struct contains a property that is not part of its hash function. It may be necessary to remove the warning in the future, but for now it is very infrequent and may provide value.

#rnx

diff --git a/Engine/Source/Runtime/CoreUObject/Private/UObject/InstanceDataObjectUtils.cpp b/Engine/Source/Runtime/CoreUObject/Private/UObject/InstanceDataObjectUtils.cpp
--- a/Engine/Source/Runtime/CoreUObject/Private/UObject/InstanceDataObjectUtils.cpp
+++ b/Engine/Source/Runtime/CoreUObject/Private/UObject/InstanceDataObjectUtils.cpp
@@ -5,6 +5,7 @@
 #if WITH_EDITORONLY_DATA
 
 #include "HAL/IConsoleManager.h"
+#include "Logging/StructuredLog.h"
 #include "Misc/ReverseIterate.h"
 #include "Serialization/ObjectReader.h"
 #include "Serialization/ObjectWriter.h"
@@ -21,6 +22,8 @@
 #include "UObject/UObjectAnnotation.h"
 #include "VerseVM/VVMVerseClass.h"
 
+DEFINE_LOG_CATEGORY_STATIC(LogInstanceDataObject, Log, VeryVerbose);
+
 // Implemented in Class.cpp
 int32 CalculatePropertyIndex(const UStruct* Struct, const FProperty* Property, int32 ArrayIndex);
 
@@ -216,7 +219,7 @@ uint32 UInstanceDataObjectStruct::GetStructTypeHash(const void* Src) const
 				BoolHash.Hash(BoolProperty->GetPropertyValue_InContainer(Src, I));
 			}
 		}
-		else if (ensure(It->HasAllPropertyFlags(CPF_HasGetValueTypeHash)))
+		else if (It->HasAllPropertyFlags(CPF_HasGetValueTypeHash))
 		{
 			for (int32 I = 0; I < It->ArrayDim; ++I)
 			{
@@ -226,6 +229,9 @@ uint32 UInstanceDataObjectStruct::GetStructTypeHash(const void* Src) const
 		}
 		else
 		{
+			UE_LOGFMT(LogInstanceDataObject, Warning,
+				"Struct {StructType} contains property {PropertyName} of type {PropertyType} that is missing GetValueTypeHash.",
+				UE::FAssetLog(this), It->GetFName(), WriteToString<128>(UE::FPropertyTypeName(*It)));
 			ValueHash = HashCombineFast(ValueHash, It->ArrayDim);
 		}
 	}
@@ -715,12 +721,7 @@ namespace UE
 
 				FField::CopyMetaData(TemplateProperty, SuperProperty);
 
-				FPropertyTypeName Type;
-				{
-					FPropertyTypeNameBuilder TypeBuilder;
-					TemplateProperty->SaveTypeName(TypeBuilder);
-					Type = TypeBuilder.Build();
-				}
+				FPropertyTypeName Type(TemplateProperty);
 
 				// Find the sub-tree containing unknown properties for this template property.
 				const FPropertyPathNameTree* SubTree = nullptr;
diff --git a/Engine/Source/Runtime/CoreUObject/Private/UObject/PropertyTypeName.cpp b/Engine/Source/Runtime/CoreUObject/Private/UObject/PropertyTypeName.cpp
--- a/Engine/Source/Runtime/CoreUObject/Private/UObject/PropertyTypeName.cpp
+++ b/Engine/Source/Runtime/CoreUObject/Private/UObject/PropertyTypeName.cpp
@@ -4,6 +4,7 @@
 
 #include "Async/Mutex.h"
 #include "Async/UniqueLock.h"
+#include "AutoRTFM/AutoRTFM.h"
 #include "Containers/Map.h"
 #include "HAL/CriticalSection.h"
 #include "HAL/LowLevelMemTracker.h"
@@ -11,7 +12,7 @@
 #include "Misc/StringBuilder.h"
 #include "String/Find.h"
 #include "UObject/Class.h"
-#include "AutoRTFM/AutoRTFM.h"
+#include "UObject/UnrealType.h"
 
 namespace UE
 {
@@ -250,6 +251,16 @@ FPropertyTypeNameNode* FPropertyTypeNameTable::AllocateBlock(int32 BlockIndex)
 
 ///////////////////////////////////////////////////////////////////////////////////////////////////////////////////////
 
+FPropertyTypeName::FPropertyTypeName(const FProperty* Property)
+	: FPropertyTypeName([](const FProperty* P)
+	{
+		FPropertyTypeNameBuilder Builder;
+		P->SaveTypeName(Builder);
+		return Builder.Build();
+	}(Property))
+{
+}
+
 FName FPropertyTypeName::GetName() const
 {
 	const FPropertyTypeNameNode* First = GPropertyTypeNameTable.ResolveByIndex(Index);
diff --git a/Engine/Source/Runtime/CoreUObject/Public/UObject/PropertyTypeName.h b/Engine/Source/Runtime/CoreUObject/Public/UObject/PropertyTypeName.h
--- a/Engine/Source/Runtime/CoreUObject/Public/UObject/PropertyTypeName.h
+++ b/Engine/Source/Runtime/CoreUObject/Public/UObject/PropertyTypeName.h
@@ -10,6 +10,7 @@
 #define UE_API COREUOBJECT_API
 
 class FArchive;
+class FProperty;
 class FStructuredArchiveSlot;
 class UField;
 struct FGuid;
@@ -44,6 +45,10 @@ struct FPropertyTypeNameNode
 class FPropertyTypeName
 {
 public:
+	FPropertyTypeName() = default;
+
+	UE_API explicit FPropertyTypeName(const FProperty* Property);
+
 	inline bool IsEmpty() const
 	{
 		return Index == 0;
