Updated to use the new constructor for FPropertyTypeName

#rnx

diff --git a/Engine/Source/Runtime/CoreUObject/Private/UObject/PropertyArray.cpp b/Engine/Source/Runtime/CoreUObject/Private/UObject/PropertyArray.cpp
--- a/Engine/Source/Runtime/CoreUObject/Private/UObject/PropertyArray.cpp
+++ b/Engine/Source/Runtime/CoreUObject/Private/UObject/PropertyArray.cpp
@@ -627,10 +627,8 @@ void FArrayProperty::SerializeItem(FStructuredArchive::FSlot Slot, void* Value,
 				}
 				else
 				{
-					UE::FPropertyTypeNameBuilder Builder;
-					StructProperty->SaveTypeName(Builder);
 					UE_LOG(LogClass, Warning, TEXT("Array Property %s has a struct type mismatch (tag %s != prop %s) in package: %s. If that struct got renamed, add an entry to ActiveStructRedirects."),
-						*WriteToString<32>(InnerTag.Name), *WriteToString<64>(InnerTag.GetType().GetParameter(0)), *WriteToString<64>(Builder.Build().GetParameter(0)), *UnderlyingArchive.GetArchiveName());
+						*WriteToString<32>(InnerTag.Name), *WriteToString<64>(InnerTag.GetType().GetParameter(0)), *WriteToString<64>(UE::FPropertyTypeName(StructProperty).GetParameter(0)), *UnderlyingArchive.GetArchiveName());
 
 #if WITH_EDITOR
 					// Ensure the structure is initialized
@@ -1271,10 +1269,8 @@ EConvertFromTypeResult FArrayProperty::ConvertFromType(const FPropertyTag& Tag,
 	}
 	else
 	{
-		UE::FPropertyTypeNameBuilder Builder;
-		Inner->SaveTypeName(Builder);
 		UE_LOG(LogClass, Warning, TEXT("Array Inner Type mismatch in %s - Previous (%s) Current(%s) in package: %s"),
-			*WriteToString<32>(Tag.Name), *WriteToString<32>(InnerPropertyTag.GetType()), *WriteToString<32>(Builder.Build()), *UnderlyingArchive.GetArchiveName());
+			*WriteToString<32>(Tag.Name), *WriteToString<32>(InnerPropertyTag.GetType()), *WriteToString<32>(UE::FPropertyTypeName(Inner)), *UnderlyingArchive.GetArchiveName());
 		return EConvertFromTypeResult::CannotConvert;
 	}
 }
diff --git a/Engine/Source/Runtime/CoreUObject/Private/UObject/PropertyBagRepository.cpp b/Engine/Source/Runtime/CoreUObject/Private/UObject/PropertyBagRepository.cpp
--- a/Engine/Source/Runtime/CoreUObject/Private/UObject/PropertyBagRepository.cpp
+++ b/Engine/Source/Runtime/CoreUObject/Private/UObject/PropertyBagRepository.cpp
@@ -189,11 +189,8 @@ static FProperty* FindPropertyByNameAndType(const UStruct* Struct, const FProper
 	// use property impersonation for SaveTypeName so that keys of IDOs and non-IDOs match
 	FUObjectSerializeContext* SerializeContext = FUObjectThreadContext::Get().GetSerializeContext();
 	TGuardValue<bool> ScopedImpersonateProperties(SerializeContext->bImpersonateProperties, true);
-	
-	FPropertyTypeNameBuilder SourceTypeNameBuilder;
-	SourceProperty->SaveTypeName(SourceTypeNameBuilder);
-	FPropertyTypeName SourceTypeName = SourceTypeNameBuilder.Build();
-	return FindPropertyByNameAndTypeName(Struct, SourceProperty->GetFName(), SourceTypeName);
+
+	return FindPropertyByNameAndTypeName(Struct, SourceProperty->GetFName(), FPropertyTypeName(SourceProperty));
 }
 
 	
diff --git a/Engine/Source/Runtime/CoreUObject/Private/UObject/PropertyMap.cpp b/Engine/Source/Runtime/CoreUObject/Private/UObject/PropertyMap.cpp
--- a/Engine/Source/Runtime/CoreUObject/Private/UObject/PropertyMap.cpp
+++ b/Engine/Source/Runtime/CoreUObject/Private/UObject/PropertyMap.cpp
@@ -1585,17 +1585,13 @@ EConvertFromTypeResult FMapProperty::ConvertFromType(const FPropertyTag& Tag, FS
 	// if we could not convert the property ourself, then indicate that calling code needs to advance the property
 	if (!bConversionSucceeded)
 	{
-		UE::FPropertyTypeNameBuilder KeyBuilder;
-		UE::FPropertyTypeNameBuilder ValueBuilder;
-		KeyProp->SaveTypeName(KeyBuilder);
-		ValueProp->SaveTypeName(ValueBuilder);
 		UE_LOG(LogClass, Warning,
 			TEXT("Map Element Type mismatch in %s - Previous (%s to %s) Current (%s to %s) for package: %s"),
 			*WriteToString<32>(Tag.Name),
 			*WriteToString<32>(KeyPropertyTag.GetType()),
 			*WriteToString<32>(ValuePropertyTag.GetType()),
-			*WriteToString<32>(KeyBuilder.Build()),
-			*WriteToString<32>(ValueBuilder.Build()),
+			*WriteToString<32>(UE::FPropertyTypeName(KeyProp)),
+			*WriteToString<32>(UE::FPropertyTypeName(ValueProp)),
 			*UnderlyingArchive.GetArchiveName());
 	}
 
diff --git a/Engine/Source/Runtime/CoreUObject/Private/UObject/PropertySet.cpp b/Engine/Source/Runtime/CoreUObject/Private/UObject/PropertySet.cpp
--- a/Engine/Source/Runtime/CoreUObject/Private/UObject/PropertySet.cpp
+++ b/Engine/Source/Runtime/CoreUObject/Private/UObject/PropertySet.cpp
@@ -1005,10 +1005,8 @@ EConvertFromTypeResult FSetProperty::ConvertFromType(const FPropertyTag& Tag, FS
 	// if we could not convert the property ourself, then indicate that calling code needs to advance the property
 	if (!bConversionSucceeded)
 	{
-		UE::FPropertyTypeNameBuilder Builder;
-		ElementProp->SaveTypeName(Builder);
 		UE_LOG(LogClass, Warning, TEXT("Set Element Type mismatch in %s - Previous (%s) Current (%s) for package: %s"),
-			*WriteToString<32>(Tag.Name), *WriteToString<32>(InnerPropertyTag.GetType()), *WriteToString<32>(Builder.Build()), *UnderlyingArchive.GetArchiveName());
+			*WriteToString<32>(Tag.Name), *WriteToString<32>(InnerPropertyTag.GetType()), *WriteToString<32>(UE::FPropertyTypeName(ElementProp)), *UnderlyingArchive.GetArchiveName());
 	}
 
 	return bConversionSucceeded ? EConvertFromTypeResult::Converted : EConvertFromTypeResult::CannotConvert;
diff --git a/Engine/Source/Runtime/CoreUObject/Private/UObject/PropertyStruct.cpp b/Engine/Source/Runtime/CoreUObject/Private/UObject/PropertyStruct.cpp
--- a/Engine/Source/Runtime/CoreUObject/Private/UObject/PropertyStruct.cpp
+++ b/Engine/Source/Runtime/CoreUObject/Private/UObject/PropertyStruct.cpp
@@ -513,10 +513,8 @@ EConvertFromTypeResult FStructProperty::ConvertFromType(const FPropertyTag& Tag,
 				}
 				else
 				{
-					UE::FPropertyTypeNameBuilder Builder;
-					SaveTypeName(Builder);
 					UE_LOG(LogClass, Warning, TEXT("SerializeFromMismatchedTag failed: Type mismatch in %s - Previous (%s) Current(%s) in package: %s"),
-						*WriteToString<32>(Tag.Name), *WriteToString<64>(Tag.GetType()), *WriteToString<64>(Builder.Build()), *UnderlyingArchive.GetArchiveName());
+						*WriteToString<32>(Tag.Name), *WriteToString<64>(Tag.GetType()), *WriteToString<64>(UE::FPropertyTypeName(this)), *UnderlyingArchive.GetArchiveName());
 					return EConvertFromTypeResult::CannotConvert;
 				}
 			}
@@ -524,10 +522,8 @@ EConvertFromTypeResult FStructProperty::ConvertFromType(const FPropertyTag& Tag,
 
 		if (Tag.Type == NAME_StructProperty && !bCanSerialize && (UnderlyingArchive.UEVer() >= VER_UE4_STRUCT_GUID_IN_PROPERTY_TAG || !UseBinaryOrNativeSerialization(UnderlyingArchive)))
 		{
-			UE::FPropertyTypeNameBuilder Builder;
-			SaveTypeName(Builder);
 			UE_LOG(LogClass, Warning, TEXT("Struct Property %s has a struct type mismatch (tag %s != prop %s) in package: %s. If that struct got renamed, add an entry to ActiveStructRedirects."),
-				*WriteToString<32>(Tag.Name), *WriteToString<64>(Tag.GetType().GetParameter(0)), *WriteToString<64>(Builder.Build().GetParameter(0)), *UnderlyingArchive.GetArchiveName());
+				*WriteToString<32>(Tag.Name), *WriteToString<64>(Tag.GetType().GetParameter(0)), *WriteToString<64>(UE::FPropertyTypeName(this).GetParameter(0)), *UnderlyingArchive.GetArchiveName());
 			return EConvertFromTypeResult::CannotConvert;
 		}
 	}
diff --git a/Engine/Source/Runtime/CoreUObject/Private/UObject/PropertyTag.cpp b/Engine/Source/Runtime/CoreUObject/Private/UObject/PropertyTag.cpp
--- a/Engine/Source/Runtime/CoreUObject/Private/UObject/PropertyTag.cpp
+++ b/Engine/Source/Runtime/CoreUObject/Private/UObject/PropertyTag.cpp
@@ -63,9 +63,7 @@ FPropertyTag::FPropertyTag(FProperty* Property, int32 InIndex, uint8* Value)
 	, ArrayIndex(InIndex)
 	, OverrideOperation(EOverriddenPropertyOperation::None)
 {
-	UE::FPropertyTypeNameBuilder TypeBuilder;
-	Property->SaveTypeName(TypeBuilder);
-	SetType(TypeBuilder.Build());
+	SetType(UE::FPropertyTypeName(Property));
 
 	if (FBoolProperty* Bool = CastField<FBoolProperty>(Property))
 	{
diff --git a/Engine/Source/Runtime/CoreUObject/Tests/UObject/PropertyPathFunctionsTest.cpp b/Engine/Source/Runtime/CoreUObject/Tests/UObject/PropertyPathFunctionsTest.cpp
--- a/Engine/Source/Runtime/CoreUObject/Tests/UObject/PropertyPathFunctionsTest.cpp
+++ b/Engine/Source/Runtime/CoreUObject/Tests/UObject/PropertyPathFunctionsTest.cpp
@@ -19,9 +19,7 @@ TEST_CASE_NAMED(FPropertyPathFunctionsFindPropertyTest, "CoreUObject::PropertyPa
 		const UStruct* TestType = UTestPropertyPathFunctionsClass::StaticClass();
 		for (const FProperty* Property = TestType->PropertyLink; Property; Property = Property->PropertyLinkNext)
 		{
-			FPropertyTypeNameBuilder TypeName;
-			Property->SaveTypeName(TypeName);
-			CHECK(Property == FindPropertyByNameAndTypeName(TestType, Property->GetFName(), TypeName.Build()));
+			CHECK(Property == FindPropertyByNameAndTypeName(TestType, Property->GetFName(), FPropertyTypeName(Property)));
 		}
 	}
 
@@ -75,7 +73,7 @@ TEST_CASE_NAMED(FPropertyPathFunctionsTryResolveTest, "CoreUObject::PropertyPath
 		for (const FPropertyReference& Segment : Segments)
 		{
 			FName PropertyName = Segment.PropertyName;
-			FPropertyTypeNameBuilder TypeName;
+			FPropertyTypeName TypeName;
 			if (Segment.Struct)
 			{
 				FProperty* Property = Segment.Struct->FindPropertyByName(PropertyName);
@@ -83,10 +81,10 @@ TEST_CASE_NAMED(FPropertyPathFunctionsTryResolveTest, "CoreUObject::PropertyPath
 				CHECK(Property);
 				if (Property)
 				{
-					Property->SaveTypeName(TypeName);
+					TypeName = FPropertyTypeName(Property);
 				}
 			}
-			PathName.Push({PropertyName, TypeName.Build(), Segment.ArrayIndex});
+			PathName.Push({PropertyName, TypeName, Segment.ArrayIndex});
 		}
 		return PathName;
 	};
diff --git a/Engine/Source/Runtime/CoreUObject/Tests/UObject/PropertyStateTrackingTest.cpp b/Engine/Source/Runtime/CoreUObject/Tests/UObject/PropertyStateTrackingTest.cpp
--- a/Engine/Source/Runtime/CoreUObject/Tests/UObject/PropertyStateTrackingTest.cpp
+++ b/Engine/Source/Runtime/CoreUObject/Tests/UObject/PropertyStateTrackingTest.cpp
@@ -448,13 +448,6 @@ TEST_CASE_NAMED(FTrackUnknownPropertiesTest, "CoreUObject::Serialization::TrackU
 		return Builder.Build();
 	};
 
-	const auto SavePropertyTypeName = [](const FProperty* Property)
-	{
-		FPropertyTypeNameBuilder Builder;
-		Property->SaveTypeName(Builder);
-		return Builder.Build();
-	};
-
 	UObject* Owner = NewObject<UTestInstanceDataObjectClass>();
 
 	FUObjectSerializeContext* SerializeContext = FUObjectThreadContext::Get().GetSerializeContext();
@@ -552,9 +545,7 @@ TEST_CASE_NAMED(FTrackUnknownPropertiesTest, "CoreUObject::Serialization::TrackU
 		FProperty* PointProperty = FTestInstanceDataObjectStructAlternate::StaticStruct()->FindPropertyByName("Point");
 		CHECKED_IF(PointProperty)
 		{
-			FPropertyTypeNameBuilder Builder;
-			PointProperty->SaveTypeName(Builder);
-			FSerializedPropertyPathScope Path(SerializeContext, {"Point", Builder.Build()});
+			FSerializedPropertyPathScope Path(SerializeContext, {"Point", FPropertyTypeName(PointProperty)});
 			CHECK(Tree->Find(SerializeContext->SerializedPropertyPath));
 		#if WITH_METADATA
 			{
@@ -573,19 +564,19 @@ TEST_CASE_NAMED(FTrackUnknownPropertiesTest, "CoreUObject::Serialization::TrackU
 		}
 	}
 	{
-		FSerializedPropertyPathScope Path(SerializeContext, {"GrainTypeChange", SavePropertyTypeName(FindFProperty<FProperty>(StaticStruct<FTestInstanceDataObjectStructAlternate>(), TEXT("GrainTypeChange")))});
+		FSerializedPropertyPathScope Path(SerializeContext, {"GrainTypeChange", FPropertyTypeName(FindFProperty<FProperty>(StaticStruct<FTestInstanceDataObjectStructAlternate>(), TEXT("GrainTypeChange")))});
 		CHECK(Tree->Find(SerializeContext->SerializedPropertyPath));
 	}
 	{
-		FSerializedPropertyPathScope Path(SerializeContext, {"FruitTypeChange", SavePropertyTypeName(FindFProperty<FProperty>(StaticStruct<FTestInstanceDataObjectStructAlternate>(), TEXT("FruitTypeChange")))});
+		FSerializedPropertyPathScope Path(SerializeContext, {"FruitTypeChange", FPropertyTypeName(FindFProperty<FProperty>(StaticStruct<FTestInstanceDataObjectStructAlternate>(), TEXT("FruitTypeChange")))});
 		CHECK(Tree->Find(SerializeContext->SerializedPropertyPath));
 	}
 	{
-		FSerializedPropertyPathScope Path(SerializeContext, {"GrainTypeAndPropertyChange", SavePropertyTypeName(FindFProperty<FProperty>(StaticStruct<FTestInstanceDataObjectStructAlternate>(), TEXT("GrainTypeAndPropertyChange")))});
+		FSerializedPropertyPathScope Path(SerializeContext, {"GrainTypeAndPropertyChange", FPropertyTypeName(FindFProperty<FProperty>(StaticStruct<FTestInstanceDataObjectStructAlternate>(), TEXT("GrainTypeAndPropertyChange")))});
 		CHECK(Tree->Find(SerializeContext->SerializedPropertyPath));
 	}
 	{
-		FSerializedPropertyPathScope Path(SerializeContext, {"FruitTypeAndPropertyChange", SavePropertyTypeName(FindFProperty<FProperty>(StaticStruct<FTestInstanceDataObjectStructAlternate>(), TEXT("FruitTypeAndPropertyChange")))});
+		FSerializedPropertyPathScope Path(SerializeContext, {"FruitTypeAndPropertyChange", FPropertyTypeName(FindFProperty<FProperty>(StaticStruct<FTestInstanceDataObjectStructAlternate>(), TEXT("FruitTypeAndPropertyChange")))});
 		CHECK(Tree->Find(SerializeContext->SerializedPropertyPath));
 	}
 	TreeAccessor.Destroy();
