PackageName: improve SplitPackageName functions.

Add flags to specify whether slashes are kept.

Move all of the SplitPackageName functions to be adjacent in the header.

Add SplitPackageName that does the same splitting done by SplitLongPackageName but does not test for IsValidPath.

Fix IsValidPath to return true for /Memory/ packages, to match the behavior of other functions such as MountPointExists.
[RN] Minor, Core

diff --git a/Engine/Source/Programs/LowLevelTests/FoundationTests/Tests/PackageNameTests.cpp b/Engine/Source/Programs/LowLevelTests/FoundationTests/Tests/PackageNameTests.cpp
--- a/Engine/Source/Programs/LowLevelTests/FoundationTests/Tests/PackageNameTests.cpp
+++ b/Engine/Source/Programs/LowLevelTests/FoundationTests/Tests/PackageNameTests.cpp
@@ -140,4 +140,217 @@ TEST_CASE("Split full object path with subobjects", "[CoreUObject][FPackageName:
 		REQUIRE(SubobjectNames[0] == ExpectedSubobject1Name);
 		REQUIRE(SubobjectNames[1] == ExpectedSubobject2Name);
 	}
+}
+
+TEST_CASE("Split PackageName", "[CoreUObject][FPackageName::SplitPackageName]")
+{
+	FString PackageName(TEXT("/root/path1/path2/leaf.umap"));
+
+	// Test that not passing outputs does not crash
+	FPackageName::SplitPackageName(PackageName, nullptr, nullptr, nullptr);
+
+	FStringView Root;
+	FStringView Path;
+	FStringView Leaf;
+	FPackageName::SplitPackageName(PackageName, &Root, &Path, &Leaf);
+	REQUIRE(Root == TEXTVIEW("/root/"));
+	REQUIRE(Path == TEXTVIEW("path1/path2/"));
+	REQUIRE(Leaf == TEXTVIEW("leaf.umap"));
+
+	PackageName = TEXTVIEW("/root/path1/path2/");
+	FPackageName::SplitPackageName(PackageName, &Root, &Path, &Leaf);
+	REQUIRE(Root == TEXTVIEW("/root/"));
+	REQUIRE(Path == TEXTVIEW("path1/path2/"));
+	REQUIRE(Leaf == TEXTVIEW(""));
+
+	PackageName = TEXTVIEW("/root/path1/leaf");
+	FPackageName::SplitPackageName(PackageName, &Root, &Path, &Leaf);
+	REQUIRE(Root == TEXTVIEW("/root/"));
+	REQUIRE(Path == TEXTVIEW("path1/"));
+	REQUIRE(Leaf == TEXTVIEW("leaf"));
+
+	PackageName = TEXTVIEW("/root/path1/");
+	FPackageName::SplitPackageName(PackageName, &Root, &Path, &Leaf);
+	REQUIRE(Root == TEXTVIEW("/root/"));
+	REQUIRE(Path == TEXTVIEW("path1/"));
+	REQUIRE(Leaf == TEXTVIEW(""));
+
+	PackageName = TEXTVIEW("/root/leaf");
+	FPackageName::SplitPackageName(PackageName, &Root, &Path, &Leaf);
+	REQUIRE(Root == TEXTVIEW("/root/"));
+	REQUIRE(Path == TEXTVIEW(""));
+	REQUIRE(Leaf == TEXTVIEW("leaf"));
+
+	PackageName = TEXTVIEW("/root/");
+	FPackageName::SplitPackageName(PackageName, &Root, &Path, &Leaf);
+	REQUIRE(Root == TEXTVIEW("/root/"));
+	REQUIRE(Path == TEXTVIEW(""));
+	REQUIRE(Leaf == TEXTVIEW(""));
+
+	PackageName = TEXTVIEW("/root");
+	FPackageName::SplitPackageName(PackageName, &Root, &Path, &Leaf);
+	REQUIRE(Root == TEXTVIEW(""));
+	REQUIRE(Path == TEXTVIEW(""));
+	REQUIRE(Leaf == TEXTVIEW(""));
+
+	PackageName = TEXTVIEW("/");
+	FPackageName::SplitPackageName(PackageName, &Root, &Path, &Leaf);
+	REQUIRE(Root == TEXTVIEW(""));
+	REQUIRE(Path == TEXTVIEW(""));
+	REQUIRE(Leaf == TEXTVIEW(""));
+
+	PackageName = TEXTVIEW("");
+	FPackageName::SplitPackageName(PackageName, &Root, &Path, &Leaf);
+	REQUIRE(Root == TEXTVIEW(""));
+	REQUIRE(Path == TEXTVIEW(""));
+	REQUIRE(Leaf == TEXTVIEW(""));
+
+	PackageName = TEXTVIEW("path1/path2");
+	FPackageName::SplitPackageName(PackageName, &Root, &Path, &Leaf);
+	REQUIRE(Root == TEXTVIEW(""));
+	REQUIRE(Path == TEXTVIEW(""));
+	REQUIRE(Leaf == TEXTVIEW(""));
+
+	// Edge cases with no contract yet; they should not crash
+	PackageName = TEXTVIEW("//");
+	FPackageName::SplitPackageName(PackageName, &Root, &Path, &Leaf);
+
+	PackageName = TEXTVIEW("///");
+	FPackageName::SplitPackageName(PackageName, &Root, &Path, &Leaf);
+
+	// Path formats
+	PackageName = TEXTVIEW("/root/path/leaf");
+	FPackageName::SplitPackageName(PackageName, &Root, &Path, &Leaf, FPackageName::EPathFormatFlags::MountPointNoSlashes);
+	REQUIRE(Root == TEXTVIEW("root"));
+	REQUIRE(Path == TEXTVIEW("path/"));
+	REQUIRE(Leaf == TEXTVIEW("leaf"));
+
+	FPackageName::SplitPackageName(PackageName, &Root, &Path, &Leaf, FPackageName::EPathFormatFlags::MountPointLeadingSlash);
+	REQUIRE(Root == TEXTVIEW("/root"));
+	REQUIRE(Path == TEXTVIEW("path/"));
+	REQUIRE(Leaf == TEXTVIEW("leaf"));
+
+	FPackageName::SplitPackageName(PackageName, &Root, &Path, &Leaf, FPackageName::EPathFormatFlags::MountPointTrailingSlash);
+	REQUIRE(Root == TEXTVIEW("root/"));
+	REQUIRE(Path == TEXTVIEW("path/"));
+	REQUIRE(Leaf == TEXTVIEW("leaf"));
+
+	FPackageName::SplitPackageName(PackageName, &Root, &Path, &Leaf, FPackageName::EPathFormatFlags::MountPointSlashes);
+	REQUIRE(Root == TEXTVIEW("/root/"));
+	REQUIRE(Path == TEXTVIEW("path/"));
+	REQUIRE(Leaf == TEXTVIEW("leaf"));
+
+	// Verify Path formats are respected in the FString version
+	FString RootStr;
+	FString PathStr;
+	FString LeafStr;
+	PackageName = TEXTVIEW("/root/path/leaf");
+	FPackageName::SplitPackageName(PackageName, &RootStr, &PathStr, &LeafStr, FPackageName::EPathFormatFlags::MountPointNoSlashes);
+	REQUIRE(RootStr == TEXTVIEW("root"));
+	REQUIRE(PathStr == TEXTVIEW("path/"));
+	REQUIRE(LeafStr == TEXTVIEW("leaf"));
+
+	FPackageName::SplitPackageName(PackageName, &RootStr, &PathStr, &LeafStr, FPackageName::EPathFormatFlags::MountPointLeadingSlash);
+	REQUIRE(RootStr == TEXTVIEW("/root"));
+	REQUIRE(PathStr == TEXTVIEW("path/"));
+	REQUIRE(LeafStr == TEXTVIEW("leaf"));
+
+	FPackageName::SplitPackageName(PackageName, &RootStr, &PathStr, &LeafStr, FPackageName::EPathFormatFlags::MountPointTrailingSlash);
+	REQUIRE(RootStr == TEXTVIEW("root/"));
+	REQUIRE(PathStr == TEXTVIEW("path/"));
+	REQUIRE(LeafStr == TEXTVIEW("leaf"));
+
+	FPackageName::SplitPackageName(PackageName, &RootStr, &PathStr, &LeafStr, FPackageName::EPathFormatFlags::MountPointSlashes);
+	REQUIRE(RootStr == TEXTVIEW("/root/"));
+	REQUIRE(PathStr == TEXTVIEW("path/"));
+	REQUIRE(LeafStr == TEXTVIEW("leaf"));
+}
+
+TEST_CASE("Split PackageNameRoot", "[CoreUObject][FPackageName::SplitPackageNameRoot]")
+{
+	FString PackageName(TEXT("/root/path1/path2/leaf.umap"));
+
+	// Test that not passing outputs does not crash
+	(void) FPackageName::SplitPackageNameRoot(PackageName, nullptr);
+
+	FStringView Root;
+	FStringView RelPath;
+	Root = FPackageName::SplitPackageNameRoot(PackageName, &RelPath);
+	REQUIRE(Root == TEXTVIEW("root"));
+	REQUIRE(RelPath == TEXTVIEW("path1/path2/leaf.umap"));
+
+	PackageName = TEXTVIEW("/root/path1/leaf");
+	Root = FPackageName::SplitPackageNameRoot(PackageName, &RelPath);
+	REQUIRE(Root == TEXTVIEW("root"));
+	REQUIRE(RelPath == TEXTVIEW("path1/leaf"));
+
+	PackageName = TEXTVIEW("/root/");
+	Root = FPackageName::SplitPackageNameRoot(PackageName, &RelPath);
+	REQUIRE(Root == TEXTVIEW("root"));
+	REQUIRE(RelPath == TEXTVIEW(""));
+
+	PackageName = TEXTVIEW("/root");
+	Root = FPackageName::SplitPackageNameRoot(PackageName, &RelPath);
+	REQUIRE(Root == TEXTVIEW("/root"));
+	REQUIRE(RelPath == TEXTVIEW(""));
+
+	PackageName = TEXTVIEW("/");
+	Root = FPackageName::SplitPackageNameRoot(PackageName, &RelPath);
+	REQUIRE(Root == TEXTVIEW("/"));
+	REQUIRE(RelPath == TEXTVIEW(""));
+
+	PackageName = TEXTVIEW("");
+	Root = FPackageName::SplitPackageNameRoot(PackageName, &RelPath);
+	REQUIRE(Root == TEXTVIEW(""));
+	REQUIRE(RelPath == TEXTVIEW(""));
+
+	PackageName = TEXTVIEW("path1/path2");
+	Root = FPackageName::SplitPackageNameRoot(PackageName, &RelPath);
+	REQUIRE(Root == TEXTVIEW(""));
+	REQUIRE(RelPath == TEXTVIEW("path1/path2"));
+
+	// Edge cases with no contract yet; they should not crash
+	PackageName = TEXTVIEW("//");
+	Root = FPackageName::SplitPackageNameRoot(PackageName, &RelPath);
+
+	PackageName = TEXTVIEW("///");
+	Root = FPackageName::SplitPackageNameRoot(PackageName, &RelPath);
+
+	// Path formats
+	PackageName = TEXTVIEW("/root/path/leaf");
+	Root = FPackageName::SplitPackageNameRoot(PackageName, &RelPath, FPackageName::EPathFormatFlags::MountPointNoSlashes);
+	REQUIRE(Root == TEXTVIEW("root"));
+	REQUIRE(RelPath == TEXTVIEW("path/leaf"));
+
+	Root = FPackageName::SplitPackageNameRoot(PackageName, &RelPath, FPackageName::EPathFormatFlags::MountPointLeadingSlash);
+	REQUIRE(Root == TEXTVIEW("/root"));
+	REQUIRE(RelPath == TEXTVIEW("path/leaf"));
+
+	Root = FPackageName::SplitPackageNameRoot(PackageName, &RelPath, FPackageName::EPathFormatFlags::MountPointTrailingSlash);
+	REQUIRE(Root == TEXTVIEW("root/"));
+	REQUIRE(RelPath == TEXTVIEW("path/leaf"));
+
+	Root = FPackageName::SplitPackageNameRoot(PackageName, &RelPath, FPackageName::EPathFormatFlags::MountPointSlashes);
+	REQUIRE(Root == TEXTVIEW("/root/"));
+	REQUIRE(RelPath == TEXTVIEW("path/leaf"));
+
+	// Verify Path formats are respected in the FName version
+	FString RootStr;
+	FString RelPathStr;
+	PackageName = TEXTVIEW("/root/path/leaf");
+	RootStr = FPackageName::SplitPackageNameRoot(FName(*PackageName), &RelPathStr, FPackageName::EPathFormatFlags::MountPointNoSlashes);
+	REQUIRE(RootStr == TEXTVIEW("root"));
+	REQUIRE(RelPathStr == TEXTVIEW("path/leaf"));
+
+	RootStr = FPackageName::SplitPackageNameRoot(FName(*PackageName), &RelPathStr, FPackageName::EPathFormatFlags::MountPointLeadingSlash);
+	REQUIRE(RootStr == TEXTVIEW("/root"));
+	REQUIRE(RelPathStr == TEXTVIEW("path/leaf"));
+
+	RootStr = FPackageName::SplitPackageNameRoot(FName(*PackageName), &RelPathStr, FPackageName::EPathFormatFlags::MountPointTrailingSlash);
+	REQUIRE(RootStr == TEXTVIEW("root/"));
+	REQUIRE(RelPathStr == TEXTVIEW("path/leaf"));
+
+	RootStr = FPackageName::SplitPackageNameRoot(FName(*PackageName), &RelPathStr, FPackageName::EPathFormatFlags::MountPointSlashes);
+	REQUIRE(RootStr == TEXTVIEW("/root/"));
+	REQUIRE(RelPathStr == TEXTVIEW("path/leaf"));
 }
\ No newline at end of file
diff --git a/Engine/Source/Runtime/CoreUObject/Private/Misc/PackageName.cpp b/Engine/Source/Runtime/CoreUObject/Private/Misc/PackageName.cpp
--- a/Engine/Source/Runtime/CoreUObject/Private/Misc/PackageName.cpp
+++ b/Engine/Source/Runtime/CoreUObject/Private/Misc/PackageName.cpp
@@ -419,7 +419,7 @@ struct FLongPackagePathsSingleton
 	}
 
 	// Checks whether the specific root path is a valid mount point.
-	bool MountPointExists(const FString& RootPath)
+	bool MountPointExists(FStringView RootPath)
 	{
 		UE::TReadScopeLock ScopeLock(MountLock);
 		return RootPathTree.Contains(RootPath) || RootPath == MemoryRootPath;
@@ -1031,56 +1031,6 @@ FString FPackageName::GetLongPackagePath(const FString& InLongPackageName)
 	}
 }
 
-bool FPackageName::SplitLongPackageName(const FString& InLongPackageName, FString& OutPackageRoot, FString& OutPackagePath, FString& OutPackageName, const bool bStripRootLeadingSlash)
-{
-	const FLongPackagePathsSingleton& Paths = FLongPackagePathsSingleton::Get();
-
-	FStringView PackageRelPath;
-	{
-		UE::TReadScopeLock ScopeLock(Paths.MountLock);
-
-		// Check to see whether our package came from a valid root
-		OutPackageRoot.Empty();
-		const FMountPoint* const* MountPointPtr = Paths.RootPathTree.FindClosestValue(InLongPackageName);
-		if (MountPointPtr)
-		{
-			// DirectoryTree is less conservative than TryMakeChildPathRelativeTo; e.g. d:/Dir//Root will match
-			// d:/Dir/Root/Child in DirectoryTree but not in TryMakeChildPathRelativeTo. Treat this as PathNotMounted,
-			// since no mountdir will match it according to TryMakeChildPathRelativeTo.
-			if (FPathViews::TryMakeChildPathRelativeTo(InLongPackageName, (*MountPointPtr)->RootPath, PackageRelPath))
-			{
-				OutPackageRoot = (*MountPointPtr)->RootPath / "";
-			}
-		}
-		else
-		{
-			const FString& ExtraPackageRoot = Paths.MemoryRootPath;
-			if (FPathViews::TryMakeChildPathRelativeTo(InLongPackageName, ExtraPackageRoot, PackageRelPath))
-			{
-				OutPackageRoot = ExtraPackageRoot / "";
-			}
-		}
-	}
-
-	if (OutPackageRoot.IsEmpty() || InLongPackageName.Len() <= OutPackageRoot.Len())
-	{
-		// Path is not part of a valid root, or the path given is too short to continue; splitting failed
-		return false;
-	}
-
-	// Use the standard path functions to get the rest
-	FString PackageRelPathStr(PackageRelPath);
-	OutPackagePath = FPaths::GetPath(PackageRelPathStr) / "";
-	OutPackageName = FPaths::GetCleanFilename(PackageRelPathStr);
-
-	if (bStripRootLeadingSlash && OutPackageRoot.StartsWith(TEXT("/"), ESearchCase::CaseSensitive))
-	{
-		OutPackageRoot.RemoveAt(0);
-	}
-
-	return true;
-}
-
 void FPackageName::SplitFullObjectPath(const FString& InFullObjectPath, FString& OutClassName,
 	FString& OutPackageName, FString& OutObjectName, FString& OutSubObjectName, bool bDetectClassName)
 {
@@ -1440,7 +1390,8 @@ bool FPackageName::IsValidPath(FStringView InPath)
 	const FLongPackagePathsSingleton& Paths = FLongPackagePathsSingleton::Get();
 	UE::TReadScopeLock ScopeLock(Paths.MountLock);
 	// Note that we return false for the root folder "/".
-	return Paths.RootPathTree.FindClosestValue(InPath) != nullptr;
+	return Paths.RootPathTree.FindClosestValue(InPath) != nullptr ||
+		IsMemoryPackage(InPath);
 }
 
 void FPackageName::RegisterMountPoint(const FString& RootPath, const FString& ContentPath)
@@ -1453,11 +1404,21 @@ void FPackageName::UnRegisterMountPoint(const FString& RootPath, const FString&
 	FLongPackagePathsSingleton::Get().RemoveMountPoint(RootPath, ContentPath);
 }
 
+bool FPackageName::MountPointExists(FStringView RootPath)
+{
+	return FLongPackagePathsSingleton::Get().MountPointExists(RootPath);
+}
+
 bool FPackageName::MountPointExists(const FString& RootPath)
 {
 	return FLongPackagePathsSingleton::Get().MountPointExists(RootPath);
 }
 
+bool FPackageName::MountPointExists(const TCHAR* RootPath)
+{
+	return FLongPackagePathsSingleton::Get().MountPointExists(FStringView(RootPath));
+}
+
 FName FPackageName::GetPackageMountPoint(const FString& InPackagePath, bool InWithoutSlashes)
 {
 	FLongPackagePathsSingleton& Paths = FLongPackagePathsSingleton::Get();
@@ -2863,13 +2824,127 @@ FString FPackageName::ExportTextPathToObjectPath(const TCHAR* InExportTextPath)
 	return ExportTextPathToObjectPath(FString(InExportTextPath));
 }
 
-FString FPackageName::SplitPackageNameRoot(FName InPackageName, FString* OutRelativePath)
+void FPackageName::SplitPackageName(FStringView InPackageName, FStringView* OutPackageRoot,
+	FStringView* OutPackagePath, FStringView* OutPackageLeafName, EPathFormatFlags Flags)
+{
+	auto SetOutput = [OutPackageRoot, OutPackagePath, OutPackageLeafName]
+	(FStringView InPackageRoot, FStringView InPackagePath, FStringView InPackageLeafName)
+		{
+			if (OutPackageRoot)
+			{
+				*OutPackageRoot = InPackageRoot;
+			}
+			if (OutPackagePath)
+			{
+				*OutPackagePath = InPackagePath;
+			}
+			if (OutPackageLeafName)
+			{
+				*OutPackageLeafName = InPackageLeafName;
+			}
+		};
+
+	if (!InPackageName.StartsWith(TEXT("/")))
+	{
+		SetOutput(FStringView(), FStringView(), FStringView());
+		return;
+	}
+
+	int32 SecondSlashIndex;
+	if (!InPackageName.RightChop(1).FindChar('/', SecondSlashIndex))
+	{
+		SetOutput(FStringView(), FStringView(), FStringView());
+		return;
+	}
+	check(SecondSlashIndex != INDEX_NONE);
+	// Change back to an index relative to the beginnning of InPackageName rather than relative to InPackageName[1].
+	SecondSlashIndex += 1; 
+
+	FStringView PackageRoot(InPackageName.Left(SecondSlashIndex + 1));
+	FStringView RemainingPath = InPackageName.RightChop(SecondSlashIndex + 1);
+	FStringView PackagePath;
+	FStringView PackageNameLeaf;
+
+	// Don't use FPathViews, because it has different rules for edge cases than we do, and adds some complexity
+	// around Extension to convert between its rules and our rules.
+
+	int32 LastSlashIndex;
+	if (!RemainingPath.FindLastChar('/', LastSlashIndex))
+	{
+		// "/root/leaf"
+		PackageNameLeaf = RemainingPath;
+	}
+	else
+	{
+		if (LastSlashIndex == RemainingPath.Len() - 1)
+		{
+			// "/root/path/"
+			PackagePath = RemainingPath;
+		}
+		else
+		{
+			// "/root/path/leaf"
+			PackagePath = RemainingPath.Left(LastSlashIndex + 1);
+			PackageNameLeaf = RemainingPath.RightChop(LastSlashIndex + 1);
+		}
+	}
+
+	// Convert to the requested output format
+	if (!EnumHasAnyFlags(Flags, EPathFormatFlags::MountPointLeadingSlash))
+	{
+		PackageRoot.RightChopInline(1);
+	}
+	if (!EnumHasAnyFlags(Flags, EPathFormatFlags::MountPointTrailingSlash))
+	{
+		PackageRoot.LeftChopInline(1);
+	}
+	SetOutput(PackageRoot, PackagePath, PackageNameLeaf);
+}
+
+void FPackageName::SplitPackageName(const FString& InLongPackageName, FString* OutPackageRoot,
+	FString* OutPackagePath, FString* OutPackageLeafName, EPathFormatFlags Flags)
+{
+	FStringView PackageRoot;
+	FStringView PackagePath;
+	FStringView PackageLeafName;
+	SplitPackageName(InLongPackageName, &PackageRoot, &PackagePath, &PackageLeafName, Flags);
+	if (OutPackageRoot)
+	{
+		*OutPackageRoot = PackageRoot;
+	}
+	if (OutPackagePath)
+	{
+		*OutPackagePath = PackagePath;
+	}
+	if (OutPackageLeafName)
+	{
+		*OutPackageLeafName = PackageLeafName;
+	}
+}
+
+bool FPackageName::SplitLongPackageName(const FString& InLongPackageName, FString& OutPackageRoot,
+	FString& OutPackagePath, FString& OutPackageName, const bool bStripRootLeadingSlash)
+{
+	if (!IsValidPath(InLongPackageName))
+	{
+		OutPackageRoot.Reset();
+		OutPackagePath.Reset();
+		OutPackageName.Reset();
+		return false;
+	}
+
+	SplitPackageName(InLongPackageName, &OutPackageRoot, &OutPackagePath, &OutPackageName,
+		bStripRootLeadingSlash ? EPathFormatFlags::MountPointTrailingSlash : EPathFormatFlags::MountPointSlashes);
+	return true;
+}
+
+FString FPackageName::SplitPackageNameRoot(FName InPackageName, FString* OutRelativePath, EPathFormatFlags Flags)
 {
 	TStringBuilder<FName::StringBufferSize> PackageNameStr(InPlace, InPackageName);
 	FStringView PackageName(PackageNameStr);
 
 	FStringView RelativePath;
-	FStringView Root = SplitPackageNameRoot(PackageName, &RelativePath);
+	FStringView Root = SplitPackageNameRoot(PackageName, &RelativePath, Flags);
 
 	if (OutRelativePath)
 	{
@@ -2878,7 +2953,8 @@ FString FPackageName::SplitPackageNameRoot(FName InPackageName, FString* OutRela
 	return FString(Root);
 }
 
-FStringView FPackageName::SplitPackageNameRoot(FStringView InPackageName, FStringView* OutRelativePath)
+FStringView FPackageName::SplitPackageNameRoot(FStringView InPackageName, FStringView* OutRelativePath,
+	EPathFormatFlags Flags)
 {
 	if (!InPackageName.StartsWith(TEXT("/")))
 	{
@@ -2889,12 +2965,8 @@ FStringView FPackageName::SplitPackageNameRoot(FStringView InPackageName, FStrin
 		return FStringView();
 	}
 
-	// Strip the first slash.
-	InPackageName.RightChopInline(1);
-
 	int32 SecondSlashIndex;
-	InPackageName.FindChar('/', SecondSlashIndex);
-	if (SecondSlashIndex == INDEX_NONE)
+	if (!InPackageName.RightChop(1).FindChar('/', SecondSlashIndex))
 	{
 		if (OutRelativePath)
 		{
@@ -2902,12 +2974,26 @@ FStringView FPackageName::SplitPackageNameRoot(FStringView InPackageName, FStrin
 		}
 		return InPackageName;
 	}
+	check(SecondSlashIndex != INDEX_NONE);
+	// Change back to an index relative to the beginnning of InPackageName rather than relative to InPackageName[1].
+	SecondSlashIndex += 1;
 
 	if (OutRelativePath)
 	{
 		*OutRelativePath = InPackageName.RightChop(SecondSlashIndex + 1);
 	}
-	return InPackageName.Left(SecondSlashIndex);
+
+	FStringView PackageRoot = InPackageName.Left(SecondSlashIndex + 1);
+	// Convert to the requested output format
+	if (!EnumHasAnyFlags(Flags, EPathFormatFlags::MountPointLeadingSlash))
+	{
+		PackageRoot.RightChopInline(1);
+	}
+	if (!EnumHasAnyFlags(Flags, EPathFormatFlags::MountPointTrailingSlash))
+	{
+		PackageRoot.LeftChopInline(1);
+	}
+	return PackageRoot;
 };
 
 
diff --git a/Engine/Source/Runtime/CoreUObject/Public/Misc/PackageName.h b/Engine/Source/Runtime/CoreUObject/Public/Misc/PackageName.h
--- a/Engine/Source/Runtime/CoreUObject/Public/Misc/PackageName.h
+++ b/Engine/Source/Runtime/CoreUObject/Public/Misc/PackageName.h
@@ -182,18 +182,6 @@ public:
 	 */
 	static COREUOBJECT_API FString GetLongPackageAssetName(const FString& InLongPackageName);
 
-	/** 
-	 * Convert a long package name into root, path, and name components
-	 *
-	 * @param InLongPackageName Package Name.
-	 * @param OutPackageRoot The package root path, eg "/Game/"
-	 * @param OutPackagePath The path from the mount point to the package, eg "Maps/TestMaps/
-	 * @param OutPackageName The name of the package, including its extension, eg "MyMap.umap"
-	 * @param bStripRootLeadingSlash String any leading / character from the returned root
-	 * @return True if the conversion was possible, false otherwise
-	 */
-	static COREUOBJECT_API bool SplitLongPackageName(const FString& InLongPackageName, FString& OutPackageRoot, FString& OutPackagePath, FString& OutPackageName, const bool bStripRootLeadingSlash = false);
-
 	/**
 	 * Split a full object path (Class /Path/To/A/Package.Object:Subobject1.Subobject2) into its constituent pieces
 	 *  
@@ -265,7 +253,7 @@ public:
 	static COREUOBJECT_API bool IsValidObjectPath(FStringView InObjectPath, FText* OutReason = nullptr);
 
 	/**
-	 * Returns true if the path starts with a valid root (i.e. /Game/, /Engine/, etc).
+	 * Returns true if the path starts with a valid root (i.e. /Game/, /Engine/, /Memory/, etc).
 	 * 
 	 * @param InObjectPath				The object path to test
 	 * @return							true if a valid object path
@@ -345,7 +333,9 @@ public:
 	/**
 	 * Returns whether the specific logical root path is a valid mount point.
 	 */
+	static COREUOBJECT_API bool MountPointExists(FStringView RootPath);
 	static COREUOBJECT_API bool MountPointExists(const FString& RootPath);
+	static COREUOBJECT_API bool MountPointExists(const TCHAR* RootPath);
 
 	/**
 	 * Get the mount point for a given package path
@@ -725,17 +715,76 @@ public:
 	static COREUOBJECT_API FString			ExportTextPathToObjectPath(const FString& InExportTextPath);
 	static COREUOBJECT_API FString			ExportTextPathToObjectPath(const TCHAR* InExportTextPath);
 
+	/** Flags for the output format of path functions. */
+	enum class EPathFormatFlags : uint32
+	{
+		None = 0x0,
+		MountPointLeadingSlash = 0x1,
+		MountPointTrailingSlash = 0x2,
+		MountPointSlashes = MountPointLeadingSlash | MountPointTrailingSlash,
+		MountPointNoSlashes = None,
+	};
+
+	/**
+	 * Convert a long package name into root, path, and leafname components.
+	 * For the FStringView version, the returned FStringViews are views into the input FStringView.
+	 *
+	 * @param InLongPackageName Package Name, e.g. "/Game/Maps/TestMaps/MyMap".
+	 *        Also supports packagenames with extension (which are not valid package names),
+	 *        e.g. /Game/Maps/TestMaps/MyMap.umap.
+	 * @param OutPackageRoot The package root path, eg "/Game/"
+	 * @param OutPackagePath The path from the mount point to the package, eg "Maps/TestMaps/
+	 * @param OutPackageLeafName The leaf name, including extension if present, eg "MyMap" or "MyMap.umap".
+	 * @param Flags Flags specifying which characters are included in OutputPackageRoot, @see EPathFormatFlags.
+	 * 
+	 * Edge cases examples:
+	 * "" -> "","",""
+	 * "/" -> "", "", ""
+	 * "/root" -> "", "", ""
+	 * "/root/" -> "/root/", "", ""
+	 * "/root/leaf" -> "/root/", "", "leaf"
+	 * "/root/path/" -> "/root/", "path/", ""
+	 * "/root/path/leaf" -> "/root/", "path/", "leaf"
+	 */
+	static COREUOBJECT_API void SplitPackageName(FStringView InPackageName, FStringView* OutPackageRoot,
+		FStringView* OutPackagePath, FStringView* OutPackageLeafName,
+		EPathFormatFlags Flags = EPathFormatFlags::MountPointSlashes);
+	static COREUOBJECT_API void SplitPackageName(const FString& InLongPackageName, FString* OutPackageRoot,
+		FString* OutPackagePath, FString* OutPackageLeafName,
+		EPathFormatFlags Flags = EPathFormatFlags::MountPointSlashes);
+
+	/**
+	 * Test whether a given packagename is in a valid root, and convert it to root, path, and leafname components.
+	 * If not in a valid root, all output arguments are reset to empty.
+	 *
+	 * @param InLongPackageName Package Name, e.g. "/Game/Maps/TestMaps/MyMap".
+	 *        Also supports packagenames with extension (which are not valid package names),
+	 *        e.g. /Game/Maps/TestMaps/MyMap.umap.
+	 * @param OutPackageRoot The package root path, eg "/Game/"
+	 * @param OutPackagePath The path from the mount point to the package, eg "Maps/TestMaps/
+	 * @param OutPackageLeafName The leaf name, including extension if present, eg "MyMap" or "MyMap.umap".
+	 * @param bStripRootLeadingSlash: If false, OutpackageRoot="/<RootName>/", otherwise "<RootName>/".
+	 * @return Whether the path is in a valid root.
+	 */
+	static COREUOBJECT_API bool SplitLongPackageName(const FString& InLongPackageName, FString& OutPackageRoot,
+		FString& OutPackagePath, FString& OutPackageLeafName, const bool bStripRootLeadingSlash = false);
+
 	/**
 	 * Returns the top level 'directory' in a package name. If the package is part of a
 	 * plugin, the root will be the name of the associated plugin, otherwise it will be
 	 * 'Engine' or 'Game'.
 	 * 
-	 * Note the slashes are removed!
+	 * Note the slashes are removed by default (EPathFormatFlags::MountPointNoSlashes), @see EPathFormatFlags.
 	 * 
+	 * Default:
 	 * "/PackageRoot/Path/Leaf" -> (return "PackageRoot"; OutRelativePath = "Path/Leaf")
+	 * With EPathFormatFlags::MountPointSlashes:
+	 * "/PackageRoot/Path/Leaf" -> (return "/PackageRoot/"; OutRelativePath = "Path/Leaf")
 	 */
-	static COREUOBJECT_API FStringView SplitPackageNameRoot(FStringView InPackageName, FStringView* OutRelativePath);
-	static COREUOBJECT_API FString SplitPackageNameRoot(FName InPackageName, FString* OutRelativePath);
+	static COREUOBJECT_API FStringView SplitPackageNameRoot(FStringView InPackageName, FStringView* OutRelativePath,
+		EPathFormatFlags Flags = EPathFormatFlags::MountPointNoSlashes);
+	static COREUOBJECT_API FString SplitPackageNameRoot(FName InPackageName, FString* OutRelativePath,
+		EPathFormatFlags Flags = EPathFormatFlags::MountPointNoSlashes);
 	
 	/** 
 	 * Returns the name of the package referred to by the specified object path
@@ -991,3 +1040,4 @@ private:
 };
 
 ENUM_CLASS_FLAGS(FPackageName::EConvertFlags);
+ENUM_CLASS_FLAGS(FPackageName::EPathFormatFlags);
