[BugFix] Fix for bug that caused unsafe memory reads in Override serialization code when IDOs are enabled

diff --git a/Engine/Source/Runtime/CoreUObject/Internal/UObject/InstanceDataObjectUtils.h b/Engine/Source/Runtime/CoreUObject/Internal/UObject/InstanceDataObjectUtils.h
--- a/Engine/Source/Runtime/CoreUObject/Internal/UObject/InstanceDataObjectUtils.h
+++ b/Engine/Source/Runtime/CoreUObject/Internal/UObject/InstanceDataObjectUtils.h
@@ -33,6 +33,7 @@ UE_API bool IsInstanceDataObjectSupportEnabledForClass(const UClass* Class);
 UE_API bool IsUninitializedAlertUIEnabled();
 
 bool StructContainsLooseProperties(const UStruct* Struct);
+UE_API bool StructIsInstanceDataObjectStruct(const UStruct* Struct);
 
 UE_API bool IsPropertyLoose(const FProperty* Property);
 
diff --git a/Engine/Source/Runtime/CoreUObject/Private/UObject/InstanceDataObjectUtils.cpp b/Engine/Source/Runtime/CoreUObject/Private/UObject/InstanceDataObjectUtils.cpp
--- a/Engine/Source/Runtime/CoreUObject/Private/UObject/InstanceDataObjectUtils.cpp
+++ b/Engine/Source/Runtime/CoreUObject/Private/UObject/InstanceDataObjectUtils.cpp
@@ -249,6 +249,7 @@ namespace UE
 	static const FName NAME_DisplayName(ANSITEXTVIEW("DisplayName"));
 	static const FName NAME_PresentAsTypeMetadata(ANSITEXTVIEW("PresentAsType"));
 	static const FName NAME_IsLooseMetadata(ANSITEXTVIEW("IsLoose"));
+	static const FName NAME_IsInstanceDataObjectStruct(ANSITEXTVIEW("IsInstanceDataObjectClass"));
 	static const FName NAME_ContainsLoosePropertiesMetadata(ANSITEXTVIEW("ContainsLooseProperties"));
 	static const FName NAME_VerseClass(ANSITEXTVIEW("VerseClass"));
 	static const FName NAME_VerseDevice(ANSITEXTVIEW("VerseDevice_C"));
@@ -405,6 +406,11 @@ namespace UE
 		return Struct->GetBoolMetaData(NAME_ContainsLoosePropertiesMetadata);
 	}
 
+	bool StructIsInstanceDataObjectStruct(const UStruct* Struct)
+	{
+		return Struct->GetBoolMetaData(NAME_IsInstanceDataObjectStruct);
+	}
+
 	static UStruct* CreateInstanceDataObjectStructRec(const UClass* StructClass, UStruct* TemplateStruct, UObject* Outer, const FPropertyPathNameTree* PropertyTree, const FUnknownEnumNames* EnumNames);
 
 	template <typename StructType>
@@ -710,6 +716,7 @@ namespace UE
 				UStruct* NewSuper = NewObject<UStruct>(Outer, SuperStructClass, MakeUniqueObjectName(nullptr, SuperStructClass, SuperName));
 				NewSuper->SetSuperStruct(Super);
 				Super = NewSuper;
+				Super->SetMetaData(NAME_IsInstanceDataObjectStruct, TEXT("True"));
 			}
 
 			// Gather properties for Super Struct
@@ -752,6 +759,7 @@ namespace UE
 		const FName InstanceDataObjectName = (TemplateStruct) ? FName(WriteToString<128>(TemplateStruct->GetName(), TEXTVIEW("_InstanceDataObject"))) : FName(TEXTVIEW("InstanceDataObject"));
 		UStruct* Result = NewObject<UStruct>(Outer, StructClass, MakeUniqueObjectName(Outer, StructClass, InstanceDataObjectName));
 		Result->SetSuperStruct(Super);
+		Result->SetMetaData(NAME_IsInstanceDataObjectStruct, TEXT("True"));
 
 		// inherit ContainsLooseProperties metadata
 		if (Super && Super->GetBoolMetaData(NAME_ContainsLoosePropertiesMetadata))
diff --git a/Engine/Source/Runtime/CoreUObject/Private/UObject/OverriddenPropertySet.cpp b/Engine/Source/Runtime/CoreUObject/Private/UObject/OverriddenPropertySet.cpp
--- a/Engine/Source/Runtime/CoreUObject/Private/UObject/OverriddenPropertySet.cpp
+++ b/Engine/Source/Runtime/CoreUObject/Private/UObject/OverriddenPropertySet.cpp
@@ -7,6 +7,8 @@
 #include "Serialization/ArchiveSerializedPropertyChain.h"
 #include "UObject/PropertyOptional.h"
 #include "Misc/ScopeExit.h"
+#include "UObject/InstanceDataObjectUtils.h"
+#include "UObject/PropertyBagRepository.h"
 #include "UObject/UObjectArchetypeHelper.h"
 #include "UObject/UObjectThreadContext.h"
 
@@ -1265,11 +1267,26 @@ void FOverriddenPropertySet::RemoveOverriddenSubProperties(FOverriddenPropertyNo
 	PropertyNode.SubPropertyNodeKeys.Empty();
 }
 
+// IDO properties are keyed as if they're instance properties. This creates some ambiguity so we need to check what we're looking at here.
+static UObject* GetOwningObjectOfProperty(UObject* ExpectedOwner, const FProperty* Property)
+{
+	UObject* PropertyOwner = ExpectedOwner;
+#if WITH_EDITOR
+	if (UE::StructIsInstanceDataObjectStruct(Property->GetOwnerClass()))
+	{
+		PropertyOwner = UE::FPropertyBagRepository::Get().FindInstanceDataObject(PropertyOwner);
+	}
+	checkf(PropertyOwner && PropertyOwner->GetClass()->IsChildOf(Property->GetOwnerClass()),
+		TEXT("Expected provided property to either be in the Owner class, or in it's IDO's class"));
+#endif
+	return PropertyOwner;
+}
+
 EOverriddenPropertyOperation FOverriddenPropertySet::GetOverriddenPropertyOperation(FPropertyVisitorPath::Iterator PropertyIterator, bool* bOutInheritedOperation /*= nullptr*/) const
 {
 	if (const FOverriddenPropertyNode* RootNode = OverriddenPropertyNodes.Find(RootNodeID))
 	{
-		return GetOverriddenPropertyOperation(*RootNode, PropertyIterator, bOutInheritedOperation, Owner);
+		return GetOverriddenPropertyOperation(*RootNode, PropertyIterator, bOutInheritedOperation, GetOwningObjectOfProperty(Owner, PropertyIterator->Property));
 	}
 	return EOverriddenPropertyOperation::None;
 }
@@ -1278,7 +1295,7 @@ bool FOverriddenPropertySet::ClearOverriddenProperty(FPropertyVisitorPath::Itera
 {
 	if (FOverriddenPropertyNode* RootNode = OverriddenPropertyNodes.Find(RootNodeID))
 	{
-		return ClearOverriddenProperty(*RootNode, PropertyIterator, Owner);
+		return ClearOverriddenProperty(*RootNode, PropertyIterator, GetOwningObjectOfProperty(Owner, PropertyIterator->Property));
 	}
 	return true;
 }
