Cooker: Add a ModifyCook delegate that has the arguments it needs to handle new cooker features since 5.1: It needs to provide separate requests per platform, and it needs access to more data from the cooker to make decisions, e.g. which plugins are enabled per platform.
#rnx

diff --git a/Engine/Source/Editor/UnrealEd/Classes/CookOnTheSide/CookOnTheFlyServer.h b/Engine/Source/Editor/UnrealEd/Classes/CookOnTheSide/CookOnTheFlyServer.h
--- a/Engine/Source/Editor/UnrealEd/Classes/CookOnTheSide/CookOnTheFlyServer.h
+++ b/Engine/Source/Editor/UnrealEd/Classes/CookOnTheSide/CookOnTheFlyServer.h
@@ -676,6 +676,7 @@ public:
 	UNREALED_API virtual bool IsIterative() override;
 	UNREALED_API virtual TArray<const ITargetPlatform*> GetSessionPlatforms() override;
 	UNREALED_API virtual FString GetCookOutputFolder(const ITargetPlatform* TargetPlatform) override;
+	UNREALED_API virtual const TSet<IPlugin*>* GetEnabledPlugins(const ITargetPlatform* TargetPlatform) override;
 
 	UNREALED_API virtual void RegisterCollector(UE::Cook::IMPCollector* Collector,
 		UE::Cook::EProcessType ProcessType = UE::Cook::EProcessType::AllMPCook) override;
diff --git a/Engine/Source/Editor/UnrealEd/Private/CookOnTheFlyServer.cpp b/Engine/Source/Editor/UnrealEd/Private/CookOnTheFlyServer.cpp
--- a/Engine/Source/Editor/UnrealEd/Private/CookOnTheFlyServer.cpp
+++ b/Engine/Source/Editor/UnrealEd/Private/CookOnTheFlyServer.cpp
@@ -8054,7 +8054,9 @@ const TCHAR* GCookRequestUsageMessage = TEXT(
 	"			// Subclass FAssetManager (Engine.ini:[/Script/Engine.Engine]:AssetManagerClassName) and override this hook.\n"
 	"           // The default AssetManager behavior cooks all packages specified by PrimaryAssetTypesToScan rules from ini.\n"
 	"		FGameDelegates::Get().GetModifyCookDelegate()\n"
-	"			// Subscribe to this delegate during your module startup.\n"
+	"			// Subscribe to this delegate during your module startup. (Will be deprecated in the future, use UE::Cook::FDelegates::ModifyCook instead).\n"
+	"       UE::Cook::FDelegates::ModifyCook\n"
+	"           // Subscribe to this delegate during your module startup.\n"
 	"		ITargetPlatform::GetExtraPackagesToCook\n"
 	"			// Override this hook on a given TargetPlatform.\n"
 );
@@ -8229,6 +8231,45 @@ void UCookOnTheFlyServer::CollectFilesToCook(TArray<FName>& FilesInPath, TMap<FN
 			UpdateInstigators(EInstigator::ModifyCookDelegate);
 		}
 
+		FCookInfoModifyCookDelegate& CookInfoModifyCook = UE::Cook::FDelegates::ModifyCook;
+		TArray<FPackageCookRule> CookRules;
+		if (CookInfoModifyCook.IsBound())
+		{
+			CookInfoModifyCook.Broadcast(*this, CookRules);
+			for (FPackageCookRule& CookRule : CookRules)
+			{
+				if (CookRule.PackageName.IsNone())
+				{
+					continue;
+				}
+				switch (CookRule.CookRule)
+				{
+				case EPackageCookRule::None:
+					break;
+				case EPackageCookRule::AddToCook:
+				{
+					FInstigator& Existing = Instigators.FindOrAdd(CookRule.PackageName);
+					if (Existing.Category == EInstigator::InvalidCategory)
+					{
+						Existing = FInstigator(EInstigator::ModifyCookDelegate, CookRule.InstigatorName);
+					}
+					FilesInPath.Add(CookRule.PackageName);
+					break;
+				}
+				case EPackageCookRule::NeverCook:
+					PackagesToNeverCook.Add(CookRule.PackageName);
+					break;
+				// case EPackageCookRule::IgnoreStartupPackage:
+				// To implement IgnoreStartupPackage, we will need to pass in StartupPackages, and
+				// delay the calculation of ProcessSoftObjectPathPackageList.
+				default:
+					checkNoEntry();
+					break;
+				}
+			}
+			UpdateInstigators(EInstigator::ModifyCookDelegate);
+		}
+
 		for (FName NeverCookPackage : PackagesToNeverCook)
 		{
 			FName PackageName;
@@ -10488,13 +10529,25 @@ UE::Cook::FCookSavePackageContext* UCookOnTheFlyServer::CreateSaveContext(const
 	DiffModeHelper->InitializePackageWriter(*this, PackageWriter, ResolvedMetadataPath, DeterminismManager.Get());
 
 	// Setup save package settings (i.e. validation)
+	TSet<IPlugin*> EnabledPlugins;
+	TArray<TSharedRef<IPlugin>> EnabledPluginRefPtrs;
+
 	FSavePackageSettings SavePackageSettings = FSavePackageSettings::GetDefaultSettings();
 	{
 		// Setup Import Validation for suppressed module native classes
 		TSet<FName> DisabledNativeScriptPackages;
 		for (TSharedRef<IPlugin>& Plugin : IPluginManager::Get().GetEnabledPlugins())
 		{
-			if (!TargetPlatform->IsEnabledForPlugin(Plugin.Get()))
+			if (TargetPlatform->IsEnabledForPlugin(Plugin.Get()))
+			{
+				bool bExists;
+				EnabledPlugins.Add(&Plugin.Get(), &bExists);
+				if (!bExists)
+				{
+					EnabledPluginRefPtrs.Add(Plugin);
+				}
+			}
+			else
 			{
 				for (const FModuleDescriptor& Module : Plugin->GetDescriptor().Modules)
 				{
@@ -10541,10 +10594,24 @@ UE::Cook::FCookSavePackageContext* UCookOnTheFlyServer::CreateSaveContext(const
 
 	FCookSavePackageContext* Context = new FCookSavePackageContext(TargetPlatform, CookArtifactReader, PackageWriter, WriterDebugName, MoveTemp(SavePackageSettings),
 		MoveTemp(DeterminismManager));
+	Context->EnabledPlugins = MoveTemp(EnabledPlugins);
+	Context->EnabledPluginRefPtrs = MoveTemp(EnabledPluginRefPtrs);
 	return Context;
 
 }
 
+const TSet<IPlugin*>* UCookOnTheFlyServer::GetEnabledPlugins(const ITargetPlatform* TargetPlatform)
+{
+	using namespace UE::Cook;
+
+	const FCookSavePackageContext* Context = FindSaveContext(TargetPlatform);
+	if (Context)
+	{
+		return &Context->EnabledPlugins;
+	}
+	return nullptr;
+}
+
 void UCookOnTheFlyServer::DeleteOutputForPackage(FName PackageName, const ITargetPlatform* TargetPlatform)
 {
 	FindOrCreatePackageWriter(TargetPlatform).RemoveCookedPackages(TArrayView<FName>(&PackageName, 1));
diff --git a/Engine/Source/Editor/UnrealEd/Private/Cooker/CookTypes.h b/Engine/Source/Editor/UnrealEd/Private/Cooker/CookTypes.h
--- a/Engine/Source/Editor/UnrealEd/Private/Cooker/CookTypes.h
+++ b/Engine/Source/Editor/UnrealEd/Private/Cooker/CookTypes.h
@@ -13,6 +13,7 @@
 #include "HAL/LowLevelMemTracker.h"
 #include "HAL/PlatformMath.h"
 #include "HAL/Platform.h"
+#include "Interfaces/IPluginManager.h"
 #include "Logging/TokenizedMessage.h"
 #include "Misc/AssertionMacros.h"
 #include "Serialization/PackageWriter.h"
@@ -411,6 +412,12 @@ namespace UE::Cook
 		ICookedPackageWriter* PackageWriter;
 		ICookedPackageWriter::FCookCapabilities PackageWriterCapabilities;
 		TUniquePtr<FDeterminismManager> DeterminismManager;
+		TSet<IPlugin*> EnabledPlugins;
+		/**
+		 * RefCount Pointers to the EnabledPlugins. We pass them to external API as a TSet of raw pointers,
+		 * but need to keep them referenced.
+		 */
+		TArray<TSharedRef<IPlugin>> EnabledPluginRefPtrs;
 	};
 
 	/* Thread Local Storage access to identify which thread is the SchedulerThread for cooking. */
diff --git a/Engine/Source/Runtime/CoreUObject/Private/UObject/ICookInfo.cpp b/Engine/Source/Runtime/CoreUObject/Private/UObject/ICookInfo.cpp
--- a/Engine/Source/Runtime/CoreUObject/Private/UObject/ICookInfo.cpp
+++ b/Engine/Source/Runtime/CoreUObject/Private/UObject/ICookInfo.cpp
@@ -88,6 +88,7 @@ FCookInfoEvent FDelegates::CookStarted;
 FCookInfoEvent FDelegates::CookFinished;
 FValidateSourcePackage FDelegates::ValidateSourcePackage;
 FPackageBlockedEvent FDelegates::PackageBlocked;
+FCookInfoModifyCookDelegate FDelegates::ModifyCook;
 
 const TCHAR* GetReferencedSetFilename()
 {
diff --git a/Engine/Source/Runtime/CoreUObject/Public/UObject/ICookInfo.h b/Engine/Source/Runtime/CoreUObject/Public/UObject/ICookInfo.h
--- a/Engine/Source/Runtime/CoreUObject/Public/UObject/ICookInfo.h
+++ b/Engine/Source/Runtime/CoreUObject/Public/UObject/ICookInfo.h
@@ -19,6 +19,7 @@
 enum class EDataValidationResult : uint8;
 class FDataValidationContext;
 class ITargetPlatform;
+class IPlugin;
 class UPackage;
 
 namespace UE::Cook { class IMPCollector; }
@@ -134,6 +135,66 @@ struct FInstigator
 	COREUOBJECT_API FString ToString() const;
 };
 
+/**
+ * Values for whether a package should be cooked, used in ModifyCook callbacks by systems that want to
+ * mark packages cooked/not-cooked independently of the usual asset-driven and config-driven cook specifications.
+ */
+enum class EPackageCookRule
+{
+	/* No action specified, the rule will be ignored. */
+	None,
+	/** The package will be cooked (unless specified as NeverCook or not cooked for the platform elsewhere). */
+	AddToCook,
+	/** The package will be not be cooked. */
+	NeverCook,
+
+	// IgnoreStartupPackage is not yet implemented, coming in a future version. For now, only FCookLoadScope
+	// provides this functionality.
+	/**
+	 * Only has an effect if the project configuration specifies that startup packages (packages already loaded before
+	 * the cook starts) are automatically cooked, and the package is a startup package. This turns off that automatic
+	 * inclusion for the given package, and the package will not be cooked unless it is referenced from another source
+	 * of requested packages or from another package that is cooked.
+	 *
+	 * Note: This will also modify the chunking of the package if it is otherwise referenced and therefore cooked.
+	 * Startup packages are added into chunk0, and marking the package as IgnoreStartupPackage will remove that chunk
+	 * rule for the package and it will be chunked only according to the AssetManager's decisions on chunking.
+	 *
+	 * Note: Currently the decision that startup packages are automatically cooked is hardcoded and there is no way
+	 * for a project to turn it off. We expect this to change in the future and eventually Ignored will become
+	 * the default for startup packages.
+	 * 
+	 * Note: This specification can also be indicated, without the need to subscribe to ModifyCook, by wrapping the
+	 * LoadPackage or CreatePackage call of the given package in FCookLoadScope Scope(ECookLoadType::EditorOnly).
+	 */
+	// IgnoreStartupPackage,
+};
+
+/**
+ * Specification of whether a package should be cooked, for use in ModifyCook callbacks by systems that want to
+ * mark packages cooked/not-cooked independently of the usual asset-driven and config-driven cook specifications.
+ */
+struct FPackageCookRule
+{
+	/** Name of the package to specify cooked/notcooked/ignored. */
+	FName PackageName;
+
+	/** Name of the subscriber system, for use by cook users debugging why a package was cooked. */
+	FName InstigatorName;
+
+	// TargetPlatform is not yet implemented, coming in a future version. For now, TargetPlatform is ignored and
+	// all rules apply to all platforms.
+	/**
+	 * For which platforms the behavior change should apply; this is necessary for multiplatform cooks,
+	 * if different platforms can have different values. nullptr indicates it applies to every platform.
+	 * To set different non-defualt cookrules for multiple platforms, use multiple FPackageCookRules.
+	 */
+	//	const ITargetPlatform* TargetPlatform = nullptr;
+
+	/** The desired cook behavior for the package, @see UE::Cook::EPackageCookRule. */
+	EPackageCookRule CookRule = EPackageCookRule::None;
+};
+
 /** Engine interface for information provided by UCookOnTheFlyServer in cook callbacks. */
 class ICookInfo
 {
@@ -181,6 +242,15 @@ public:
 	 */
 	virtual FString GetCookOutputFolder(const ITargetPlatform* TargetPlatform) = 0;
 
+	/**
+	 * Return the set of plugins that the cooker has determined are enabled on the current platform. This is different
+	 * than IPluginManager::GetEnabledPlugins. The IPluginManager version returns the plugins enabled for the editor
+	 * process, this GetEnabledPlugins returns the plugins that will be enabled on the TargetPlatform.
+	 * 
+	 * Returns nullptr if the given Platform is not a session platform in the current cook.
+	 */
+	virtual const TSet<IPlugin*>* GetEnabledPlugins(const ITargetPlatform* TargetPlatform) = 0;
+
 	/**
 	 * MPCook: register in the current process a collector that replicates system-specific and package-specific
 	 * information between CookWorkers and the CookDirector. Registration will be skipped if the current cook is
@@ -207,8 +277,11 @@ public:
 };
 
 DECLARE_MULTICAST_DELEGATE_OneParam(FCookInfoEvent, ICookInfo&);
-DECLARE_DELEGATE_RetVal_TwoParams(EDataValidationResult, FValidateSourcePackage, UPackage* /*Package*/, FDataValidationContext& /*ValidationContext*/);
+DECLARE_DELEGATE_RetVal_TwoParams(EDataValidationResult, FValidateSourcePackage, UPackage* /*Package*/,
+	FDataValidationContext& /*ValidationContext*/);
 DECLARE_MULTICAST_DELEGATE_TwoParams(FPackageBlockedEvent, const UObject* /* Object */, FStringBuilderBase& OutDebugInfo)
+DECLARE_MULTICAST_DELEGATE_TwoParams(FCookInfoModifyCookDelegate, ICookInfo& /* CookInfo */,
+	TArray<FPackageCookRule>& /* InOutPackageCookRules */);
 
 /** UE::Cook::FDelegates: callbacks for cook events. */
 struct FDelegates
@@ -226,6 +299,11 @@ public:
 	static COREUOBJECT_API FValidateSourcePackage ValidateSourcePackage;
 	/** Called when the given package has been blocked from saving for beyond the configured time threshold; can be registered to append additional debug info in this case. */
 	static COREUOBJECT_API FPackageBlockedEvent PackageBlocked;
+	/**
+	 * Called during CookByTheBook to specify packages that should be cooked, not cooked even if referenced
+	 * (aka nevercooked), or have their reference behavior changed in other ways.
+	 */
+	static COREUOBJECT_API FCookInfoModifyCookDelegate ModifyCook;
 };
 
 
