Fix issue with UHT created UVerseClasses where the Outer chain wasn't currectly initialized to UVerseClass

#rnx

diff --git a/Engine/Source/Runtime/CoreUObject/Private/UObject/UObjectConstructInternal.h b/Engine/Source/Runtime/CoreUObject/Private/UObject/UObjectConstructInternal.h
--- a/Engine/Source/Runtime/CoreUObject/Private/UObject/UObjectConstructInternal.h
+++ b/Engine/Source/Runtime/CoreUObject/Private/UObject/UObjectConstructInternal.h
@@ -114,7 +114,7 @@ void ConstructUClassNoInitHelper(
 	InPostNewFn(NewClass);
 
 	InitializePrivateStaticClass(
-		UClass::StaticClass,
+		UClassClass::StaticClass,
 		InSuperClassFn(),
 		NewClass,
 		InWithinClassFn(),
