Fixed the struct serialization type check when using impersonation on a struct with a guid

diff --git a/Engine/Source/Runtime/CoreUObject/Private/UObject/PropertyStruct.cpp b/Engine/Source/Runtime/CoreUObject/Private/UObject/PropertyStruct.cpp
--- a/Engine/Source/Runtime/CoreUObject/Private/UObject/PropertyStruct.cpp
+++ b/Engine/Source/Runtime/CoreUObject/Private/UObject/PropertyStruct.cpp
@@ -618,15 +618,17 @@ bool FStructProperty::CanSerializeFromTypeName(UE::FPropertyTypeName Type) const
 	}
 
 	const FName StructGuidName = Type.GetParameterName(1);
-	if (FGuid StructGuid; !StructGuidName.IsNone() && FGuid::Parse(StructGuidName.ToString(), StructGuid) && StructGuid.IsValid())
+	if (FGuid StructGuid; !StructGuidName.IsNone() && FGuid::Parse(StructGuidName.ToString(), StructGuid) &&
+		StructGuid.IsValid() && StructGuid == LocalStruct->GetCustomGuid())
 	{
-		return StructGuid == LocalStruct->GetCustomGuid();
+		return true;
 	}
 
 #if WITH_EDITORONLY_DATA
-	if (const UE::FPropertyTypeName OriginalType = UE::FindOriginalType(this); !OriginalType.IsEmpty())
+	if (const UE::FPropertyTypeName OriginalType = UE::FindOriginalType(this);
+		!OriginalType.IsEmpty() && StructName == OriginalType.GetName())
 	{
-		return StructName == OriginalType.GetName();
+		return true;
 	}
 #endif // WITH_EDITORONLY_DATA
 
