Fixed property name redirects when using impersonation

diff --git a/Engine/Source/Runtime/CoreUObject/Private/UObject/Property.cpp b/Engine/Source/Runtime/CoreUObject/Private/UObject/Property.cpp
--- a/Engine/Source/Runtime/CoreUObject/Private/UObject/Property.cpp
+++ b/Engine/Source/Runtime/CoreUObject/Private/UObject/Property.cpp
@@ -2320,6 +2320,26 @@ const TCHAR* FProperty::ImportSingleProperty( const TCHAR* Str, void* DestData,
 	return Str;
 }
 
+#if WITH_EDITORONLY_DATA
+static FCoreRedirectObjectName BuildCoreRedirectPropertyName(FName PropertyName, UE::FPropertyTypeName TypeName)
+{
+	FName OuterName = TypeName.GetName();
+	if (const int32 OuterCount = TypeName.GetParameterCount() - 1; OuterCount > 0)
+	{
+		TStringBuilder<256> OuterChain;
+		for (int32 OuterIndex = 0; OuterIndex < OuterCount; ++OuterIndex)
+		{
+			OuterChain << TypeName.GetParameterName(OuterIndex + 1) << (OuterIndex == 0 ? TEXT(':') : TEXT('.'));
+		}
+		OuterChain << OuterName;
+		OuterName = FName(OuterChain);
+	}
+
+	const FName PackageName = TypeName.GetParameterName(0);
+	return FCoreRedirectObjectName(PropertyName, OuterName, PackageName);
+}
+#endif // WITH_EDITORONLY_DATA
+
 FName FProperty::FindRedirectedPropertyName(const UStruct* ObjectStruct, FName OldName)
 {
 	DECLARE_SCOPE_CYCLE_COUNTER(TEXT("FProperty::FindRedirectedPropertyName"), STAT_LinkerLoad_FindRedirectedPropertyName, STATGROUP_LoadTimeVerbose);
@@ -2328,16 +2348,29 @@ FName FProperty::FindRedirectedPropertyName(const UStruct* ObjectStruct, FName O
 	UPackage* StructPackage = ObjectStruct->GetOutermost();
 	FName PackageName = StructPackage->GetFName();
 	// Avoid GetPathName string allocation and FName initialization when there is only one outer
-	FName OuterName = (StructPackage == ObjectStruct->GetOuter()) ? ObjectStruct->GetFName() : FName(*ObjectStruct->GetPathName(StructPackage));
+	FName OuterName = (StructPackage == ObjectStruct->GetOuter()) ? ObjectStruct->GetFName() : FName(ObjectStruct->GetPathName(StructPackage));
 
 	FCoreRedirectObjectName OldRedirectName(OldName, OuterName, PackageName);
 	FCoreRedirectObjectName NewRedirectName = FCoreRedirects::GetRedirectedName(ECoreRedirectFlags::Type_Property, OldRedirectName);
-
 	if (NewRedirectName != OldRedirectName)
 	{
 		return NewRedirectName.ObjectName;
 	}
 
+#if WITH_EDITORONLY_DATA
+	// Compare against the original type of an impersonated type.
+	if (const UE::FPropertyTypeName OriginalType = UE::FindOriginalType(ObjectStruct); !OriginalType.IsEmpty())
+	{
+		OldRedirectName = BuildCoreRedirectPropertyName(OldName, OriginalType);
+		NewRedirectName = FCoreRedirects::GetRedirectedName(ECoreRedirectFlags::Type_Property, OldRedirectName);
+		if (NewRedirectName != OldRedirectName)
+		{
+			return NewRedirectName.ObjectName;
+		}
+	}
+#endif // WITH_EDITORONLY_DATA
+
+
 	return NAME_None;
 }
 
