Handle empty property chain, behavior regressed by 39491016

#rnx

diff --git a/Engine/Source/Runtime/CoreUObject/Private/UObject/OverriddenPropertySet.cpp b/Engine/Source/Runtime/CoreUObject/Private/UObject/OverriddenPropertySet.cpp
--- a/Engine/Source/Runtime/CoreUObject/Private/UObject/OverriddenPropertySet.cpp
+++ b/Engine/Source/Runtime/CoreUObject/Private/UObject/OverriddenPropertySet.cpp
@@ -1268,15 +1268,21 @@ void FOverriddenPropertySet::RemoveOverriddenSubProperties(FOverriddenPropertyNo
 }
 
 // IDO properties are keyed as if they're instance properties. This creates some ambiguity so we need to check what we're looking at here.
-static UObject* GetOwningObjectOfProperty(UObject* ExpectedOwner, const FProperty* Property)
+static UObject* GetOwningObjectOfProperty(UObject* ExpectedOwner, FPropertyVisitorPath::Iterator PropertyIterator)
 {
+	if(!PropertyIterator)
+	{
+		return ExpectedOwner;
+	}
+
+	const UClass* Class = PropertyIterator->Property->GetOwnerClass();
 	UObject* PropertyOwner = ExpectedOwner;
 #if WITH_EDITOR
-	if (UE::StructIsInstanceDataObjectStruct(Property->GetOwnerClass()))
+	if (UE::StructIsInstanceDataObjectStruct(Class))
 	{
 		PropertyOwner = UE::FPropertyBagRepository::Get().FindInstanceDataObject(PropertyOwner);
 	}
-	checkf(PropertyOwner && PropertyOwner->GetClass()->IsChildOf(Property->GetOwnerClass()),
+	checkf(PropertyOwner && PropertyOwner->GetClass()->IsChildOf(Class),
 		TEXT("Expected provided property to either be in the Owner class, or in it's IDO's class"));
 #endif
 	return PropertyOwner;
@@ -1286,7 +1292,8 @@ EOverriddenPropertyOperation FOverriddenPropertySet::GetOverriddenPropertyOperat
 {
 	if (const FOverriddenPropertyNode* RootNode = OverriddenPropertyNodes.Find(RootNodeID))
 	{
-		return GetOverriddenPropertyOperation(*RootNode, PropertyIterator, bOutInheritedOperation, GetOwningObjectOfProperty(Owner, PropertyIterator->Property));
+		return GetOverriddenPropertyOperation(*RootNode, PropertyIterator, bOutInheritedOperation, 
+			GetOwningObjectOfProperty(Owner, PropertyIterator));
 	}
 	return EOverriddenPropertyOperation::None;
 }
@@ -1295,7 +1302,8 @@ bool FOverriddenPropertySet::ClearOverriddenProperty(FPropertyVisitorPath::Itera
 {
 	if (FOverriddenPropertyNode* RootNode = OverriddenPropertyNodes.Find(RootNodeID))
 	{
-		return ClearOverriddenProperty(*RootNode, PropertyIterator, GetOwningObjectOfProperty(Owner, PropertyIterator->Property));
+		return ClearOverriddenProperty(*RootNode, PropertyIterator, 
+			GetOwningObjectOfProperty(Owner, PropertyIterator));
 	}
 	return true;
 }
