Fix assert when looking up localized package names in -game or -server mode

diff --git a/Engine/Source/Runtime/CoreUObject/Private/Serialization/AsyncLoading2.cpp b/Engine/Source/Runtime/CoreUObject/Private/Serialization/AsyncLoading2.cpp
--- a/Engine/Source/Runtime/CoreUObject/Private/Serialization/AsyncLoading2.cpp
+++ b/Engine/Source/Runtime/CoreUObject/Private/Serialization/AsyncLoading2.cpp
@@ -542,6 +542,7 @@ FName ApplyPackageNameRedirections(FName PackageName)
 		// Handle localization when running in -game
 		if (!GIsEditor)
 		{
+			UE::Core::Private::FPlayInEditorLoadingScope PlayInEditorIDScope(INDEX_NONE);
 			NewPackageName = FPackageLocalizationManager::Get().FindLocalizedPackageName(PackageName);
 			if (!NewPackageName.IsNone())
 			{
