Modified and added to `OnTransaction*Delegate` for remote executor so delegates know what transaction just started or ended. Intended for debugging/profiling purposes.

diff --git a/Engine/Source/Runtime/CoreUObject/Private/UObject/RemoteExecutor.cpp b/Engine/Source/Runtime/CoreUObject/Private/UObject/RemoteExecutor.cpp
--- a/Engine/Source/Runtime/CoreUObject/Private/UObject/RemoteExecutor.cpp
+++ b/Engine/Source/Runtime/CoreUObject/Private/UObject/RemoteExecutor.cpp
@@ -12,8 +12,9 @@ TDelegate<void()> TickNetworkDelegate;
 TDelegate<TTuple<FRemoteWorkPriority, TFunction<void(void)>>(void)> FetchNextDeferredRPCDelegate;
 TDelegate<void()> WaitForServersToSyncDelegate;
 
-TMulticastDelegate<void()> OnTransactionCompletedDelegate;
-TMulticastDelegate<void()> OnTransactionAbortedDelegate;
+TMulticastDelegate<void(FRemoteTransactionId)> OnTransactionStartingDelegate;
+TMulticastDelegate<void(FRemoteTransactionId)> OnTransactionCompletedDelegate;
+TMulticastDelegate<void(FRemoteTransactionId)> OnTransactionAbortedDelegate;
 }
 
 FArchive& operator<<(FArchive& Ar, FRemoteTransactionId& Id)
@@ -178,21 +179,22 @@ public:
 				// all of the subsystems are ready, try to perform the work
 				DebugLastAbortReason = TEXT("");
 
+				FRemoteTransactionId RequestId = ExecutingWork.RequestId;
+				UE::RemoteExecutor::OnTransactionStartingDelegate.Broadcast(RequestId);
 				AutoRTFM::ETransactionResult WorkTransactResult = AutoRTFM::Transact([&ExecutingWork]() { ExecutingWork.Work(); });
 
 				if (WorkTransactResult == AutoRTFM::ETransactionResult::AbortedByRequest)
 				{
 					// the work aborted, it should have inserted a new requested object, retry
 					UE_LOG(LogRemoteExec, Verbose, TEXT("\twork ABORTED (%s)"), *DebugLastAbortReason);
-					UE::RemoteExecutor::OnTransactionAbortedDelegate.Broadcast();
+					UE::RemoteExecutor::OnTransactionAbortedDelegate.Broadcast(RequestId);
 
 					PendingWorkIndex++;
 				}
 				else
 				{
 					UE_LOG(LogRemoteExec, VeryVerbose, TEXT("\twork COMPLETED!"));
-
-					UE::RemoteExecutor::OnTransactionCompletedDelegate.Broadcast();
+					UE::RemoteExecutor::OnTransactionCompletedDelegate.Broadcast(RequestId);
 
 					for (int32 SubsystemIndex = 0; SubsystemIndex < Subsystems.Num(); SubsystemIndex++)
 					{
diff --git a/Engine/Source/Runtime/CoreUObject/Public/UObject/RemoteExecutor.h b/Engine/Source/Runtime/CoreUObject/Public/UObject/RemoteExecutor.h
--- a/Engine/Source/Runtime/CoreUObject/Public/UObject/RemoteExecutor.h
+++ b/Engine/Source/Runtime/CoreUObject/Public/UObject/RemoteExecutor.h
@@ -189,11 +189,14 @@ namespace UE::RemoteExecutor
 	/** Delegate that syncs all servers */
 	extern COREUOBJECT_API TDelegate<void()> WaitForServersToSyncDelegate;
 
+	/** Delegate executed when a ExecuteTransactional starts an AutoRTFM transaction */
+	extern COREUOBJECT_API TMulticastDelegate<void(FRemoteTransactionId)> OnTransactionStartingDelegate;
+
 	/** Delegate executed when a ExecuteTransactional observes a completed AutoRTFM transaction */
-	extern COREUOBJECT_API TMulticastDelegate<void()> OnTransactionCompletedDelegate;
+	extern COREUOBJECT_API TMulticastDelegate<void(FRemoteTransactionId)> OnTransactionCompletedDelegate;
 
 	/** Delegate executed when a ExecuteTransactional observes an aborted AutoRTFM transaction */
-	extern COREUOBJECT_API TMulticastDelegate<void()> OnTransactionAbortedDelegate;
+	extern COREUOBJECT_API TMulticastDelegate<void(FRemoteTransactionId)> OnTransactionAbortedDelegate;
 
 	COREUOBJECT_API void RegisterRemoteSubsystem(FRemoteSubsystemBase* Subsystem);
 
