[Feature] IDO CVar support: Impersonation
added IDO.Impersonation.EnableOnSave CVar - When disabled we save data from instances directly instead of their IDOs, but continue to create IDOs on load
Note: This CVar has been temporarily toggled off to unblock users.

diff --git a/Engine/Source/Runtime/CoreUObject/Internal/UObject/InstanceDataObjectUtils.h b/Engine/Source/Runtime/CoreUObject/Internal/UObject/InstanceDataObjectUtils.h
--- a/Engine/Source/Runtime/CoreUObject/Internal/UObject/InstanceDataObjectUtils.h
+++ b/Engine/Source/Runtime/CoreUObject/Internal/UObject/InstanceDataObjectUtils.h
@@ -31,6 +31,8 @@ UE_API bool IsInstanceDataObjectSupportEnabled(const UObject* Object);
 UE_API bool IsInstanceDataObjectSupportEnabledForClass(const UClass* Class);
 /** Query if uninitialized alert UI is enabled **/
 UE_API bool IsUninitializedAlertUIEnabled();
+/** Query if InstanceDataObjects should be saved instead of instances. */
+UE_API bool IsInstanceDataObjectImpersonationEnabledOnSave();
 
 bool StructContainsLooseProperties(const UStruct* Struct);
 UE_API bool StructIsInstanceDataObjectStruct(const UStruct* Struct);
diff --git a/Engine/Source/Runtime/CoreUObject/Private/UObject/InstanceDataObjectUtils.cpp b/Engine/Source/Runtime/CoreUObject/Private/UObject/InstanceDataObjectUtils.cpp
--- a/Engine/Source/Runtime/CoreUObject/Private/UObject/InstanceDataObjectUtils.cpp
+++ b/Engine/Source/Runtime/CoreUObject/Private/UObject/InstanceDataObjectUtils.cpp
@@ -288,6 +288,14 @@ namespace UE
 		TEXT("Enables alert information for uninitalized properties. Requires IDO.Enable=true")
 	);
 
+	// TODO: re-enable this ASAP. This disables most IDO features but disabling was necessary to unblock those experiencing IDO bugs
+	bool bEnableIDOImpersonationOnSave = false;
+	FAutoConsoleVariableRef EnableIDOImpersonationOnSave(
+		TEXT("IDO.Impersonation.EnableOnSave"),
+		bEnableIDOImpersonationOnSave,
+		TEXT("When enabled, IDOs will be saved instead of instances. Disabling this will stop data retention on save.")
+	);
+
 	static TSet<FString> GetExcludedLoosePropertyTypes()
 	{
 		TArray<FString> Result;
@@ -305,6 +313,11 @@ namespace UE
 		return bEnableUninitializedUI;
 	}
 
+	bool IsInstanceDataObjectImpersonationEnabledOnSave()
+	{
+		return bEnableIDOImpersonationOnSave;
+	}
+
 	bool IsInstanceDataObjectSupportEnabled(const UObject* InObject)
 	{
 		if (!IsInstanceDataObjectSupportEnabled() || !InObject)
diff --git a/Engine/Source/Runtime/CoreUObject/Private/UObject/UObjectThreadContext.cpp b/Engine/Source/Runtime/CoreUObject/Private/UObject/UObjectThreadContext.cpp
--- a/Engine/Source/Runtime/CoreUObject/Private/UObject/UObjectThreadContext.cpp
+++ b/Engine/Source/Runtime/CoreUObject/Private/UObject/UObjectThreadContext.cpp
@@ -158,6 +158,9 @@ FScopedObjectSerializeContext::FScopedObjectSerializeContext(UObject* InObject,
 	// Disable if cooking because any extra data in the IDO will not be understood outside of the editor.
 	const bool bSupportsIDO = !Archive.IsCooking() && IsInstanceDataObjectSupportEnabled(Object);
 
+	// determine whether impersonation should be enabled when saving
+	const bool bImpersonateOnSave = bSupportsIDO && UE::IsInstanceDataObjectImpersonationEnabledOnSave();
+
 	// Disable if a newer version of the class exists because there is no point in creating an IDO of an obsolete type.
 	bCreateInstanceDataObject = bIsLoading && !bIsIDO && bSupportsIDO && !Object->GetClass()->HasAnyClassFlags(CLASS_NewerVersionExists);
 
@@ -171,7 +174,7 @@ FScopedObjectSerializeContext::FScopedObjectSerializeContext(UObject* InObject,
 	SerializeContext->bTrackSerializedProperties = bIsLoading && FSerializedPropertyValueState(Object).IsTracking();
 
 	// Enable impersonation when loading to an IDO or saving an object that may have an IDO.
-	SerializeContext->bImpersonateProperties = bIsLoading ? bIsIDO : bSupportsIDO;
+	SerializeContext->bImpersonateProperties = bIsLoading ? bIsIDO : bImpersonateOnSave;
 #endif // WITH_EDITORONLY_DATA
 }
 
