Mark package as having been EndLoaded when saving (this deals with editor-time created packages)

diff --git a/Engine/Source/Runtime/CoreUObject/Private/UObject/SavePackage2.cpp b/Engine/Source/Runtime/CoreUObject/Private/UObject/SavePackage2.cpp
--- a/Engine/Source/Runtime/CoreUObject/Private/UObject/SavePackage2.cpp
+++ b/Engine/Source/Runtime/CoreUObject/Private/UObject/SavePackage2.cpp
@@ -3148,6 +3148,9 @@ ESavePackageResult InnerSave(FSaveContext& SaveContext)
 	// Mark the Package RF_Loaded after its been serialized. 
 	// This was already done for each object in the package in SaveHarvestedRealms
 	SaveContext.GetPackage()->SetFlags(RF_WasLoaded | RF_LoadCompleted);
+#if WITH_EDITORONLY_DATA
+	SaveContext.GetPackage()->SetHasBeenEndLoaded(true);
+#endif // WITH_EDITORONLY_DATA
 
 	return SaveContext.Result;
 }
diff --git a/Engine/Source/Runtime/CoreUObject/Public/UObject/Package.h b/Engine/Source/Runtime/CoreUObject/Public/UObject/Package.h
--- a/Engine/Source/Runtime/CoreUObject/Public/UObject/Package.h
+++ b/Engine/Source/Runtime/CoreUObject/Public/UObject/Package.h
@@ -239,11 +239,12 @@ private:
 	*/
 	std::atomic<bool> bLoadedByEditorPropertiesOnly{ false };
 
-	/** Set to true after serialization and postload and before returning from LoadPackage or calling load completion delegate, for loaded packages.			
-	*   Note: This flag is manipulated on different threads, do not mix with other bitfields above as they might get corrupted.
-	*/
+	/** This flag is manipulated on different threads, do not mix with other bitfields above as they might get corrupted. */
 	std::atomic<bool> bHasBeenEndLoaded{ false };
 public:
+	/** This flag becomes true for loaded packages after serialization and postload and before returning from LoadPackage or calling load completion delegate.
+	  *  For newly created packages in editor, it becomes true once the package is saved. 
+	*/
 	bool GetHasBeenEndLoaded() const
 	{
 		return bHasBeenEndLoaded.load(std::memory_order_acquire) != 0;
