[FSoftObjectPath]

- Change FixupCoreRedirects to construct the old FCoreRedirectObjectName used for fixup from the FSoftObjectPath itself instead of FString. This avoids needing to always breakdown the FString into three FNames inside of FCoreRedirectObjectName which can cause contention as part of FName construction. Instead FNames now only need to be constructed if the FSoftObjectPath's SubPathString is not empty.
- Use a local buffer for PackageName strings
- Changes /Script/ mount check to use the more standard IsScriptPackage check from FPackageName.
- We now only call ToString when we have a redirection which is the less common case.

#rnx

diff --git a/Engine/Source/Runtime/CoreUObject/Private/UObject/SoftObjectPath.cpp b/Engine/Source/Runtime/CoreUObject/Private/UObject/SoftObjectPath.cpp
--- a/Engine/Source/Runtime/CoreUObject/Private/UObject/SoftObjectPath.cpp
+++ b/Engine/Source/Runtime/CoreUObject/Private/UObject/SoftObjectPath.cpp
@@ -781,18 +781,24 @@ bool FSoftObjectPath::FixupForPIE(TFunctionRef<void(int32, FSoftObjectPath&)> In
 
 bool FSoftObjectPath::FixupCoreRedirects()
 {
-	FString OldString = ToString();
-	FCoreRedirectObjectName OldName = FCoreRedirectObjectName(OldString);
+	// Construct from FSoftObjectPath to avoid unnecessary string copying and possible FName creation
+	FCoreRedirectObjectName OldName(*this);
+	FCoreRedirectObjectName NewName;
 
-	// Always try the object redirect, this will pick up any package redirects as well
-	// For things that look like native objects, try all types as we don't know which it would be
-	const bool bIsNative = OldString.StartsWith(TEXT("/Script/"));
-	FCoreRedirectObjectName NewName = FCoreRedirects::GetRedirectedName(bIsNative ? ECoreRedirectFlags::Type_AllMask : ECoreRedirectFlags::Type_Object, OldName);
+	{
+		TStringBuilder<NAME_SIZE> OldPackageNameString;
+		OldName.PackageName.ToString(OldPackageNameString);
+
+		// Always try the object redirect, this will pick up any package redirects as well
+		// For things that look like native objects, try all types as we don't know which it would be
+		const bool bIsNative = FPackageName::IsScriptPackage(OldPackageNameString);
+		NewName = FCoreRedirects::GetRedirectedName(bIsNative ? ECoreRedirectFlags::Type_AllMask : ECoreRedirectFlags::Type_Object, OldName);
+	}
 
 	if (OldName != NewName)
 	{
 		// Only do the fixup if the old object isn't in memory (or was redirected to new name), this avoids false positives
-		UObject* FoundOldObject = FindObjectSafe<UObject>(nullptr, *OldString);
+		UObject* FoundOldObject = FindObjectSafe<UObject>(nullptr, *OldName.ToString());
 		FString NewString = NewName.ToString();
 
 		if (!FoundOldObject || FoundOldObject->GetPathName() == NewString)
