DSTM:
- thread the RemoteOwnerServerId on the resident UObject so that this can persist when an object gets borrowed and we can know where to send it back to
- fixed bad code merges in RemoteObjectTransfer that broke the canonical root object matching and the bOwnershipChange detection

diff --git a/Engine/Source/Runtime/CoreUObject/Private/UObject/Class.cpp b/Engine/Source/Runtime/CoreUObject/Private/UObject/Class.cpp
--- a/Engine/Source/Runtime/CoreUObject/Private/UObject/Class.cpp
+++ b/Engine/Source/Runtime/CoreUObject/Private/UObject/Class.cpp
@@ -5027,13 +5027,13 @@ FName UClass::GetDefaultObjectName() const
 //
 void UClass::DeferredRegister(UClass *UClassStaticClass, const TCHAR* PackageName, const TCHAR* Name
 #if UE_WITH_REMOTE_OBJECT_HANDLE
-	, FRemoteObjectId RemoteId
+	, FRemoteObjectId RemoteId, FRemoteServerId RemoteOwnerServerId
 #endif
 )
 {
 	Super::DeferredRegister(UClassStaticClass, PackageName, Name
 #if UE_WITH_REMOTE_OBJECT_HANDLE
-		, RemoteId
+		, RemoteId, RemoteOwnerServerId
 #endif
 	);
 
diff --git a/Engine/Source/Runtime/CoreUObject/Private/UObject/ObjectHandlePrivate.h b/Engine/Source/Runtime/CoreUObject/Private/UObject/ObjectHandlePrivate.h
--- a/Engine/Source/Runtime/CoreUObject/Private/UObject/ObjectHandlePrivate.h
+++ b/Engine/Source/Runtime/CoreUObject/Private/UObject/ObjectHandlePrivate.h
@@ -43,6 +43,26 @@ namespace UE::CoreUObject::Private
 #endif // UE_WITH_REMOTE_OBJECT_HANDLE
 			return FRemoteObjectId();
 		}
+
+		FORCEINLINE static FRemoteServerId GetRemoteOwnerServerId(const UObjectBase* Object)
+		{
+			FRemoteServerId Result;
+
+#if UE_WITH_REMOTE_OBJECT_HANDLE
+			if (Object)
+			{
+				if (UNLIKELY(Object->InternalIndex < 0))
+				{
+					Result = Object->GetPendingRegistrantRemoteOwnerServerId();
+				}
+				else
+				{
+					Result = GUObjectArray.GetRemoteOwnerServerId(Object->InternalIndex);
+				}
+			}
+#endif // UE_WITH_REMOTE_OBJECT_HANDLE
+			return Result;
+		}
 	};
 }
 
diff --git a/Engine/Source/Runtime/CoreUObject/Private/UObject/RemoteObject.cpp b/Engine/Source/Runtime/CoreUObject/Private/UObject/RemoteObject.cpp
--- a/Engine/Source/Runtime/CoreUObject/Private/UObject/RemoteObject.cpp
+++ b/Engine/Source/Runtime/CoreUObject/Private/UObject/RemoteObject.cpp
@@ -239,6 +239,35 @@ bool IsRemote(const UObject* Object)
 	return bIsRemote;
 }
 
+bool IsOwned(const UObject* Object)
+{
+	bool Result = false;
+#if UE_WITH_REMOTE_OBJECT_HANDLE
+	FUObjectItem* ObjectItem = GUObjectArray.ObjectToObjectItem(Object);
+	checkf(ObjectItem, TEXT("Attempting to get owner id for an object that does not exist in the global UObject array (it's possible GUObjectArray is not initialized yet, ObjectIndex=%d)"), GUObjectArray.ObjectToIndex(Object));
+	Result = (ObjectItem->RemoteOwnerServerId == GetGlobalServerId());
+#endif
+	return Result;
+}
+
+FRemoteServerId GetOwnerServerId(const UObject* Object)
+{
+	FRemoteServerId Result;
+#if UE_WITH_REMOTE_OBJECT_HANDLE
+	FUObjectItem* ObjectItem = GUObjectArray.ObjectToObjectItem(Object);
+	Result = ObjectItem->RemoteOwnerServerId;
+#endif
+	return Result;
+}
+
+void ChangeOwnerServerId(const UObject* Object, FRemoteServerId NewOwnerServerId)
+{
+#if UE_WITH_REMOTE_OBJECT_HANDLE
+	FUObjectItem* ObjectItem = GUObjectArray.ObjectToObjectItem(Object);
+	ObjectItem->RemoteOwnerServerId = NewOwnerServerId;
+#endif
+}
+
 UObject* ResolveObject(const FRemoteObjectStub* Stub, ERemoteReferenceType RefType /*= ERemoteReferenceType::Strong*/)
 {
 	using namespace UE::RemoteObject::Transfer;
diff --git a/Engine/Source/Runtime/CoreUObject/Private/UObject/RemoteObjectSerialization.cpp b/Engine/Source/Runtime/CoreUObject/Private/UObject/RemoteObjectSerialization.cpp
--- a/Engine/Source/Runtime/CoreUObject/Private/UObject/RemoteObjectSerialization.cpp
+++ b/Engine/Source/Runtime/CoreUObject/Private/UObject/RemoteObjectSerialization.cpp
@@ -82,10 +82,10 @@ void LoadObjectFromDisk(FRemoteTransactionId RequestId, FRemoteWorkPriority Requ
 	FileReader->Close();
 	IFileManager::Get().Delete(*Filename, false, true, true);
 
-	UE::RemoteObject::Transfer::OnObjectDataReceived(RequestId, ObjectId, CurrentOwnerServerId, ObjectData);
+	UE::RemoteObject::Transfer::OnObjectDataReceived(RequestId, FRemoteServerId(), ObjectId, CurrentOwnerServerId, ObjectData);
 }
 
-void SaveObjectToDisk(FRemoteTransactionId RequestId, FRemoteObjectId ObjectId, FRemoteServerId DestinationServerId, FRemoteObjectData& ObjectData)
+void SaveObjectToDisk(FRemoteTransactionId RequestId, FRemoteObjectId ObjectId, FRemoteServerId DestinationServerId, FRemoteServerId OwnerServerId, FRemoteObjectData& ObjectData)
 {
 	TUniquePtr<FArchive> FileWriter{ IFileManager::Get().CreateFileWriter(*GenerateRemoteObjectFilename(ObjectId, DestinationServerId)) };
 	checkf(FileWriter.IsValid(), TEXT("Unable to create file writer for remote object %s"), *ObjectId.ToString());
@@ -612,7 +612,7 @@ FRemoteObjectConstructionOverrides::FRemoteObjectConstructionOverrides(const TAr
 	}
 }
 
-UObject* ConstructRemoteObject(const FRemoteObjectHeader& Header, FRemoteObjectConstructionOverrides* ConstructionOverrides)
+UObject* ConstructRemoteObject(const FRemoteServerId OwnerServerId, const FRemoteObjectHeader& Header, FRemoteObjectConstructionOverrides* ConstructionOverrides)
 {
 	using namespace UE::RemoteObject::Handle;
 	using namespace UE::RemoteObject::Private;
@@ -642,6 +642,8 @@ UObject* ConstructRemoteObject(const FRemoteObjectHeader& Header, FRemoteObjectC
 			Name = FName();
 			Object = nullptr;
 		}
+
+		ChangeOwnerServerId(Object, OwnerServerId);
 	}
 	else
 	{
@@ -650,6 +652,8 @@ UObject* ConstructRemoteObject(const FRemoteObjectHeader& Header, FRemoteObjectC
 		{
 			// The object already exists on this server but has been renamed
 			Name = Object->GetFName();
+
+			ChangeOwnerServerId(Object, OwnerServerId);
 		}
 	}
 
@@ -663,6 +667,7 @@ UObject* ConstructRemoteObject(const FRemoteObjectHeader& Header, FRemoteObjectC
 		Params.Template = Header.Archetype.Get();
 #if UE_WITH_REMOTE_OBJECT_HANDLE
 		Params.RemoteId = Header.RemoteId;
+		Params.RemoteOwnerServerId = OwnerServerId;
 		Params.RemoteSubObjectOverrides = ConstructionOverrides;
 #endif
 
@@ -676,6 +681,7 @@ UObject* ConstructRemoteObject(const FRemoteObjectHeader& Header, FRemoteObjectC
 	}
 
 	checkf(FRemoteObjectId(Object) == Header.RemoteId, TEXT("Created an object with a different ID:%s than requested:%s"), *FRemoteObjectId(Object).ToString(), *Header.RemoteId.ToString());
+	checkf(GetOwnerServerId(Object) == OwnerServerId, TEXT("Created an object with a different Owner:%s than requested:%s"), *GetOwnerServerId(Object).ToString(), *OwnerServerId.ToString());
 
 	MarkAsLocal(Object);
 
@@ -794,7 +800,7 @@ FRemoteObjectData* SerializeObjectData(UObject* InObject, TSet<UObject*>& OutObj
 	return ObjectData;
 }
 
-int32 DeserializeObjectData(FRemoteObjectData& ObjectData, TArray<UObject*>& OutObjects)
+int32 DeserializeObjectData(FRemoteServerId OwnerServerId, FRemoteObjectData& ObjectData, TArray<UObject*>& OutObjects)
 {
 	using namespace UE::RemoteObject::Handle;
 	using namespace UE::RemoteObject::Private;
@@ -843,7 +849,7 @@ int32 DeserializeObjectData(FRemoteObjectData& ObjectData, TArray<UObject*>& Out
 		OutObjects.Reserve(ObjectHeaders.Num());
 		for (int32 ObjectIndex = 0; ObjectIndex < ObjectHeaders.Num(); ++ObjectIndex)
 		{
-			OutObjects.Add(ConstructRemoteObject(ObjectHeaders[ObjectIndex], &ConstructionOverrides));
+			OutObjects.Add(ConstructRemoteObject(OwnerServerId, ObjectHeaders[ObjectIndex], &ConstructionOverrides));
 		}
 
 		// Deserialize all objects
diff --git a/Engine/Source/Runtime/CoreUObject/Private/UObject/RemoteObjectSerialization.h b/Engine/Source/Runtime/CoreUObject/Private/UObject/RemoteObjectSerialization.h
--- a/Engine/Source/Runtime/CoreUObject/Private/UObject/RemoteObjectSerialization.h
+++ b/Engine/Source/Runtime/CoreUObject/Private/UObject/RemoteObjectSerialization.h
@@ -69,7 +69,7 @@ namespace UE::RemoteObject::Serialization
 	* @return Index of an object in OutReceivedObjects that was the main object the migration request was triggered for 
 	* (usually 0 but if a migration requests a default subobject then its parent is also migrated and the return value will be > 0)
 	*/
-	int32 DeserializeObjectData(FRemoteObjectData& ObjectData, TArray<UObject*>& OutReceivedObjects);
+	int32 DeserializeObjectData(FRemoteServerId OwnerServerId, FRemoteObjectData& ObjectData, TArray<UObject*>& OutReceivedObjects);
 
 	/**
 	* Finds the canonical 'root' object that is used for remote object serialization
@@ -83,5 +83,5 @@ namespace UE::RemoteObject::Serialization
 namespace UE::RemoteObject::Serialization::Disk
 {
 	void LoadObjectFromDisk(FRemoteTransactionId RequestId, FRemoteWorkPriority RequestPriority, FRemoteObjectId ObjectId, FRemoteServerId CurrentOwnerServerId, FRemoteServerId DestinationServerId);
-	void SaveObjectToDisk(FRemoteTransactionId RequestId, FRemoteObjectId ObjectId, FRemoteServerId ServerId, FRemoteObjectData& ObjectData);
+	void SaveObjectToDisk(FRemoteTransactionId RequestId, FRemoteObjectId ObjectId, FRemoteServerId ServerId, FRemoteServerId OwnerServerId, FRemoteObjectData& ObjectData);
 }
\ No newline at end of file
diff --git a/Engine/Source/Runtime/CoreUObject/Private/UObject/RemoteObjectTransfer.cpp b/Engine/Source/Runtime/CoreUObject/Private/UObject/RemoteObjectTransfer.cpp
--- a/Engine/Source/Runtime/CoreUObject/Private/UObject/RemoteObjectTransfer.cpp
+++ b/Engine/Source/Runtime/CoreUObject/Private/UObject/RemoteObjectTransfer.cpp
@@ -22,7 +22,7 @@ namespace UE::RemoteObject::Transfer
 const FRemoteServerId DatabaseId(0x3ff);
 
 TDelegate<void(FRemoteTransactionId /*RequestId*/, FRemoteWorkPriority /*RequestPriority*/, FRemoteObjectId /*ObjectId*/, FRemoteServerId /*CurrentOwnerServerId*/, FRemoteServerId /*DestinationServerId*/)> RequestRemoteObjectDelegate;
-TDelegate<void(FRemoteTransactionId /*RequestId*/, FRemoteObjectId /*ObjectId*/, FRemoteServerId /*DestinationServerId*/, FRemoteObjectData& /*ObjectData*/)> RemoteObjectTransferDelegate;
+TDelegate<void(FRemoteTransactionId /*RequestId*/, FRemoteObjectId /*ObjectId*/, FRemoteServerId /*DestinationServerId*/, FRemoteServerId /*OwnerServerId*/, FRemoteObjectData& /*ObjectData*/)> RemoteObjectTransferDelegate;
 TDelegate<void(FRemoteTransactionId /*RequestId*/, FRemoteObjectId /*ObjectId*/, FRemoteServerId /*DestinationServerId*/)> RemoteObjectDeniedTransferDelegate;
 
 TMulticastDelegate<void(const FRemoteObjectData& /*ObjectData*/, FRemoteServerId /*RemoteServerId*/)> OnObjectDataSentDelegate;
@@ -31,10 +31,10 @@ TMulticastDelegate<void(const FRemoteObjectData& /*ObjectData*/, FRemoteServerId
 TMulticastDelegate<void(FRemoteObjectId /*ObjectId*/, FRemoteServerId /*RemoteServerId*/)> OnObjectDataDeniedDelegate;
 TMulticastDelegate<void(const TArray<UObject*>& /*MigratedObjects*/, FRemoteServerId /*RemoteServerId*/)> OnObjectsReceivedDelegate;
 
-TDelegate<void(FRemoteTransactionId /*RequestId*/, FRemoteObjectId /*ObjectId*/, FRemoteServerId /*DatabaseId*/, FRemoteObjectData& /*ObjectData*/)> StoreRemoteObjectDataDelegate;
+TDelegate<void(FRemoteTransactionId /*RequestId*/, FRemoteObjectId /*ObjectId*/, FRemoteServerId /*DatabaseId*/, FRemoteServerId /*OwnerServerId*/, FRemoteObjectData& /*ObjectData*/)> StoreRemoteObjectDataDelegate;
 TDelegate<void(FRemoteTransactionId /*RequestId*/, FRemoteWorkPriority /*RequestPriority*/, FRemoteObjectId /*ObjectId*/, FRemoteServerId /*DatabaseId*/, FRemoteServerId /*DestinationServerId*/)> RestoreRemoteObjectDataDelegate;
 
-UObject* MigrateObjectFromObjectData(FRemoteObjectData& ObjectData, FRemoteServerId RemoteServerId, EObjectMigrationFlags ObjectFlags);
+UObject* MigrateObjectFromObjectData(FRemoteServerId OwnerServerId, FRemoteObjectData& ObjectData, FRemoteServerId RemoteServerId, EObjectMigrationFlags ObjectFlags);
 
 /**
  * Remote object transfer queue. 
@@ -156,10 +156,10 @@ public:
 	* @param ObjectData Object data to send. Send request takes ownership of the data
 	* @return Future that's going to receive the object data
 	*/
-	void SendRemoteObject(FRemoteTransactionId RequestId, FRemoteObjectId ObjectId, FRemoteServerId DestinationServerId, FRemoteObjectData* ObjectData)
+	void SendRemoteObject(FRemoteTransactionId RequestId, FRemoteObjectId ObjectId, FRemoteServerId DestinationServerId, FRemoteServerId OwnerServerId, FRemoteObjectData* ObjectData)
 	{
 		check(!AutoRTFM::IsClosed());
-		RemoteObjectTransferDelegate.Execute(RequestId, ObjectId, DestinationServerId, *ObjectData);
+		RemoteObjectTransferDelegate.Execute(RequestId, ObjectId, DestinationServerId, OwnerServerId, *ObjectData);
 	}
 
 	void SendDeniedRemoteObject(FRemoteTransactionId RequestId, FRemoteObjectId ObjectId, FRemoteServerId DestinationServerId)
@@ -174,7 +174,7 @@ public:
 	* @param Data Object data. Data ownership is transferred to the receive request (if it exists)
 	* @return True if a receive request for the specified object exists (and object data ownership was transferred), false otherwise
 	*/
-	void FulfillReceiveRequest(FRemoteTransactionId RequestId, FRemoteObjectId ObjectId, FRemoteServerId RemoteServerId, FRemoteObjectData* ObjectData)
+	void FulfillReceiveRequest(FRemoteTransactionId RequestId, FRemoteServerId OwnerServerId, FRemoteObjectId ObjectId, FRemoteServerId RemoteServerId, FRemoteObjectData* ObjectData)
 	{
 		using namespace UE::RemoteObject::Serialization;
 
@@ -183,13 +183,14 @@ public:
 		const int32 PendingRequestIndex = PendingObjectRequests.Find(ObjectId);
 		const bool bFoundRequest = (PendingRequestIndex != INDEX_NONE);
 
-		UE_LOG(LogRemoteObject, Verbose, TEXT("     : FulfillReceiveRequest(%s) %s (was requested: %d)"), *RequestId.ToString(), *ObjectId.ToString(), bFoundRequest);
+		UE_LOG(LogRemoteObject, Verbose, TEXT("     : FulfillReceiveRequest(%s) %s (was requested: %d, owner: %s)"), *RequestId.ToString(), *ObjectId.ToString(), bFoundRequest, *OwnerServerId.ToString());
 
-		const bool bAssignedOwnership = !RequestId.IsValid();
+		const bool bAssignedOwnership = (RequestId == PushMigrationRequestId);
 		EObjectMigrationFlags MigrationFlags = bAssignedOwnership ? EObjectMigrationFlags::Acquire : EObjectMigrationFlags::Receive;
-		UObject* Result = MigrateObjectFromObjectData(*ObjectData, RemoteServerId, MigrationFlags);
+		UObject* Result = MigrateObjectFromObjectData(OwnerServerId, *ObjectData, RemoteServerId, MigrationFlags);
 		checkf(UE::CoreUObject::Private::FObjectHandleUtils::GetRemoteId(Result) == ObjectId, TEXT("PeerId:%s requested object %s migration from PeerId:%s but received %s"),
 			*GetGlobalServerId().ToString(), *ObjectId.ToString(), *RemoteServerId.ToString(), *UE::CoreUObject::Private::FObjectHandleUtils::GetRemoteId(Result).ToString());
+		checkf(OwnerServerId == Handle::GetOwnerServerId(Result), TEXT("Expected owner id %s doesn't match post-deserialize owner: %s"), *OwnerServerId.ToString(), *Handle::GetOwnerServerId(Result).ToString());
 
 		if (bFoundRequest)
 		{
@@ -243,7 +244,7 @@ void RequestObjectData(FRemoteObjectId ObjectId, FRemoteServerId CurrentOwnerSer
 		DestinationServerId);
 }
 
-UObject* MigrateObjectFromObjectData(FRemoteObjectData& ObjectData, FRemoteServerId RemoteServerId, EObjectMigrationFlags MigrationFlags)
+UObject* MigrateObjectFromObjectData(FRemoteServerId OwnerServerId, FRemoteObjectData& ObjectData, FRemoteServerId RemoteServerId, EObjectMigrationFlags MigrationFlags)
 {
 	using namespace UE::RemoteObject::Serialization;
 	check(!AutoRTFM::IsClosed());
@@ -252,7 +253,7 @@ UObject* MigrateObjectFromObjectData(FRemoteObjectData& ObjectData, FRemoteServe
 	OnObjectDataReceivedDelegate.Broadcast(ObjectData, RemoteServerId);
 
 	TArray<UObject*> ReceivedObjects;
-	int32 RequestedObjectIndex = DeserializeObjectData(ObjectData, ReceivedObjects);
+	int32 RequestedObjectIndex = DeserializeObjectData(OwnerServerId, ObjectData, ReceivedObjects);
 
 	checkf(ReceivedObjects.Num() > 0, TEXT("PeerId:%s unable to deserialize object data (%d bytes)"), *GetGlobalServerId().ToString(), ObjectData.Bytes.Num());
 
@@ -359,17 +360,17 @@ void MigrateObjectToRemoteServerWithExplicitPriority(FRemoteTransactionId Reques
 	const FRemoteObjectRequest* EqualPriorityRequest = nullptr;
 	const FRemoteObjectRequest* LowerPriorityRequest = nullptr;
 
-	UE_LOG(LogRemoteObject, Verbose, TEXT("     : MigrateObjectToRemoteServer(%s) sending obj %s [%s] (isremote %d)"), *RequestId.ToString(), *ObjectId.ToString(), *Object->GetName(), Handle::IsRemote(Object));
+	UE_LOG(LogRemoteObject, Verbose, TEXT("     : MigrateObjectToRemoteServer(%s) sending obj %s [%s] (root %s) (isremote %d)"), *RequestId.ToString(), *ObjectId.ToString(), *Object->GetName(), *CanonicalRootObjectId.ToString(), Handle::IsRemote(Object));
 
 	for (int32 RequestIndex = 0; RequestIndex < GTransferQueue.GetRequestCount(); RequestIndex++)
 	{
 		FRemoteObjectRequest* ExistingRequest = GTransferQueue.GetRequestByIndex(RequestIndex);
 
-		for (int32 RequiredObjectIndex = 0; RequiredObjectIndex < ExistingRequest->RequiredObjects.Num(); RequiredObjectIndex++)
+		for (int32 RequiredObjectIndex = 0; RequiredObjectIndex < ExistingRequest->RequiredObjectsCanonicalRoot.Num(); RequiredObjectIndex++)
 		{
-			FRemoteObjectId RequiredObjectId = ExistingRequest->RequiredObjects[RequiredObjectIndex];
+			FRemoteObjectId RequiredObjectId = ExistingRequest->RequiredObjectsCanonicalRoot[RequiredObjectIndex];
 
-			if (RequiredObjectId == ObjectId)
+			if (RequiredObjectId == CanonicalRootObjectId)
 			{
 				// do we give it up, or deny their request?
 				bool bIsHigherPriority = IsHigherPriority(ExistingRequest->Priority, RequestPriority);
@@ -430,7 +431,10 @@ void MigrateObjectToRemoteServerWithExplicitPriority(FRemoteTransactionId Reques
 	}
 	else
 	{
-		UE_LOG(LogRemoteObject, VeryVerbose, TEXT("     : MigrateObjectToRemoteServer(%s) serializing obj %s (isremote %d)"), *RequestId.ToString(), *ObjectId.ToString(), Handle::IsRemote(Object));
+		const bool bOwnershipChange = (RequestId == PushMigrationRequestId);
+		FRemoteServerId OwnerServerId = bOwnershipChange ? DestinationServerId : Handle::GetOwnerServerId(RootObject);
+
+		UE_LOG(LogRemoteObject, VeryVerbose, TEXT("     : MigrateObjectToRemoteServer(%s) serializing obj %s (isremote %d owner changing %d owner %s)"), *RequestId.ToString(), *ObjectId.ToString(), Handle::IsRemote(Object), bOwnershipChange, *OwnerServerId.ToString());
 
 		// send the object
 		TSet<UObject*> SentObjects;
@@ -447,7 +451,7 @@ void MigrateObjectToRemoteServerWithExplicitPriority(FRemoteTransactionId Reques
 		// Notify that object data has been migrated (actually will be but SendObjectData does a MoveTemp on the data)
 		OnObjectDataSentDelegate.Broadcast(*ObjectData, DestinationServerId);
 
-		GTransferQueue.SendRemoteObject(RequestId, ObjectId, DestinationServerId, ObjectData);
+		GTransferQueue.SendRemoteObject(RequestId, ObjectId, DestinationServerId, OwnerServerId, ObjectData);
 
 		checkf(SentObjects.Num() > 0, TEXT("PeerId:%s unable to migrate object %s to PeerId:%s"), *GetGlobalServerId().ToString(), *UE::CoreUObject::Private::FObjectHandleUtils::GetRemoteId(Object).ToString(), *DestinationServerId.ToString());
 
@@ -460,7 +464,6 @@ void MigrateObjectToRemoteServerWithExplicitPriority(FRemoteTransactionId Reques
 			OnObjectsSentDelegate.Broadcast(SentObjects, DestinationServerId);
 
 			// Are we loaning it out or changing ownership?
-			const bool bOwnershipChange = !RequestId.IsValid();
 			FPostMigrateParams PostMigrateParams;
 			PostMigrateParams.RequestingServerId = RequestPriority.GetRootServerId();
 			PostMigrateParams.DestinationServerId = DestinationServerId;
@@ -536,10 +539,10 @@ void MigrateObjectToRemoteServerWithExplicitPriority(FRemoteTransactionId Reques
 	}
 }
 
-void OnObjectDataReceived(FRemoteTransactionId RequestId, FRemoteObjectId ObjectId, FRemoteServerId RemoteServerId, FRemoteObjectData* Data)
+void OnObjectDataReceived(FRemoteTransactionId RequestId, FRemoteServerId OwnerServerId, FRemoteObjectId ObjectId, FRemoteServerId RemoteServerId, FRemoteObjectData* Data)
 {
 	check(!AutoRTFM::IsClosed());
-	GTransferQueue.FulfillReceiveRequest(RequestId, ObjectId, RemoteServerId, Data);
+	GTransferQueue.FulfillReceiveRequest(RequestId, OwnerServerId, ObjectId, RemoteServerId, Data);
 }
 
 void OnObjectDataDenied(FRemoteTransactionId RequestId, FRemoteObjectId ObjectId, FRemoteServerId RemoteServerId)
@@ -619,7 +622,7 @@ void StoreObjectToDatabase(UObject* Object)
 		TSet<UObject*> SentObjects;
 		{
 			TUniquePtr<FRemoteObjectData> ObjectData(SerializeObjectData(OutermostUnreachableRemoteReference, SentObjects));
-			StoreRemoteObjectDataDelegate.Execute(PushMigrationRequestId, FRemoteObjectId(OutermostUnreachableRemoteReference), DatabaseId, *ObjectData);
+			StoreRemoteObjectDataDelegate.Execute(PushMigrationRequestId, FRemoteObjectId(OutermostUnreachableRemoteReference), DatabaseId, FRemoteServerId(), *ObjectData);
 		}
 
 		checkf(SentObjects.Num() > 0, TEXT("PeerId:%s unable to store object %s to DatabaseId:%s"), *GetGlobalServerId().ToString(), *FRemoteObjectId(Object).ToString(), *DatabaseId.ToString());
diff --git a/Engine/Source/Runtime/CoreUObject/Private/UObject/UObjectArray.cpp b/Engine/Source/Runtime/CoreUObject/Private/UObject/UObjectArray.cpp
--- a/Engine/Source/Runtime/CoreUObject/Private/UObject/UObjectArray.cpp
+++ b/Engine/Source/Runtime/CoreUObject/Private/UObject/UObjectArray.cpp
@@ -226,7 +226,7 @@ void FUObjectArray::DisableDisregardForGC()
 	}
 }
 
-void FUObjectArray::AllocateUObjectIndex(UObjectBase* Object, EInternalObjectFlags InitialFlags, int32 AlreadyAllocatedIndex, int32 SerialNumber, FRemoteObjectId RemoteId)
+void FUObjectArray::AllocateUObjectIndex(UObjectBase* Object, EInternalObjectFlags InitialFlags, int32 AlreadyAllocatedIndex, int32 SerialNumber, FRemoteObjectId RemoteId, FRemoteServerId RemoteOwnerServerId)
 {
 	LLM_SCOPE(ELLMTag::UObject);
 	// Clear asset scopes
@@ -242,6 +242,10 @@ void FUObjectArray::AllocateUObjectIndex(UObjectBase* Object, EInternalObjectFla
 	{
 		RemoteId = FRemoteObjectId::Generate();
 	}
+	if (!RemoteOwnerServerId.IsValid())
+	{
+		RemoteOwnerServerId = UE::RemoteObject::GetGlobalServerId();
+	}
 #endif
 
 	LockInternalArray();
@@ -298,6 +302,7 @@ void FUObjectArray::AllocateUObjectIndex(UObjectBase* Object, EInternalObjectFla
 	ObjectItem->SerialNumber = SerialNumber;
 #if UE_WITH_REMOTE_OBJECT_HANDLE
 	ObjectItem->RemoteId = RemoteId;
+	ObjectItem->RemoteOwnerServerId = RemoteOwnerServerId;
 #endif // UE_WITH_REMOTE_OBJECT_HANDLE
 	Object->InternalIndex = Index;
 
@@ -329,6 +334,7 @@ void FUObjectArray::ResetSerialNumber(UObjectBase* Object)
 	ObjectItem->SerialNumber = 0;
 #if UE_WITH_REMOTE_OBJECT_HANDLE
 	ObjectItem->RemoteId = FRemoteObjectId();
+	ObjectItem->RemoteOwnerServerId = FRemoteServerId();
 #endif
 }
 
@@ -397,6 +403,7 @@ void FUObjectArray::FreeUObjectIndex(UObjectBase* Object)
 	ObjectItem->SerialNumber = 0;
 #if UE_WITH_REMOTE_OBJECT_HANDLE
 	ObjectItem->RemoteId = FRemoteObjectId();
+	ObjectItem->RemoteOwnerServerId = FRemoteServerId();
 #endif
 	Object->InternalIndex = INDEX_NONE;
 
diff --git a/Engine/Source/Runtime/CoreUObject/Private/UObject/UObjectBase.cpp b/Engine/Source/Runtime/CoreUObject/Private/UObject/UObjectBase.cpp
--- a/Engine/Source/Runtime/CoreUObject/Private/UObject/UObjectBase.cpp
+++ b/Engine/Source/Runtime/CoreUObject/Private/UObject/UObjectBase.cpp
@@ -65,11 +65,13 @@ struct FPendingRegistrantInfo
 	class UClass*  (*StaticClassFn)();
 #if UE_WITH_REMOTE_OBJECT_HANDLE
 	const FRemoteObjectId RemoteId;
+	const FRemoteServerId RemoteOwnerServerId;
 #endif // UE_WITH_REMOTE_OBJECT_HANDLE
 
 	FPendingRegistrantInfo(class UClass* (*InStaticClassFn)(), const TCHAR* InName, const TCHAR* InPackageName
 #if UE_WITH_REMOTE_OBJECT_HANDLE
 		, const FRemoteObjectId InRemoteId
+		, const FRemoteServerId InRemoteOwnerServerId
 #endif
 		)
 		: Name(InName)
@@ -77,6 +79,7 @@ struct FPendingRegistrantInfo
 		, StaticClassFn(InStaticClassFn)
 #if UE_WITH_REMOTE_OBJECT_HANDLE
 		, RemoteId(InRemoteId)
+		, RemoteOwnerServerId(InRemoteOwnerServerId)
 #endif
 	{}
 	static TMap<UObjectBase*, FPendingRegistrantInfo>& GetMap()
@@ -146,7 +149,8 @@ UObjectBase::UObjectBase(UClass* InClass,
 	FName InName,
 	int32 InInternalIndex,
 	int32 InSerialNumber,
-	FRemoteObjectId InRemoteId)
+	FRemoteObjectId InRemoteId,
+	FRemoteServerId InRemoteOwnerServerId)
 :	ObjectFlags			(InFlags)
 ,	InternalIndex		(INDEX_NONE)
 ,	ClassPrivate		(InClass)
@@ -154,7 +158,7 @@ UObjectBase::UObjectBase(UClass* InClass,
 {
 	check(ClassPrivate);
 	// Add to global table.
-	AddObject(InName, InInternalFlags, InInternalIndex, InSerialNumber, InRemoteId);
+	AddObject(InName, InInternalFlags, InInternalIndex, InSerialNumber, InRemoteId, InRemoteOwnerServerId);
 	
 #if CSV_PROFILER_STATS && CSV_TRACK_UOBJECT_COUNT
 	UObjectStats::IncrementUObjectCount();
@@ -192,6 +196,7 @@ UObjectBase::~UObjectBase()
 void UObjectBase::DeferredRegister(UClass *UClassStaticClass, const TCHAR* PackageName, const TCHAR* InName
 #if UE_WITH_REMOTE_OBJECT_HANDLE
 	, FRemoteObjectId RemoteId
+	, FRemoteServerId RemoteOwnerServerId
 #endif
 )
 {
@@ -209,7 +214,7 @@ void UObjectBase::DeferredRegister(UClass *UClassStaticClass, const TCHAR* Packa
 	// Add to the global object table.
 	AddObject(FName(InName), EInternalObjectFlags::None
 #if UE_WITH_REMOTE_OBJECT_HANDLE
-		, /*IndernalIndex = */ -1, /*SerialNumber = */ 0, RemoteId
+		, /*IndernalIndex = */ -1, /*SerialNumber = */ 0, RemoteId, RemoteOwnerServerId
 #endif
 	);
 
@@ -235,7 +240,7 @@ void UObjectBase::DeferredRegister(UClass *UClassStaticClass, const TCHAR* Packa
  * @param Name name to assign to this uobject
  */
 UE_AUTORTFM_ALWAYS_OPEN
-void UObjectBase::AddObject(FName InName, EInternalObjectFlags InSetInternalFlags, int32 InInternalIndex, int32 InSerialNumber, FRemoteObjectId InRemoteId)
+void UObjectBase::AddObject(FName InName, EInternalObjectFlags InSetInternalFlags, int32 InInternalIndex, int32 InSerialNumber, FRemoteObjectId InRemoteId, FRemoteServerId InRemoteOwnerServerId)
 {
 	NamePrivate = InName;
 	EInternalObjectFlags InternalFlagsToSet = InSetInternalFlags;
@@ -253,7 +258,7 @@ void UObjectBase::AddObject(FName InName, EInternalObjectFlags InSetInternalFlag
 		InternalFlagsToSet |= EInternalObjectFlags::Native;
 		ObjectFlags &= ~RF_MarkAsNative;
 	}
-	GUObjectArray.AllocateUObjectIndex(this, InternalFlagsToSet, InInternalIndex, InSerialNumber, InRemoteId);
+	GUObjectArray.AllocateUObjectIndex(this, InternalFlagsToSet, InInternalIndex, InSerialNumber, InRemoteId, InRemoteOwnerServerId);
 	check(InName != NAME_None && InternalIndex >= 0);
 	HashObject(this);
 	check(IsValidLowLevel());
@@ -537,6 +542,7 @@ void UObjectBase::Register(class UClass* (*StaticClassFn)(), const TCHAR* Packag
 	PendingRegistrants.Add(this, FPendingRegistrantInfo(StaticClassFn, InName, PackageName
 #if UE_WITH_REMOTE_OBJECT_HANDLE
 		, FRemoteObjectId::Generate()
+		, FRemoteServerId()
 #endif
 	));
 
@@ -570,6 +576,13 @@ FRemoteObjectId UObjectBase::GetPendingRegistrantRemoteId() const
 	FPendingRegistrantInfo& Info = PendingRegistrants.FindChecked(this);
 	return Info.RemoteId;
 }
+
+FRemoteServerId UObjectBase::GetPendingRegistrantRemoteOwnerServerId() const
+{
+	TMap<UObjectBase*, FPendingRegistrantInfo>& PendingRegistrants = FPendingRegistrantInfo::GetMap();
+	FPendingRegistrantInfo& Info = PendingRegistrants.FindChecked(this);
+	return Info.RemoteOwnerServerId;
+}
 #endif
 
 /**
@@ -636,11 +649,12 @@ void UObjectForceRegistration(UObjectBase* Object, bool bCheckForModuleRelease)
 		UClass* StaticClass = Info->StaticClassFn();
 #if UE_WITH_REMOTE_OBJECT_HANDLE
 		FRemoteObjectId RemoteId = Info->RemoteId;
+		FRemoteServerId RemoteOwnerServerId = Info->RemoteOwnerServerId;
 #endif
 		PendingRegistrants.Remove(Object);  // delete this first so that it doesn't try to do it twice
 		Object->DeferredRegister(StaticClass, PackageName, Name
 #if UE_WITH_REMOTE_OBJECT_HANDLE
-			, RemoteId
+			, RemoteId, RemoteOwnerServerId
 #endif
 		);
 	}
diff --git a/Engine/Source/Runtime/CoreUObject/Private/UObject/UObjectGlobals.cpp b/Engine/Source/Runtime/CoreUObject/Private/UObject/UObjectGlobals.cpp
--- a/Engine/Source/Runtime/CoreUObject/Private/UObject/UObjectGlobals.cpp
+++ b/Engine/Source/Runtime/CoreUObject/Private/UObject/UObjectGlobals.cpp
@@ -3554,6 +3554,7 @@ UObject* StaticAllocateObject
 	UPackage* ExternalPackage,
 	int32 SerialNumber,
 	FRemoteObjectId RemoteId,
+	FRemoteServerId RemoteOwnerServerId,
 	FGCReconstructionGuard* GCGuard)
 {
 	LLM_SCOPE(ELLMTag::UObject);
@@ -3830,7 +3831,7 @@ UObject* StaticAllocateObject
 #if UE_WITH_REMOTE_OBJECT_HANDLE
 				OldRemoteId.IsValid() ? OldRemoteId :
 #endif
-				RemoteId);
+				RemoteId, RemoteOwnerServerId);
 		};
 		if (!bWasConstructedOnOldObject)
 		{
@@ -4887,6 +4888,12 @@ UObject* StaticConstructObject_Internal(const FStaticConstructObjectParameters&
 	UObject* InTemplate = Params.Template;
 	int32 SerialNumber = Params.SerialNumber;
 	FRemoteObjectId RemoteId;
+#if UE_WITH_REMOTE_OBJECT_HANDLE
+	FRemoteServerId RemoteOwnerServerId = Params.RemoteOwnerServerId;
+#else
+	FRemoteServerId RemoteOwnerServerId;
+#endif
+
 	UE::RemoteObject::Serialization::FRemoteObjectConstructionOverrides* RemoteSubObjectOverrides = nullptr;
 
 	LLM_SCOPE(ELLMTag::UObject);
@@ -4925,7 +4932,7 @@ UObject* StaticConstructObject_Internal(const FStaticConstructObjectParameters&
 
 	FGCReconstructionGuard GCGuard;
 	bool bRecycledSubobject = false;	
-	Result = StaticAllocateObject(InClass, InOuter, InName, InFlags, Params.InternalSetFlags, bCanRecycleSubobjects, &bRecycledSubobject, Params.ExternalPackage, SerialNumber, RemoteId, &GCGuard);
+	Result = StaticAllocateObject(InClass, InOuter, InName, InFlags, Params.InternalSetFlags, bCanRecycleSubobjects, &bRecycledSubobject, Params.ExternalPackage, SerialNumber, RemoteId, RemoteOwnerServerId, &GCGuard);
 	check(Result != nullptr);
 	// Don't call the constructor on recycled subobjects, they haven't been destroyed.
 	if (!bRecycledSubobject)
diff --git a/Engine/Source/Runtime/CoreUObject/Public/UObject/Class.h b/Engine/Source/Runtime/CoreUObject/Public/UObject/Class.h
--- a/Engine/Source/Runtime/CoreUObject/Public/UObject/Class.h
+++ b/Engine/Source/Runtime/CoreUObject/Public/UObject/Class.h
@@ -3248,6 +3248,7 @@ public:
 	COREUOBJECT_API virtual void DeferredRegister(UClass* UClassStaticClass, const TCHAR* PackageName, const TCHAR* InName
 #if UE_WITH_REMOTE_OBJECT_HANDLE
 		, FRemoteObjectId RemoteId
+		, FRemoteServerId RemoteOwnerServerId
 #endif // UE_WITH_REMOTE_OBJECT_HANDLE
 	) override;
 	COREUOBJECT_API virtual bool Rename(const TCHAR* NewName = nullptr, UObject* NewOuter = nullptr, ERenameFlags Flags = REN_None) override;
diff --git a/Engine/Source/Runtime/CoreUObject/Public/UObject/RemoteObject.h b/Engine/Source/Runtime/CoreUObject/Public/UObject/RemoteObject.h
--- a/Engine/Source/Runtime/CoreUObject/Public/UObject/RemoteObject.h
+++ b/Engine/Source/Runtime/CoreUObject/Public/UObject/RemoteObject.h
@@ -93,4 +93,11 @@ namespace UE::RemoteObject::Handle
 	* @return True if the object is remote
 	*/
 	COREUOBJECT_API bool IsRemote(const UObject* Object);
-}
+
+	/**
+	* Checks if a locally resident object is owned by this server
+	*/
+	COREUOBJECT_API bool IsOwned(const UObject* Object);
+	COREUOBJECT_API FRemoteServerId GetOwnerServerId(const UObject* Object);
+	COREUOBJECT_API void ChangeOwnerServerId(const UObject* Object, FRemoteServerId NewOwnerServerId);
+}
\ No newline at end of file
diff --git a/Engine/Source/Runtime/CoreUObject/Public/UObject/RemoteObjectTransfer.h b/Engine/Source/Runtime/CoreUObject/Public/UObject/RemoteObjectTransfer.h
--- a/Engine/Source/Runtime/CoreUObject/Public/UObject/RemoteObjectTransfer.h
+++ b/Engine/Source/Runtime/CoreUObject/Public/UObject/RemoteObjectTransfer.h
@@ -58,7 +58,7 @@ namespace UE::RemoteObject::Transfer
 	* @param RemoteServerId Server that sent object data
 	* @param Data Remote object data. Data ownership is transferred to the remote object queue
 	*/
-	COREUOBJECT_API void OnObjectDataReceived(FRemoteTransactionId RequestId, FRemoteObjectId ObjectId, FRemoteServerId RemoteServerId, FRemoteObjectData* Data);
+	COREUOBJECT_API void OnObjectDataReceived(FRemoteTransactionId RequestId, FRemoteServerId OwnerServerId, FRemoteObjectId ObjectId, FRemoteServerId RemoteServerId, FRemoteObjectData* Data);
 
 	/**
 	* Called when a remote object request was denied by a remote server
@@ -131,7 +131,7 @@ namespace UE::RemoteObject::Transfer
 	extern COREUOBJECT_API void InitRemoteObjectTransfer();
 
 	/** Delegate that transfers object data to another server */
-	extern COREUOBJECT_API TDelegate<void(FRemoteTransactionId /*RequestId*/, FRemoteObjectId /*ObjectId*/, FRemoteServerId /*DestinationServerId*/, FRemoteObjectData& /*ObjectData*/)> RemoteObjectTransferDelegate;
+	extern COREUOBJECT_API TDelegate<void(FRemoteTransactionId /*RequestId*/, FRemoteObjectId /*ObjectId*/, FRemoteServerId /*DestinationServerId*/, FRemoteServerId /*OwnerServerId*/, FRemoteObjectData& /*ObjectData*/)> RemoteObjectTransferDelegate;
 
 	/** Delegate that handles an object request being denied */
 	extern COREUOBJECT_API TDelegate<void(FRemoteTransactionId /*RequestId*/, FRemoteObjectId /*ObjectId*/, FRemoteServerId /*DestinationServerId*/)> RemoteObjectDeniedTransferDelegate;
@@ -159,7 +159,7 @@ namespace UE::RemoteObject::Transfer
 	extern COREUOBJECT_API TMulticastDelegate<void(const TSet<UObject*>& /*MigratedObjects*/, FRemoteServerId /*RemoteServerId*/)> OnObjectsSentDelegate;
 
 	/** Delegate that stores locally unrachable object data into a database. */
-	extern COREUOBJECT_API TDelegate<void(FRemoteTransactionId /*RequestId*/, FRemoteObjectId /*ObjectId*/, FRemoteServerId /*DatabaseId*/, FRemoteObjectData& /*ObjectData*/)> StoreRemoteObjectDataDelegate;
+	extern COREUOBJECT_API TDelegate<void(FRemoteTransactionId /*RequestId*/, FRemoteObjectId /*ObjectId*/, FRemoteServerId /*DatabaseId*/, FRemoteServerId /*OwnerServerId*/, FRemoteObjectData& /*ObjectData*/)> StoreRemoteObjectDataDelegate;
 
 	/** Delegate that restores object data from a database. */
 	extern COREUOBJECT_API TDelegate<void(FRemoteTransactionId /*RequestId*/, FRemoteWorkPriority /*RequestPriority*/, FRemoteObjectId /*ObjectId*/, FRemoteServerId /*DatabaseId*/, FRemoteServerId /*DestinationServerId*/)> RestoreRemoteObjectDataDelegate;
diff --git a/Engine/Source/Runtime/CoreUObject/Public/UObject/UObjectArray.h b/Engine/Source/Runtime/CoreUObject/Public/UObject/UObjectArray.h
--- a/Engine/Source/Runtime/CoreUObject/Public/UObject/UObjectArray.h
+++ b/Engine/Source/Runtime/CoreUObject/Public/UObject/UObjectArray.h
@@ -63,6 +63,7 @@ public:
 #if UE_WITH_REMOTE_OBJECT_HANDLE
 	// Globally unqiue id of this object
 	FRemoteObjectId RemoteId;
+	FRemoteServerId RemoteOwnerServerId;
 #endif
 
 #if STATS || ENABLE_STATNAMEDEVENTS_UOBJECT
@@ -135,6 +136,11 @@ public:
 	{
 		return RemoteId;
 	}
+
+	FORCEINLINE FRemoteServerId GetRemoteOwnerServerId() const
+	{
+		return RemoteOwnerServerId;
+	}
 #endif // UE_WITH_REMOTE_OBJECT_HANDLE
 
 	FORCEINLINE void SetFlags(EInternalObjectFlags FlagsToSet)
@@ -803,7 +809,7 @@ public:
 class FUObjectArray
 {
 	friend class UObject;
-	friend COREUOBJECT_API UObject* StaticAllocateObject(const UClass*, UObject*, FName, EObjectFlags, EInternalObjectFlags, bool, bool*, UPackage*, int32, FRemoteObjectId, class FGCReconstructionGuard*);
+	friend COREUOBJECT_API UObject* StaticAllocateObject(const UClass*, UObject*, FName, EObjectFlags, EInternalObjectFlags, bool, bool*, UPackage*, int32, FRemoteObjectId, FRemoteServerId, class FGCReconstructionGuard*);
 
 private:
 	/**
@@ -924,7 +930,7 @@ public:
 	 * @param	AlreadyAllocatedIndex already allocated internal index to use, negative value means allocate a new index
 	 * @param	SerialNumber serial number to use
 	 */
-	COREUOBJECT_API void AllocateUObjectIndex(class UObjectBase* Object, EInternalObjectFlags InitialFlags, int32 AlreadyAllocatedIndex = -1, int32 SerialNumber = 0, FRemoteObjectId RemoteId = FRemoteObjectId());
+	COREUOBJECT_API void AllocateUObjectIndex(class UObjectBase* Object, EInternalObjectFlags InitialFlags, int32 AlreadyAllocatedIndex = -1, int32 SerialNumber = 0, FRemoteObjectId RemoteId = FRemoteObjectId(), FRemoteServerId RemoteOwnerServerId = FRemoteServerId());
 
 	/**
 	 * Returns a UObject index top to the global uobject array
@@ -1182,6 +1188,13 @@ public:
 		checkSlow(ObjectItem);
 		return ObjectItem->GetRemoteId();
 	}
+
+	FORCEINLINE FRemoteServerId GetRemoteOwnerServerId(int32 Index)
+	{
+		FUObjectItem* ObjectItem = IndexToObject(Index);
+		checkSlow(ObjectItem);
+		return ObjectItem->GetRemoteOwnerServerId();
+	}
 #endif
 
 	/** Locks the internal object array mutex */
diff --git a/Engine/Source/Runtime/CoreUObject/Public/UObject/UObjectBase.h b/Engine/Source/Runtime/CoreUObject/Public/UObject/UObjectBase.h
--- a/Engine/Source/Runtime/CoreUObject/Public/UObject/UObjectBase.h
+++ b/Engine/Source/Runtime/CoreUObject/Public/UObject/UObjectBase.h
@@ -92,7 +92,8 @@ public:
 			FName InName,
 			int32 InInternalIndex = -1,
 			int32 InSerialNumber = 0,
-			FRemoteObjectId InRemoteId = FRemoteObjectId());
+			FRemoteObjectId InRemoteId = FRemoteObjectId(),
+			FRemoteServerId InRemoteOwnerServerId = FRemoteServerId());
 
 	/**
 	 * Final destructor, removes the object from the object array, and indirectly, from any annotations
@@ -124,6 +125,7 @@ protected:
 	COREUOBJECT_API virtual void DeferredRegister(UClass* UClassStaticClass, const TCHAR* PackageName, const TCHAR* Name
 #if UE_WITH_REMOTE_OBJECT_HANDLE
 		, FRemoteObjectId InRemoteId
+		, FRemoteServerId InRemoteOwnerServerId
 #endif
 	);
 
@@ -136,7 +138,7 @@ private:
 	 * @param InInternalIndex already allocated internal index to use, negative value means allocate a new index
 	 * @param InSerialNumber already allocated serial number to re-use
 	 */
-	COREUOBJECT_API void AddObject(FName Name, EInternalObjectFlags InSetInternalFlags, int32 InInternalIndex = -1, int32 InSerialNumber = 0, FRemoteObjectId InRemoteId = FRemoteObjectId());
+	COREUOBJECT_API void AddObject(FName Name, EInternalObjectFlags InSetInternalFlags, int32 InInternalIndex = -1, int32 InSerialNumber = 0, FRemoteObjectId InRemoteId = FRemoteObjectId(), FRemoteServerId InRemoteOwnerServerId = FRemoteServerId());
 
 	/**
 	 * Checks to see if the object appears to be valid when being destroyed
@@ -146,6 +148,7 @@ private:
 
 #if UE_WITH_REMOTE_OBJECT_HANDLE
 	FRemoteObjectId GetPendingRegistrantRemoteId() const;
+	FRemoteServerId GetPendingRegistrantRemoteOwnerServerId() const;
 #endif
 
 public:
diff --git a/Engine/Source/Runtime/CoreUObject/Public/UObject/UObjectGlobals.h b/Engine/Source/Runtime/CoreUObject/Public/UObject/UObjectGlobals.h
--- a/Engine/Source/Runtime/CoreUObject/Public/UObject/UObjectGlobals.h
+++ b/Engine/Source/Runtime/CoreUObject/Public/UObject/UObjectGlobals.h
@@ -1188,7 +1188,7 @@ bool StaticAllocateObjectErrorTests( const UClass* Class, UObject* InOuter, FNam
  * @param GCGuard		Special synchronization object that prevents GC from running when reconstructing an object on top of existing one
  * @return	a pointer to a fully initialized object of the specified class.
  */
-COREUOBJECT_API UObject* StaticAllocateObject(const UClass* Class, UObject* InOuter, FName Name, EObjectFlags SetFlags, EInternalObjectFlags InternalSetFlags = EInternalObjectFlags::None, bool bCanReuseSubobjects = false, bool* bOutReusedSubobject = nullptr, UPackage* ExternalPackage = nullptr, int32 SerialNumber = 0, FRemoteObjectId RemoteId = FRemoteObjectId(), class FGCReconstructionGuard* GCGuard = nullptr);
+COREUOBJECT_API UObject* StaticAllocateObject(const UClass* Class, UObject* InOuter, FName Name, EObjectFlags SetFlags, EInternalObjectFlags InternalSetFlags = EInternalObjectFlags::None, bool bCanReuseSubobjects = false, bool* bOutReusedSubobject = nullptr, UPackage* ExternalPackage = nullptr, int32 SerialNumber = 0, FRemoteObjectId RemoteId = FRemoteObjectId(), FRemoteServerId RemoteOwnerServerId = FRemoteServerId(), class FGCReconstructionGuard* GCGuard = nullptr);
 
 /** FObjectInitializer options */
 enum class EObjectInitializerOptions
@@ -1722,6 +1722,9 @@ struct FStaticConstructObjectParameters
 	/** Existing remote id for the object */
 	FRemoteObjectId RemoteId;
 
+	/** ServerId of the current owner for the object */
+	FRemoteServerId RemoteOwnerServerId;
+
 	UE::RemoteObject::Serialization::FRemoteObjectConstructionOverrides* RemoteSubObjectOverrides = nullptr;
 #endif // UE_WITH_REMOTE_OBJECT_HANDLE
 
