[SharedStruct] added missing call to the destructor of the wrapped struct.

diff --git a/Engine/Source/Developer/StructUtilsTestSuite/Private/StructSharedStructTest.cpp b/Engine/Source/Developer/StructUtilsTestSuite/Private/StructSharedStructTest.cpp
--- a/Engine/Source/Developer/StructUtilsTestSuite/Private/StructSharedStructTest.cpp
+++ b/Engine/Source/Developer/StructUtilsTestSuite/Private/StructSharedStructTest.cpp
@@ -223,6 +223,31 @@ namespace FStructSharedStructTest
 	};
 
 	IMPLEMENT_AI_INSTANT_TEST(FTest_ConstSharedStructBasic, "System.StructUtils.SharedStruct.Const.Basic");
+
+	struct FTest_SharedStructDestruction : FAITestBase
+	{
+		virtual bool InstantTest() override
+		{
+			int32 NumConstructedStructBefore = FTestStructNonTrivialDestructor::NumConstructed;
+			{
+				FTestStructNonTrivialDestructor Struct;	
+				AITEST_EQUAL("Number of constructed FTestStructNonTrivialDestructor should have increased by 1 after local struct ctor", FTestStructNonTrivialDestructor::NumConstructed, NumConstructedStructBefore+1);
+			}
+			AITEST_EQUAL("Number of constructed FTestStructNonTrivialDestructor should be decreased by 1 after local struct dtor", FTestStructNonTrivialDestructor::NumConstructed, NumConstructedStructBefore);
+
+			NumConstructedStructBefore = FTestStructNonTrivialDestructor::NumConstructed;
+			{
+				FTestStructNonTrivialDestructor Struct;
+				FSharedStruct SharedStruct = FSharedStruct::Make(Struct);
+				AITEST_EQUAL("Number of constructed FTestStructNonTrivialDestructor should have increased by 2 after local struct and shared struct ctors", FTestStructNonTrivialDestructor::NumConstructed, NumConstructedStructBefore+2);
+			}
+			AITEST_EQUAL("Number of constructed FTestStructNonTrivialDestructor should be decreased by 2 after local struct and shared struct dtors", FTestStructNonTrivialDestructor::NumConstructed, NumConstructedStructBefore);
+
+			return true;
+		}
+	};
+
+	IMPLEMENT_AI_INSTANT_TEST(FTest_SharedStructDestruction, "System.StructUtils.SharedStruct.Destruction");
 }
 
 
diff --git a/Engine/Source/Developer/StructUtilsTestSuite/Public/StructUtilsTestTypes.h b/Engine/Source/Developer/StructUtilsTestSuite/Public/StructUtilsTestTypes.h
--- a/Engine/Source/Developer/StructUtilsTestSuite/Public/StructUtilsTestTypes.h
+++ b/Engine/Source/Developer/StructUtilsTestSuite/Public/StructUtilsTestTypes.h
@@ -157,6 +157,24 @@ struct FTestStructHashable1 : public FTestStructSimpleBase
 
 uint32 GetTypeHash(const FTestStructHashable1& other);
 
+USTRUCT()
+struct FTestStructNonTrivialDestructor
+{
+	GENERATED_BODY()
+
+	FTestStructNonTrivialDestructor()
+	{
+		NumConstructed++;
+	}
+
+	~FTestStructNonTrivialDestructor()
+	{
+		NumConstructed--;
+	}
+
+	static inline int32 NumConstructed = 0;
+};
+
 UCLASS()
 class UBagTestObject1 : public UObject
 {
diff --git a/Engine/Source/Runtime/CoreUObject/Private/StructUtils/SharedStruct.cpp b/Engine/Source/Runtime/CoreUObject/Private/StructUtils/SharedStruct.cpp
--- a/Engine/Source/Runtime/CoreUObject/Private/StructUtils/SharedStruct.cpp
+++ b/Engine/Source/Runtime/CoreUObject/Private/StructUtils/SharedStruct.cpp
@@ -1,6 +1,5 @@
 // Copyright Epic Games, Inc. All Rights Reserved.
 #include "StructUtils/SharedStruct.h"
-#include "StructUtils/StructView.h"
 
 #include UE_INLINE_GENERATED_CPP_BY_NAME(SharedStruct)
 
@@ -24,20 +23,6 @@ void FConstSharedStruct::AddStructReferencedObjects(class FReferenceCollector& C
 
 ///////////////////////////////////////////////////////////////// FSharedStruct /////////////////////////////////////////////////////////////////
 
-FSharedStruct::FSharedStruct(const FConstStructView& InOther)
-{
-	InitializeAs(InOther.GetScriptStruct(), InOther.GetMemory());
-}
-
-FSharedStruct& FSharedStruct::operator=(const FConstStructView& InOther)
-{
-	if (*this != InOther)
-	{
-		InitializeAs(InOther.GetScriptStruct(), InOther.GetMemory());
-	}
-	return *this;
-}
-
 bool FSharedStruct::Identical(const FSharedStruct* Other, uint32 PortFlags) const
 {
 	// Only empty or strictly identical is considered equal
@@ -46,7 +31,7 @@ bool FSharedStruct::Identical(const FSharedStruct* Other, uint32 PortFlags) cons
 		&& GetScriptStruct() == Other->GetScriptStruct();
 }
 
-void FSharedStruct::AddStructReferencedObjects(class FReferenceCollector& Collector)
+void FSharedStruct::AddStructReferencedObjects(class FReferenceCollector& Collector) const
 {
 	if (auto* Struct = GetScriptStructPtr(); Struct && *Struct)
 	{
diff --git a/Engine/Source/Runtime/CoreUObject/Public/StructUtils/SharedStruct.h b/Engine/Source/Runtime/CoreUObject/Public/StructUtils/SharedStruct.h
--- a/Engine/Source/Runtime/CoreUObject/Public/StructUtils/SharedStruct.h
+++ b/Engine/Source/Runtime/CoreUObject/Public/StructUtils/SharedStruct.h
@@ -23,12 +23,12 @@ template<typename BaseStructT> struct TConstStructView;
  * structure described by the ScriptStruct. This avoids two pointer referencing (cache misses). 
  * Look at Create() to understand more.
  * 
- * A 'const FStructSharedMemory' the memory is immutable. We restrict shallow copies of StructMemory where its not appropriate
- * in the owning types that compose this type. ie:
+ * A 'const FStructSharedMemory' the memory is immutable. We restrict shallow copies of StructMemory where it's not appropriate
+ * in the owning types that compose this type:
  * - FSharedStruct A; ConstSharedStruct B = A; is allowed
  * - ConstSharedStruct A; FSharedStruct B = A; is not allowed
  * 
- * This type is designed to be used in composition and should not be used outside of the types that compose it.
+ * This type is designed to be used in composition and should not be used outside the types that compose it.
  */
 struct COREUOBJECT_API FStructSharedMemory
 {
@@ -46,6 +46,11 @@ struct COREUOBJECT_API FStructSharedMemory
 	{
 		FORCEINLINE void operator()(FStructSharedMemory* StructSharedMemory) const
 		{
+			checkf(StructSharedMemory != nullptr
+				, TEXT("FStructSharedMemoryDeleter is expected to be used only with a valid FStructSharedMemory. "
+					"See FStructSharedMemory::Create/CreateArgs"));
+
+			StructSharedMemory->~FStructSharedMemory();
 			FMemory::Free(StructSharedMemory);
 		}
 	};
@@ -55,6 +60,7 @@ struct COREUOBJECT_API FStructSharedMemory
 		checkf(InScriptStruct, TEXT("FStructSharedMemory::Create - script struct cannot be null"));
 
 		FStructSharedMemory* const StructMemory = CreateImpl(InScriptStruct, InStructMemory);
+		checkf(StructMemory != nullptr, TEXT("CreateImpl is expected to always return a valid pointer to FStructSharedMemory."));
 
 		return MakeShareable(StructMemory, FStructSharedMemoryDeleter());
 	}
@@ -71,6 +77,7 @@ struct COREUOBJECT_API FStructSharedMemory
 		UE::StructUtils::CheckStructType<T>();
 
 		FStructSharedMemory* const StructMemory = CreateImpl(TBaseStructure<T>::Get());
+		checkf(StructMemory != nullptr, TEXT("CreateImpl is expected to always return a valid pointer to FStructSharedMemory."));
 
 		new (StructMemory->GetMutableMemory()) T(Forward<TArgs>(InArgs)...);
 		return MakeShareable(StructMemory, FStructSharedMemoryDeleter());
@@ -137,13 +144,12 @@ private:
  * (One pointer for the FInstancedStruct and one pointer for the struct memory it is wrapping).
  * Also note that because of its implementation, it is not possible for now to go from a struct reference or struct view back to a shared struct.
  *
- * This struct type is also convertible to a FStructView / FConstStructView and like FInstancedStruct it is the preferable way of passing it as a parameter.
- * If the calling code would like to keep a shared pointer to the struct, you may pass the FSharedStruct as a parameter but it is recommended to pass it as
+ * This struct type is also convertible to a FStructView / FConstStructView, and like FInstancedStruct, it is the preferable way of passing it as a parameter.
+ * If the calling code would like to keep a shared pointer to the struct, you may pass the FSharedStruct as a parameter, but it is recommended to pass it as
  * a "const FSharedStruct&" to limit the unnecessary recounting.
  * 
- * A 'const FSharedStruct' can not be made to point at another instance of a struct, whilst a vanila FSharedStruct can. In either case the shared struct memory /data is
- * mutable.
- *
+ * A 'const FSharedStruct' cannot be made to point at another instance of a struct, whilst a vanilla FSharedStruct can.
+ * In either case, the shared struct memory /data is mutable.
  */
 USTRUCT()
 struct COREUOBJECT_API FSharedStruct
@@ -154,9 +160,6 @@ struct COREUOBJECT_API FSharedStruct
 
 	FSharedStruct() = default;
 
-	UE_DEPRECATED(5.3, "Use Make() or InitializeAs() instead")
-	explicit FSharedStruct(const struct FConstStructView& InOther); // Note passing view by reference on purpose to allow code to compile
-
 	~FSharedStruct()
 	{
 	}
@@ -169,12 +172,9 @@ struct COREUOBJECT_API FSharedStruct
 	FSharedStruct& operator=(const FSharedStruct& InOther) = default;
 	FSharedStruct& operator=(FSharedStruct&& InOther) = default;
 
-	UE_DEPRECATED(5.3, "Use Make() instead")
-	FSharedStruct& operator=(const struct FConstStructView& InOther); // Note passing view by reference on purpose to allow code to compile
-
 	/** For StructOpsTypeTraits */
 	bool Identical(const FSharedStruct* Other, uint32 PortFlags) const;
-	void AddStructReferencedObjects(class FReferenceCollector& Collector);
+	void AddStructReferencedObjects(class FReferenceCollector& Collector) const;
 
 	/** Returns struct type. */
 	const UScriptStruct* GetScriptStruct() const
@@ -182,7 +182,7 @@ struct COREUOBJECT_API FSharedStruct
 		return StructMemoryPtr ? &(StructMemoryPtr.Get()->GetScriptStruct()) : nullptr;
 	}
 
-	TObjectPtr<const UScriptStruct>* const GetScriptStructPtr() const
+	TObjectPtr<const UScriptStruct>* GetScriptStructPtr() const
 	{
 		return StructMemoryPtr ? &StructMemoryPtr.Get()->GetScriptStructPtr() : nullptr;
 	}
@@ -199,7 +199,7 @@ struct COREUOBJECT_API FSharedStruct
 		StructMemoryPtr.Reset();
 	}
 
-	/** Initializes from templated struct type. This will create a new instance of the shared struct memory. */
+	/** Initializes from a templated struct type. This will create a new instance of the shared struct memory. */
 	template<typename T>
 	void InitializeAs()
 	{
@@ -222,7 +222,7 @@ struct COREUOBJECT_API FSharedStruct
 		InitializeAs(Struct.GetScriptStruct(), Struct.GetMemory());
 	}
 
-	/** Initializes from struct type and optional data. This will create a new instance of the shared struct memory. */
+	/** Initializes from a struct type and optional data. This will create a new instance of the shared struct memory. */
 	void InitializeAs(const UScriptStruct* InScriptStruct, const uint8* InStructMemory = nullptr)
 	{
 		if (InScriptStruct)
@@ -244,7 +244,7 @@ struct COREUOBJECT_API FSharedStruct
 		StructMemoryPtr = FStructSharedMemory::CreateArgs<T>(Forward<TArgs>(InArgs)...);
 	}
 
-	/** Creates a new FSharedStruct from templated struct type. This will create a new instance of the shared struct memory. */
+	/** Creates a new FSharedStruct from a templated struct type. This will create a new instance of the shared struct memory. */
 	template<typename T>
 	static FSharedStruct Make()
 	{
@@ -293,29 +293,6 @@ struct COREUOBJECT_API FSharedStruct
 		return UE::StructUtils::GetStructPtr<T>(GetScriptStruct(), GetMemory());
 	}
 
-	/** Returns a mutable pointer to struct memory. */
-	UE_DEPRECATED(5.3, "Use GetMemory() instead")
-	uint8* GetMutableMemory()
-	{
-		return GetMemory();
-	}
-
-	/** Returns mutable reference to the struct, this getter assumes that all data is valid. */
-	template<typename T>
-	UE_DEPRECATED(5.3, "Use Get() instead")
-	T& GetMutable()
-	{
-		return Get<T>();
-	}
-
-	/** Returns mutable pointer to the struct, or nullptr if cast is not valid. */
-	template<typename T>
-	UE_DEPRECATED(5.3, "Use GetPtr() instead")
-	T* GetMutablePtr()
-	{
-		return GetPtr<T>();
-	}
-
 	/** Returns True if the struct is valid.*/
 	bool IsValid() const
 	{
@@ -336,7 +313,7 @@ struct COREUOBJECT_API FSharedStruct
 	}
 
 	/**
-	 * Determines whether Other contains same values as `this`
+	 * Determines whether Other contains the same values as `this`
 	 * @return whether the values stored are equal
 	 */
 	template <typename OtherType>
@@ -407,7 +384,7 @@ struct TSharedStruct
 		return SharedStruct.GetScriptStruct();
 	}
 
-	TObjectPtr<const UScriptStruct>* const GetScriptStructPtr() const
+	TObjectPtr<const UScriptStruct>* GetScriptStructPtr() const
 	{
 		return SharedStruct.GetScriptStructPtr();
 	}
@@ -424,7 +401,7 @@ struct TSharedStruct
 		SharedStruct.Reset();
 	}
 
-	/** Initializes from templated struct type. This will create a new instance of the shared struct memory. */
+	/** Initializes from a templated struct type. This will create a new instance of the shared struct memory. */
 	template<typename T = BaseStructT, typename = std::enable_if_t<std::is_base_of_v<BaseStructT, std::decay_t<T>>>>
 	void Initialize()
 	{
@@ -510,7 +487,7 @@ private:
 	/**
 	 * Note:
 	 *   TSharedStruct is a wrapper for a FSharedStruct (rather than inheriting) so that it can provide a locked-down type-safe 
-	 *   API for use in C++, without being able to accidentally take a reference to the untyped API to workaround the restrictions.
+	 *   API for use in C++, without being able to accidentally take a reference to the untyped API to work around the restrictions.
 	 * 
 	 *   TSharedStruct MUST be the same size as FSharedStruct, as the reflection layer treats a TSharedStruct as a FSharedStruct.
 	 *   This means that any reflected APIs (like ExportText) that accept an FSharedStruct pointer can also accept a TSharedStruct pointer.
@@ -532,10 +509,10 @@ struct TStructOpsTypeTraits<FSharedStruct> : public TStructOpsTypeTraitsBase2<FS
 
 ///////////////////////////////////////////////////////////////// FConstSharedStruct /////////////////////////////////////////////////////////////////
 /**
- * FConstSharedStruct is the same as the FSharedStruct but restricts the API to return const struct type. 
+ * FConstSharedStruct is the same as the FSharedStruct but restricts the API to return a const struct type. 
  * 
- * A 'const FConstSharedStruct' can not be made to point at another instance of a struct, whilst a vanila FConstSharedStruct can. In either case the struct data is
- * immutable.
+ * A 'const FConstSharedStruct' cannot be made to point at another instance of a struct, whilst a vanilla FConstSharedStruct can.
+ * In either case, the struct data is immutable.
  * 
  * See FSharedStruct for more information.
  */
@@ -598,7 +575,7 @@ struct COREUOBJECT_API FConstSharedStruct
 		StructMemoryPtr.Reset();
 	}
 
-	/** Initializes from templated struct type. */
+	/** Initializes from a templated struct type. */
 	template<typename T>
 	void InitializeAs()
 	{
@@ -621,7 +598,7 @@ struct COREUOBJECT_API FConstSharedStruct
 		InitializeAs(Struct.GetScriptStruct(), Struct.GetMemory());
 	}
 
-	/** Initializes from struct type and optional data. This will create a new instance of the shared struct memory. */
+	/** Initializes from a struct type and optional data. This will create a new instance of the shared struct memory. */
 	void InitializeAs(const UScriptStruct* InScriptStruct, const uint8* InStructMemory = nullptr)
 	{
 		if (InScriptStruct)
@@ -643,7 +620,7 @@ struct COREUOBJECT_API FConstSharedStruct
 		StructMemoryPtr = FStructSharedMemory::CreateArgs<T>(Forward<TArgs>(InArgs)...);
 	}
 
-	/** Creates a new FSharedStruct from templated struct type. This will create a new instance of the shared struct memory. */
+	/** Creates a new FSharedStruct from a templated struct type. This will create a new instance of the shared struct memory. */
 	template<typename T>
 	static FConstSharedStruct Make()
 	{
@@ -661,7 +638,7 @@ struct COREUOBJECT_API FConstSharedStruct
 		return SharedStruct;
 	}
 
-	/** Creates a new FSharedStruct from struct type and optional data. This will create a new instance of the shared struct memory. */
+	/** Creates a new FSharedStruct from a struct type and optional data. This will create a new instance of the shared struct memory. */
 	static FConstSharedStruct Make(const UScriptStruct* InScriptStruct, const uint8* InStructMemory = nullptr)
 	{
 		FConstSharedStruct SharedStruct;
@@ -685,14 +662,6 @@ struct COREUOBJECT_API FConstSharedStruct
 		return UE::StructUtils::GetStructRef<T>(GetScriptStruct(), GetMemory());
 	}
 
-	/** Returns const reference to the struct, this getter assumes that all data is valid. */
-	template<typename T>
-	UE_DEPRECATED(5.3, "Use the version of this function which takes const in the template type")
-	constexpr typename TEnableIf<!TIsConst<T>::Value, const T&>::Type Get() const
-	{
-		return UE::StructUtils::GetStructRef<T>(GetScriptStruct(), GetMemory());
-	}
-
 	/** Returns const pointer to the struct, or nullptr if cast is not valid. */
 	template<typename T>
 	constexpr typename TEnableIf<TIsConst<T>::Value, T*>::Type GetPtr() const
@@ -700,14 +669,6 @@ struct COREUOBJECT_API FConstSharedStruct
 		return UE::StructUtils::GetStructPtr<T>(GetScriptStruct(), GetMemory());
 	}
 
-	/** Returns const pointer to the struct, or nullptr if cast is not valid. */
-	template<typename T>
-	UE_DEPRECATED(5.3, "Use the version of this function which takes const in the template type")
-	constexpr typename TEnableIf<!TIsConst<T>::Value, const T*>::Type GetPtr() const
-	{
-		return UE::StructUtils::GetStructPtr<T>(GetScriptStruct(), GetMemory());
-	}
-
 	/** Returns True if the struct is valid.*/
 	bool IsValid() const
 	{
@@ -728,7 +689,7 @@ struct COREUOBJECT_API FConstSharedStruct
 	}
 
 	/** 
-	 * Determines whether Other contains same values as `this` 
+	 * Determines whether Other contains the same values as `this`
 	 * @return whether the values stored are equal
 	 */
 	template <typename OtherType>
@@ -812,7 +773,7 @@ struct TConstSharedStruct
 		ConstSharedStruct.Reset();
 	}
 
-	/** Initializes from templated struct type. */
+	/** Initializes from a templated struct type. */
 	void Initialize()
 	{
 		ConstSharedStruct.InitializeAs<BaseStructT>();
@@ -839,7 +800,7 @@ struct TConstSharedStruct
 		ConstSharedStruct.InitializeAs<T>(Forward<TArgs>(InArgs)...);
 	}
 
-	/** Creates a new TSharedStruct from templated struct type. This will create a new instance of the shared struct memory. */
+	/** Creates a new TSharedStruct from a templated struct type. This will create a new instance of the shared struct memory. */
 	template<typename T = BaseStructT, typename = std::enable_if_t<std::is_base_of_v<BaseStructT, std::decay_t<T>>>>
 	static TConstSharedStruct Make()
 	{
@@ -857,7 +818,7 @@ struct TConstSharedStruct
 		return SharedStruct;
 	}
 
-	/** Creates a new TSharedStruct from struct type and optional data. This will create a new instance of the shared struct memory. */
+	/** Creates a new TSharedStruct from a struct type and optional data. This will create a new instance of the shared struct memory. */
 	template<typename T = BaseStructT, typename = std::enable_if_t<std::is_base_of_v<BaseStructT, std::decay_t<T>>>>
 	static TConstSharedStruct Make(const uint8* InStructMemory = nullptr)
 	{
@@ -926,7 +887,7 @@ private:
 	/**
 	 * Note:
 	 *   TConstSharedStruct is a wrapper for a FConstSharedStruct (rather than inheriting) so that it can provide a locked-down type-safe 
-	 *   API for use in C++, without being able to accidentally take a reference to the untyped API to workaround the restrictions.
+	 *   API for use in C++, without being able to accidentally take a reference to the untyped API to work around the restrictions.
 	 * 
 	 *   TConstSharedStruct MUST be the same size as FConstSharedStruct, as the reflection layer treats a TConstSharedStruct as a FConstSharedStruct.
 	 *   This means that any reflected APIs (like ExportText) that accept an FConstSharedStruct pointer can also accept a TConstSharedStruct pointer.
