[Zenloader][EngineTest]

- Prevent recursively calling ForceRegenerateClass for the same type. This is to prevent stack overflow when a type has circular references during deferred ExportCreation. Deferred creation is already trying to handle circular references however if there is more than one cycle ForceRegenerateClass would cause a stack overflow. BlueprintSupport in particular takes a slightly different path when FinalizingBlueprints such that is won't regenerate the class if the class is cooked. I believe this breaks the recursion in the non-cooked cases which is why the issue is cook only.
- Add test case of loading a BlueprintActor with BlueprintActorComponents that contain circular references to the other components as well as back to the BlueprintActor. Test level CoreAssetLoading level will crash in zenloader when loading as loose cooked data without this fix. Due to how blueprint references no FunctionalTest case is added since in editor we are unable to unload the blueprint due to circular references preventing garbage collection. Instead the BlueprintActor is placed in the level CoreAssetLoading and will be tested implicitly upon running any of the other CoreAssetLoading tests (or loading the level directly)

#tests Adds new TestCase running implicitly as part of CoreAssetLoading tests
#rnx

#virtualized

diff --git a/Engine/Source/Runtime/CoreUObject/Private/Blueprint/BlueprintSupport.cpp b/Engine/Source/Runtime/CoreUObject/Private/Blueprint/BlueprintSupport.cpp
--- a/Engine/Source/Runtime/CoreUObject/Private/Blueprint/BlueprintSupport.cpp
+++ b/Engine/Source/Runtime/CoreUObject/Private/Blueprint/BlueprintSupport.cpp
@@ -1424,13 +1424,17 @@ bool FLinkerLoad::DeferExportCreation(const int32 Index, UObject* Outer)
 	// the midst of doing so somewhere up the callstack)
 	if (!bIsLoadingExportClass || (LoadFlags & LOAD_ResolvingDeferredExports) != 0)
 	{
-		DEFERRED_DEPENDENCY_CHECK(!IsExportBeingResolved(Index));
+		bool bAlreadyResolvingExport = IsExportBeingResolved(Index);
 		FScopedResolvingExportTracker ReentranceGuard(this, Index);
 
 		// we want to be very careful, since we haven't filled in the export yet,
 		// we could get stuck in a recursive loop here (force-finalizing the 
 		// class here ends us back 
-		ForceRegenerateClass(LoadClass);
+		if (!bAlreadyResolvingExport)
+		{
+			ForceRegenerateClass(LoadClass);
+		}
+
 		return false;
 	}
 	
