Added a new check() during config loading to ensure that if we are loading a config file for a class that the class's reflection information is available.

[RN] minor foundation/core

diff --git a/Engine/Source/Runtime/CoreUObject/Private/UObject/Class.cpp b/Engine/Source/Runtime/CoreUObject/Private/UObject/Class.cpp
--- a/Engine/Source/Runtime/CoreUObject/Private/UObject/Class.cpp
+++ b/Engine/Source/Runtime/CoreUObject/Private/UObject/Class.cpp
@@ -96,6 +96,14 @@ UE_IMPLEMENT_STRUCT("/Script/CoreUObject", TestUninitializedScriptStructMembersT
 // casting it to uint32 that need to be fixed up (mostly printfs but also some serialization code)
 static_assert(sizeof(__underlying_type(EClassFlags)) == sizeof(uint32), "expecting ClassFlags enum to fit in a uint32");
 
+enum class EStructStateFlags : uint16
+{
+	None = 0,
+	PropertyDataAvailable = 1 << 0
+};
+
+ENUM_CLASS_FLAGS(EStructStateFlags)
+
 //////////////////////////////////////////////////////////////////////////
 
 namespace { UE_CALL_ONCE(UE::GC::RegisterSlowImplementation, &UClass::AddReferencedObjects, UE::GC::EAROFlags::Unbalanced); }
@@ -648,7 +656,7 @@ UStruct::UStruct(EStaticConstructor, int32 InSize, int32 InMinAlignment, EObject
 	, Children(nullptr)
 	, ChildProperties(nullptr)
 	, PropertiesSize(InSize)
-	, MinAlignment(InMinAlignment)
+	, MinAlignment(IntCastChecked<int16>(InMinAlignment))
 	, PropertyLink(nullptr)
 	, RefLink(nullptr)
 	, DestructorLink(nullptr)
@@ -667,7 +675,7 @@ UStruct::UStruct(UStruct* InSuperStruct, SIZE_T ParamsSize, SIZE_T Alignment)
 	, Children(nullptr)
 	, ChildProperties(nullptr)
 	, PropertiesSize(ParamsSize ? IntCastChecked<int32>(ParamsSize) : (InSuperStruct ? InSuperStruct->GetPropertiesSize() : 0))
-	, MinAlignment(Alignment ? IntCastChecked<int32>(Alignment) : (FMath::Max(InSuperStruct ? InSuperStruct->GetMinAlignment() : 1, 1)))
+	, MinAlignment(IntCastChecked<int16>(Alignment ? Alignment : (FMath::Max(InSuperStruct ? InSuperStruct->GetMinAlignment() : 1, 1))))
 	, PropertyLink(nullptr)
 	, RefLink(nullptr)
 	, DestructorLink(nullptr)
@@ -689,7 +697,7 @@ UStruct::UStruct(const FObjectInitializer& ObjectInitializer, UStruct* InSuperSt
 	, Children(nullptr)
 	, ChildProperties(nullptr)
 	, PropertiesSize(ParamsSize ? IntCastChecked<int32>(ParamsSize) : (InSuperStruct ? InSuperStruct->GetPropertiesSize() : 0))
-	, MinAlignment(Alignment ? IntCastChecked<int32>(Alignment) : (FMath::Max(InSuperStruct ? InSuperStruct->GetMinAlignment() : 1, 1)))
+	, MinAlignment(IntCastChecked<int16>(Alignment ? Alignment : (FMath::Max(InSuperStruct ? InSuperStruct->GetMinAlignment() : 1, 1))))
 	, PropertyLink(nullptr)
 	, RefLink(nullptr)
 	, DestructorLink(nullptr)
@@ -907,7 +915,7 @@ void UStruct::Link(FArchive& Ar, bool bRelinkExistingProperties)
 				TotalFieldCount = InheritanceSuper->TotalFieldCount;
 			#endif
 				PropertiesSize = InheritanceSuper->GetPropertiesSize();
-				MinAlignment = InheritanceSuper->GetMinAlignment();
+				MinAlignment = IntCastChecked<int16>(InheritanceSuper->GetMinAlignment());
 			}
 
 			for (FField* Field = ChildProperties; Field; Field = Field->Next)
@@ -946,7 +954,7 @@ void UStruct::Link(FArchive& Ar, bool bRelinkExistingProperties)
 						break;
 					}
 
-					MinAlignment = FMath::Max(MinAlignment, Property->GetMinAlignment());
+					MinAlignment = IntCastChecked<int16>(FMath::Max(MinAlignment, Property->GetMinAlignment()));
 				}
 			}
 		}
@@ -974,7 +982,7 @@ void UStruct::Link(FArchive& Ar, bool bRelinkExistingProperties)
 
 			if (UScriptStruct::ICppStructOps* CppStructOps = ScriptStruct.GetCppStructOps())
 			{
-				MinAlignment = CppStructOps->GetAlignment();
+				MinAlignment = IntCastChecked<int16>(CppStructOps->GetAlignment());
 				PropertiesSize = CppStructOps->GetSize();
 				bHandledWithCppStructOps = true;
 			}
@@ -1153,6 +1161,8 @@ void UStruct::Link(FArchive& Ar, bool bRelinkExistingProperties)
 	}
 	PropertyWrappers.Empty();
 #endif
+
+	StructStateFlags.fetch_or(static_cast<uint16>(EStructStateFlags::PropertyDataAvailable), std::memory_order_release);
 }
 
 void UStruct::InitializeStruct(void* InDest, int32 ArrayDim/* = 1*/) const
@@ -1206,6 +1216,15 @@ void UStruct::DestroyStruct(void* Dest, int32 ArrayDim) const
 	}
 }
 
+#if DO_CHECK
+
+bool UStruct::DebugIsPropertyChainReady() const
+{
+	return HasAllStructStateFlags(EStructStateFlags::PropertyDataAvailable);
+}
+
+#endif // DO_CHECK
+
 //
 // Serialize all of the class's data that belongs in a particular
 // bin and resides in Data.
@@ -1553,6 +1572,11 @@ struct FSerializationControlExtensionContext
 	}
 };
 
+bool UStruct::HasAllStructStateFlags(EStructStateFlags TestStructStateFlags) const
+{
+	return EnumHasAllFlags(static_cast<EStructStateFlags>(StructStateFlags.load(std::memory_order_acquire)), TestStructStateFlags);
+}
+
 void UStruct::SerializeVersionedTaggedProperties(FStructuredArchive::FSlot Slot, uint8* Data, UStruct* DefaultsStruct, uint8* Defaults, const UObject* BreakRecursionIfFullyLoad) const
 {
 	using namespace UE;
diff --git a/Engine/Source/Runtime/CoreUObject/Private/UObject/Obj.cpp b/Engine/Source/Runtime/CoreUObject/Private/UObject/Obj.cpp
--- a/Engine/Source/Runtime/CoreUObject/Private/UObject/Obj.cpp
+++ b/Engine/Source/Runtime/CoreUObject/Private/UObject/Obj.cpp
@@ -2771,6 +2771,12 @@ void UObject::LoadConfig(UClass* ConfigClass/*=nullptr*/, const TCHAR* InFilenam
 		return;
 	}
 
+
+#if DO_CHECK
+	// When a CDO is created before the Class is linked, it will have an empty set of properties. Among other issues, this can cause loading of config data to fail silently. 
+	checkf(ConfigClass->DebugIsPropertyChainReady(), TEXT("Attempting to load config data for %s before the Class has been constructed/registered/linked (likely during module loading or early startup)."), *GetName());
+#endif
+
 #if !WITH_EDITOR
 	if (OutAccessedValues)
 	{
diff --git a/Engine/Source/Runtime/CoreUObject/Public/UObject/Class.h b/Engine/Source/Runtime/CoreUObject/Public/UObject/Class.h
--- a/Engine/Source/Runtime/CoreUObject/Public/UObject/Class.h
+++ b/Engine/Source/Runtime/CoreUObject/Public/UObject/Class.h
@@ -388,6 +388,7 @@ private:
 /*-----------------------------------------------------------------------------
 	UStruct.
 -----------------------------------------------------------------------------*/
+enum class EStructStateFlags : uint16;
 
 /**
  * Base class for all UObject types that contain fields.
@@ -415,8 +416,14 @@ public:
 	/** Total size of all UProperties, the allocated structure may be larger due to alignment */
 	int32 PropertiesSize;
 	/** Alignment of structure in memory, structure will be at least this large */
-	int32 MinAlignment;
-	
+	int16 MinAlignment;
+
+private:
+	/** State flags -- Access via HasAllStructStateFlags() **/
+	std::atomic<uint16> StructStateFlags;
+
+public:
+
 	/** Script bytecode associated with this object */
 	TArray<uint8> Script;
 
@@ -588,6 +595,11 @@ public:
 	 */
 	COREUOBJECT_API virtual void DestroyStruct(void* Dest, int32 ArrayDim = 1) const;
 
+#if DO_CHECK
+	/** Returns true if this object has been initialized to the point that its property chain data is available. */
+	bool DebugIsPropertyChainReady() const;
+#endif
+
 	/** Look up a property by an alternate name if it was not found in the first search, this is overridden for user structs */
 	virtual FProperty* CustomFindProperty(const FName InName) const { return nullptr; };
 
@@ -828,6 +840,9 @@ private:
 	friend class FBlueprintCompileReinstancer;
 #endif
 
+	/** This method is private to discourage extension of the StructStateFlags. Returns true if StructStateFlags contains all the TestStructStateFlags specified. */
+	bool HasAllStructStateFlags(EStructStateFlags TestStructStateFlags) const;
+
 	COREUOBJECT_API void SerializeVersionedTaggedProperties(FStructuredArchive::FSlot Slot, uint8* Data, UStruct* DefaultsStruct, uint8* Defaults, const UObject* BreakRecursionIfFullyLoad) const;
 };
 
