RemoteObject:

- Move the owner server ID from FUObjectItem into the remote object stub
- ensure stubs are created for all received objects

diff --git a/Engine/Source/Runtime/CoreUObject/Private/UObject/Class.cpp b/Engine/Source/Runtime/CoreUObject/Private/UObject/Class.cpp
--- a/Engine/Source/Runtime/CoreUObject/Private/UObject/Class.cpp
+++ b/Engine/Source/Runtime/CoreUObject/Private/UObject/Class.cpp
@@ -5051,13 +5051,13 @@ FName UClass::GetDefaultObjectName() const
 //
 void UClass::DeferredRegister(UClass *UClassStaticClass, const TCHAR* PackageName, const TCHAR* Name
 #if UE_WITH_REMOTE_OBJECT_HANDLE
-	, FRemoteObjectId RemoteId, FRemoteServerId RemoteOwnerServerId
+	, FRemoteObjectId RemoteId
 #endif
 )
 {
 	Super::DeferredRegister(UClassStaticClass, PackageName, Name
 #if UE_WITH_REMOTE_OBJECT_HANDLE
-		, RemoteId, RemoteOwnerServerId
+		, RemoteId
 #endif
 	);
 
diff --git a/Engine/Source/Runtime/CoreUObject/Private/UObject/ObjectHandlePrivate.h b/Engine/Source/Runtime/CoreUObject/Private/UObject/ObjectHandlePrivate.h
--- a/Engine/Source/Runtime/CoreUObject/Private/UObject/ObjectHandlePrivate.h
+++ b/Engine/Source/Runtime/CoreUObject/Private/UObject/ObjectHandlePrivate.h
@@ -43,26 +43,6 @@ namespace UE::CoreUObject::Private
 #endif // UE_WITH_REMOTE_OBJECT_HANDLE
 			return FRemoteObjectId();
 		}
-
-		FORCEINLINE static FRemoteServerId GetRemoteOwnerServerId(const UObjectBase* Object)
-		{
-			FRemoteServerId Result;
-
-#if UE_WITH_REMOTE_OBJECT_HANDLE
-			if (Object)
-			{
-				if (UNLIKELY(Object->InternalIndex < 0))
-				{
-					Result = Object->GetPendingRegistrantRemoteOwnerServerId();
-				}
-				else
-				{
-					Result = GUObjectArray.GetRemoteOwnerServerId(Object->InternalIndex);
-				}
-			}
-#endif // UE_WITH_REMOTE_OBJECT_HANDLE
-			return Result;
-		}
 	};
 }
 
diff --git a/Engine/Source/Runtime/CoreUObject/Private/UObject/RemoteObject.cpp b/Engine/Source/Runtime/CoreUObject/Private/UObject/RemoteObject.cpp
--- a/Engine/Source/Runtime/CoreUObject/Private/UObject/RemoteObject.cpp
+++ b/Engine/Source/Runtime/CoreUObject/Private/UObject/RemoteObject.cpp
@@ -241,30 +241,54 @@ bool IsRemote(const UObject* Object)
 
 bool IsOwned(const UObject* Object)
 {
+	using namespace UE::CoreUObject::Private;
+	using namespace UE::RemoteObject::Private;
+
 	bool Result = false;
 #if UE_WITH_REMOTE_OBJECT_HANDLE
-	FUObjectItem* ObjectItem = GUObjectArray.ObjectToObjectItem(Object);
-	checkf(ObjectItem, TEXT("Attempting to get owner id for an object that does not exist in the global UObject array (it's possible GUObjectArray is not initialized yet, ObjectIndex=%d)"), GUObjectArray.ObjectToIndex(Object));
-	Result = (ObjectItem->RemoteOwnerServerId == GetGlobalServerId());
+	UE::RemoteObject::Handle::FRemoteObjectStub* RemoteStub = FindRemoteObjectStub(FObjectHandleUtils::GetRemoteId(Object));
+	if (RemoteStub)
+	{
+		Result = (RemoteStub->OwningServerId == GetGlobalServerId());
+	}
+	else
+	{
+		// if the object wasn't received or ever migrated, we own it locally
+		Result = true;
+	}
 #endif
 	return Result;
 }
 
 FRemoteServerId GetOwnerServerId(const UObject* Object)
 {
+	using namespace UE::CoreUObject::Private;
+	using namespace UE::RemoteObject::Private;
+
 	FRemoteServerId Result;
 #if UE_WITH_REMOTE_OBJECT_HANDLE
-	FUObjectItem* ObjectItem = GUObjectArray.ObjectToObjectItem(Object);
-	Result = ObjectItem->RemoteOwnerServerId;
+	UE::RemoteObject::Handle::FRemoteObjectStub* RemoteStub = FindRemoteObjectStub(FObjectHandleUtils::GetRemoteId(Object));
+	if (RemoteStub)
+	{
+		Result = RemoteStub->OwningServerId;
+	}
+	else
+	{
+		// if the object wasn't received or ever migrated, we own it locally
+		Result = GetGlobalServerId();
+	}
 #endif
 	return Result;
 }
 
 void ChangeOwnerServerId(const UObject* Object, FRemoteServerId NewOwnerServerId)
 {
+	using namespace UE::CoreUObject::Private;
+	using namespace UE::RemoteObject::Private;
+	
 #if UE_WITH_REMOTE_OBJECT_HANDLE
-	FUObjectItem* ObjectItem = GUObjectArray.ObjectToObjectItem(Object);
-	ObjectItem->RemoteOwnerServerId = NewOwnerServerId;
+	UE::RemoteObject::Handle::FRemoteObjectStub* RemoteStub = FindRemoteObjectStub(FObjectHandleUtils::GetRemoteId(Object));
+	RemoteStub->OwningServerId = NewOwnerServerId;
 #endif
 }
 
diff --git a/Engine/Source/Runtime/CoreUObject/Private/UObject/RemoteObjectPrivate.h b/Engine/Source/Runtime/CoreUObject/Private/UObject/RemoteObjectPrivate.h
--- a/Engine/Source/Runtime/CoreUObject/Private/UObject/RemoteObjectPrivate.h
+++ b/Engine/Source/Runtime/CoreUObject/Private/UObject/RemoteObjectPrivate.h
@@ -28,9 +28,9 @@ namespace UE::RemoteObject::Private
 	void MarkAsLocal(UObject* Object);
 
 	/**
-	* Registers a stub for a remote object that is known to be owned by a specific server
+	* Registers a stub for a remote object that is known to be resident on a specific server
 	*/
-	void RegisterRemoteObjectId(FRemoteObjectId ObjectId, FRemoteServerId OwnerServerId);
+	void RegisterRemoteObjectId(FRemoteObjectId ObjectId, FRemoteServerId ResidentServerId);
 
 	/**
 	* Registers object for sharing, marking it as owned by the current server
diff --git a/Engine/Source/Runtime/CoreUObject/Private/UObject/RemoteObjectSerialization.cpp b/Engine/Source/Runtime/CoreUObject/Private/UObject/RemoteObjectSerialization.cpp
--- a/Engine/Source/Runtime/CoreUObject/Private/UObject/RemoteObjectSerialization.cpp
+++ b/Engine/Source/Runtime/CoreUObject/Private/UObject/RemoteObjectSerialization.cpp
@@ -612,7 +612,7 @@ FRemoteObjectConstructionOverrides::FRemoteObjectConstructionOverrides(const TAr
 	}
 }
 
-UObject* ConstructRemoteObject(const FRemoteServerId OwnerServerId, const FRemoteObjectHeader& Header, FRemoteObjectConstructionOverrides* ConstructionOverrides)
+UObject* ConstructRemoteObject(const FRemoteObjectHeader& Header, FRemoteObjectConstructionOverrides* ConstructionOverrides)
 {
 	using namespace UE::RemoteObject::Handle;
 	using namespace UE::RemoteObject::Private;
@@ -642,8 +642,6 @@ UObject* ConstructRemoteObject(const FRemoteServerId OwnerServerId, const FRemot
 			Name = FName();
 			Object = nullptr;
 		}
-
-		ChangeOwnerServerId(Object, OwnerServerId);
 	}
 	else
 	{
@@ -652,8 +650,6 @@ UObject* ConstructRemoteObject(const FRemoteServerId OwnerServerId, const FRemot
 		{
 			// The object already exists on this server but has been renamed
 			Name = Object->GetFName();
-
-			ChangeOwnerServerId(Object, OwnerServerId);
 		}
 	}
 
@@ -667,7 +663,6 @@ UObject* ConstructRemoteObject(const FRemoteServerId OwnerServerId, const FRemot
 		Params.Template = Header.Archetype.Get();
 #if UE_WITH_REMOTE_OBJECT_HANDLE
 		Params.RemoteId = Header.RemoteId;
-		Params.RemoteOwnerServerId = OwnerServerId;
 		Params.RemoteSubObjectOverrides = ConstructionOverrides;
 #endif
 
@@ -681,7 +676,6 @@ UObject* ConstructRemoteObject(const FRemoteServerId OwnerServerId, const FRemot
 	}
 
 	checkf(FRemoteObjectId(Object) == Header.RemoteId, TEXT("Created an object with a different ID:%s than requested:%s"), *FRemoteObjectId(Object).ToString(), *Header.RemoteId.ToString());
-	checkf(GetOwnerServerId(Object) == OwnerServerId, TEXT("Created an object with a different Owner:%s than requested:%s"), *GetOwnerServerId(Object).ToString(), *OwnerServerId.ToString());
 
 	MarkAsLocal(Object);
 
@@ -849,7 +843,15 @@ int32 DeserializeObjectData(FRemoteServerId OwnerServerId, FRemoteObjectData& Ob
 		OutObjects.Reserve(ObjectHeaders.Num());
 		for (int32 ObjectIndex = 0; ObjectIndex < ObjectHeaders.Num(); ++ObjectIndex)
 		{
-			OutObjects.Add(ConstructRemoteObject(OwnerServerId, ObjectHeaders[ObjectIndex], &ConstructionOverrides));
+			OutObjects.Add(ConstructRemoteObject(ObjectHeaders[ObjectIndex], &ConstructionOverrides));
+		}
+
+		// ensure remote object stubs are created and ownership is assigned
+		for (int32 ObjectIndex = 0; ObjectIndex < OutObjects.Num(); ++ObjectIndex)
+		{
+			UObject* Object = OutObjects[ObjectIndex];
+			UE::RemoteObject::Private::RegisterRemoteObjectId(FRemoteObjectId(Object), LocalServerId);
+			ChangeOwnerServerId(Object, OwnerServerId);
 		}
 
 		// Deserialize all objects
diff --git a/Engine/Source/Runtime/CoreUObject/Private/UObject/RemoteObjectTransfer.cpp b/Engine/Source/Runtime/CoreUObject/Private/UObject/RemoteObjectTransfer.cpp
--- a/Engine/Source/Runtime/CoreUObject/Private/UObject/RemoteObjectTransfer.cpp
+++ b/Engine/Source/Runtime/CoreUObject/Private/UObject/RemoteObjectTransfer.cpp
@@ -484,9 +484,9 @@ void RegisterSharedObject(UObject* Object)
 	UE::RemoteObject::Private::RegisterSharedObject(Object);
 }
 
-void RegisterRemoteObjectId(FRemoteObjectId Id, FRemoteServerId OwnerServerId)
+void RegisterRemoteObjectId(FRemoteObjectId Id, FRemoteServerId ResidentServerId)
 {
-	UE::RemoteObject::Private::RegisterRemoteObjectId(Id, OwnerServerId);
+	UE::RemoteObject::Private::RegisterRemoteObjectId(Id, ResidentServerId);
 }
 
 void MigrateObjectToRemoteServer(FRemoteTransactionId RequestId, FRemoteObjectId ObjectId, FRemoteServerId DestinationServerId)
diff --git a/Engine/Source/Runtime/CoreUObject/Private/UObject/UObjectArray.cpp b/Engine/Source/Runtime/CoreUObject/Private/UObject/UObjectArray.cpp
--- a/Engine/Source/Runtime/CoreUObject/Private/UObject/UObjectArray.cpp
+++ b/Engine/Source/Runtime/CoreUObject/Private/UObject/UObjectArray.cpp
@@ -226,7 +226,7 @@ void FUObjectArray::DisableDisregardForGC()
 	}
 }
 
-void FUObjectArray::AllocateUObjectIndex(UObjectBase* Object, EInternalObjectFlags InitialFlags, int32 AlreadyAllocatedIndex, int32 SerialNumber, FRemoteObjectId RemoteId, FRemoteServerId RemoteOwnerServerId)
+void FUObjectArray::AllocateUObjectIndex(UObjectBase* Object, EInternalObjectFlags InitialFlags, int32 AlreadyAllocatedIndex, int32 SerialNumber, FRemoteObjectId RemoteId)
 {
 	LLM_SCOPE(ELLMTag::UObject);
 	// Clear asset scopes
@@ -242,10 +242,6 @@ void FUObjectArray::AllocateUObjectIndex(UObjectBase* Object, EInternalObjectFla
 	{
 		RemoteId = FRemoteObjectId::Generate();
 	}
-	if (!RemoteOwnerServerId.IsValid())
-	{
-		RemoteOwnerServerId = UE::RemoteObject::GetGlobalServerId();
-	}
 #endif
 
 	LockInternalArray();
@@ -302,7 +298,6 @@ void FUObjectArray::AllocateUObjectIndex(UObjectBase* Object, EInternalObjectFla
 	ObjectItem->SerialNumber = SerialNumber;
 #if UE_WITH_REMOTE_OBJECT_HANDLE
 	ObjectItem->RemoteId = RemoteId;
-	ObjectItem->RemoteOwnerServerId = RemoteOwnerServerId;
 #endif // UE_WITH_REMOTE_OBJECT_HANDLE
 	Object->InternalIndex = Index;
 
@@ -334,7 +329,6 @@ void FUObjectArray::ResetSerialNumber(UObjectBase* Object)
 	ObjectItem->SerialNumber = 0;
 #if UE_WITH_REMOTE_OBJECT_HANDLE
 	ObjectItem->RemoteId = FRemoteObjectId();
-	ObjectItem->RemoteOwnerServerId = FRemoteServerId();
 #endif
 }
 
@@ -403,7 +397,6 @@ void FUObjectArray::FreeUObjectIndex(UObjectBase* Object)
 	ObjectItem->SerialNumber = 0;
 #if UE_WITH_REMOTE_OBJECT_HANDLE
 	ObjectItem->RemoteId = FRemoteObjectId();
-	ObjectItem->RemoteOwnerServerId = FRemoteServerId();
 #endif
 	Object->InternalIndex = INDEX_NONE;
 
diff --git a/Engine/Source/Runtime/CoreUObject/Private/UObject/UObjectBase.cpp b/Engine/Source/Runtime/CoreUObject/Private/UObject/UObjectBase.cpp
--- a/Engine/Source/Runtime/CoreUObject/Private/UObject/UObjectBase.cpp
+++ b/Engine/Source/Runtime/CoreUObject/Private/UObject/UObjectBase.cpp
@@ -65,13 +65,11 @@ struct FPendingRegistrantInfo
 	class UClass*  (*StaticClassFn)();
 #if UE_WITH_REMOTE_OBJECT_HANDLE
 	const FRemoteObjectId RemoteId;
-	const FRemoteServerId RemoteOwnerServerId;
 #endif // UE_WITH_REMOTE_OBJECT_HANDLE
 
 	FPendingRegistrantInfo(class UClass* (*InStaticClassFn)(), const TCHAR* InName, const TCHAR* InPackageName
 #if UE_WITH_REMOTE_OBJECT_HANDLE
 		, const FRemoteObjectId InRemoteId
-		, const FRemoteServerId InRemoteOwnerServerId
 #endif
 		)
 		: Name(InName)
@@ -79,7 +77,6 @@ struct FPendingRegistrantInfo
 		, StaticClassFn(InStaticClassFn)
 #if UE_WITH_REMOTE_OBJECT_HANDLE
 		, RemoteId(InRemoteId)
-		, RemoteOwnerServerId(InRemoteOwnerServerId)
 #endif
 	{}
 	static TMap<UObjectBase*, FPendingRegistrantInfo>& GetMap()
@@ -149,8 +146,7 @@ UObjectBase::UObjectBase(UClass* InClass,
 	FName InName,
 	int32 InInternalIndex,
 	int32 InSerialNumber,
-	FRemoteObjectId InRemoteId,
-	FRemoteServerId InRemoteOwnerServerId)
+	FRemoteObjectId InRemoteId)
 :	ObjectFlags			(InFlags)
 ,	InternalIndex		(INDEX_NONE)
 ,	ClassPrivate		(InClass)
@@ -158,7 +154,7 @@ UObjectBase::UObjectBase(UClass* InClass,
 {
 	check(ClassPrivate);
 	// Add to global table.
-	AddObject(InName, InInternalFlags, InInternalIndex, InSerialNumber, InRemoteId, InRemoteOwnerServerId);
+	AddObject(InName, InInternalFlags, InInternalIndex, InSerialNumber, InRemoteId);
 	
 #if CSV_PROFILER_STATS && CSV_TRACK_UOBJECT_COUNT
 	UObjectStats::IncrementUObjectCount();
@@ -196,7 +192,6 @@ UObjectBase::~UObjectBase()
 void UObjectBase::DeferredRegister(UClass *UClassStaticClass, const TCHAR* PackageName, const TCHAR* InName
 #if UE_WITH_REMOTE_OBJECT_HANDLE
 	, FRemoteObjectId RemoteId
-	, FRemoteServerId RemoteOwnerServerId
 #endif
 )
 {
@@ -214,7 +209,7 @@ void UObjectBase::DeferredRegister(UClass *UClassStaticClass, const TCHAR* Packa
 	// Add to the global object table.
 	AddObject(FName(InName), EInternalObjectFlags::None
 #if UE_WITH_REMOTE_OBJECT_HANDLE
-		, /*IndernalIndex = */ -1, /*SerialNumber = */ 0, RemoteId, RemoteOwnerServerId
+		, /*IndernalIndex = */ -1, /*SerialNumber = */ 0, RemoteId
 #endif
 	);
 
@@ -240,7 +235,7 @@ void UObjectBase::DeferredRegister(UClass *UClassStaticClass, const TCHAR* Packa
  * @param Name name to assign to this uobject
  */
 UE_AUTORTFM_ALWAYS_OPEN
-void UObjectBase::AddObject(FName InName, EInternalObjectFlags InSetInternalFlags, int32 InInternalIndex, int32 InSerialNumber, FRemoteObjectId InRemoteId, FRemoteServerId InRemoteOwnerServerId)
+void UObjectBase::AddObject(FName InName, EInternalObjectFlags InSetInternalFlags, int32 InInternalIndex, int32 InSerialNumber, FRemoteObjectId InRemoteId)
 {
 	NamePrivate = InName;
 	EInternalObjectFlags InternalFlagsToSet = InSetInternalFlags;
@@ -258,7 +253,7 @@ void UObjectBase::AddObject(FName InName, EInternalObjectFlags InSetInternalFlag
 		InternalFlagsToSet |= EInternalObjectFlags::Native;
 		ObjectFlags &= ~RF_MarkAsNative;
 	}
-	GUObjectArray.AllocateUObjectIndex(this, InternalFlagsToSet, InInternalIndex, InSerialNumber, InRemoteId, InRemoteOwnerServerId);
+	GUObjectArray.AllocateUObjectIndex(this, InternalFlagsToSet, InInternalIndex, InSerialNumber, InRemoteId);
 	check(InName != NAME_None && InternalIndex >= 0);
 	HashObject(this);
 	check(IsValidLowLevel());
@@ -542,7 +537,6 @@ void UObjectBase::Register(class UClass* (*StaticClassFn)(), const TCHAR* Packag
 	PendingRegistrants.Add(this, FPendingRegistrantInfo(StaticClassFn, InName, PackageName
 #if UE_WITH_REMOTE_OBJECT_HANDLE
 		, FRemoteObjectId::Generate()
-		, FRemoteServerId()
 #endif
 	));
 
@@ -576,13 +570,6 @@ FRemoteObjectId UObjectBase::GetPendingRegistrantRemoteId() const
 	FPendingRegistrantInfo& Info = PendingRegistrants.FindChecked(this);
 	return Info.RemoteId;
 }
-
-FRemoteServerId UObjectBase::GetPendingRegistrantRemoteOwnerServerId() const
-{
-	TMap<UObjectBase*, FPendingRegistrantInfo>& PendingRegistrants = FPendingRegistrantInfo::GetMap();
-	FPendingRegistrantInfo& Info = PendingRegistrants.FindChecked(this);
-	return Info.RemoteOwnerServerId;
-}
 #endif
 
 /**
@@ -649,12 +636,11 @@ void UObjectForceRegistration(UObjectBase* Object, bool bCheckForModuleRelease)
 		UClass* StaticClass = Info->StaticClassFn();
 #if UE_WITH_REMOTE_OBJECT_HANDLE
 		FRemoteObjectId RemoteId = Info->RemoteId;
-		FRemoteServerId RemoteOwnerServerId = Info->RemoteOwnerServerId;
 #endif
 		PendingRegistrants.Remove(Object);  // delete this first so that it doesn't try to do it twice
 		Object->DeferredRegister(StaticClass, PackageName, Name
 #if UE_WITH_REMOTE_OBJECT_HANDLE
-			, RemoteId, RemoteOwnerServerId
+			, RemoteId
 #endif
 		);
 	}
diff --git a/Engine/Source/Runtime/CoreUObject/Private/UObject/UObjectGlobals.cpp b/Engine/Source/Runtime/CoreUObject/Private/UObject/UObjectGlobals.cpp
--- a/Engine/Source/Runtime/CoreUObject/Private/UObject/UObjectGlobals.cpp
+++ b/Engine/Source/Runtime/CoreUObject/Private/UObject/UObjectGlobals.cpp
@@ -3554,7 +3554,6 @@ UObject* StaticAllocateObject
 	UPackage* ExternalPackage,
 	int32 SerialNumber,
 	FRemoteObjectId RemoteId,
-	FRemoteServerId RemoteOwnerServerId,
 	FGCReconstructionGuard* GCGuard)
 {
 	LLM_SCOPE(ELLMTag::UObject);
@@ -3831,7 +3830,7 @@ UObject* StaticAllocateObject
 #if UE_WITH_REMOTE_OBJECT_HANDLE
 				OldRemoteId.IsValid() ? OldRemoteId :
 #endif
-				RemoteId, RemoteOwnerServerId);
+				RemoteId);
 		};
 		if (!bWasConstructedOnOldObject)
 		{
@@ -4888,12 +4887,6 @@ UObject* StaticConstructObject_Internal(const FStaticConstructObjectParameters&
 	UObject* InTemplate = Params.Template;
 	int32 SerialNumber = Params.SerialNumber;
 	FRemoteObjectId RemoteId;
-#if UE_WITH_REMOTE_OBJECT_HANDLE
-	FRemoteServerId RemoteOwnerServerId = Params.RemoteOwnerServerId;
-#else
-	FRemoteServerId RemoteOwnerServerId;
-#endif
-
 	UE::RemoteObject::Serialization::FRemoteObjectConstructionOverrides* RemoteSubObjectOverrides = nullptr;
 
 	LLM_SCOPE(ELLMTag::UObject);
@@ -4932,7 +4925,7 @@ UObject* StaticConstructObject_Internal(const FStaticConstructObjectParameters&
 
 	FGCReconstructionGuard GCGuard;
 	bool bRecycledSubobject = false;	
-	Result = StaticAllocateObject(InClass, InOuter, InName, InFlags, Params.InternalSetFlags, bCanRecycleSubobjects, &bRecycledSubobject, Params.ExternalPackage, SerialNumber, RemoteId, RemoteOwnerServerId, &GCGuard);
+	Result = StaticAllocateObject(InClass, InOuter, InName, InFlags, Params.InternalSetFlags, bCanRecycleSubobjects, &bRecycledSubobject, Params.ExternalPackage, SerialNumber, RemoteId, &GCGuard);
 	check(Result != nullptr);
 	// Don't call the constructor on recycled subobjects, they haven't been destroyed.
 	if (!bRecycledSubobject)
diff --git a/Engine/Source/Runtime/CoreUObject/Public/UObject/Class.h b/Engine/Source/Runtime/CoreUObject/Public/UObject/Class.h
--- a/Engine/Source/Runtime/CoreUObject/Public/UObject/Class.h
+++ b/Engine/Source/Runtime/CoreUObject/Public/UObject/Class.h
@@ -3263,7 +3263,6 @@ public:
 	COREUOBJECT_API virtual void DeferredRegister(UClass* UClassStaticClass, const TCHAR* PackageName, const TCHAR* InName
 #if UE_WITH_REMOTE_OBJECT_HANDLE
 		, FRemoteObjectId RemoteId
-		, FRemoteServerId RemoteOwnerServerId
 #endif // UE_WITH_REMOTE_OBJECT_HANDLE
 	) override;
 	COREUOBJECT_API virtual bool Rename(const TCHAR* NewName = nullptr, UObject* NewOuter = nullptr, ERenameFlags Flags = REN_None) override;
diff --git a/Engine/Source/Runtime/CoreUObject/Public/UObject/RemoteObject.h b/Engine/Source/Runtime/CoreUObject/Public/UObject/RemoteObject.h
--- a/Engine/Source/Runtime/CoreUObject/Public/UObject/RemoteObject.h
+++ b/Engine/Source/Runtime/CoreUObject/Public/UObject/RemoteObject.h
@@ -41,6 +41,10 @@ namespace UE::RemoteObject::Handle
 		/** Server id that where the object currently resides */
 		FRemoteServerId ResidentServerId;
 
+		/** Server id of the server that has ownership of the object
+		    (note: only valid if the object is Local) */
+		FRemoteServerId OwningServerId;
+
 		FRemoteObjectStub() = default;
 		COREUOBJECT_API explicit FRemoteObjectStub(UObject* Object);
 	};
@@ -100,4 +104,4 @@ namespace UE::RemoteObject::Handle
 	COREUOBJECT_API bool IsOwned(const UObject* Object);
 	COREUOBJECT_API FRemoteServerId GetOwnerServerId(const UObject* Object);
 	COREUOBJECT_API void ChangeOwnerServerId(const UObject* Object, FRemoteServerId NewOwnerServerId);
-}
\ No newline at end of file
+}
diff --git a/Engine/Source/Runtime/CoreUObject/Public/UObject/RemoteObjectTransfer.h b/Engine/Source/Runtime/CoreUObject/Public/UObject/RemoteObjectTransfer.h
--- a/Engine/Source/Runtime/CoreUObject/Public/UObject/RemoteObjectTransfer.h
+++ b/Engine/Source/Runtime/CoreUObject/Public/UObject/RemoteObjectTransfer.h
@@ -118,7 +118,7 @@ namespace UE::RemoteObject::Transfer
 	* @param ObjectId Remote object id
 	* @param OwnerServerId Id of a remote server that currently owns the object
 	*/
-	COREUOBJECT_API void RegisterRemoteObjectId(FRemoteObjectId Id, FRemoteServerId OwnerServerId);
+	COREUOBJECT_API void RegisterRemoteObjectId(FRemoteObjectId Id, FRemoteServerId ResidentServerId);
 
 	/**
 	* Registers object for sharing, marking it as owned by the current server
diff --git a/Engine/Source/Runtime/CoreUObject/Public/UObject/UObjectArray.h b/Engine/Source/Runtime/CoreUObject/Public/UObject/UObjectArray.h
--- a/Engine/Source/Runtime/CoreUObject/Public/UObject/UObjectArray.h
+++ b/Engine/Source/Runtime/CoreUObject/Public/UObject/UObjectArray.h
@@ -63,7 +63,6 @@ public:
 #if UE_WITH_REMOTE_OBJECT_HANDLE
 	// Globally unqiue id of this object
 	FRemoteObjectId RemoteId;
-	FRemoteServerId RemoteOwnerServerId;
 #endif
 
 #if STATS || ENABLE_STATNAMEDEVENTS_UOBJECT
@@ -136,11 +135,6 @@ public:
 	{
 		return RemoteId;
 	}
-
-	FORCEINLINE FRemoteServerId GetRemoteOwnerServerId() const
-	{
-		return RemoteOwnerServerId;
-	}
 #endif // UE_WITH_REMOTE_OBJECT_HANDLE
 
 	FORCEINLINE void SetFlags(EInternalObjectFlags FlagsToSet)
@@ -809,7 +803,7 @@ public:
 class FUObjectArray
 {
 	friend class UObject;
-	friend COREUOBJECT_API UObject* StaticAllocateObject(const UClass*, UObject*, FName, EObjectFlags, EInternalObjectFlags, bool, bool*, UPackage*, int32, FRemoteObjectId, FRemoteServerId, class FGCReconstructionGuard*);
+	friend COREUOBJECT_API UObject* StaticAllocateObject(const UClass*, UObject*, FName, EObjectFlags, EInternalObjectFlags, bool, bool*, UPackage*, int32, FRemoteObjectId, class FGCReconstructionGuard*);
 
 private:
 	/**
@@ -930,7 +924,7 @@ public:
 	 * @param	AlreadyAllocatedIndex already allocated internal index to use, negative value means allocate a new index
 	 * @param	SerialNumber serial number to use
 	 */
-	COREUOBJECT_API void AllocateUObjectIndex(class UObjectBase* Object, EInternalObjectFlags InitialFlags, int32 AlreadyAllocatedIndex = -1, int32 SerialNumber = 0, FRemoteObjectId RemoteId = FRemoteObjectId(), FRemoteServerId RemoteOwnerServerId = FRemoteServerId());
+	COREUOBJECT_API void AllocateUObjectIndex(class UObjectBase* Object, EInternalObjectFlags InitialFlags, int32 AlreadyAllocatedIndex = -1, int32 SerialNumber = 0, FRemoteObjectId RemoteId = FRemoteObjectId());
 
 	/**
 	 * Returns a UObject index top to the global uobject array
@@ -1188,13 +1182,6 @@ public:
 		checkSlow(ObjectItem);
 		return ObjectItem->GetRemoteId();
 	}
-
-	FORCEINLINE FRemoteServerId GetRemoteOwnerServerId(int32 Index)
-	{
-		FUObjectItem* ObjectItem = IndexToObject(Index);
-		checkSlow(ObjectItem);
-		return ObjectItem->GetRemoteOwnerServerId();
-	}
 #endif
 
 	/** Locks the internal object array mutex */
diff --git a/Engine/Source/Runtime/CoreUObject/Public/UObject/UObjectBase.h b/Engine/Source/Runtime/CoreUObject/Public/UObject/UObjectBase.h
--- a/Engine/Source/Runtime/CoreUObject/Public/UObject/UObjectBase.h
+++ b/Engine/Source/Runtime/CoreUObject/Public/UObject/UObjectBase.h
@@ -92,8 +92,7 @@ public:
 			FName InName,
 			int32 InInternalIndex = -1,
 			int32 InSerialNumber = 0,
-			FRemoteObjectId InRemoteId = FRemoteObjectId(),
-			FRemoteServerId InRemoteOwnerServerId = FRemoteServerId());
+			FRemoteObjectId InRemoteId = FRemoteObjectId());
 
 	/**
 	 * Final destructor, removes the object from the object array, and indirectly, from any annotations
@@ -125,7 +124,6 @@ protected:
 	COREUOBJECT_API virtual void DeferredRegister(UClass* UClassStaticClass, const TCHAR* PackageName, const TCHAR* Name
 #if UE_WITH_REMOTE_OBJECT_HANDLE
 		, FRemoteObjectId InRemoteId
-		, FRemoteServerId InRemoteOwnerServerId
 #endif
 	);
 
@@ -138,7 +136,7 @@ private:
 	 * @param InInternalIndex already allocated internal index to use, negative value means allocate a new index
 	 * @param InSerialNumber already allocated serial number to re-use
 	 */
-	COREUOBJECT_API void AddObject(FName Name, EInternalObjectFlags InSetInternalFlags, int32 InInternalIndex = -1, int32 InSerialNumber = 0, FRemoteObjectId InRemoteId = FRemoteObjectId(), FRemoteServerId InRemoteOwnerServerId = FRemoteServerId());
+	COREUOBJECT_API void AddObject(FName Name, EInternalObjectFlags InSetInternalFlags, int32 InInternalIndex = -1, int32 InSerialNumber = 0, FRemoteObjectId InRemoteId = FRemoteObjectId());
 
 	/**
 	 * Checks to see if the object appears to be valid when being destroyed
@@ -148,7 +146,6 @@ private:
 
 #if UE_WITH_REMOTE_OBJECT_HANDLE
 	FRemoteObjectId GetPendingRegistrantRemoteId() const;
-	FRemoteServerId GetPendingRegistrantRemoteOwnerServerId() const;
 #endif
 
 public:
diff --git a/Engine/Source/Runtime/CoreUObject/Public/UObject/UObjectGlobals.h b/Engine/Source/Runtime/CoreUObject/Public/UObject/UObjectGlobals.h
--- a/Engine/Source/Runtime/CoreUObject/Public/UObject/UObjectGlobals.h
+++ b/Engine/Source/Runtime/CoreUObject/Public/UObject/UObjectGlobals.h
@@ -1188,7 +1188,7 @@ bool StaticAllocateObjectErrorTests( const UClass* Class, UObject* InOuter, FNam
  * @param GCGuard		Special synchronization object that prevents GC from running when reconstructing an object on top of existing one
  * @return	a pointer to a fully initialized object of the specified class.
  */
-COREUOBJECT_API UObject* StaticAllocateObject(const UClass* Class, UObject* InOuter, FName Name, EObjectFlags SetFlags, EInternalObjectFlags InternalSetFlags = EInternalObjectFlags::None, bool bCanReuseSubobjects = false, bool* bOutReusedSubobject = nullptr, UPackage* ExternalPackage = nullptr, int32 SerialNumber = 0, FRemoteObjectId RemoteId = FRemoteObjectId(), FRemoteServerId RemoteOwnerServerId = FRemoteServerId(), class FGCReconstructionGuard* GCGuard = nullptr);
+COREUOBJECT_API UObject* StaticAllocateObject(const UClass* Class, UObject* InOuter, FName Name, EObjectFlags SetFlags, EInternalObjectFlags InternalSetFlags = EInternalObjectFlags::None, bool bCanReuseSubobjects = false, bool* bOutReusedSubobject = nullptr, UPackage* ExternalPackage = nullptr, int32 SerialNumber = 0, FRemoteObjectId RemoteId = FRemoteObjectId(), class FGCReconstructionGuard* GCGuard = nullptr);
 
 /** FObjectInitializer options */
 enum class EObjectInitializerOptions
@@ -1722,9 +1722,6 @@ struct FStaticConstructObjectParameters
 	/** Existing remote id for the object */
 	FRemoteObjectId RemoteId;
 
-	/** ServerId of the current owner for the object */
-	FRemoteServerId RemoteOwnerServerId;
-
 	UE::RemoteObject::Serialization::FRemoteObjectConstructionOverrides* RemoteSubObjectOverrides = nullptr;
 #endif // UE_WITH_REMOTE_OBJECT_HANDLE
 
