RemoteObject:
- move the 'Aborting transaction' print to the top so in case we hit the ActiveRequest or some other check/assert/crash, we will have that in the logs
- extend the MigrateObjectToRemoteServer log to include the ownership information

diff --git a/Engine/Source/Runtime/CoreUObject/Private/UObject/RemoteObjectTransfer.cpp b/Engine/Source/Runtime/CoreUObject/Private/UObject/RemoteObjectTransfer.cpp
--- a/Engine/Source/Runtime/CoreUObject/Private/UObject/RemoteObjectTransfer.cpp
+++ b/Engine/Source/Runtime/CoreUObject/Private/UObject/RemoteObjectTransfer.cpp
@@ -296,6 +296,10 @@ void MigrateObjectFromRemoteServer(FRemoteObjectId ObjectId, FRemoteServerId Cur
 	// ensure the request is added in the Open so after we abort it is preserved
 	bool bAdded = false;
 
+	UE_LOG(LogRemoteObject, Verbose, TEXT("Aborting transaction, need obj id %s on server %s"),
+		*ObjectId.ToString(),
+		*CurrentOwnerServerId.ToString());
+
 	check(GTransferQueue.ActiveRequest);
 
 	UE_AUTORTFM_OPEN
@@ -308,10 +312,6 @@ void MigrateObjectFromRemoteServer(FRemoteObjectId ObjectId, FRemoteServerId Cur
 			GTransferQueue.ActiveRequest->RequiredObjects.Add(ObjectId);
 			GTransferQueue.ActiveRequest->RequiredObjectsCanonicalRoot.Add(ObjectId);
 		}
-
-		UE_LOG(LogRemoteObject, Verbose, TEXT("Aborting transaction, need obj id %s on server %s"),
-			*ObjectId.ToString(),
-			*CurrentOwnerServerId.ToString());
 		
 		if (UE_LOG_ACTIVE(LogRemoteObject, VeryVerbose))
 		{
@@ -360,7 +360,18 @@ void MigrateObjectToRemoteServerWithExplicitPriority(FRemoteTransactionId Reques
 	const FRemoteObjectRequest* EqualPriorityRequest = nullptr;
 	const FRemoteObjectRequest* LowerPriorityRequest = nullptr;
 
-	UE_LOG(LogRemoteObject, Verbose, TEXT("     : MigrateObjectToRemoteServer(%s) sending obj %s [%s] (root %s) (isremote %d)"), *RequestId.ToString(), *ObjectId.ToString(), *Object->GetName(), *CanonicalRootObjectId.ToString(), Handle::IsRemote(Object));
+	const bool bOwnershipChange = (RequestId == PushMigrationRequestId);
+	FRemoteServerId OwnerServerId = bOwnershipChange ? DestinationServerId : Handle::GetOwnerServerId(RootObject);
+
+	UE_LOG(LogRemoteObject, Verbose, 
+		TEXT("     : MigrateObjectToRemoteServer(%s) sending obj %s [%s] (root %s, isremote %d, changing owner %d, owner %s)"),
+		*RequestId.ToString(),
+		*ObjectId.ToString(),
+		*Object->GetName(),
+		*CanonicalRootObjectId.ToString(),
+		Handle::IsRemote(Object),
+		bOwnershipChange,
+		*OwnerServerId.ToString());
 
 	for (int32 RequestIndex = 0; RequestIndex < GTransferQueue.GetRequestCount(); RequestIndex++)
 	{
@@ -431,9 +442,6 @@ void MigrateObjectToRemoteServerWithExplicitPriority(FRemoteTransactionId Reques
 	}
 	else
 	{
-		const bool bOwnershipChange = (RequestId == PushMigrationRequestId);
-		FRemoteServerId OwnerServerId = bOwnershipChange ? DestinationServerId : Handle::GetOwnerServerId(RootObject);
-
 		UE_LOG(LogRemoteObject, VeryVerbose, TEXT("     : MigrateObjectToRemoteServer(%s) serializing obj %s (isremote %d owner changing %d owner %s)"), *RequestId.ToString(), *ObjectId.ToString(), Handle::IsRemote(Object), bOwnershipChange, *OwnerServerId.ToString());
 
 		// send the object
