Adding PPF_MigratingRemoteObjects used when migrating remote objects to avoid unnecessary remote object migrations when serializing properties

diff --git a/Engine/Source/Runtime/Core/Public/UObject/PropertyPortFlags.h b/Engine/Source/Runtime/Core/Public/UObject/PropertyPortFlags.h
--- a/Engine/Source/Runtime/Core/Public/UObject/PropertyPortFlags.h
+++ b/Engine/Source/Runtime/Core/Public/UObject/PropertyPortFlags.h
@@ -92,7 +92,8 @@ enum EPropertyPortFlags
 	/** indicates that we're exporting properties for display in the property window. - used to hide EditHide items in collapsed structs */
 	PPF_PropertyWindow				= 0x00080000,
 
-	//								= 0x00100000,
+	/** indicates that we're serializing remote objects for migration */
+	PPF_MigratingRemoteObjects		= 0x00100000,
 
 	/** Force fully qualified object names (for debug dumping) */
 	PPF_DebugDump					= 0x00200000,
diff --git a/Engine/Source/Runtime/Core/Public/UObject/ScriptDelegates.h b/Engine/Source/Runtime/Core/Public/UObject/ScriptDelegates.h
--- a/Engine/Source/Runtime/Core/Public/UObject/ScriptDelegates.h
+++ b/Engine/Source/Runtime/Core/Public/UObject/ScriptDelegates.h
@@ -405,6 +405,11 @@ public:
 		return Object;
 	}
 
+	const WeakPtrType& GetUObjectRef() const
+	{
+		return Object;
+	}
+
 	/**
 	 * Gets the name of the function to call on the bound object
 	 *
diff --git a/Engine/Source/Runtime/CoreUObject/Private/UObject/PropertyMulticastDelegate.cpp b/Engine/Source/Runtime/CoreUObject/Private/UObject/PropertyMulticastDelegate.cpp
--- a/Engine/Source/Runtime/CoreUObject/Private/UObject/PropertyMulticastDelegate.cpp
+++ b/Engine/Source/Runtime/CoreUObject/Private/UObject/PropertyMulticastDelegate.cpp
@@ -147,6 +147,16 @@ bool FMulticastDelegateProperty::Identical( const void* A, const void* B, uint32
 		const FMulticastScriptDelegate::UnicastDelegateType& BindingA = DelegateA.InvocationList[CurInvocationIndex];
 		const FMulticastScriptDelegate::UnicastDelegateType& BindingB = DelegateB.InvocationList[CurInvocationIndex];
 
+#if UE_WITH_REMOTE_OBJECT_HANDLE
+		if ((PortFlags & PPF_MigratingRemoteObjects) != 0)
+		{
+			if (BindingA.GetUObjectRef().GetRemoteId() != BindingB.GetUObjectRef().GetRemoteId())
+			{
+				return false;
+			}
+		}
+		else
+#endif
 		if (BindingA.GetUObject() != BindingB.GetUObject())
 		{
 			return false;
diff --git a/Engine/Source/Runtime/CoreUObject/Private/UObject/RemoteObjectSerialization.cpp b/Engine/Source/Runtime/CoreUObject/Private/UObject/RemoteObjectSerialization.cpp
--- a/Engine/Source/Runtime/CoreUObject/Private/UObject/RemoteObjectSerialization.cpp
+++ b/Engine/Source/Runtime/CoreUObject/Private/UObject/RemoteObjectSerialization.cpp
@@ -224,6 +224,7 @@ public:
 		SetIsPersistent(false);
 		SetUseUnversionedPropertySerialization(true);
 		SetMigratingRemoteObjects(true);
+		SetPortFlags(PPF_MigratingRemoteObjects);
 	}
 
 	virtual ~FArchiveRemoteObjectWriter()
@@ -499,6 +500,7 @@ public:
 		SetIsPersistent(false);
 		SetUseUnversionedPropertySerialization(true);
 		SetMigratingRemoteObjects(true);
+		SetPortFlags(PPF_MigratingRemoteObjects);
 	}
 
 	virtual FArchive& operator<<(UObject*& Obj) override
@@ -582,6 +584,7 @@ public:
 		SetIsPersistent(false);
 		InitObjectMap(DeserializedObjects);
 		SetMigratingRemoteObjects(true);
+		SetPortFlags(PPF_MigratingRemoteObjects);
 	}
 
 	virtual FArchive& operator<<(FObjectPtr& Value) override
