RemoteObject: allow querying IsOwned from just a FRemoteObjectId. For nonresident objects we are able to query whether we are the owner, but for objects we don't own, we don't necessarily have an up to date owner id (but we do know that we definitely don't own it)

diff --git a/Engine/Source/Runtime/CoreUObject/Private/UObject/RemoteObject.cpp b/Engine/Source/Runtime/CoreUObject/Private/UObject/RemoteObject.cpp
--- a/Engine/Source/Runtime/CoreUObject/Private/UObject/RemoteObject.cpp
+++ b/Engine/Source/Runtime/CoreUObject/Private/UObject/RemoteObject.cpp
@@ -57,6 +57,10 @@ public:
 			Stub = new FRemoteObjectStub();
 			Stub->Id = Id;
 			Stub->ResidentServerId = ResidentServerId.IsValid() ? ResidentServerId : Id.GetServerId();
+
+			// if we are creating the stub, then this object wasn't received by
+			// a remote server, therefore we are the owner
+			Stub->OwningServerId = GetGlobalServerId();
 		}		
 		return Stub;
 	}
@@ -70,6 +74,10 @@ public:
 		if (!Stub)
 		{
 			Stub = new FRemoteObjectStub(Object);
+
+			// if we are creating the stub, then this object wasn't received by
+			// a remote server, therefore we are the owner
+			Stub->OwningServerId = GetGlobalServerId();
 		}
 		Stub->ResidentServerId = DestinationServerId;
 		FUObjectItem* ObjectItem = GUObjectArray.ObjectToObjectItem(Object);
@@ -240,13 +248,19 @@ bool IsRemote(const UObject* Object)
 }
 
 bool IsOwned(const UObject* Object)
+{
+	using namespace UE::CoreUObject::Private;
+	return IsOwned(FObjectHandleUtils::GetRemoteId(Object));
+}
+
+bool IsOwned(FRemoteObjectId ObjectId)
 {
 	using namespace UE::CoreUObject::Private;
 	using namespace UE::RemoteObject::Private;
 
 	bool Result = false;
 #if UE_WITH_REMOTE_OBJECT_HANDLE
-	UE::RemoteObject::Handle::FRemoteObjectStub* RemoteStub = FindRemoteObjectStub(FObjectHandleUtils::GetRemoteId(Object));
+	UE::RemoteObject::Handle::FRemoteObjectStub* RemoteStub = FindRemoteObjectStub(ObjectId);
 	if (RemoteStub)
 	{
 		Result = (RemoteStub->OwningServerId == GetGlobalServerId());
diff --git a/Engine/Source/Runtime/CoreUObject/Public/UObject/RemoteObject.h b/Engine/Source/Runtime/CoreUObject/Public/UObject/RemoteObject.h
--- a/Engine/Source/Runtime/CoreUObject/Public/UObject/RemoteObject.h
+++ b/Engine/Source/Runtime/CoreUObject/Public/UObject/RemoteObject.h
@@ -102,6 +102,23 @@ namespace UE::RemoteObject::Handle
 	* Checks if a locally resident object is owned by this server
 	*/
 	COREUOBJECT_API bool IsOwned(const UObject* Object);
+
+	/**
+	* Checks if an object id is owned by this server
+	* We are only able to check if we own the object. If we don't
+	* own the object then we don't have a reliable way of knowing
+	* who the owner is, which is why the below GetOwnerServerId
+	* function requires the object be locally resident
+	*/
+	COREUOBJECT_API bool IsOwned(FRemoteObjectId ObjectId);
+
+	/**
+	* Get the owner server id for a locally resident object
+	*/
 	COREUOBJECT_API FRemoteServerId GetOwnerServerId(const UObject* Object);
+
+	/**
+	* Sets the owner server id for a locally resident object
+	*/
 	COREUOBJECT_API void ChangeOwnerServerId(const UObject* Object, FRemoteServerId NewOwnerServerId);
 }
