Return reference to initialized instanced struct

diff --git a/Engine/Source/Runtime/CoreUObject/Public/StructUtils/InstancedStruct.h b/Engine/Source/Runtime/CoreUObject/Public/StructUtils/InstancedStruct.h
--- a/Engine/Source/Runtime/CoreUObject/Public/StructUtils/InstancedStruct.h
+++ b/Engine/Source/Runtime/CoreUObject/Public/StructUtils/InstancedStruct.h
@@ -84,7 +84,7 @@ public:
 
 	/** Initializes from struct type and emplace construct. */
 	template<typename T, typename... TArgs>
-	void InitializeAs(TArgs&&... InArgs)
+	T& InitializeAs(TArgs&&... InArgs)
 	{
 		UE::StructUtils::CheckStructType<T>();
 
@@ -111,6 +111,8 @@ public:
 
 		check(Memory);
 		new (Memory) T(Forward<TArgs>(InArgs)...);
+
+		return *((T*)Memory);
 	}
 
 	/** Creates a new FInstancedStruct from templated struct type. */
