Build fix. Reducing warning to display for now.

diff --git a/Engine/Source/Runtime/CoreUObject/Private/UObject/Obj.cpp b/Engine/Source/Runtime/CoreUObject/Private/UObject/Obj.cpp
--- a/Engine/Source/Runtime/CoreUObject/Private/UObject/Obj.cpp
+++ b/Engine/Source/Runtime/CoreUObject/Private/UObject/Obj.cpp
@@ -88,7 +88,7 @@ static int32 GPrematureConfigLoadErrorLevel = 0;
 static FAutoConsoleVariableRef CVarPrematureConfigLoadErrorLevel(
 	TEXT("obj.PrematureConfigLoadErrorLevel"),
 	GPrematureConfigLoadErrorLevel,
-	TEXT("Determines the level of warning issued if LoadConfig is called on an object that is not yet sufficiently initialized. 0=Warning, 1=Ensure, 2=Check"));
+	TEXT("Determines the level of warning issued if LoadConfig is called on an object that is not yet sufficiently initialized. 0=Display, 1=Warning, 2=Ensure, 3=Check"));
 }
 
 /*-----------------------------------------------------------------------------
@@ -2784,12 +2784,16 @@ void UObject::LoadConfig(UClass* ConfigClass/*=nullptr*/, const TCHAR* InFilenam
 	// When a CDO is created before the Class is linked, it will have an empty set of properties. Among other issues, this can cause loading of config data to fail silently. 
 	if (!ConfigClass->DebugIsPropertyChainReady())
 	{
-		FString ErrorMessage = FString::Printf(TEXT("Attempting to load config data for %s before the Class has been constructed/registered/linked (likely during module loading or early startup)."), *GetName());
+		FString ErrorMessage = FString::Printf(TEXT("Attempting to load config data for %s before the Class has been constructed/registered/linked (likely during module loading or early startup). This will result in the load silently failing and should be fixed."), *GetName());
 		if (GPrematureConfigLoadErrorLevel == 0)
 		{
-			UE_LOG(LogObj, Warning, TEXT("%s"), *ErrorMessage);
+			UE_LOG(LogObj, Display, TEXT("%s"), *ErrorMessage);
 		}
 		else if (GPrematureConfigLoadErrorLevel == 1)
+		{
+			UE_LOG(LogObj, Warning, TEXT("%s"), *ErrorMessage);
+		}
+		else if (GPrematureConfigLoadErrorLevel == 2)
 		{
 			// Repeating the condition here just so it ends up in the output message
 			ensureAlwaysMsgf(ConfigClass->DebugIsPropertyChainReady(), TEXT("%s"), *ErrorMessage);
