Fixed braces and added [[nodiscard]] to FSoftObjectPath where appropriate.

diff --git a/Engine/Source/Runtime/CoreUObject/Public/UObject/SoftObjectPath.h b/Engine/Source/Runtime/CoreUObject/Public/UObject/SoftObjectPath.h
--- a/Engine/Source/Runtime/CoreUObject/Public/UObject/SoftObjectPath.h
+++ b/Engine/Source/Runtime/CoreUObject/Public/UObject/SoftObjectPath.h
@@ -44,44 +44,65 @@ DECLARE_DELEGATE_TwoParams(FLoadSoftObjectPathAsyncDelegate, const FSoftObjectPa
  */
 struct FSoftObjectPath
 {
-	FSoftObjectPath() = default;
-	FSoftObjectPath(const FSoftObjectPath& Other) = default;
-	FSoftObjectPath(FSoftObjectPath&& Other) = default;
+	[[nodiscard]] FSoftObjectPath() = default;
+	[[nodiscard]] FSoftObjectPath(const FSoftObjectPath& Other) = default;
+	[[nodiscard]] FSoftObjectPath(FSoftObjectPath&& Other) = default;
 	~FSoftObjectPath() = default;
 	FSoftObjectPath& operator=(const FSoftObjectPath& Path) = default;
 	FSoftObjectPath& operator=(FSoftObjectPath&& Path) = default;
 
 	/** Construct from a path string. Non-explicit for backwards compatibility. */
-	FSoftObjectPath(const FString& Path)						{ SetPath(FStringView(Path)); }
-	FSoftObjectPath(FTopLevelAssetPath InAssetPath, FString InSubPathString)
+	[[nodiscard]] FSoftObjectPath(const FString& Path)
+	{
+		SetPath(FStringView(Path));
+	}
+	[[nodiscard]] explicit FSoftObjectPath(FTopLevelAssetPath InAssetPath, FString InSubPathString)
 	{
 		SetPath(InAssetPath, MoveTemp(InSubPathString));
 	}
-	FSoftObjectPath(FName InPackageName, FName InAssetName, FString InSubPathString)
+	[[nodiscard]] explicit FSoftObjectPath(FName InPackageName, FName InAssetName, FString InSubPathString)
 	{
 		SetPath(FTopLevelAssetPath(InPackageName, InAssetName), MoveTemp(InSubPathString));
 	}
+
 	/** Explicitly extend a top-level object path with an empty subobject path. */
-	explicit FSoftObjectPath(FTopLevelAssetPath InAssetPath)	{ SetPath(InAssetPath, FString()); }
-	explicit FSoftObjectPath(FWideStringView Path)				{ SetPath(Path); }
-	explicit FSoftObjectPath(FAnsiStringView Path)				{ SetPath(Path); }
-	explicit FSoftObjectPath(const WIDECHAR* Path)				{ SetPath(FWideStringView(Path)); }
-	explicit FSoftObjectPath(const ANSICHAR* Path)				{ SetPath(FAnsiStringView(Path)); }
-	explicit FSoftObjectPath(TYPE_OF_NULLPTR)					{}
+	[[nodiscard]] explicit FSoftObjectPath(FTopLevelAssetPath InAssetPath)
+	{
+		SetPath(InAssetPath, FString());
+	}
+	[[nodiscard]] explicit FSoftObjectPath(FWideStringView Path)
+	{
+		SetPath(Path);
+	}
+	[[nodiscard]] explicit FSoftObjectPath(FAnsiStringView Path)
+	{
+		SetPath(Path);
+	}
+	[[nodiscard]] explicit FSoftObjectPath(const WIDECHAR* Path)
+	{
+		SetPath(FWideStringView(Path));
+	}
+	[[nodiscard]] explicit FSoftObjectPath(const ANSICHAR* Path)
+	{
+		SetPath(FAnsiStringView(Path));
+	}
+	[[nodiscard]] explicit FSoftObjectPath(TYPE_OF_NULLPTR)
+	{
+	}
 
 	template <typename T>
-	FSoftObjectPath(const TObjectPtr<T>& InObject)
+	[[nodiscard]] FSoftObjectPath(const TObjectPtr<T>& InObject)
 	{
 		SetPath(InObject.GetPathName());
 	}
 
-	FSoftObjectPath(const FObjectPtr& InObject)
+	[[nodiscard]] FSoftObjectPath(const FObjectPtr& InObject)
 	{
 		SetPath(InObject.GetPathName());
 	}
 
 	/** Construct from an existing object in memory */
-	FSoftObjectPath(const UObject* InObject)
+	[[nodiscard]] FSoftObjectPath(const UObject* InObject)
 	{
 		if (InObject)
 		{
@@ -90,31 +111,59 @@ struct FSoftObjectPath
 	}
 
 	/** Static methods for more meaningful construction sites. */
-	COREUOBJECT_API static FSoftObjectPath ConstructFromPackageAssetSubpath(FName InPackageName, FName InAssetName, const FString& InSubPathString);
-	COREUOBJECT_API static FSoftObjectPath ConstructFromPackageAssetSubpath(FName InPackageName, FName InAssetName, FString&& InSubPathString);
-	COREUOBJECT_API static FSoftObjectPath ConstructFromPackageAsset(FName InPackageName, FName InAssetName);
-	COREUOBJECT_API static FSoftObjectPath ConstructFromAssetPath(FTopLevelAssetPath InAssetPath);
-	COREUOBJECT_API static FSoftObjectPath ConstructFromStringPath(FString&& InPath);
-	COREUOBJECT_API static FSoftObjectPath ConstructFromStringPath(FStringView InPath);
-	COREUOBJECT_API static FSoftObjectPath ConstructFromStringPath(FUtf8StringView InPath);
-	COREUOBJECT_API static FSoftObjectPath ConstructFromObject(const UObject* InObject);
-	COREUOBJECT_API static FSoftObjectPath ConstructFromObject(const FObjectPtr& InObject);
+	[[nodiscard]] COREUOBJECT_API static FSoftObjectPath ConstructFromPackageAssetSubpath(FName InPackageName, FName InAssetName, const FString& InSubPathString);
+	[[nodiscard]] COREUOBJECT_API static FSoftObjectPath ConstructFromPackageAssetSubpath(FName InPackageName, FName InAssetName, FString&& InSubPathString);
+	[[nodiscard]] COREUOBJECT_API static FSoftObjectPath ConstructFromPackageAsset(FName InPackageName, FName InAssetName);
+	[[nodiscard]] COREUOBJECT_API static FSoftObjectPath ConstructFromAssetPath(FTopLevelAssetPath InAssetPath);
+	[[nodiscard]] COREUOBJECT_API static FSoftObjectPath ConstructFromStringPath(FString&& InPath);
+	[[nodiscard]] COREUOBJECT_API static FSoftObjectPath ConstructFromStringPath(FStringView InPath);
+	[[nodiscard]] COREUOBJECT_API static FSoftObjectPath ConstructFromStringPath(FUtf8StringView InPath);
+	[[nodiscard]] COREUOBJECT_API static FSoftObjectPath ConstructFromObject(const UObject* InObject);
+	[[nodiscard]] COREUOBJECT_API static FSoftObjectPath ConstructFromObject(const FObjectPtr& InObject);
 	template <typename T>
-	static FSoftObjectPath ConstructFromObject(const TObjectPtr<T>& InObject)
+	[[nodiscard]] static FSoftObjectPath ConstructFromObject(const TObjectPtr<T>& InObject)
 	{
 		return ConstructFromObject(InObject.Get());
 	}
 
-	FSoftObjectPath& operator=(const FTopLevelAssetPath Path)			{ SetPath(Path); return *this; }
-	FSoftObjectPath& operator=(const FString& Path)						{ SetPath(FStringView(Path)); return *this; }
-	FSoftObjectPath& operator=(FWideStringView Path)					{ SetPath(Path); return *this; }
-	FSoftObjectPath& operator=(FAnsiStringView Path)					{ SetPath(Path); return *this; }
-	FSoftObjectPath& operator=(const WIDECHAR* Path)					{ SetPath(FWideStringView(Path)); return *this; }
-	FSoftObjectPath& operator=(const ANSICHAR* Path)					{ SetPath(FAnsiStringView(Path)); return *this; }
-	FSoftObjectPath& operator=(TYPE_OF_NULLPTR)							{ Reset(); return *this; }
+	FSoftObjectPath& operator=(const FTopLevelAssetPath Path)
+	{
+		SetPath(Path);
+		return *this;
+	}
+	FSoftObjectPath& operator=(const FString& Path)
+	{
+		SetPath(FStringView(Path));
+		return *this;
+	}
+	FSoftObjectPath& operator=(FWideStringView Path)
+	{
+		SetPath(Path);
+		return *this;
+	}
+	FSoftObjectPath& operator=(FAnsiStringView Path)
+	{
+		SetPath(Path);
+		return *this;
+	}
+	FSoftObjectPath& operator=(const WIDECHAR* Path)
+	{
+		SetPath(FWideStringView(Path));
+		return *this;
+	}
+	FSoftObjectPath& operator=(const ANSICHAR* Path)
+	{
+		SetPath(FAnsiStringView(Path));
+		return *this;
+	}
+	FSoftObjectPath& operator=(TYPE_OF_NULLPTR)
+	{
+		Reset();
+		return *this;
+	}
 
 	/** Returns string representation of reference, in form /package/path.assetname[:subpath] */
-	COREUOBJECT_API FString ToString() const;
+	[[nodiscard]] COREUOBJECT_API FString ToString() const;
 
 	/** Append string representation of reference, in form /package/path.assetname[:subpath] */
 	COREUOBJECT_API void ToString(FStringBuilderBase& Builder) const;
@@ -126,19 +175,19 @@ struct FSoftObjectPath
 	COREUOBJECT_API void AppendString(FUtf8StringBuilderBase& Builder) const;
 
 	/** Returns the top-level asset part of this path, without the subobject path. */
-	FTopLevelAssetPath GetAssetPath() const
+	[[nodiscard]] FTopLevelAssetPath GetAssetPath() const
 	{
 		return AssetPath;
 	}
 
 	/** Returns this path without the SubPath component, restricting the result to a top level asset but keeping the type as FSoftObjectPath in contrast to GetAssetPath. */
-	FSoftObjectPath GetWithoutSubPath() const
+	[[nodiscard]] FSoftObjectPath GetWithoutSubPath() const
 	{
 		return FSoftObjectPath(AssetPath, {});
 	}
 
 	/** Returns string version of asset path, including both package and asset but not sub object */
-	FORCEINLINE FString GetAssetPathString() const
+	[[nodiscard]] FORCEINLINE FString GetAssetPathString() const
 	{
 		if (AssetPath.IsNull())
 		{
@@ -149,7 +198,7 @@ struct FSoftObjectPath
 	}
 
 	/** Returns the sub path, which is often empty */
-	FORCEINLINE const FString& GetSubPathString() const
+	[[nodiscard]] FORCEINLINE const FString& GetSubPathString() const
 	{
 		return SubPathString;
 	}
@@ -160,27 +209,27 @@ struct FSoftObjectPath
 	}
 
 	/** Returns /package/path, leaving off the asset name and sub object */
-	FString GetLongPackageName() const
+	[[nodiscard]] FString GetLongPackageName() const
 	{
 		FName PackageName = GetAssetPath().GetPackageName();
 		return PackageName.IsNone() ? FString() : PackageName.ToString();
 	}
 
 	/** Returns /package/path, leaving off the asset name and sub object */
-	FName GetLongPackageFName() const
+	[[nodiscard]] FName GetLongPackageFName() const
 	{
 		return GetAssetPath().GetPackageName();
 	}
 
 	/** Returns assetname string, leaving off the /package/path part and sub object */
-	FString GetAssetName() const
+	[[nodiscard]] FString GetAssetName() const
 	{
 		FName AssetName = GetAssetPath().GetAssetName();
 		return AssetName.IsNone() ? FString() : AssetName.ToString();
 	}
 
 	/** Returns assetname string, leaving off the /package/path part and sub object */
-	FName GetAssetFName() const
+	[[nodiscard]] FName GetAssetFName() const
 	{
 		return GetAssetPath().GetAssetName();
 	}
@@ -190,9 +239,18 @@ struct FSoftObjectPath
 	COREUOBJECT_API void SetPath(FWideStringView Path);
 	COREUOBJECT_API void SetPath(FAnsiStringView Path);
 	COREUOBJECT_API void SetPath(FUtf8StringView Path);
-	void SetPath(const WIDECHAR* Path)			{ SetPath(FWideStringView(Path)); }
-	void SetPath(const ANSICHAR* Path)			{ SetPath(FAnsiStringView(Path)); }
-	void SetPath(const FString& Path)			{ SetPath(FStringView(Path)); }
+	void SetPath(const WIDECHAR* Path)
+	{
+		SetPath(FWideStringView(Path));
+	}
+	void SetPath(const ANSICHAR* Path)
+	{
+		SetPath(FAnsiStringView(Path));
+	}
+	void SetPath(const FString& Path)
+	{
+		SetPath(FStringView(Path));
+	}
 
 	/**
 	 * Attempts to load the asset, this will call LoadObject which can be very slow
@@ -217,41 +275,41 @@ struct FSoftObjectPath
 	 *
 	 * @return Found UObject, or nullptr if not currently in memory
 	 */
-	COREUOBJECT_API UObject* ResolveObject() const;
+	[[nodiscard]] COREUOBJECT_API UObject* ResolveObject() const;
 
 	/** Resets reference to point to null */
 	void Reset()
-	{		
+	{
 		AssetPath.Reset();
 		SubPathString.Reset();
 	}
 	
 	/** Check if this could possibly refer to a real object, or was initialized to null */
-	FORCEINLINE bool IsValid() const
+	[[nodiscard]] FORCEINLINE bool IsValid() const
 	{
 		return AssetPath.IsValid();
 	}
 
 	/** Checks to see if this is initialized to null */
-	FORCEINLINE bool IsNull() const
+	[[nodiscard]] FORCEINLINE bool IsNull() const
 	{
 		return AssetPath.IsNull();
 	}
 
 	/** Check if this represents an asset, meaning it is not null but does not have a sub path */
-	FORCEINLINE bool IsAsset() const
+	[[nodiscard]] FORCEINLINE bool IsAsset() const
 	{
 		return !AssetPath.IsNull() && SubPathString.IsEmpty();
 	}
 
 	/** Check if this represents a sub object, meaning it has a sub path */
-	FORCEINLINE bool IsSubobject() const
+	[[nodiscard]] FORCEINLINE bool IsSubobject() const
 	{
 		return !AssetPath.IsNull() && !SubPathString.IsEmpty();
 	}
 
 	/** Return true if this path appears before Other in lexical order */
-	bool LexicalLess(const FSoftObjectPath& Other) const
+	[[nodiscard]] bool LexicalLess(const FSoftObjectPath& Other) const
 	{
 		int32 PathCompare = AssetPath.Compare(Other.AssetPath);
 		if (PathCompare != 0)
@@ -262,7 +320,7 @@ struct FSoftObjectPath
 	}
 
 	/** Return true if this path appears before Other using fast index-based fname order */
-	bool FastLess(const FSoftObjectPath& Other) const
+	[[nodiscard]] bool FastLess(const FSoftObjectPath& Other) const
 	{
 		int32 PathCompare = AssetPath.CompareFast(Other.AssetPath);
 		if (PathCompare != 0)
@@ -275,14 +333,17 @@ struct FSoftObjectPath
 	/** Struct overrides */
 	COREUOBJECT_API bool Serialize(FArchive& Ar);
 	COREUOBJECT_API bool Serialize(FStructuredArchive::FSlot Slot);
-	COREUOBJECT_API bool operator==(FSoftObjectPath const& Other) const;
-	bool operator!=(FSoftObjectPath const& Other) const
+	[[nodiscard]] COREUOBJECT_API bool operator==(FSoftObjectPath const& Other) const;
+
+#if !PLATFORM_COMPILER_HAS_GENERATED_COMPARISON_OPERATORS
+	[[nodiscard]] bool operator!=(FSoftObjectPath const& Other) const
 	{
 		return !(*this == Other);
 	}
+#endif
 
 	COREUOBJECT_API bool ExportTextItem(FString& ValueStr, FSoftObjectPath const& DefaultValue, UObject* Parent, int32 PortFlags, UObject* ExportRootScope) const;
-	COREUOBJECT_API bool ImportTextItem( const TCHAR*& Buffer, int32 PortFlags, UObject* Parent, FOutputDevice* ErrorText, FArchive* InSerializingArchive = nullptr );
+	COREUOBJECT_API bool ImportTextItem(const TCHAR*& Buffer, int32 PortFlags, UObject* Parent, FOutputDevice* ErrorText, FArchive* InSerializingArchive = nullptr);
 	COREUOBJECT_API bool SerializeFromMismatchedTag(struct FPropertyTag const& Tag, FStructuredArchive::FSlot Slot);
 
 	/** Serializes the internal path and also handles save/PIE fixups. Call this from the archiver overrides */
@@ -318,13 +379,13 @@ struct FSoftObjectPath
 		return Hash;
 	}
 
-	static COREUOBJECT_API FSoftObjectPath GetOrCreateIDForObject(FObjectPtr Object);
-	static FORCEINLINE FSoftObjectPath GetOrCreateIDForObject(const UObject* Object)
+	[[nodiscard]] static COREUOBJECT_API FSoftObjectPath GetOrCreateIDForObject(FObjectPtr Object);
+	[[nodiscard]] static FORCEINLINE FSoftObjectPath GetOrCreateIDForObject(const UObject* Object)
 	{
 		return GetOrCreateIDForObject(FObjectPtr(const_cast<UObject*>(Object)));
 	}
 	template <typename T>
-	static FORCEINLINE FSoftObjectPath GetOrCreateIDForObject(TObjectPtr<T> Object)
+	[[nodiscard]] static FORCEINLINE FSoftObjectPath GetOrCreateIDForObject(TObjectPtr<T> Object)
 	{
 		// This needs to be a template instead of TObjectPtr<const UObject> because C++ does derived-to-base
 		// pointer conversions ('standard conversion sequences') in more cases than TSmartPtr<Derived>-to-TSmartPtr<Base>
@@ -358,8 +419,8 @@ private:
 	/** Package names currently being duplicated, needed by FixupForPIE */
 	static COREUOBJECT_API TSet<FName> PIEPackageNames;
 
-	COREUOBJECT_API UObject* ResolveObjectInternal() const;
-	COREUOBJECT_API UObject* ResolveObjectInternal(const TCHAR* PathString) const;
+	[[nodiscard]] COREUOBJECT_API UObject* ResolveObjectInternal() const;
+	[[nodiscard]] COREUOBJECT_API UObject* ResolveObjectInternal(const TCHAR* PathString) const;
 
 	COREUOBJECT_API friend void SerializeForLog(FCbWriter& Writer, const FSoftObjectPath& Value);
 	friend struct Z_Construct_UScriptStruct_FSoftObjectPath_Statics;
@@ -368,7 +429,7 @@ private:
 /** Fast non-alphabetical order that is only stable during this process' lifetime. */
 struct FSoftObjectPathFastLess
 {
-	bool operator()(const FSoftObjectPath& Lhs, const FSoftObjectPath& Rhs) const
+	[[nodiscard]] bool operator()(const FSoftObjectPath& Lhs, const FSoftObjectPath& Rhs) const
 	{
 		return Lhs.FastLess(Rhs);
 	}
@@ -377,7 +438,7 @@ struct FSoftObjectPathFastLess
 /** Slow alphabetical order that is stable / deterministic over process runs. */
 struct FSoftObjectPathLexicalLess
 {
-	bool operator()(const FSoftObjectPath& Lhs, const FSoftObjectPath& Rhs) const
+	[[nodiscard]] bool operator()(const FSoftObjectPath& Lhs, const FSoftObjectPath& Rhs) const
 	{
 		return Lhs.LexicalLess(Rhs);
 	}
@@ -399,9 +460,9 @@ inline FUtf8StringBuilderBase& operator<<(FUtf8StringBuilderBase& Builder, const
  */
 struct FSoftClassPath : public FSoftObjectPath
 {
-	FSoftClassPath() = default;
-	FSoftClassPath(const FSoftClassPath& Other) = default;
-	FSoftClassPath(FSoftClassPath&& Other) = default;
+	[[nodiscard]] FSoftClassPath() = default;
+	[[nodiscard]] FSoftClassPath(const FSoftClassPath& Other) = default;
+	[[nodiscard]] FSoftClassPath(FSoftClassPath&& Other) = default;
 	~FSoftClassPath() = default;
 	FSoftClassPath& operator=(const FSoftClassPath& Path) = default;
 	FSoftClassPath& operator=(FSoftClassPath&& Path) = default;
@@ -409,16 +470,18 @@ struct FSoftClassPath : public FSoftObjectPath
 	/**
 	 * Construct from a path string
 	 */
-	FSoftClassPath(const FString& PathString)
+	[[nodiscard]] FSoftClassPath(const FString& PathString)
 		: FSoftObjectPath(PathString)
-	{ }
+	{
+	}
 
 	/**
 	 * Construct from an existing class, will do some string processing
 	 */
-	FSoftClassPath(const UClass* InClass)
+	[[nodiscard]] FSoftClassPath(const UClass* InClass)
 		: FSoftObjectPath(InClass)
-	{ }
+	{
+	}
 
 	/**
 	* Attempts to load the class.
@@ -439,15 +502,17 @@ struct FSoftClassPath : public FSoftObjectPath
 	 * Attempts to find a currently loaded object that matches this object ID
 	 * @return Found UClass, or NULL if not currently loaded
 	 */
-	COREUOBJECT_API UClass* ResolveClass() const;
+	[[nodiscard]] COREUOBJECT_API UClass* ResolveClass() const;
 
 	COREUOBJECT_API bool SerializeFromMismatchedTag(const FPropertyTag& Tag, FStructuredArchive::FSlot Slot);
 
-	static COREUOBJECT_API FSoftClassPath GetOrCreateIDForClass(const UClass *InClass);
+	[[nodiscard]] static COREUOBJECT_API FSoftClassPath GetOrCreateIDForClass(const UClass *InClass);
 
 private:
 	/** Forbid creation for UObject. This class is for UClass only. Use FSoftObjectPath instead. */
-	FSoftClassPath(const UObject* InObject) { }
+	[[nodiscard]] explicit FSoftClassPath(const UObject* InObject)
+	{
+	}
 
 	/** Forbidden. This class is for UClass only. Use FSoftObjectPath instead. */
 	static COREUOBJECT_API FSoftObjectPath GetOrCreateIDForObject(const UObject *Object);
@@ -482,7 +547,9 @@ class COREUOBJECT_API FSoftObjectPathThreadContext : public TThreadSingleton<FSo
 	friend TThreadSingleton<FSoftObjectPathThreadContext>;
 	friend struct FSoftObjectPathSerializationScope;
 
-	FSoftObjectPathThreadContext() {}
+	[[nodiscard]] FSoftObjectPathThreadContext()
+	{
+	}
 
 	struct FSerializationOptions
 	{
@@ -491,8 +558,18 @@ class COREUOBJECT_API FSoftObjectPathThreadContext : public TThreadSingleton<FSo
 		ESoftObjectPathCollectType CollectType;
 		ESoftObjectPathSerializeType SerializeType;
 
-		FSerializationOptions() : CollectType(ESoftObjectPathCollectType::AlwaysCollect) {}
-		FSerializationOptions(FName InPackageName, FName InPropertyName, ESoftObjectPathCollectType InCollectType, ESoftObjectPathSerializeType InSerializeType) : PackageName(InPackageName), PropertyName(InPropertyName), CollectType(InCollectType), SerializeType(InSerializeType) {}
+		[[nodiscard]] FSerializationOptions()
+			: CollectType(ESoftObjectPathCollectType::AlwaysCollect)
+		{
+		}
+
+		[[nodiscard]] explicit FSerializationOptions(FName InPackageName, FName InPropertyName, ESoftObjectPathCollectType InCollectType, ESoftObjectPathSerializeType InSerializeType)
+			: PackageName(InPackageName)
+			, PropertyName(InPropertyName)
+			, CollectType(InCollectType)
+			, SerializeType(InSerializeType)
+		{
+		}
 	};
 
 	TArray<FSerializationOptions> OptionStack;
@@ -518,12 +595,12 @@ struct FSoftObjectPathSerializationScope
 	 * @param SerializingPropertyName Property that this path belongs to
 	 * @param CollectType Set type of collecting that should be done, can be used to disable tracking entirely
 	 */
-	FSoftObjectPathSerializationScope(FName SerializingPackageName, FName SerializingPropertyName, ESoftObjectPathCollectType CollectType, ESoftObjectPathSerializeType SerializeType)
+	[[nodiscard]] explicit FSoftObjectPathSerializationScope(FName SerializingPackageName, FName SerializingPropertyName, ESoftObjectPathCollectType CollectType, ESoftObjectPathSerializeType SerializeType)
 	{
 		FSoftObjectPathThreadContext::Get().OptionStack.Emplace(SerializingPackageName, SerializingPropertyName, CollectType, SerializeType);
 	}
 
-	explicit FSoftObjectPathSerializationScope(ESoftObjectPathCollectType CollectType)
+	[[nodiscard]] explicit FSoftObjectPathSerializationScope(ESoftObjectPathCollectType CollectType)
 	{
 		FSoftObjectPathThreadContext::Get().OptionStack.Emplace(NAME_None, NAME_None, CollectType, ESoftObjectPathSerializeType::AlwaysSerialize);
 	}
@@ -555,7 +632,7 @@ struct FDirectoryPath
 // Fixup archive
 struct FSoftObjectPathFixupArchive : public FArchiveUObject
 {
-	FSoftObjectPathFixupArchive(TFunction<void(FSoftObjectPath&)> InFixupFunction)
+	[[nodiscard]] explicit FSoftObjectPathFixupArchive(TFunction<void(FSoftObjectPath&)> InFixupFunction)
 		: FixupFunction(InFixupFunction)
 	{
 		this->SetIsSaving(true);
@@ -563,7 +640,7 @@ struct FSoftObjectPathFixupArchive : public FArchiveUObject
 		this->SetShouldSkipCompilingAssets(true);
 	}
 
-	FSoftObjectPathFixupArchive(const FString& InOldAssetPathString, const FString& InNewAssetPathString)
+	[[nodiscard]] explicit FSoftObjectPathFixupArchive(const FString& InOldAssetPathString, const FString& InNewAssetPathString)
 		: FSoftObjectPathFixupArchive([OldAssetPathString = InOldAssetPathString, NewAssetPath = FTopLevelAssetPath(InNewAssetPathString)](FSoftObjectPath& Value)
 		{
 			if (!Value.IsNull() && Value.GetAssetPathString().Equals(OldAssetPathString, ESearchCase::IgnoreCase))
@@ -603,7 +680,7 @@ struct FSoftObjectPathFixupArchive : public FArchiveUObject
 namespace UE::SoftObjectPath::Private
 {
 	UE_DEPRECATED(5.1, "This function is only for use in fixing up deprecated APIs.")
-	inline TArray<FName> ConvertSoftObjectPaths(TConstArrayView<FSoftObjectPath> InPaths)
+	[[nodiscard]] inline TArray<FName> ConvertSoftObjectPaths(TConstArrayView<FSoftObjectPath> InPaths)
 	{
 		TArray<FName> Out;
 		Out.Reserve(InPaths.Num());
@@ -617,7 +694,7 @@ PRAGMA_ENABLE_DEPRECATION_WARNINGS
 	}
 
 	UE_DEPRECATED(5.1, "This function is only for use in fixing up deprecated APIs.")
-	inline TArray<FSoftObjectPath> ConvertObjectPathNames(TConstArrayView<FName> InPaths)
+	[[nodiscard]] inline TArray<FSoftObjectPath> ConvertObjectPathNames(TConstArrayView<FName> InPaths)
 	{
 		TArray<FSoftObjectPath> Out;
 		Out.Reserve(InPaths.Num());
