RemoteObject:
- the priority for object requests is now equal to the task requesting it rather than being of 'dependent task' priority. This solves an issue where we could live lock by saying an object request is higher priority than an equal depth dependent task, but the task requesting the object is also blocked on the dependent task completing. Since object requests aren't themselves tasks or transactions, the priority we send along with it is just for arbitrating and so now we use the priority of the task that would be using that object
- we now send object requests at the priority of the highest priority task that we have pending
- fixed an issue where bRunTransactionally wasn't initialized

diff --git a/Engine/Source/Runtime/CoreUObject/Private/UObject/RemoteObject.cpp b/Engine/Source/Runtime/CoreUObject/Private/UObject/RemoteObject.cpp
--- a/Engine/Source/Runtime/CoreUObject/Private/UObject/RemoteObject.cpp
+++ b/Engine/Source/Runtime/CoreUObject/Private/UObject/RemoteObject.cpp
@@ -342,30 +342,6 @@ UObject* ResolveObject(const FRemoteObjectStub* Stub, ERemoteReferenceType RefTy
 	return Object;
 }
 
-bool RequestObject(const FRemoteObjectStub* Stub)
-{
-	using namespace UE::RemoteObject::Transfer;
-	using namespace UE::RemoteObject::Private;
-
-	UObject* Object = StaticFindObjectFastInternal(Stub->Id);
-	bool bRemoteObject = false;
-	if (Object)
-	{
-		bRemoteObject = IsRemote(Object);
-	}
-	else
-	{
-		bRemoteObject = IsRemote(Stub->Id);
-	}
-
-	if (bRemoteObject)
-	{
-		RequestObjectData(Stub->Id, Stub->ResidentServerId, GetGlobalServerId());
-	}
-
-	return !bRemoteObject;
-}
-
 UObject* ResolveObject(UObject* Object, ERemoteReferenceType RefType /*= ERemoteReferenceType::Strong*/)
 {
 	using namespace UE::RemoteObject::Private;
@@ -382,22 +358,6 @@ UObject* ResolveObject(UObject* Object, ERemoteReferenceType RefType /*= ERemote
 	return Object;
 }
 
-bool RequestObject(FRemoteObjectId ObjectId)
-{
-	using namespace UE::RemoteObject::Private;
-
-	if (IsRemote(ObjectId))
-	{
-		FRemoteObjectStub* Stub = ObjectMaps->FindRemoteObjectStub(ObjectId);
-		checkf(Stub, TEXT("Failed to find remote object stub for %s"), *ObjectId.ToString());
-		return RequestObject(Stub);
-	}
-	else
-	{
-		return true;
-	}
-}
-
 bool CanResolveObject(FRemoteObjectId ObjectId)
 {
 	using namespace UE::RemoteObject::Private;
diff --git a/Engine/Source/Runtime/CoreUObject/Private/UObject/RemoteObjectSerialization.cpp b/Engine/Source/Runtime/CoreUObject/Private/UObject/RemoteObjectSerialization.cpp
--- a/Engine/Source/Runtime/CoreUObject/Private/UObject/RemoteObjectSerialization.cpp
+++ b/Engine/Source/Runtime/CoreUObject/Private/UObject/RemoteObjectSerialization.cpp
@@ -797,7 +797,7 @@ FRemoteObjectData* SerializeObjectData(UObject* InObject, TSet<UObject*>& OutObj
 	return ObjectData;
 }
 
-int32 DeserializeObjectData(FRemoteServerId OwnerServerId, FRemoteObjectData& ObjectData, TArray<UObject*>& OutObjects)
+int32 DeserializeObjectData(FRemoteServerId OwnerServerId, FRemoteObjectData& ObjectData, TArray<FRemoteObjectId>& OutObjectRemoteIds, TArray<UObject*>& OutObjects)
 {
 	using namespace UE::RemoteObject::Handle;
 	using namespace UE::RemoteObject::Private;
@@ -843,9 +843,11 @@ int32 DeserializeObjectData(FRemoteServerId OwnerServerId, FRemoteObjectData& Ob
 
 		// Construct all objects first
 		FRemoteObjectConstructionOverrides ConstructionOverrides(ObjectHeaders);
+		OutObjectRemoteIds.Reserve(ObjectHeaders.Num());
 		OutObjects.Reserve(ObjectHeaders.Num());
 		for (int32 ObjectIndex = 0; ObjectIndex < ObjectHeaders.Num(); ++ObjectIndex)
 		{
+			OutObjectRemoteIds.Add(ObjectHeaders[ObjectIndex].RemoteId);
 			OutObjects.Add(ConstructRemoteObject(ObjectHeaders[ObjectIndex], &ConstructionOverrides));
 		}
 
diff --git a/Engine/Source/Runtime/CoreUObject/Private/UObject/RemoteObjectSerialization.h b/Engine/Source/Runtime/CoreUObject/Private/UObject/RemoteObjectSerialization.h
--- a/Engine/Source/Runtime/CoreUObject/Private/UObject/RemoteObjectSerialization.h
+++ b/Engine/Source/Runtime/CoreUObject/Private/UObject/RemoteObjectSerialization.h
@@ -4,7 +4,7 @@
 
 #include "UObject/RemoteObjectTransfer.h"
 
-struct FRemoteObjectStub;
+namespace UE::RemoteObject::Handle { struct FRemoteObjectStub; }
 
 namespace UE::RemoteObject::Serialization
 {
@@ -69,7 +69,7 @@ namespace UE::RemoteObject::Serialization
 	* @return Index of an object in OutReceivedObjects that was the main object the migration request was triggered for 
 	* (usually 0 but if a migration requests a default subobject then its parent is also migrated and the return value will be > 0)
 	*/
-	int32 DeserializeObjectData(FRemoteServerId OwnerServerId, FRemoteObjectData& ObjectData, TArray<UObject*>& OutReceivedObjects);
+	int32 DeserializeObjectData(FRemoteServerId OwnerServerId, FRemoteObjectData& ObjectData, TArray<FRemoteObjectId>& OutObjectRemoteIds, TArray<UObject*>& OutReceivedObjects);
 
 	/**
 	* Finds the canonical 'root' object that is used for remote object serialization
diff --git a/Engine/Source/Runtime/CoreUObject/Private/UObject/RemoteObjectTransfer.cpp b/Engine/Source/Runtime/CoreUObject/Private/UObject/RemoteObjectTransfer.cpp
--- a/Engine/Source/Runtime/CoreUObject/Private/UObject/RemoteObjectTransfer.cpp
+++ b/Engine/Source/Runtime/CoreUObject/Private/UObject/RemoteObjectTransfer.cpp
@@ -68,6 +68,8 @@ public:
 
 	void Tick() final
 	{
+		using namespace UE::RemoteObject::Private;
+		using namespace UE::RemoteObject::Handle;
 		using namespace UE::RemoteObject::Serialization;
 
 		check(ActiveRequest->RequiredObjects.Num() == ActiveRequest->RequiredObjectsCanonicalRoot.Num());
@@ -75,9 +77,9 @@ public:
 		UE_LOG(LogRemoteObject, VeryVerbose, TEXT("\t%d required objs (%d pending requests):"), ActiveRequest->RequiredObjects.Num(), PendingObjectRequests.Num());
 		for (int32 RequiredObjectIndex = 0; RequiredObjectIndex < ActiveRequest->RequiredObjects.Num(); RequiredObjectIndex++)
 		{
-			FRemoteObjectId RequestedObjectId = ActiveRequest->RequiredObjects[RequiredObjectIndex];
+			FRemoteObjectId RequiredObjectId = ActiveRequest->RequiredObjects[RequiredObjectIndex];
 			FRemoteObjectId CanonicalObjectId = ActiveRequest->RequiredObjectsCanonicalRoot[RequiredObjectIndex];
-			UE_LOG(LogRemoteObject, VeryVerbose, TEXT("\t\t%s [root %s]"), *RequestedObjectId.ToString(), *CanonicalObjectId.ToString());
+			UE_LOG(LogRemoteObject, VeryVerbose, TEXT("\t\t%s [root %s] %s"), *RequiredObjectId.ToString(), *CanonicalObjectId.ToString(), Handle::IsRemote(RequiredObjectId) ? TEXT("<remote>") : TEXT("<local>"));
 		}
 
 		// send the requests... slowly for now (to avoid straining the network subsystem)
@@ -89,26 +91,49 @@ public:
 			}
 
 			FRemoteObjectId RequiredObjectId = ActiveRequest->RequiredObjects[RequiredObjectIndex];
+			FRemoteObjectId RequiredCanonicalRootObjectId = ActiveRequest->RequiredObjectsCanonicalRoot[RequiredObjectIndex];
+			
+			FRemoteObjectId ObjectIdToRequest = RequiredCanonicalRootObjectId.IsValid() ? RequiredCanonicalRootObjectId : RequiredObjectId;
 
-			if (Handle::IsRemote(RequiredObjectId) && (PendingObjectRequests.Find(RequiredObjectId) == INDEX_NONE))
+			if (Handle::IsRemote(ObjectIdToRequest))
 			{
-				// send one request
-				bool bAlreadyHaveObject = UE::RemoteObject::Handle::RequestObject(RequiredObjectId);
+				if (PendingObjectRequests.Find(ObjectIdToRequest) == INDEX_NONE)
+				{
+					// find the highest priority request that needs this object
+					FRemoteObjectRequest* HighestPriorityRequest = ActiveRequest;
+
+					for (int32 RequestIndex = 0; RequestIndex < GetRequestCount(); RequestIndex++)
+					{
+						FRemoteObjectRequest* ExistingRequest = GetRequestByIndex(RequestIndex);
+
+						bool bFoundObject =
+							(ExistingRequest->RequiredObjectsCanonicalRoot.Find(ObjectIdToRequest) != INDEX_NONE) ||
+							(ExistingRequest->RequiredObjects.Find(ObjectIdToRequest) != INDEX_NONE);
+
+						if (bFoundObject)
+						{
+							if (IsHigherPriority(ExistingRequest->Priority, HighestPriorityRequest->Priority))
+							{
+								HighestPriorityRequest = ExistingRequest;
+							}
+						}
+					}
+
+					FRemoteObjectStub* RemoteObjectStub = FindRemoteObjectStub(ObjectIdToRequest);
+					FRemoteServerId CurrentResidentServerId = RemoteObjectStub ? RemoteObjectStub->ResidentServerId : ObjectIdToRequest.GetServerId();
+
+					// send one request
+					RequestRemoteObject(
+						HighestPriorityRequest->GetRequestId(),
+						HighestPriorityRequest->GetPriority(),
+						ObjectIdToRequest,
+						CurrentResidentServerId,
+						GetGlobalServerId());
+			
+					PendingObjectRequests.Add(ObjectIdToRequest);
 				
-				// bAlreadyHaveObject should be false since we already tested IsRemote()
-				check (!bAlreadyHaveObject);
-
-				UE_LOG(LogRemoteObject, Verbose, TEXT("\tsent request for %s %s (%d pending requests)"), *RequiredObjectId.ToString(), *ActiveRequest->Priority.ToString(), PendingObjectRequests.Num());
-
-				PendingObjectRequests.Add(RequiredObjectId);
-			}
-
-			if (!Handle::IsRemote(RequiredObjectId))
-			{
-				// discover the canonical root object id
-				// todo what about the other requests that might have this?
-				UObject* const RootObject = FindCanonicalRootObjectForSerialization(StaticFindObjectFastInternal(RequiredObjectId));
-				ActiveRequest->RequiredObjectsCanonicalRoot[RequiredObjectIndex] = UE::CoreUObject::Private::FObjectHandleUtils::GetRemoteId(RootObject);
+					UE_LOG(LogRemoteObject, Verbose, TEXT("\tsent request for %s %s (%d pending requests)"), *RequiredObjectId.ToString(), *HighestPriorityRequest->Priority.ToString(), PendingObjectRequests.Num());
+				}
 			}
 		}
 	}
@@ -195,21 +220,6 @@ public:
 		if (bFoundRequest)
 		{
 			PendingObjectRequests.RemoveAt(PendingRequestIndex);
-
-			// now that we have the object locally we can discover and update the canonical root object id
-			UObject* RootObject = FindCanonicalRootObjectForSerialization(Result);
-			FRemoteObjectId RootObjectRemoteId = UE::CoreUObject::Private::FObjectHandleUtils::GetRemoteId(RootObject);
-
-			for (int32 RequestIndex = 0; RequestIndex < GetRequestCount(); RequestIndex++)
-			{
-				FRemoteObjectRequest* ExistingRequest = GetRequestByIndex(RequestIndex);
-
-				const int32 RequiredObjectIndex = ExistingRequest->RequiredObjects.Find(ObjectId);
-				if (RequiredObjectIndex != INDEX_NONE)
-				{
-					ExistingRequest->RequiredObjectsCanonicalRoot[RequiredObjectIndex] = RootObjectRemoteId;
-				}
-			}
 		}
 	}
 
@@ -231,19 +241,6 @@ public:
 
 FRemoteObjectTransferQueue GTransferQueue;
 
-void RequestObjectData(FRemoteObjectId ObjectId, FRemoteServerId CurrentOwnerServerId, FRemoteServerId DestinationServerId)
-{
-	check(!AutoRTFM::IsClosed());
-
-	check(GTransferQueue.ActiveRequest != nullptr);
-	GTransferQueue.RequestRemoteObject(
-		GTransferQueue.ActiveRequest->GetRequestId(),
-		GTransferQueue.ActiveRequest->GetPriority().CreateDependentWorkPriority(),
-		ObjectId,
-		CurrentOwnerServerId,
-		DestinationServerId);
-}
-
 UObject* MigrateObjectFromObjectData(FRemoteServerId OwnerServerId, FRemoteObjectData& ObjectData, FRemoteServerId RemoteServerId, EObjectMigrationFlags MigrationFlags)
 {
 	using namespace UE::RemoteObject::Serialization;
@@ -252,8 +249,9 @@ UObject* MigrateObjectFromObjectData(FRemoteServerId OwnerServerId, FRemoteObjec
 	// Notify that object data has been migrated from a remote server
 	OnObjectDataReceivedDelegate.Broadcast(ObjectData, RemoteServerId);
 
+	TArray<FRemoteObjectId> ReceivedObjectRemoteIds;
 	TArray<UObject*> ReceivedObjects;
-	int32 RequestedObjectIndex = DeserializeObjectData(OwnerServerId, ObjectData, ReceivedObjects);
+	int32 RequestedObjectIndex = DeserializeObjectData(OwnerServerId, ObjectData, ReceivedObjectRemoteIds, ReceivedObjects);
 
 	checkf(ReceivedObjects.Num() > 0, TEXT("PeerId:%s unable to deserialize object data (%d bytes)"), *GetGlobalServerId().ToString(), ObjectData.Bytes.Num());
 
@@ -264,6 +262,8 @@ UObject* MigrateObjectFromObjectData(FRemoteServerId OwnerServerId, FRemoteObjec
 			for(int32 SubObjectIndex = 0; SubObjectIndex < ReceivedObjects.Num(); SubObjectIndex++)
 			{
 				FRemoteObjectId SubObjectId = UE::CoreUObject::Private::FObjectHandleUtils::GetRemoteId(ReceivedObjects[SubObjectIndex]);
+				check(ReceivedObjectRemoteIds[SubObjectIndex] == SubObjectId);
+
 				UE_LOG(LogRemoteObject, VeryVerbose, TEXT("         subobject[%2d] : %s %s"),
 					SubObjectIndex,
 					*SubObjectId.ToString(),
@@ -288,29 +288,63 @@ UObject* MigrateObjectFromObjectData(FRemoteServerId OwnerServerId, FRemoteObjec
 	// Notify that objects have been migrated from a remote server
 	OnObjectsReceivedDelegate.Broadcast(ReceivedObjects, RemoteServerId);
 
+	// do any existing requests require any of the objects in this hierarchy?
+	// if so, we can now update the canonical root id
+	FRemoteObjectId RootObjectRemoteId = UE::CoreUObject::Private::FObjectHandleUtils::GetRemoteId(ReceivedObjects[0]);
+
+	for (int32 RequestIndex = 0; RequestIndex < GTransferQueue.GetRequestCount(); RequestIndex++)
+	{
+		FRemoteObjectRequest* ExistingRequest = GTransferQueue.GetRequestByIndex(RequestIndex);
+
+		for (FRemoteObjectId SubObjectId : ReceivedObjectRemoteIds)
+		{
+			const int32 RequiredObjectIndex = ExistingRequest->RequiredObjects.Find(SubObjectId);
+			if (RequiredObjectIndex != INDEX_NONE)
+			{
+				ExistingRequest->RequiredObjectsCanonicalRoot[RequiredObjectIndex] = RootObjectRemoteId;
+			}
+		}
+	}
+
 	return ReceivedObjects[RequestedObjectIndex];
 } 
 
 void MigrateObjectFromRemoteServer(FRemoteObjectId ObjectId, FRemoteServerId CurrentOwnerServerId, UObject* DestinationOuter)
 {
-	// ensure the request is added in the Open so after we abort it is preserved
-	bool bAdded = false;
-
 	UE_LOG(LogRemoteObject, Verbose, TEXT("Aborting transaction, need obj id %s on server %s"),
 		*ObjectId.ToString(),
 		*CurrentOwnerServerId.ToString());
 
-	check(GTransferQueue.ActiveRequest);
+	checkf(GTransferQueue.ActiveRequest, TEXT("Attempting to access remote object %s but we are outside of a transaction"), *ObjectId.ToString());
 
+	// ensure the request is added in the Open so after we abort it is preserved
 	UE_AUTORTFM_OPEN
 	{
-		// it's possible a prefetch already added this to the RequestedObjects list
-		bAdded = (GTransferQueue.ActiveRequest->RequiredObjects.Find(ObjectId) == INDEX_NONE);
-
-		if (bAdded)
+		// add it to the active request's list of required objects
+		if (GTransferQueue.ActiveRequest->RequiredObjects.Find(ObjectId) == INDEX_NONE)
 		{
+			// does any other existing request also require this object? if so, try to
+			// grab the cached canonical id from them
+			FRemoteObjectId CanonicalRootObjectId;
+
+			for (int32 RequestIndex = 0; RequestIndex < GTransferQueue.GetRequestCount(); RequestIndex++)
+			{
+				FRemoteObjectRequest* ExistingRequest = GTransferQueue.GetRequestByIndex(RequestIndex);
+
+				for (int32 RequiredObjectIndex = 0; RequiredObjectIndex < ExistingRequest->RequiredObjects.Num(); RequiredObjectIndex++)
+				{
+					FRemoteObjectId RequiredObjectId = ExistingRequest->RequiredObjects[RequiredObjectIndex];
+
+					if (ObjectId == RequiredObjectId)
+					{
+						CanonicalRootObjectId = ExistingRequest->RequiredObjectsCanonicalRoot[RequiredObjectIndex];
+						break;
+					}
+				}
+			}
+
 			GTransferQueue.ActiveRequest->RequiredObjects.Add(ObjectId);
-			GTransferQueue.ActiveRequest->RequiredObjectsCanonicalRoot.Add(ObjectId);
+			GTransferQueue.ActiveRequest->RequiredObjectsCanonicalRoot.Add(CanonicalRootObjectId);
 		}
 		
 		if (UE_LOG_ACTIVE(LogRemoteObject, VeryVerbose))
diff --git a/Engine/Source/Runtime/CoreUObject/Public/UObject/RemoteObject.h b/Engine/Source/Runtime/CoreUObject/Public/UObject/RemoteObject.h
--- a/Engine/Source/Runtime/CoreUObject/Public/UObject/RemoteObject.h
+++ b/Engine/Source/Runtime/CoreUObject/Public/UObject/RemoteObject.h
@@ -71,13 +71,6 @@ namespace UE::RemoteObject::Handle
 	*/
 	COREUOBJECT_API UObject* ResolveObject(UObject* Object, ERemoteReferenceType RefType = ERemoteReferenceType::Strong);
 
-	/**
-	* Sends a request to a remote server for the object
-	* @param Object Object to resolve
-	* @return true if the object is already local, false if a request was sent
-	*/
-	COREUOBJECT_API bool RequestObject(FRemoteObjectId ObjectId);
-
 	/**
 	* Resolves a remote object
 	* @param Object Object id to resolve
diff --git a/Engine/Source/Runtime/CoreUObject/Public/UObject/RemoteObjectTransfer.h b/Engine/Source/Runtime/CoreUObject/Public/UObject/RemoteObjectTransfer.h
--- a/Engine/Source/Runtime/CoreUObject/Public/UObject/RemoteObjectTransfer.h
+++ b/Engine/Source/Runtime/CoreUObject/Public/UObject/RemoteObjectTransfer.h
@@ -67,14 +67,6 @@ namespace UE::RemoteObject::Transfer
 	*/
 	COREUOBJECT_API void OnObjectDataDenied(FRemoteTransactionId RequestId, FRemoteObjectId ObjectId, FRemoteServerId RemoteServerId);
 
-	/**
-	* Requests remote object data to be transferred from Current Owner Server to Destination Server
-	* @param ObjectId Remote object id
-	* @param CurrentOwnerServerId Id of a remote server that is known to currently own the object
-	* @param DestinationServerId Id of a remote server that wants to own the object (may be different than the the server that sent the request)
-	*/
-	COREUOBJECT_API void RequestObjectData(FRemoteObjectId ObjectId, FRemoteServerId CurrentOwnerServerId, FRemoteServerId DestinationServerId);
-
 	/**
 	* Migrates an object to a remote server
 	* @param Object Local object to migrate
diff --git a/Engine/Source/Runtime/Engine/Private/TickTaskManager.cpp b/Engine/Source/Runtime/Engine/Private/TickTaskManager.cpp
--- a/Engine/Source/Runtime/Engine/Private/TickTaskManager.cpp
+++ b/Engine/Source/Runtime/Engine/Private/TickTaskManager.cpp
@@ -2135,6 +2135,7 @@ FTickFunction::FTickFunction()
 	, bAllowTickBatching(false)
 	, bHighPriority(false)
 	, bRunOnAnyThread(false)
+	, bRunTransactionally(false)
 	, TickState(ETickState::Enabled)
 	, TickInterval(0.f)
 {
