Removed redundant FSoftObjectPath::ResolveObjectInternal overload.

diff --git a/Engine/Source/Runtime/CoreUObject/Private/UObject/SoftObjectPath.cpp b/Engine/Source/Runtime/CoreUObject/Private/UObject/SoftObjectPath.cpp
--- a/Engine/Source/Runtime/CoreUObject/Private/UObject/SoftObjectPath.cpp
+++ b/Engine/Source/Runtime/CoreUObject/Private/UObject/SoftObjectPath.cpp
@@ -686,12 +686,8 @@ UObject* FSoftObjectPath::ResolveObjectInternal() const
 {
 	TStringBuilder<FName::StringBufferSize> Builder;
 	Builder << *this;
-	return ResolveObjectInternal(*Builder);
-}
 
-UObject* FSoftObjectPath::ResolveObjectInternal(const TCHAR* PathString) const
-{
-	UObject* FoundObject = FindObject<UObject>(nullptr, PathString);
+	UObject* FoundObject = FindObject<UObject>(nullptr, *Builder);
 
 	if (!FoundObject && IsSubobject())
 	{
diff --git a/Engine/Source/Runtime/CoreUObject/Public/UObject/SoftObjectPath.h b/Engine/Source/Runtime/CoreUObject/Public/UObject/SoftObjectPath.h
--- a/Engine/Source/Runtime/CoreUObject/Public/UObject/SoftObjectPath.h
+++ b/Engine/Source/Runtime/CoreUObject/Public/UObject/SoftObjectPath.h
@@ -420,7 +420,6 @@ private:
 	static COREUOBJECT_API TSet<FName> PIEPackageNames;
 
 	[[nodiscard]] COREUOBJECT_API UObject* ResolveObjectInternal() const;
-	[[nodiscard]] COREUOBJECT_API UObject* ResolveObjectInternal(const TCHAR* PathString) const;
 
 	COREUOBJECT_API friend void SerializeForLog(FCbWriter& Writer, const FSoftObjectPath& Value);
 	friend struct Z_Construct_UScriptStruct_FSoftObjectPath_Statics;
