Remote object reader will now also dump stats to log when remote object archive debugging is enabled

diff --git a/Engine/Source/Runtime/CoreUObject/Private/Serialization/UnversionedPropertySerialization.cpp b/Engine/Source/Runtime/CoreUObject/Private/Serialization/UnversionedPropertySerialization.cpp
--- a/Engine/Source/Runtime/CoreUObject/Private/Serialization/UnversionedPropertySerialization.cpp
+++ b/Engine/Source/Runtime/CoreUObject/Private/Serialization/UnversionedPropertySerialization.cpp
@@ -8,6 +8,7 @@
 #include "Misc/ScopeRWLock.h"
 #include "UObject/OverridableManager.h"
 #include "UObject/PropertyOptional.h"
+#include "UObject/RemoteObjectSerialization.h"
 #include "UObject/UnrealType.h"
 
 #if WITH_EDITORONLY_DATA
@@ -1000,7 +1001,7 @@ void SerializeUnversionedProperties(const UStruct* Struct, FStructuredArchive::F
 				{
 					if (It.IsNonZero())
 					{
-#if WITH_EDITOR // Skip this scope to save time in the runtime; it is only needed for reference collection in editor
+#if WITH_EDITOR || UE_WITH_REMOTEOBJECTARCHIVE_DEBUGGING// Skip this scope to save time in the runtime; it is only needed for reference collection in editor
 						FSerializedPropertyScope SerializedProperty(UnderlyingArchive, It.GetSerializer().GetProperty());
 #endif
 
diff --git a/Engine/Source/Runtime/CoreUObject/Private/UObject/RemoteObjectSerialization.cpp b/Engine/Source/Runtime/CoreUObject/Private/UObject/RemoteObjectSerialization.cpp
--- a/Engine/Source/Runtime/CoreUObject/Private/UObject/RemoteObjectSerialization.cpp
+++ b/Engine/Source/Runtime/CoreUObject/Private/UObject/RemoteObjectSerialization.cpp
@@ -18,10 +18,6 @@
 #include "HAL/FileManager.h"
 #include "Serialization/ArchiveSerializedPropertyChain.h"
 
-#ifndef UE_WITH_REMOTEOBJECTWRITER_DEBUGGING
-#define UE_WITH_REMOTEOBJECTWRITER_DEBUGGING 0
-#endif
-
 int32 GRemoteObjectsMigrateFullHierarchy = 1;
 static FAutoConsoleVariableRef CVarRemoteObjectsMigrateFullHierarchy(
 	TEXT("ro.MigrateFullHierarchy"),
@@ -129,25 +125,140 @@ FArchive& operator<<(FArchive& Ar, ERemoteReferenceType& RefType)
 	return Ar;
 }
 
-class FArchiveRemoteObjectWriter : public FMemoryWriter
+template <typename T>
+class TArchiveRemoteObjectBase : public T
 {
+protected:
+
 	UObject* RootObject = nullptr;
 	FRemoteObjectData& ObjectData;
-	TArray<UObject*>& ObjectsToSerialize;
-	const TSet<UObject*>& SerializedObjects;
-	TMap<FName, FNameIndexType> NameMap;
-	TMap<UObject*, FNameIndexType> PathNameMap;
-
-#if UE_WITH_REMOTEOBJECTWRITER_DEBUGGING
-	FString SerializationScope;	
+	FString ArchiveName;
 
-	struct FSerializeStats
+#if UE_WITH_REMOTEOBJECTARCHIVE_DEBUGGING
+	struct FMigratedPropertyStats
 	{
 		int64 Size = 0;
 		int64 Count = 0;
 	};
-	TMap<FString, TMap<FProperty*, FSerializeStats>> ObjectPropertyStats;
+
+	FString SerializationScope;
+	TMap<FString, TMap<FProperty*, FMigratedPropertyStats>> ObjectPropertyStats;
+
+	void DumpStatsToLog()
+	{
+		int64 NameTableSize = 0;
+		{
+			TArray<uint8> NameTableData;
+			FMemoryWriter NameTableWriter(NameTableData);
+			UE::RemoteObject::Serialization::Disk::SerializeNameTables(NameTableWriter, ObjectData);
+			NameTableSize = NameTableData.Num();
+		}
+
+		const int64 TotalSize = NameTableSize + ObjectData.Bytes.Num();
+
+		UE_LOG(LogRemoteObject, Log, TEXT("%s Object Data stats for %s %s (Object Data toal: %d, total: %lld):"), *GetArchiveName(), *FRemoteObjectId(RootObject).ToString(), *GetFullNameSafe(RootObject), ObjectData.Bytes.Num(), TotalSize);
+		UE_LOG(LogRemoteObject, Log, TEXT("  Name Table total size: %lld (Names: %d, Paths: %d)"), NameTableSize, ObjectData.Names.Num(), ObjectData.PathNames.Num());
+
+		for (const TPair<FString, TMap<FProperty*, FMigratedPropertyStats>>& ObjectStatsPair : ObjectPropertyStats)
+		{
+			const FString& Obj = ObjectStatsPair.Key;
+			const TMap<FProperty*, FMigratedPropertyStats>& PropStats = ObjectStatsPair.Value;
+
+			int64 Total = 0;
+			for (const TPair<FProperty*, FMigratedPropertyStats>& PropertyPair : PropStats)
+			{
+				const FMigratedPropertyStats& Stats = PropertyPair.Value;
+				Total += Stats.Size;
+			}
+
+			UE_LOG(LogRemoteObject, Log, TEXT("  Data serialized for %s (Total: %lld):"), Obj.Len() ? *Obj : TEXT("Native Serialize"), Total);
+
+			for (const TPair<FProperty*, FMigratedPropertyStats>& PropertyPair : PropStats)
+			{
+				FProperty* Prop = PropertyPair.Key;
+				const FMigratedPropertyStats& Stats = PropertyPair.Value;
+				UE_LOG(LogRemoteObject, Log, TEXT("    %s: size: %lld, count: %lld"), Prop ? *GetFullNameSafe(Prop) : TEXT("Native Serialize"), Stats.Size, Stats.Count);
+			}
+		}
+	}
+#endif // #if UE_WITH_REMOTEOBJECTARCHIVE_DEBUGGING
+
+public:
+
+	TArchiveRemoteObjectBase(FRemoteObjectData& InObjectData, const TCHAR* InArchiveName)
+		: T(InObjectData.Bytes)
+		, ObjectData(InObjectData)
+		, ArchiveName(InArchiveName)
+	{
+		T::SetIsPersistent(false);
+		T::SetUseUnversionedPropertySerialization(true);
+		T::SetMigratingRemoteObjects(true);
+		T::SetPortFlags(PPF_MigratingRemoteObjects);
+	}
+
+	virtual ~TArchiveRemoteObjectBase()
+	{
+#if UE_WITH_REMOTEOBJECTARCHIVE_DEBUGGING
+		DumpStatsToLog();
 #endif
+	}
+
+	virtual FString GetArchiveName() const override
+	{
+		return ArchiveName;
+	}
+
+	void SetRootObject(UObject* InRoot)
+	{
+		RootObject = InRoot;
+	}
+
+	UObject* GetRootObject() const
+	{
+		return RootObject;
+	}
+
+#if UE_WITH_REMOTEOBJECTARCHIVE_DEBUGGING
+	void SetSerializationScope(const TCHAR* InScope)
+	{
+		SerializationScope = InScope;
+	}
+
+	virtual void Serialize(void* Data, int64 Num) override
+	{
+		if (ArchiveName == TEXT("RemoteObjectReader"))
+		{
+			static volatile int32 xx = 0;
+			xx++;
+		}
+		FProperty* CurrentProperty = nullptr;
+		if (const FArchiveSerializedPropertyChain* PropChain = GetSerializedPropertyChain())
+		{
+			if (PropChain->GetNumProperties())
+			{
+				CurrentProperty = PropChain->GetPropertyFromRoot(0);
+			}
+		}
+
+		TMap<FProperty*, FMigratedPropertyStats>& ObjectStats = ObjectPropertyStats.FindOrAdd(SerializationScope);
+		FMigratedPropertyStats& PropertyStats = ObjectStats.FindOrAdd(CurrentProperty);
+
+		const int64 StartPos = Tell();
+
+		T::Serialize(Data, Num);
+
+		PropertyStats.Size += Tell() - StartPos;
+		PropertyStats.Count++;
+	}
+#endif // UE_WITH_REMOTEOBJECTARCHIVE_DEBUGGING
+};
+
+class FArchiveRemoteObjectWriter : public TArchiveRemoteObjectBase<FMemoryWriter>
+{
+	TArray<UObject*>& ObjectsToSerialize;
+	const TSet<UObject*>& SerializedObjects;
+	TMap<FName, FNameIndexType> NameMap;
+	TMap<UObject*, FNameIndexType> PathNameMap;
 
 	void WriteObjectReference(const FObjectPtr& ObjPtr)
 	{
@@ -213,25 +324,13 @@ class FArchiveRemoteObjectWriter : public FMemoryWriter
 public:
 
 	FArchiveRemoteObjectWriter(UObject* InRootObject, FRemoteObjectData& OutObjectData, TArray<UObject*>& OutObjectsToSerialize, const TSet<UObject*>& InSerializedObjects)
-		: FMemoryWriter(OutObjectData.Bytes)
-		, RootObject(InRootObject)
-		, ObjectData(OutObjectData)
+		: TArchiveRemoteObjectBase<FMemoryWriter>(OutObjectData, TEXT("RemoteObjectWriter"))
 		, ObjectsToSerialize(OutObjectsToSerialize)
 		, SerializedObjects(InSerializedObjects)
 	{
 		static_assert(sizeof(FNameIndexType) == sizeof(FRemoteObjectPathName::FNameIndexType), "Name Index type must match the type used for storing remote object path name indices");
 
-		SetIsPersistent(false);
-		SetUseUnversionedPropertySerialization(true);
-		SetMigratingRemoteObjects(true);
-		SetPortFlags(PPF_MigratingRemoteObjects);
-	}
-
-	virtual ~FArchiveRemoteObjectWriter()
-	{
-#if UE_WITH_REMOTEOBJECTWRITER_DEBUGGING
-		DumpStatsToLog();
-#endif
+		SetRootObject(InRootObject);
 	}
 
 	FNameIndexType AddNameToNameMap(const FName& Name)
@@ -273,78 +372,6 @@ public:
 		return PathNameIndex;
 	}
 
-	FORCEINLINE UObject* GetRootObject() const
-	{
-		return RootObject;
-	}
-
-#if UE_WITH_REMOTEOBJECTWRITER_DEBUGGING
-	void DumpStatsToLog()
-	{
-		int64 NameTableSize = 0;
-		{
-			TArray<uint8> NameTableData;
-			FMemoryWriter NameTableWriter(NameTableData);			
-			UE::RemoteObject::Serialization::Disk::SerializeNameTables(NameTableWriter, ObjectData);
-			NameTableSize = NameTableData.Num();
-		}
-
-		const int64 TotalSize = NameTableSize + ObjectData.Bytes.Num();
-
-		UE_LOG(LogRemoteObject, Log, TEXT("Object Data stats for %s %s (Object Data toal: %d, total: %lld):"), *FRemoteObjectId(RootObject).ToString(), *GetFullNameSafe(RootObject), ObjectData.Bytes.Num(), TotalSize);
-		UE_LOG(LogRemoteObject, Log, TEXT("  Name Table total size: %lld (Names: %d, Paths: %d)"), NameTableSize, ObjectData.Names.Num(), ObjectData.PathNames.Num());
-
-		for (const TPair<FString, TMap<FProperty*, FSerializeStats>>& ObjectStatsPair : ObjectPropertyStats)
-		{
-			const FString& Obj = ObjectStatsPair.Key;
-			const TMap<FProperty*, FSerializeStats>& PropStats = ObjectStatsPair.Value;
-
-			int64 Total = 0;
-			for (const TPair<FProperty*, FSerializeStats>& PropertyPair : PropStats)
-			{
-				const FSerializeStats& Stats = PropertyPair.Value;
-				Total += Stats.Size;
-			}
-
-			UE_LOG(LogRemoteObject, Log, TEXT("  Data serialized for %s (Total: %lld):"), Obj.Len() ? *Obj : TEXT("Native Serialize"), Total);
-			
-			for (const TPair<FProperty*, FSerializeStats>& PropertyPair : PropStats)
-			{
-				FProperty* Prop = PropertyPair.Key;
-				const FSerializeStats& Stats = PropertyPair.Value;
-				UE_LOG(LogRemoteObject, Log, TEXT("    %s: size: %lld, count: %lld"), Prop ? *GetFullNameSafe(Prop) : TEXT("Native Serialize"), Stats.Size, Stats.Count);
-			}
-		}
-	}
-
-	void SetSerializationScope(const TCHAR* InScope)
-	{
-		SerializationScope = InScope;
-	}
-
-	virtual void Serialize(void* Data, int64 Num) override
-	{
-		FProperty* CurrentProperty = nullptr;
-		if (const FArchiveSerializedPropertyChain* PropChain = GetSerializedPropertyChain())
-		{
-			if (PropChain->GetNumProperties())
-			{
-				CurrentProperty = PropChain->GetPropertyFromRoot(0);
-			}
-		}
-
-		TMap<FProperty*, FSerializeStats>& ObjectStats = ObjectPropertyStats.FindOrAdd(SerializationScope);
-		FSerializeStats& PropertyStats = ObjectStats.FindOrAdd(CurrentProperty);
-
-		const int64 StartPos = Tell();
-		
-		FMemoryWriter::Serialize(Data, Num);
-
-		PropertyStats.Size += Tell() - StartPos;
-		PropertyStats.Count++;		
-	}
-#endif // UE_WITH_REMOTEOBJECTWRITER_DEBUGGING
-
 	virtual FArchive& operator<<(UObject*& Obj) override
 	{
 		WriteObjectReference(FObjectPtr(Obj));
@@ -386,34 +413,38 @@ public:
 	}
 };
 
-struct FRemoteObjectWriterScope
+template <typename T>
+struct TRemoteObjectArchiveScope
 {
-	FArchiveRemoteObjectWriter& Writer;
+	TArchiveRemoteObjectBase<T>& Ar;
 
-	FRemoteObjectWriterScope(FArchiveRemoteObjectWriter& InWriter, const TCHAR* InScope)
-		: Writer(InWriter)
+	TRemoteObjectArchiveScope(TArchiveRemoteObjectBase<T>& InAr, const TCHAR* InScope)
+		: Ar(InAr)
 	{
-#if UE_WITH_REMOTEOBJECTWRITER_DEBUGGING
-		Writer.SetSerializationScope(InScope);
+#if UE_WITH_REMOTEOBJECTARCHIVE_DEBUGGING
+		Ar.SetSerializationScope(InScope);
 #endif
 	}
-	FRemoteObjectWriterScope(FArchiveRemoteObjectWriter& InWriter, UObject* InObjectScope)
-		: Writer(InWriter)
+	TRemoteObjectArchiveScope(TArchiveRemoteObjectBase<T>& InAr, UObject* InObjectScope)
+		: Ar(InAr)
 	{
-#if UE_WITH_REMOTEOBJECTWRITER_DEBUGGING
-		UObject* Root = Writer.GetRootObject();
+#if UE_WITH_REMOTEOBJECTARCHIVE_DEBUGGING
+		UObject* Root = Ar.GetRootObject();
 		UObject* RootOuter = UE::CoreUObject::Private::FObjectHandleUtils::GetNonAccessTrackedOuterNoResolve(Root);
-		Writer.SetSerializationScope(*InObjectScope->GetFullName(RootOuter));
+		Ar.SetSerializationScope(*InObjectScope->GetFullName(RootOuter));
 #endif
 	}
-	~FRemoteObjectWriterScope()
+	~TRemoteObjectArchiveScope()
 	{
-#if UE_WITH_REMOTEOBJECTWRITER_DEBUGGING
-		Writer.SetSerializationScope(nullptr);
+#if UE_WITH_REMOTEOBJECTARCHIVE_DEBUGGING
+		Ar.SetSerializationScope(nullptr);
 #endif
 	}
 };
 
+using FRemoteObjectWriterScope = TRemoteObjectArchiveScope<FMemoryWriter>;
+using FRemoteObjectReaderScope = TRemoteObjectArchiveScope<FMemoryReader>;
+
 FArchive& operator<<(FArchive& Ar, FRemoteObjectHeader& Header)
 {
 	Ar << Header.Name;
@@ -443,7 +474,7 @@ void SerializeRemoteObject(UObject* Object, FArchive& Ar, FRemoteObjectHeader& O
 	Object->Serialize(Ar);
 }
 
-class FArchiveRemoteObjectReader : public FMemoryReader
+class FArchiveRemoteObjectReader : public TArchiveRemoteObjectBase<FMemoryReader>
 {
 	const TArray<FName>& Names;
 	const TArray<UObject*>& ResolvedPathNameObjects;
@@ -492,15 +523,11 @@ class FArchiveRemoteObjectReader : public FMemoryReader
 
 public:
 
-	FArchiveRemoteObjectReader(const FRemoteObjectData& InObjectData, const TArray<UObject*>& InResolvedPathNames)
-		: FMemoryReader(InObjectData.Bytes)
+	FArchiveRemoteObjectReader(FRemoteObjectData& InObjectData, const TArray<UObject*>& InResolvedPathNames)
+		: TArchiveRemoteObjectBase<FMemoryReader>(InObjectData, TEXT("RemoteObjectReader"))
 		, Names(InObjectData.Names)
 		, ResolvedPathNameObjects(InResolvedPathNames)
 	{
-		SetIsPersistent(false);
-		SetUseUnversionedPropertySerialization(true);
-		SetMigratingRemoteObjects(true);
-		SetPortFlags(PPF_MigratingRemoteObjects);
 	}
 
 	virtual FArchive& operator<<(UObject*& Obj) override
@@ -813,14 +840,20 @@ int32 DeserializeObjectData(FRemoteServerId OwnerServerId, FRemoteObjectData& Ob
 
 	FArchiveRemoteObjectReader Ar(ObjectData, ResolvedPathNameObjects);
 
-	Ar << Version;
-	Ar << RequestedObjectId;
-	Ar << HeaderOffset;
+	{
+		FRemoteObjectReaderScope Scope(Ar, TEXT("Header"));
+		Ar << Version;
+		Ar << RequestedObjectId;
+		Ar << HeaderOffset;
+	}
 
 	const int64 ObjectDataOffset = Ar.Tell();
 
 	Ar.Seek(HeaderOffset);
-	Ar << ObjectHeaders;
+	{
+		FRemoteObjectReaderScope Scope(Ar, TEXT("ObjectHeaders"));
+		Ar << ObjectHeaders;
+	}
 	Ar.Seek(ObjectDataOffset);	
 
 	if (ObjectHeaders.Num())
@@ -860,9 +893,11 @@ int32 DeserializeObjectData(FRemoteServerId OwnerServerId, FRemoteObjectData& Ob
 		}
 
 		// Deserialize all objects
+		Ar.SetRootObject(OutObjects.Num() ? OutObjects[0] : nullptr);
 		for (int32 ObjectIndex = 0; ObjectIndex < OutObjects.Num(); ++ObjectIndex)
 		{
 			UObject* Object = OutObjects[ObjectIndex];
+			FRemoteObjectReaderScope Scope(Ar, Object);
 			Object->Serialize(Ar);
 			if (RequestedObjectId == FRemoteObjectId(Object))
 			{
diff --git a/Engine/Source/Runtime/CoreUObject/Private/UObject/RemoteObjectSerialization.h b/Engine/Source/Runtime/CoreUObject/Private/UObject/RemoteObjectSerialization.h
--- a/Engine/Source/Runtime/CoreUObject/Private/UObject/RemoteObjectSerialization.h
+++ b/Engine/Source/Runtime/CoreUObject/Private/UObject/RemoteObjectSerialization.h
@@ -6,6 +6,10 @@
 
 namespace UE::RemoteObject::Handle { struct FRemoteObjectStub; }
 
+#ifndef UE_WITH_REMOTEOBJECTARCHIVE_DEBUGGING
+#define UE_WITH_REMOTEOBJECTARCHIVE_DEBUGGING 0
+#endif
+
 namespace UE::RemoteObject::Serialization
 {
 	/**
