Deprecated FSoftObjectPath functions which take an package/asset FName pair.

diff --git a/Engine/Plugins/Editor/ContentBrowser/ContentBrowserAliasDataSource/Source/ContentBrowserAliasDataSource/Private/ContentBrowserAliasDataSource.cpp b/Engine/Plugins/Editor/ContentBrowser/ContentBrowserAliasDataSource/Source/ContentBrowserAliasDataSource/Private/ContentBrowserAliasDataSource.cpp
--- a/Engine/Plugins/Editor/ContentBrowser/ContentBrowserAliasDataSource/Source/ContentBrowserAliasDataSource/Private/ContentBrowserAliasDataSource.cpp
+++ b/Engine/Plugins/Editor/ContentBrowser/ContentBrowserAliasDataSource/Source/ContentBrowserAliasDataSource/Private/ContentBrowserAliasDataSource.cpp
@@ -149,7 +149,7 @@ void UContentBrowserAliasDataSource::CompileFilter(const FName InPath, const FCo
 							{
 								if (const FAliasData* AliasData = AllAliases.Find(FContentBrowserUniqueAlias(ObjectPath, Alias)))
 								{
-									Callback(FSoftObjectPath::ConstructFromPackageAsset(AliasData->PackageName, AliasData->AssetData.AssetName));
+									Callback(FSoftObjectPath::ConstructFromAssetPath(FTopLevelAssetPath(AliasData->PackageName, AliasData->AssetData.AssetName)));
 								}
 							}
 						}
diff --git a/Engine/Plugins/Editor/ContentBrowser/ContentBrowserAliasDataSource/Source/ContentBrowserAliasDataSource/Public/ContentBrowserAliasDataSource.h b/Engine/Plugins/Editor/ContentBrowser/ContentBrowserAliasDataSource/Source/ContentBrowserAliasDataSource/Public/ContentBrowserAliasDataSource.h
--- a/Engine/Plugins/Editor/ContentBrowser/ContentBrowserAliasDataSource/Source/ContentBrowserAliasDataSource/Public/ContentBrowserAliasDataSource.h
+++ b/Engine/Plugins/Editor/ContentBrowser/ContentBrowserAliasDataSource/Source/ContentBrowserAliasDataSource/Public/ContentBrowserAliasDataSource.h
@@ -209,7 +209,7 @@ private:
 				PathBuilder << AssetNameBuilder.ToView();
 				PackageName = FName(PathBuilder.ToView());
 			}
-			ObjectPath = FSoftObjectPath(PackageName, FName(AssetNameBuilder.ToView()), {});
+			ObjectPath = FSoftObjectPath(FTopLevelAssetPath(PackageName, FName(AssetNameBuilder.ToView())));
 		}
 
 		/** The source asset for this alias */
diff --git a/Engine/Plugins/Experimental/PythonScriptPlugin/Source/PythonScriptPlugin/Private/PyWrapperObject.cpp b/Engine/Plugins/Experimental/PythonScriptPlugin/Source/PythonScriptPlugin/Private/PyWrapperObject.cpp
--- a/Engine/Plugins/Experimental/PythonScriptPlugin/Source/PythonScriptPlugin/Private/PyWrapperObject.cpp
+++ b/Engine/Plugins/Experimental/PythonScriptPlugin/Source/PythonScriptPlugin/Private/PyWrapperObject.cpp
@@ -1985,7 +1985,7 @@ void UPythonGeneratedClass::PostRename(UObject* OldOuter, const FName OldName)
 
 	if (PyType)
 	{
-		FPyWrapperTypeRegistry::Get().UnregisterWrappedClassType(FSoftObjectPath::ConstructFromPackageAsset(OldOuter->GetFName(), OldName), PyType, false);
+		FPyWrapperTypeRegistry::Get().UnregisterWrappedClassType(FSoftObjectPath::ConstructFromAssetPath(FTopLevelAssetPath(OldOuter->GetFName(), OldName)), PyType, false);
 		FPyWrapperTypeRegistry::Get().RegisterWrappedClassType(this, PyType, false);
 	}
 }
diff --git a/Engine/Plugins/Experimental/PythonScriptPlugin/Source/PythonScriptPlugin/Private/PyWrapperStruct.cpp b/Engine/Plugins/Experimental/PythonScriptPlugin/Source/PythonScriptPlugin/Private/PyWrapperStruct.cpp
--- a/Engine/Plugins/Experimental/PythonScriptPlugin/Source/PythonScriptPlugin/Private/PyWrapperStruct.cpp
+++ b/Engine/Plugins/Experimental/PythonScriptPlugin/Source/PythonScriptPlugin/Private/PyWrapperStruct.cpp
@@ -1645,7 +1645,7 @@ void UPythonGeneratedStruct::PostRename(UObject* OldOuter, const FName OldName)
 
 	if (PyType)
 	{
-		FPyWrapperTypeRegistry::Get().UnregisterWrappedStructType(FSoftObjectPath::ConstructFromPackageAsset(OldOuter->GetFName(), OldName), PyType, false);
+		FPyWrapperTypeRegistry::Get().UnregisterWrappedStructType(FSoftObjectPath::ConstructFromAssetPath(FTopLevelAssetPath(OldOuter->GetFName(), OldName)), PyType, false);
 		FPyWrapperTypeRegistry::Get().RegisterWrappedStructType(this, PyType, false);
 	}
 }
diff --git a/Engine/Plugins/Interchange/Runtime/Source/Pipelines/Private/InterchangeGenericAssetsPipeline.cpp b/Engine/Plugins/Interchange/Runtime/Source/Pipelines/Private/InterchangeGenericAssetsPipeline.cpp
--- a/Engine/Plugins/Interchange/Runtime/Source/Pipelines/Private/InterchangeGenericAssetsPipeline.cpp
+++ b/Engine/Plugins/Interchange/Runtime/Source/Pipelines/Private/InterchangeGenericAssetsPipeline.cpp
@@ -1114,7 +1114,7 @@ void UInterchangeGenericAssetsPipeline::ShowConflictDialog(const FGuid& Conflict
 								for (const FName& AssetDependencyName : AssetDependencies)
 								{
 									const FString PackageString = AssetDependencyName.ToString();
-									const FSoftObjectPath FullAssetPath(*PackageString, *FPackageName::GetLongPackageAssetName(PackageString), {});
+									const FSoftObjectPath FullAssetPath(FTopLevelAssetPath(*PackageString, *FPackageName::GetLongPackageAssetName(PackageString)));
 									FAssetData AssetData = AssetRegistryModule.Get().GetAssetByObjectPath(FullAssetPath);
 									if (AssetData.GetClass() != nullptr)
 									{
diff --git a/Engine/Source/Editor/UnrealEd/Private/Commandlets/ExportDialogueScriptCommandlet.cpp b/Engine/Source/Editor/UnrealEd/Private/Commandlets/ExportDialogueScriptCommandlet.cpp
--- a/Engine/Source/Editor/UnrealEd/Private/Commandlets/ExportDialogueScriptCommandlet.cpp
+++ b/Engine/Source/Editor/UnrealEd/Private/Commandlets/ExportDialogueScriptCommandlet.cpp
@@ -267,7 +267,7 @@ int32 UExportDialogueScriptCommandlet::Main(const FString& Params)
 			UDialogueWave* LocalizedDialogueWave = nullptr;
 			{
 				const FString LocalizedDialogueWavePackagePath = FPackageName::GetLocalizedPackagePath(AssetData.PackageName.ToString(), CultureName);
-				const FAssetData LocalizedDialogueWaveAssetData = AssetRegistry.GetAssetByObjectPath(FSoftObjectPath(*LocalizedDialogueWavePackagePath, AssetData.AssetName, {}));
+				const FAssetData LocalizedDialogueWaveAssetData = AssetRegistry.GetAssetByObjectPath(FSoftObjectPath(FTopLevelAssetPath(*LocalizedDialogueWavePackagePath, AssetData.AssetName)));
 				LocalizedDialogueWave = Cast<UDialogueWave>(LocalizedDialogueWaveAssetData.GetAsset());
 				if (LocalizedDialogueWave == DialogueWave)
 				{
diff --git a/Engine/Source/Editor/UnrealEd/Private/Cooker/CookDiagnostics.cpp b/Engine/Source/Editor/UnrealEd/Private/Cooker/CookDiagnostics.cpp
--- a/Engine/Source/Editor/UnrealEd/Private/Cooker/CookDiagnostics.cpp
+++ b/Engine/Source/Editor/UnrealEd/Private/Cooker/CookDiagnostics.cpp
@@ -233,7 +233,7 @@ void FDiagnostics::AnalyzeHiddenDependencies(UCookOnTheFlyServer& COTFS, FPackag
 		if (Resave.Realms[0].bValid)
 		{
 			FPackageReaderResults& Realm = Resave.Realms[0];
-			FSoftObjectPath PackagePath(UnsolicitedPackageName, NAME_None, FString());
+			FSoftObjectPath PackagePath(FTopLevelAssetPath(UnsolicitedPackageName, NAME_None), FString());
 			FSoftObjectPath PackageClassPath(UPackage::StaticClass());
 			bool bIsImportOfResavedPackage = false;
 			bool bIsEditorOnlyImportOfResavedPackage = false;
diff --git a/Engine/Source/Editor/UnrealEd/Private/Factories/ActorFactory.cpp b/Engine/Source/Editor/UnrealEd/Private/Factories/ActorFactory.cpp
--- a/Engine/Source/Editor/UnrealEd/Private/Factories/ActorFactory.cpp
+++ b/Engine/Source/Editor/UnrealEd/Private/Factories/ActorFactory.cpp
@@ -584,11 +584,11 @@ FQuat UActorFactoryStaticMesh::AlignObjectToSurfaceNormal(const FVector& InSurfa
 UActorFactoryBasicShape
 -----------------------------------------------------------------------------*/
 
-const FSoftObjectPath UActorFactoryBasicShape::BasicCube("/Engine/BasicShapes/Cube", "Cube", {});
-const FSoftObjectPath UActorFactoryBasicShape::BasicSphere("/Engine/BasicShapes/Sphere", "Sphere", {});
-const FSoftObjectPath UActorFactoryBasicShape::BasicCylinder("/Engine/BasicShapes/Cylinder", "Cylinder", {});
-const FSoftObjectPath UActorFactoryBasicShape::BasicCone("/Engine/BasicShapes/Cone", "Cone", {});
-const FSoftObjectPath UActorFactoryBasicShape::BasicPlane("/Engine/BasicShapes/Plane", "Plane", {});
+const FSoftObjectPath UActorFactoryBasicShape::BasicCube(FTopLevelAssetPath("/Engine/BasicShapes/Cube", "Cube"));
+const FSoftObjectPath UActorFactoryBasicShape::BasicSphere(FTopLevelAssetPath("/Engine/BasicShapes/Sphere", "Sphere"));
+const FSoftObjectPath UActorFactoryBasicShape::BasicCylinder(FTopLevelAssetPath("/Engine/BasicShapes/Cylinder", "Cylinder"));
+const FSoftObjectPath UActorFactoryBasicShape::BasicCone(FTopLevelAssetPath("/Engine/BasicShapes/Cone", "Cone"));
+const FSoftObjectPath UActorFactoryBasicShape::BasicPlane(FTopLevelAssetPath("/Engine/BasicShapes/Plane", "Plane"));
 
 UActorFactoryBasicShape::UActorFactoryBasicShape(const FObjectInitializer& ObjectInitializer)
 	: Super(ObjectInitializer)
diff --git a/Engine/Source/Editor/UnrealEd/Private/Fbx/FbxReimportDialogs.cpp b/Engine/Source/Editor/UnrealEd/Private/Fbx/FbxReimportDialogs.cpp
--- a/Engine/Source/Editor/UnrealEd/Private/Fbx/FbxReimportDialogs.cpp
+++ b/Engine/Source/Editor/UnrealEd/Private/Fbx/FbxReimportDialogs.cpp
@@ -149,7 +149,7 @@ void FFbxImporter::ShowFbxSkeletonConflictWindow(USkeletalMesh* SkeletalMesh, US
 				for (const FName& AssetDependencyName : AllDependencies)
 				{
 					const FString PackageString = AssetDependencyName.ToString();
-					const FSoftObjectPath FullAssetPath(*PackageString, *FPackageName::GetLongPackageAssetName(PackageString), {});
+					const FSoftObjectPath FullAssetPath(FTopLevelAssetPath(*PackageString, *FPackageName::GetLongPackageAssetName(PackageString)));
 					
 					FAssetData AssetData = AssetRegistryModule.Get().GetAssetByObjectPath(FullAssetPath);
 					if (AssetData.GetClass() != nullptr)
diff --git a/Engine/Source/Editor/UnrealEd/Private/MRUList.cpp b/Engine/Source/Editor/UnrealEd/Private/MRUList.cpp
--- a/Engine/Source/Editor/UnrealEd/Private/MRUList.cpp
+++ b/Engine/Source/Editor/UnrealEd/Private/MRUList.cpp
@@ -198,7 +198,7 @@ bool FMRUList::VerifyMRUFile(int32 InItem, FString& OutPackageName)
 
 	// Handle redirector
 	const FString OriginalPackageName = Items[InItem];
-	const FSoftObjectPath OriginalObjectPath = FSoftObjectPath(*OriginalPackageName, *FPackageName::GetShortName(OriginalPackageName), {});
+	const FSoftObjectPath OriginalObjectPath = FSoftObjectPath(FTopLevelAssetPath(*OriginalPackageName, *FPackageName::GetShortName(OriginalPackageName)));
 	FAssetRegistryModule& AssetRegistryModule = FModuleManager::LoadModuleChecked<FAssetRegistryModule>("AssetRegistry");
 	const FSoftObjectPath RedirectedObjectPath = AssetRegistryModule.Get().GetRedirectedObjectPath(OriginalObjectPath);
 
diff --git a/Engine/Source/Runtime/AssetRegistry/Private/AssetRegistryState.cpp b/Engine/Source/Runtime/AssetRegistry/Private/AssetRegistryState.cpp
--- a/Engine/Source/Runtime/AssetRegistry/Private/AssetRegistryState.cpp
+++ b/Engine/Source/Runtime/AssetRegistry/Private/AssetRegistryState.cpp
@@ -2297,7 +2297,7 @@ FDependsNode* FAssetRegistryState::ResolveRedirector(FDependsNode* InDependency,
 							return; // Already looked at this dependency node
 						}
 						const FAssetIdentifier& AssetId = InDepends->GetIdentifier();
-						FSoftObjectPath AssetPath(AssetId.PackageName, AssetId.ObjectName, FString());
+						FSoftObjectPath AssetPath(FTopLevelAssetPath(AssetId.PackageName, AssetId.ObjectName), FString());
 						if (InAllowedAssets.Contains(FCachedAssetKey(AssetPath)))
 						{
 							// This asset is in the allowed asset list, so take this as the redirect target
diff --git a/Engine/Source/Runtime/AssetRegistry/Private/PackageReader.cpp b/Engine/Source/Runtime/AssetRegistry/Private/PackageReader.cpp
--- a/Engine/Source/Runtime/AssetRegistry/Private/PackageReader.cpp
+++ b/Engine/Source/Runtime/AssetRegistry/Private/PackageReader.cpp
@@ -545,7 +545,7 @@ bool FPackageReader::ReadLinkerObjects(TMap<FSoftObjectPath, FObjectData>& OutEx
 	{
 		FObjectData Data;
 		const FObjectImport& Import = ImportMap[Index];
-		Data.ClassPath = FSoftObjectPath::ConstructFromPackageAsset(Import.ClassPackage, Import.ClassName);
+		Data.ClassPath = FSoftObjectPath::ConstructFromAssetPath(FTopLevelAssetPath(Import.ClassPackage, Import.ClassName));
 		Data.bUsedInGame = ImportUsedInGame[Index];
 		OutImports.Add(ImportsPaths[Index], MoveTemp(Data));
 	}
@@ -1510,25 +1510,27 @@ void FPackageReader::ConvertLinkerTableToPaths(FName InPackageName, TArray<FObje
 			{
 				if (Index.IsExport())
 				{
-					Result = FSoftObjectPath::ConstructFromPackageAsset(InPackageName, ObjectName);
+					Result = FSoftObjectPath::ConstructFromAssetPath(FTopLevelAssetPath( InPackageName, ObjectName));
 				}
 				else
 				{
-					Result = FSoftObjectPath::ConstructFromPackageAsset(ObjectName, NAME_None);
+					Result = FSoftObjectPath::ConstructFromAssetPath(FTopLevelAssetPath(ObjectName, NAME_None));
 				}
 			}
 			else if (ParentPath.GetAssetFName().IsNone())
 			{
-				Result = FSoftObjectPath::ConstructFromPackageAsset(ParentPath.GetLongPackageFName(), ObjectName);
+				Result = FSoftObjectPath::ConstructFromAssetPath(FTopLevelAssetPath(ParentPath.GetLongPackageFName(), ObjectName));
 			}
 			else if (ParentPath.GetSubPathString().IsEmpty())
 			{
-				Result = FSoftObjectPath(ParentPath.GetLongPackageFName(), ParentPath.GetAssetFName(),
+				Result = FSoftObjectPath::ConstructFromAssetPathAndSubpath(
+					FTopLevelAssetPath(ParentPath.GetLongPackageFName(), ParentPath.GetAssetFName()),
 					ObjectName.ToString());
 			}
 			else
 			{
-				Result = FSoftObjectPath(ParentPath.GetLongPackageFName(), ParentPath.GetAssetFName(),
+				Result = FSoftObjectPath::ConstructFromAssetPathAndSubpath(
+					FTopLevelAssetPath(ParentPath.GetLongPackageFName(), ParentPath.GetAssetFName()),
 					ParentPath.GetSubPathString() + TEXT(".") + ObjectName.ToString());
 			}
 		}
diff --git a/Engine/Source/Runtime/CoreUObject/Private/AssetRegistry/AssetData.cpp b/Engine/Source/Runtime/CoreUObject/Private/AssetRegistry/AssetData.cpp
--- a/Engine/Source/Runtime/CoreUObject/Private/AssetRegistry/AssetData.cpp
+++ b/Engine/Source/Runtime/CoreUObject/Private/AssetRegistry/AssetData.cpp
@@ -398,7 +398,7 @@ FSoftObjectPath FAssetData::GetSoftObjectPath() const
 {
 	if (IsTopLevelAsset())
 	{
-		return FSoftObjectPath::ConstructFromPackageAsset(PackageName, AssetName);
+		return FSoftObjectPath::ConstructFromAssetPath(FTopLevelAssetPath(PackageName, AssetName));
 	}
 	else
 	{
diff --git a/Engine/Source/Runtime/CoreUObject/Private/Misc/RedirectCollector.cpp b/Engine/Source/Runtime/CoreUObject/Private/Misc/RedirectCollector.cpp
--- a/Engine/Source/Runtime/CoreUObject/Private/Misc/RedirectCollector.cpp
+++ b/Engine/Source/Runtime/CoreUObject/Private/Misc/RedirectCollector.cpp
@@ -114,7 +114,7 @@ void FRedirectCollector::CollectSavedSoftPackageReferences(FName ReferencingPack
 	TArray<FSoftObjectPathProperty, TInlineAllocator<4>> SoftObjectPathArray;
 	Algo::Transform(PackageNames, SoftObjectPathArray, [ReferencingPackage, bEditorOnlyReferences](const FName& PackageName)
 		{
-			return FSoftObjectPathProperty(FSoftObjectPath::ConstructFromPackageAsset(PackageName, NAME_None), NAME_None, bEditorOnlyReferences);
+			return FSoftObjectPathProperty(FSoftObjectPath::ConstructFromAssetPath(FTopLevelAssetPath(PackageName, NAME_None)), NAME_None, bEditorOnlyReferences);
 		});
 
 	FScopeLock ScopeLock(&CriticalSection);
diff --git a/Engine/Source/Runtime/CoreUObject/Private/UObject/LinkerInstancingContext.cpp b/Engine/Source/Runtime/CoreUObject/Private/UObject/LinkerInstancingContext.cpp
--- a/Engine/Source/Runtime/CoreUObject/Private/UObject/LinkerInstancingContext.cpp
+++ b/Engine/Source/Runtime/CoreUObject/Private/UObject/LinkerInstancingContext.cpp
@@ -401,7 +401,7 @@ bool FLinkerInstancedPackageMap::FixupSoftObjectPath(FSoftObjectPath& InOutSoftO
 	{
 		if (FName LongPackageName = InOutSoftObjectPath.GetLongPackageFName(), RemappedPackage = RemapPackage(LongPackageName); RemappedPackage != LongPackageName)
 		{
-			InOutSoftObjectPath = FSoftObjectPath(RemappedPackage, InOutSoftObjectPath.GetAssetFName(), InOutSoftObjectPath.GetSubPathString());
+			InOutSoftObjectPath = FSoftObjectPath(FTopLevelAssetPath(RemappedPackage, InOutSoftObjectPath.GetAssetFName()), InOutSoftObjectPath.GetSubPathString());
 			return true;
 		}
 		else if (!GeneratedPackagesFolder.IsEmpty())
diff --git a/Engine/Source/Runtime/CoreUObject/Private/UObject/SavePackage/SavePackageUtilities.cpp b/Engine/Source/Runtime/CoreUObject/Private/UObject/SavePackage/SavePackageUtilities.cpp
--- a/Engine/Source/Runtime/CoreUObject/Private/UObject/SavePackage/SavePackageUtilities.cpp
+++ b/Engine/Source/Runtime/CoreUObject/Private/UObject/SavePackage/SavePackageUtilities.cpp
@@ -801,7 +801,7 @@ void HarvestCookRuntimeDependencies(FObjectSaveContextData& Data, UObject* Harve
 			// Ignore native imports; we don't need to mark them for cooking
 			continue;
 		}
-		FSoftObjectPath PackageSoftPath(PackageName, NAME_None, FString());
+		FSoftObjectPath PackageSoftPath(FTopLevelAssetPath(PackageName, NAME_None));
 		Data.CookRuntimeDependencies.Add(MoveTemp(PackageSoftPath));
 	}
 }
diff --git a/Engine/Source/Runtime/CoreUObject/Private/UObject/SoftObjectPath.cpp b/Engine/Source/Runtime/CoreUObject/Private/UObject/SoftObjectPath.cpp
--- a/Engine/Source/Runtime/CoreUObject/Private/UObject/SoftObjectPath.cpp
+++ b/Engine/Source/Runtime/CoreUObject/Private/UObject/SoftObjectPath.cpp
@@ -35,17 +35,27 @@ FName FSoftObjectPath::NAME_Untracked(TEXT("Untracked"));
 /** Static methods for more meaningful construction sites. */
 FSoftObjectPath FSoftObjectPath::ConstructFromPackageAssetSubpath(FName InPackageName, FName InAssetName, const FString& InSubPathString)
 {
-	return FSoftObjectPath(InPackageName, InAssetName, InSubPathString);
+	return FSoftObjectPath(FTopLevelAssetPath(InPackageName, InAssetName), InSubPathString);
 }
 
 FSoftObjectPath FSoftObjectPath::ConstructFromPackageAssetSubpath(FName InPackageName, FName InAssetName, FString&& InSubPathString)
 {
-	return FSoftObjectPath(InPackageName, InAssetName, MoveTemp(InSubPathString));
+	return FSoftObjectPath(FTopLevelAssetPath(InPackageName, InAssetName), MoveTemp(InSubPathString));
 }
 
 FSoftObjectPath FSoftObjectPath::ConstructFromPackageAsset(FName InPackageName, FName InAssetName)
 {
-	return FSoftObjectPath(InPackageName, InAssetName, {});
+	return FSoftObjectPath(FTopLevelAssetPath(InPackageName, InAssetName), {});
+}
+
+FSoftObjectPath FSoftObjectPath::ConstructFromAssetPathAndSubpath(FTopLevelAssetPath InAssetPath, const FString& InSubPathString)
+{
+	return FSoftObjectPath(InAssetPath, InSubPathString);
+}
+
+FSoftObjectPath FSoftObjectPath::ConstructFromAssetPathAndSubpath(FTopLevelAssetPath InAssetPath, FString&& InSubPathString)
+{
+	return FSoftObjectPath(InAssetPath, MoveTemp(InSubPathString));
 }
 
 FSoftObjectPath FSoftObjectPath::ConstructFromAssetPath(FTopLevelAssetPath InAssetPath)
diff --git a/Engine/Source/Runtime/CoreUObject/Public/UObject/SoftObjectPath.h b/Engine/Source/Runtime/CoreUObject/Public/UObject/SoftObjectPath.h
--- a/Engine/Source/Runtime/CoreUObject/Public/UObject/SoftObjectPath.h
+++ b/Engine/Source/Runtime/CoreUObject/Public/UObject/SoftObjectPath.h
@@ -60,6 +60,7 @@ struct FSoftObjectPath
 	{
 		SetPath(InAssetPath, MoveTemp(InSubPathString));
 	}
+	UE_DEPRECATED(5.6, "FSoftObjectPath which takes a package name and asset name has been deprecated - please pass FTopLevelAssetPath instead.")
 	[[nodiscard]] explicit FSoftObjectPath(FName InPackageName, FName InAssetName, FString InSubPathString)
 	{
 		SetPath(FTopLevelAssetPath(InPackageName, InAssetName), MoveTemp(InSubPathString));
@@ -111,9 +112,14 @@ struct FSoftObjectPath
 	}
 
 	/** Static methods for more meaningful construction sites. */
+	UE_DEPRECATED(5.6, "FSoftObjectPath::ConstructFromPackageAssetSubpath has been deprecated - please use ConstructFromAssetPathAndSubpath instead.")
 	[[nodiscard]] COREUOBJECT_API static FSoftObjectPath ConstructFromPackageAssetSubpath(FName InPackageName, FName InAssetName, const FString& InSubPathString);
+	UE_DEPRECATED(5.6, "FSoftObjectPath::ConstructFromPackageAssetSubpath has been deprecated - please use ConstructFromAssetPathAndSubpath instead.")
 	[[nodiscard]] COREUOBJECT_API static FSoftObjectPath ConstructFromPackageAssetSubpath(FName InPackageName, FName InAssetName, FString&& InSubPathString);
+	UE_DEPRECATED(5.6, "FSoftObjectPath::ConstructFromPackageAsset has been deprecated - please use ConstructFromAssetPath instead.")
 	[[nodiscard]] COREUOBJECT_API static FSoftObjectPath ConstructFromPackageAsset(FName InPackageName, FName InAssetName);
+	[[nodiscard]] COREUOBJECT_API static FSoftObjectPath ConstructFromAssetPathAndSubpath(FTopLevelAssetPath InAssetPath, const FString& InSubPathString);
+	[[nodiscard]] COREUOBJECT_API static FSoftObjectPath ConstructFromAssetPathAndSubpath(FTopLevelAssetPath InAssetPath, FString&& InSubPathString);
 	[[nodiscard]] COREUOBJECT_API static FSoftObjectPath ConstructFromAssetPath(FTopLevelAssetPath InAssetPath);
 	[[nodiscard]] COREUOBJECT_API static FSoftObjectPath ConstructFromStringPath(FString&& InPath);
 	[[nodiscard]] COREUOBJECT_API static FSoftObjectPath ConstructFromStringPath(FStringView InPath);
diff --git a/Engine/Source/Runtime/Engine/Private/BlueprintGeneratedClass.cpp b/Engine/Source/Runtime/Engine/Private/BlueprintGeneratedClass.cpp
--- a/Engine/Source/Runtime/Engine/Private/BlueprintGeneratedClass.cpp
+++ b/Engine/Source/Runtime/Engine/Private/BlueprintGeneratedClass.cpp
@@ -2601,7 +2601,7 @@ void UBlueprintGeneratedClass::Serialize(FArchive& Ar)
 				// Ignore native imports; we don't need to mark them for cooking
 				continue;
 			}
-			FSoftObjectPath PackageSoftPath(PackageName, NAME_None, FString());
+			FSoftObjectPath PackageSoftPath(FTopLevelAssetPath(PackageName, NAME_None));
 			Ar << PackageSoftPath;
 		}
 	}
diff --git a/Engine/Source/Runtime/Engine/Private/ObjectLibrary.cpp b/Engine/Source/Runtime/Engine/Private/ObjectLibrary.cpp
--- a/Engine/Source/Runtime/Engine/Private/ObjectLibrary.cpp
+++ b/Engine/Source/Runtime/Engine/Private/ObjectLibrary.cpp
@@ -507,7 +507,7 @@ int32 UObjectLibrary::LoadAssetsFromAssetData()
 		for (int32 AssetIdx = 0; AssetIdx < AssetDataList.Num(); AssetIdx++)
 		{
 			FAssetData& Data = AssetDataList[AssetIdx];
-			AssetsToStream.AddUnique(FSoftObjectPath(Data.PackageName, {}, {}));
+			AssetsToStream.AddUnique(FSoftObjectPath(FTopLevelAssetPath(Data.PackageName, NAME_None)));
 		}
 
 		if (AssetsToStream.Num())
diff --git a/Engine/Source/Runtime/Engine/Private/WorldComposition.cpp b/Engine/Source/Runtime/Engine/Private/WorldComposition.cpp
--- a/Engine/Source/Runtime/Engine/Private/WorldComposition.cpp
+++ b/Engine/Source/Runtime/Engine/Private/WorldComposition.cpp
@@ -183,7 +183,7 @@ struct FWorldTilesGatherer
 		for (const FString& FullPath : MapFilesToConsider)
 		{
 			FString TilePackageName = FPackageName::FilenameToLongPackageName(FullPath);
-			const FSoftObjectPath TilePackagePath(FName(*TilePackageName), FName(*FPackageName::GetLongPackageAssetName(TilePackageName)), FString{});
+			const FSoftObjectPath TilePackagePath(FTopLevelAssetPath(FName(*TilePackageName), FName(*FPackageName::GetLongPackageAssetName(TilePackageName))));
 
 			const FAssetData MapAssetData = AssetRegistry.GetAssetByObjectPath(TilePackagePath);
 
diff --git a/Engine/Source/Runtime/Interchange/Engine/Private/InterchangeSceneImportAsset.cpp b/Engine/Source/Runtime/Interchange/Engine/Private/InterchangeSceneImportAsset.cpp
--- a/Engine/Source/Runtime/Interchange/Engine/Private/InterchangeSceneImportAsset.cpp
+++ b/Engine/Source/Runtime/Interchange/Engine/Private/InterchangeSceneImportAsset.cpp
@@ -110,7 +110,7 @@ void UInterchangeSceneImportAsset::OnPostWorldRename(UWorld* World)
 		if(ensure(FactoryNode))
 		{
 			const FString DisplayName = FactoryNode->GetDisplayLabel();
-			const FSoftObjectPath ObjectPath{ FName(NewWorldPath), FName(NewWorldName), NewPrefix + DisplayName };
+			const FSoftObjectPath ObjectPath{ FTopLevelAssetPath(FName(NewWorldPath), FName(NewWorldName)), NewPrefix + DisplayName };
 
 			FactoryNode->SetCustomReferenceObject(ObjectPath);
 			SceneObjects.Add(ObjectPath, UniqueID);
@@ -224,7 +224,7 @@ void UInterchangeSceneImportAsset::UpdateSceneObjects()
 UObject* UInterchangeSceneImportAsset::GetSceneObject(const FString& PackageName, const FString& AssetName, const FString& SubPathString) const
 {
 #if WITH_EDITORONLY_DATA
-	const FSoftObjectPath ObjectPath(FName(PackageName), FName(AssetName), SubPathString);
+	const FSoftObjectPath ObjectPath( FTopLevelAssetPath(FName(PackageName), FName(AssetName)), SubPathString);
 	const FName ObjectSubPathBaseName = FActorSpawnUtils::GetBaseName(*SubPathString);
 
 	UObject* SceneObject = nullptr;
@@ -282,7 +282,7 @@ UObject* UInterchangeSceneImportAsset::GetSceneObject(const FString& PackageName
 const UInterchangeFactoryBaseNode* UInterchangeSceneImportAsset::GetFactoryNode(const FString& PackageName, const FString& AssetName, const FString& SubPathString) const
 {
 #if WITH_EDITORONLY_DATA
-	const FSoftObjectPath ObjectPath(FName(PackageName), FName(AssetName), SubPathString);
+	const FSoftObjectPath ObjectPath(FTopLevelAssetPath(FName(PackageName), FName(AssetName)), SubPathString);
 	return GetFactoryNode(ObjectPath);
 #else
 	return nullptr;
diff --git a/Engine/Source/Runtime/TypedElementFramework/Tests/Elements/Common/TypedElementMapKeyTests.cpp b/Engine/Source/Runtime/TypedElementFramework/Tests/Elements/Common/TypedElementMapKeyTests.cpp
--- a/Engine/Source/Runtime/TypedElementFramework/Tests/Elements/Common/TypedElementMapKeyTests.cpp
+++ b/Engine/Source/Runtime/TypedElementFramework/Tests/Elements/Common/TypedElementMapKeyTests.cpp
@@ -8,7 +8,7 @@
 namespace UE::Editor::DataStorage::Tests::Private
 {
 	FName TestFName = FName("FMapKey_TestName");
-	FSoftObjectPath TestPath = FSoftObjectPath(FName("FMapKey_TestPackage"), FName("FMapKey_TestAsset"), TEXT("FMapKey_TestSubPath"));
+	FSoftObjectPath TestPath(FTopLevelAssetPath(FName("FMapKey_TestPackage"), FName("FMapKey_TestAsset")), TEXT("FMapKey_TestSubPath"));
 
 	void CopyKeyTest(const FMapKey& Original)
 	{
