Print referenced objects' pathnames when dumping remote object archive stats

diff --git a/Engine/Source/Runtime/CoreUObject/Private/UObject/RemoteObjectSerialization.cpp b/Engine/Source/Runtime/CoreUObject/Private/UObject/RemoteObjectSerialization.cpp
--- a/Engine/Source/Runtime/CoreUObject/Private/UObject/RemoteObjectSerialization.cpp
+++ b/Engine/Source/Runtime/CoreUObject/Private/UObject/RemoteObjectSerialization.cpp
@@ -144,6 +144,29 @@ protected:
 	FString SerializationScope;
 	TMap<FString, TMap<FProperty*, FMigratedPropertyStats>> ObjectPropertyStats;
 
+	FString PathNameToString(const FRemoteObjectPathName& PathName) const
+	{
+		TStringBuilder<256> Result;
+		for (int32 PathSegment = PathName.Names.Num() - 1; PathSegment >= 0; --PathSegment)
+		{
+			ObjectData.Names[PathName.Names[PathSegment]].AppendString(Result);
+			if (PathSegment > 0)
+			{
+				if (PathSegment != (PathName.Names.Num() - 2))
+				{
+					Result << TEXT('.');
+				}
+				else
+				{
+					// at this point the constructed path consists of /Package/Name.ObjectName so we need to use the subobject delimiter
+					// see UObjectBaseUtility::GetPathName
+					Result << SUBOBJECT_DELIMITER_CHAR;
+				}
+			}
+		}
+		return Result.ToString();
+	}
+
 	void DumpStatsToLog()
 	{
 		int64 NameTableSize = 0;
@@ -159,6 +182,17 @@ protected:
 		UE_LOG(LogRemoteObject, Log, TEXT("%s Object Data stats for %s %s (Object Data toal: %d, total: %lld):"), *GetArchiveName(), *FRemoteObjectId(RootObject).ToString(), *GetFullNameSafe(RootObject), ObjectData.Bytes.Num(), TotalSize);
 		UE_LOG(LogRemoteObject, Log, TEXT("  Name Table total size: %lld (Names: %d, Paths: %d)"), NameTableSize, ObjectData.Names.Num(), ObjectData.PathNames.Num());
 
+		TArray<FString> SortedPathNames;
+		for (const FRemoteObjectPathName& PathName : ObjectData.PathNames)
+		{
+			SortedPathNames.Add(PathNameToString(PathName));
+		}
+		SortedPathNames.Sort();
+		for (const FString& PathName : SortedPathNames)
+		{
+			UE_LOG(LogRemoteObject, Log, TEXT("    %s"), *PathName);
+		}
+
 		for (const TPair<FString, TMap<FProperty*, FMigratedPropertyStats>>& ObjectStatsPair : ObjectPropertyStats)
 		{
 			const FString& Obj = ObjectStatsPair.Key;
@@ -226,13 +260,8 @@ public:
 
 	virtual void Serialize(void* Data, int64 Num) override
 	{
-		if (ArchiveName == TEXT("RemoteObjectReader"))
-		{
-			static volatile int32 xx = 0;
-			xx++;
-		}
 		FProperty* CurrentProperty = nullptr;
-		if (const FArchiveSerializedPropertyChain* PropChain = GetSerializedPropertyChain())
+		if (const FArchiveSerializedPropertyChain* PropChain = T::GetSerializedPropertyChain())
 		{
 			if (PropChain->GetNumProperties())
 			{
@@ -243,11 +272,11 @@ public:
 		TMap<FProperty*, FMigratedPropertyStats>& ObjectStats = ObjectPropertyStats.FindOrAdd(SerializationScope);
 		FMigratedPropertyStats& PropertyStats = ObjectStats.FindOrAdd(CurrentProperty);
 
-		const int64 StartPos = Tell();
+		const int64 StartPos = T::Tell();
 
 		T::Serialize(Data, Num);
 
-		PropertyStats.Size += Tell() - StartPos;
+		PropertyStats.Size += T::Tell() - StartPos;
 		PropertyStats.Count++;
 	}
 #endif // UE_WITH_REMOTEOBJECTARCHIVE_DEBUGGING
