Change how the transactionally safe async loader works to only cause retries to happen when flushing.

This means that only synchronous loads (the bad ones, that people shouldn't be doing) will cause transaction retries.

I've also changed the caching mechanism so that even packages loaded outside a transaction are now visible in the package cache for transactions to be able to see them.

diff --git a/Engine/Source/Programs/AutoRTFMTests/Private/UECoreTests.cpp b/Engine/Source/Programs/AutoRTFMTests/Private/UECoreTests.cpp
--- a/Engine/Source/Programs/AutoRTFMTests/Private/UECoreTests.cpp
+++ b/Engine/Source/Programs/AutoRTFMTests/Private/UECoreTests.cpp
@@ -2469,6 +2469,33 @@ TEST_CASE("UECore.AsyncLoading")
 		});
 	}
 
+	SECTION("FlushAsyncLoading Empty")
+	{
+		FString Name(FString::Printf(TEXT("/AutoRTFMTestPackage%d"), __LINE__));
+		int32 RequestId = LoadPackageAsync(Name);
+
+		AutoRTFM::Testing::Commit([&]
+		{
+			FlushAsyncLoading();
+		});
+	}
+
+	SECTION("FlushAsyncLoading One In One Out")
+	{
+		FString Name(FString::Printf(TEXT("/AutoRTFMTestPackage%d"), __LINE__));
+		int32 RequestId1 = LoadPackageAsync(Name);
+
+		AutoRTFM::Testing::Commit([&]
+		{
+			FString Name(FString::Printf(TEXT("/AutoRTFMTestPackage%d"), __LINE__));
+			int32 RequestId2 = LoadPackageAsync(Name);
+			TArray<int32> RequestIds;
+			RequestIds.Add(RequestId1);
+			RequestIds.Add(RequestId2);
+			FlushAsyncLoading(RequestIds);
+		});
+	}
+
 	SECTION("CompletionDelegate is called closed")
 	{
 		AutoRTFM::Testing::Abort([&]
@@ -2585,7 +2612,7 @@ TEST_CASE("UECore.AsyncLoading")
 			RequestIds.Add(LoadPackageAsync(FString::Printf(TEXT("/AutoRTFMTestPackage%d"), __LINE__), CompletionDelegate));
 
 			FlushAsyncLoading(RequestIds);
-			REQUIRE(6 == NumCompletionCallbacks);
+			REQUIRE(3 == NumCompletionCallbacks);
 		});
 	}
 
@@ -2609,7 +2636,7 @@ TEST_CASE("UECore.AsyncLoading")
 			RequestIds.Add(LoadPackageAsync(FString::Printf(TEXT("/AutoRTFMTestPackage%d"), __LINE__), CompletionDelegate));
 
 			FlushAsyncLoading(RequestIds);
-			REQUIRE(6 == NumCompletionCallbacks);
+			REQUIRE(3 == NumCompletionCallbacks);
 
 			AutoRTFM::AbortTransaction();
 		});
diff --git a/Engine/Source/Runtime/CoreUObject/Private/Serialization/TransactionallySafeAsyncLoading.cpp b/Engine/Source/Runtime/CoreUObject/Private/Serialization/TransactionallySafeAsyncLoading.cpp
--- a/Engine/Source/Runtime/CoreUObject/Private/Serialization/TransactionallySafeAsyncLoading.cpp
+++ b/Engine/Source/Runtime/CoreUObject/Private/Serialization/TransactionallySafeAsyncLoading.cpp
@@ -5,7 +5,41 @@
 #include "AutoRTFM/AutoRTFM.h"
 #include "Containers/Map.h"
 #include "UObject/LinkerInstancingContext.h"
-
+#include "Async/Mutex.h"
+#include "Misc/ScopeLock.h"
+
+// A transactionally safe async package loader that wraps an underlying actual
+// async package loader but allows it to be interacted with safely while also
+// inside a transaction.
+//
+// The fundamental issue with async loading is, by its nature, it is touching
+// deep bits of Unreal Engine. It's going to be creating a bunch of UObject's
+// and thus touch all the deep and gnarly stuff that backs that system. This
+// does not meld well with how we handle modifications to these *same* core
+// bits of the engine from within AutoRTFM - namely by using transactionally
+// safe locks that are held until the transaction completes. If we just tried
+// to use the existing async package loader we'd deadlock because AutoRTFM
+// would be holding locks that the async loader would be trying to take, and
+// the AutoRTFM transaction could be blocked on the async loader while trying
+// to flush load a package. Nasty!
+//
+// We get around this issue by making it so that if a user does anything that
+// requires flushing the async loader (EG. if they are doing bad things like
+// synchronous loads of packages...), we have to abort the entire transaction
+// nest to release any locks we hold, flush the async loader, then retry the
+// transaction. To get this to work we keep a cache of loaded packages so
+// that we do not have to interact with the underlying async loader for a
+// given package after we have successfully loaded that package.
+//
+// The two fundamental changes when in a transaction are:
+// - When loading a package in the closed we check whether the package cache
+//   already has the package, and if so we just return that package object. If
+//   the package is not in the cache we instead just remember to load the
+//   package when the transaction commits. Loading a package is an async action
+//   so it is fine for us to just defer the async nature of it.
+// - When flushing a previous request-id, we check if the request-id was one
+//   that we know has already been flushed. If not, we need to abort and retry
+//   the transaction with the flush happening in between the abort and the retry.
 class FTransactionallySafeAsyncPackageLoader final : public IAsyncPackageLoader
 {
 public:
@@ -51,74 +85,85 @@ public:
 	{
 		if (AutoRTFM::IsClosed())
 		{
-			FString PackagePathAsString = PackagePath.GetDebugNameWithExtension();
-			FMapPayload* const Payload = RetriedPackageCache.Find(PackagePathAsString);
-			if (!Payload)
+			return AutoRTFM::Open([&]
 			{
-				// We do not have the package cached and ready to be used in a
-				// transaction. So we need to do the expensive (but necessary)
-				// operation of retrying the transaction, and loading the
-				// package in the gap between the initial transaction aborting
-				// and the retried one beginning.
-				AutoRTFM::Open([&]
+				FString PackagePathAsString = PackagePath.GetDebugNameWithExtension();
+
+				TOptional<FMapPayload> MapPayload;
+
+				{
+					UE::TScopeLock _(Mutex);
+
+					if (FMapPayload* const Result = PackageCache.Find(PackagePathAsString))
+					{
+						// Need to copy it because as soon as we give up our mutex the package cache
+						// could have been resized thus invalidating our underlying memory allocation.
+						MapPayload = *Result;
+					}
+				}
+
+				if (!MapPayload)
 				{
+					// We do not have the package cached and ready to be used in our transaction. So
+					// we copy the required state into the open, and cache that so we can process it
+					// later when the transaction has completed, or if we call `FlushLoading`, during
+					// transaction retry.
+				
 					FPackagePath PackagePathClone = PackagePath;
 					FName CustomPackageNameClone = FName(CustomPackageName.ToString());
-					TSharedPtr<FLinkerInstancingContext> LinkerInstancingContextClone;
+					TOptional<FLinkerInstancingContext> LinkerInstancingContextClone;
 
 					if (InInstancingContext)
 					{
-						LinkerInstancingContextClone = MakeShared<FLinkerInstancingContext>(FLinkerInstancingContext::DuplicateContext(*InInstancingContext));
+						LinkerInstancingContextClone = FLinkerInstancingContext::DuplicateContext(*InInstancingContext);
 					}
 
-					const AutoRTFM::EContextStatus Status = AutoRTFM::Close([&]
+					const bool bFirst = TransactionallyDeferredLoadPackages.IsEmpty();
+
+					FTransactionallyDeferredLoadPackagePayload Payload;
+					Payload.PackagePath = PackagePathClone;
+					Payload.CustomPackageName = CustomPackageNameClone;
+					Payload.CompletionDelegate = InCompletionDelegate;
+					Payload.PackageFlags = InPackageFlags;
+					Payload.PIEInstanceID = InPIEInstanceID;
+					Payload.PackagePriority = InPackagePriority;
+					Payload.InstancingContext = LinkerInstancingContextClone;
+					Payload.LoadFlags = InLoadFlags;
+					TransactionallyDeferredLoadPackages.Push(Payload);
+
+					// If we are the first, we need to register our handlers so that
+					// on-commit we issue the loads, and on-abort we purge the deferred
+					// array of them.
+					if (bFirst)
 					{
-						AutoRTFM::CascadingRetryTransaction([this, PackagePathClone, PackagePathAsString, CustomPackageNameClone, InPackageFlags, InPIEInstanceID, InPackagePriority, LinkerInstancingContextClone = ::MoveTemp(LinkerInstancingContextClone), InLoadFlags]
+						const AutoRTFM::EContextStatus Status = AutoRTFM::Close([&]
 						{
-							UE_LOG(LogStreaming, Display, TEXT("A call to LoadPackage within a transaction is causing Verse to abort and retry."));
-
-							FLoadPackageAsyncDelegate CompletionDelegate;
-							CompletionDelegate.BindLambda([this, PackagePathAsString](const FName& PackageName, UPackage* const Package, EAsyncLoadingResult::Type Result)
-							{
-								this->RetriedPackageCache.Add(PackagePathAsString, FMapPayload(PackageName, Package, Result));
-							});
-
-							const int32 RequestId = this->WrappedPackageLoader->LoadPackage(
-								PackagePathClone,
-								CustomPackageNameClone,
-								CompletionDelegate,
-								InPackageFlags,
-								InPIEInstanceID,
-								InPackagePriority,
-								LinkerInstancingContextClone.Get(),
-								InLoadFlags);
-							this->WrappedPackageLoader->FlushLoading({ RequestId });
-
-							// The completion delegate should have triggered so double check that now.
-							FMapPayload* const Payload = this->RetriedPackageCache.Find(PackagePathAsString);
-							if (ensure(Payload))
-							{
-								Payload->RequestId = RequestId;
-								this->MaxFlushedRequestId = FMath::Max(this->MaxFlushedRequestId, RequestId);
-							}
+							AutoRTFM::OnCommit([this] { DrainTransactionallyDeferredLoadPackages(); });
+							AutoRTFM::PushOnAbortHandler(this, [this] { TransactionallyDeferredLoadPackages.Empty(); });
 						});
-					});
-
-					ensure(Status == AutoRTFM::EContextStatus::AbortedByCascadingRetry);
-				});
+						ensure(Status == AutoRTFM::EContextStatus::OnTrack);
+					}
 
-				return -1;
-			}
-			else
-			{
-				InCompletionDelegate.ExecuteIfBound(Payload->PackageName, Payload->LoadedPackage, Payload->Result);
-				return Payload->RequestId;
-			}
+					return PackageCacheMiss;
+				}
+				else
+				{
+					// Whether we cause an abort in `InCompletionDelegate`, all we are doing in the open
+					// after this call is to return from the wrapped open, which will do the right thing
+					// and continue with the abort. That is why it is safe to ignore the result from the
+					// close.
+					(void)AutoRTFM::Close([&]
+					{
+						InCompletionDelegate.ExecuteIfBound(MapPayload->PackageName, MapPayload->LoadedPackage, MapPayload->Result);
+					});
+					return PackageCacheHit;
+				}
+			});
 		}
 		else
 		{
 			AutoRTFM::UnreachableIfTransactional(UnreachableMessage);
-			return WrappedPackageLoader->LoadPackage(PackagePath,
+			return LoadAndCachePackage(PackagePath,
 				CustomPackageName,
 				InCompletionDelegate,
 				InPackageFlags,
@@ -131,29 +176,21 @@ public:
 
 	int32 LoadPackage(const FPackagePath& PackagePath, FLoadPackageAsyncOptionalParams OptionalParams) override
 	{
-		if (AutoRTFM::IsClosed())
-		{
-			FLoadPackageAsyncDelegate CompletionDelegate;
+		FLoadPackageAsyncDelegate CompletionDelegate;
 
-			if (OptionalParams.CompletionDelegate.IsValid())
-			{
-				CompletionDelegate = *OptionalParams.CompletionDelegate.Get();
-			}
-
-			return LoadPackage(PackagePath,
-				OptionalParams.CustomPackageName,
-				CompletionDelegate,
-				OptionalParams.PackageFlags,
-				OptionalParams.PIEInstanceID,
-				OptionalParams.PackagePriority,
-				OptionalParams.InstancingContext,
-				OptionalParams.LoadFlags);
-		}
-		else
+		if (OptionalParams.CompletionDelegate.IsValid())
 		{
-			AutoRTFM::UnreachableIfTransactional(UnreachableMessage);
-			return WrappedPackageLoader->LoadPackage(PackagePath, MoveTemp(OptionalParams));
+			CompletionDelegate = *OptionalParams.CompletionDelegate.Get();
 		}
+
+		return LoadPackage(PackagePath,
+			OptionalParams.CustomPackageName,
+			CompletionDelegate,
+			OptionalParams.PackageFlags,
+			OptionalParams.PIEInstanceID,
+			OptionalParams.PackagePriority,
+			OptionalParams.InstancingContext,
+			OptionalParams.LoadFlags);
 	}
 
 	EAsyncPackageState::Type ProcessLoading(bool bUseTimeLimit, bool bUseFullTimeLimit, double TimeLimit) override
@@ -192,32 +229,115 @@ public:
 		{
 			AutoRTFM::Open([&]
 			{
-				for (const int32 RequestId : RequestIds)
+				if (RequestIds.IsEmpty() && TransactionallyDeferredLoadPackages.IsEmpty())
 				{
-					if (RequestId > MaxFlushedRequestId)
+					const AutoRTFM::EContextStatus Status = AutoRTFM::Close([&]
 					{
-						TArray<int32> RequestIdsClone(RequestIds);
-						
-						const AutoRTFM::EContextStatus Status = AutoRTFM::Close([&]
+						AutoRTFM::OnCommit([this]
 						{
-							AutoRTFM::CascadingRetryTransaction([this, RequestIdsClone]
-							{
-								UE_LOG(LogStreaming, Display, TEXT("A call to FlushLoading within a transaction is causing Verse to abort and retry."));
-								WrappedPackageLoader->FlushLoading(RequestIdsClone);
-
-								MaxFlushedRequestId = FMath::Max(MaxFlushedRequestId, FMath::Max(RequestIdsClone));
-							});
+							WrappedPackageLoader->FlushLoading({});
 						});
+					});
+					ensure(Status == AutoRTFM::EContextStatus::OnTrack);
+					return;
+				}
+
+				const int32 LastMaxFlushedRequestId = GetLastMaxFlushedRequestId();
+
+				bool bSkipRetry = true;
+				for (const int32 RequestId : RequestIds)
+				{
+					if (PackageCacheHit == RequestId)
+					{
+						continue;
+					}
 
-						ensure(Status == AutoRTFM::EContextStatus::AbortedByCascadingRetry);
+					if ((RequestId <= LastMaxFlushedRequestId) && (PackageCacheMiss != RequestId))
+					{
+						continue;
 					}
+
+					bSkipRetry = false;
+					break;
+				}
+
+				if (!bSkipRetry)
+				{
+					TArray<int32> RequestIdsClone;
+
+					for (const int32 RequestId : RequestIds)
+					{
+						// Filter out our special return status' for package cache hit and miss.
+						// The wrapped underlying async package loader does not understand them.
+						if ((PackageCacheHit == RequestId) || (PackageCacheMiss == RequestId))
+						{
+							continue;
+						}
+
+						RequestIdsClone.Add(RequestId);
+					}
+
+					const bool bShouldPopOnAbortHandler = !TransactionallyDeferredLoadPackages.IsEmpty();
+
+					const AutoRTFM::EContextStatus Status = AutoRTFM::Close([&]
+					{
+						if (bShouldPopOnAbortHandler)
+						{
+							// Pop the on-abort handler so that we don't purge the list of transactionally
+							// deferred packages to load, which we need to do in the retry below.
+							AutoRTFM::PopOnAbortHandler(this);
+						}
+
+						AutoRTFM::CascadingRetryTransaction([this, RequestIdsClone = ::MoveTemp(RequestIdsClone)]() mutable
+						{
+							UE_LOG(LogStreaming, Display, TEXT("A call to `FlushLoading` that is flushing non-cached packages is causing Verse to abort and retry."));
+							DrainTransactionallyDeferredLoadPackages(false, &RequestIdsClone);
+							UpdateMaxFlushedRequestId(RequestIdsClone);
+							WrappedPackageLoader->FlushLoading(RequestIdsClone);
+							UE_LOG(LogStreaming, Display, TEXT("`FlushLoading` has completed after Verse aborted, and we are now retrying."));
+						});
+					});
+					ensure(Status == AutoRTFM::EContextStatus::AbortedByCascadingRetry);
 				}
 			});
 		}
 		else
 		{
 			AutoRTFM::UnreachableIfTransactional(UnreachableMessage);
-			WrappedPackageLoader->FlushLoading(RequestIds);
+			UpdateMaxFlushedRequestId(RequestIds);
+			
+			if (INDEX_NONE != RequestIds.Find(PackageCacheMiss))
+			{
+				// If any request id was a package cache miss, it means we kicked
+				// off an async load package request inside a transaction, and are
+				// now trying to flush that from outside the transaction. Since we
+				// don't have a mapping from the miss to the actual request id, we
+				// just have to flush all packages.
+				WrappedPackageLoader->FlushLoading({});
+			}
+			else if (INDEX_NONE != RequestIds.Find(PackageCacheHit))
+			{
+				// If any request id was a package cache hit, we just need to filter
+				// it out from the list of request ids we are going to ask the actual
+				// underlying wrapped async package loader to flush.
+				TArray<int32> SubsetOfRequestIds;
+
+				for (const int32 RequestId : RequestIds)
+				{
+					if (PackageCacheHit == RequestId)
+					{
+						continue;
+					}
+
+					SubsetOfRequestIds.Add(RequestId);
+				}
+
+				WrappedPackageLoader->FlushLoading(SubsetOfRequestIds);
+			}
+			else
+			{
+				WrappedPackageLoader->FlushLoading(RequestIds);
+			}
 		}
 	}
 
@@ -286,8 +406,25 @@ public:
 	}
 
 private:
+	static constexpr int32 PackageCacheHit = INT_MIN;
+	static constexpr int32 PackageCacheMiss = PackageCacheHit + 1;
+
 	static const TCHAR* const UnreachableMessage;
 	IAsyncPackageLoader* const WrappedPackageLoader;
+	
+	struct FTransactionallyDeferredLoadPackagePayload final 
+	{
+		FPackagePath PackagePath;
+		FName CustomPackageName;
+		FLoadPackageAsyncDelegate CompletionDelegate;
+		EPackageFlags PackageFlags;
+		int32 PIEInstanceID;
+		int32 PackagePriority;
+		TOptional<FLinkerInstancingContext> InstancingContext;
+		uint32 LoadFlags;
+	};
+
+	TArray<FTransactionallyDeferredLoadPackagePayload> TransactionallyDeferredLoadPackages;
 
 	struct FMapPayload final
 	{
@@ -297,11 +434,88 @@ private:
 		FName PackageName;
 		UPackage* LoadedPackage;
 		EAsyncLoadingResult::Type Result;
-		int32 RequestId = -1;
 	};
 
-	TMap<FString, FMapPayload> RetriedPackageCache;
-	int32 MaxFlushedRequestId = INT32_MIN;
+	TMap<FString, FMapPayload> PackageCache;
+	int32 MaxFlushedRequestId = -1;
+	UE::FMutex Mutex;
+
+	void UpdateMaxFlushedRequestId(const TConstArrayView<int32> RequestIds)
+	{
+		UE::TScopeLock _(Mutex);
+
+		for (const int32 RequestId : RequestIds)
+		{
+			MaxFlushedRequestId = std::max(MaxFlushedRequestId, RequestId);
+		}
+	}
+
+	int32 GetLastMaxFlushedRequestId()
+	{
+		UE::TScopeLock _(Mutex);
+		return MaxFlushedRequestId;
+	}
+
+	int32 LoadAndCachePackage(
+			const FPackagePath& PackagePath,
+			FName CustomPackageName,
+			FLoadPackageAsyncDelegate InCompletionDelegate,
+			EPackageFlags InPackageFlags,
+			int32 InPIEInstanceID,
+			int32 InPackagePriority,
+			const FLinkerInstancingContext* InInstancingContext,
+			uint32 InLoadFlags)
+	{
+		AutoRTFM::UnreachableIfTransactional(UnreachableMessage);
+
+		FLoadPackageAsyncDelegate WrapperDelegate;
+		WrapperDelegate.BindLambda([this, PackagePathStr = PackagePath.GetDebugNameWithExtension(), CompletionDelegate = ::MoveTemp(InCompletionDelegate)](
+			const FName& Name, UPackage* const Package, EAsyncLoadingResult::Type Type)
+		{
+			if (EAsyncLoadingResult::Canceled != Type)
+			{
+				FMapPayload Payload(Name, Package, Type);
+
+				UE::TScopeLock _(Mutex);
+				PackageCache.Add(PackagePathStr, Payload);
+			}
+
+			CompletionDelegate.ExecuteIfBound(Name, Package, Type);
+		});
+
+		return WrappedPackageLoader->LoadPackage(PackagePath,
+			CustomPackageName,
+			WrapperDelegate,
+			InPackageFlags,
+			InPIEInstanceID,
+			InPackagePriority,
+			InInstancingContext,
+			InLoadFlags);
+	}
+
+	void DrainTransactionallyDeferredLoadPackages(const bool bCallOriginalCompletionDelegate = true, TArray<int32>* RequestIds = nullptr)
+	{
+		AutoRTFM::UnreachableIfTransactional(UnreachableMessage);
+
+		for (const FTransactionallyDeferredLoadPackagePayload& Payload : TransactionallyDeferredLoadPackages)
+		{
+			const int32 RequestId = LoadAndCachePackage(Payload.PackagePath,
+				Payload.CustomPackageName,
+				bCallOriginalCompletionDelegate ? Payload.CompletionDelegate : FLoadPackageAsyncDelegate(),
+				Payload.PackageFlags,
+				Payload.PIEInstanceID,
+				Payload.PackagePriority,
+				Payload.InstancingContext.GetPtrOrNull(),
+				Payload.LoadFlags);
+
+			if (RequestIds)
+			{
+				RequestIds->Add(RequestId);
+			}
+		}
+
+		TransactionallyDeferredLoadPackages.Empty();
+	}
 };
 
 const TCHAR* const FTransactionallySafeAsyncPackageLoader::UnreachableMessage = TEXT("Cannot call function within a transaction!");
