PropertyBag : Support for Sets in GetContainerTypesFromProperty, GetValueTypeFromProperty, and GetValueTypeObjectFromProperty

Support for Sets in PropertyBags was added in changelist 36156471 but these functions were missing support.
These functions are used by the FPropertyBagPropertyDesc(const FName InName, const FProperty* InSourceProperty) constructor.

diff --git a/Engine/Source/Runtime/CoreUObject/Private/StructUtils/PropertyBag.cpp b/Engine/Source/Runtime/CoreUObject/Private/StructUtils/PropertyBag.cpp
--- a/Engine/Source/Runtime/CoreUObject/Private/StructUtils/PropertyBag.cpp
+++ b/Engine/Source/Runtime/CoreUObject/Private/StructUtils/PropertyBag.cpp
@@ -83,17 +83,34 @@ namespace UE::StructUtils::Private
 	{
 		FPropertyBagContainerTypes ContainerTypes;
 
-		while (const FArrayProperty* ArrayProperty = CastField<FArrayProperty>(InSourceProperty))
+		while (InSourceProperty)
 		{
-			if (ContainerTypes.Add(EPropertyBagContainerType::Array))
+			if (const FArrayProperty* ArrayProperty = CastField<FArrayProperty>(InSourceProperty))
 			{
-				InSourceProperty = CastField<FArrayProperty>(ArrayProperty->Inner);
+				if (ContainerTypes.Add(EPropertyBagContainerType::Array))
+				{
+					InSourceProperty = ArrayProperty->Inner;
+					continue;
+				}
+				else // we reached the nested containers limit
+				{
+					ContainerTypes.Reset();
+				}
 			}
-			else // we reached the nested containers limit
+			if (const FSetProperty* SetProperty = CastField<FSetProperty>(InSourceProperty))
 			{
-				ContainerTypes.Reset();
-				break;
+				if (ContainerTypes.Add(EPropertyBagContainerType::Set))
+				{
+					InSourceProperty = SetProperty->ElementProp;
+					continue;
+				}
+				else // we reached the nested containers limit
+				{
+					ContainerTypes.Reset();
+				}
 			}
+			
+			InSourceProperty = nullptr;
 		}
 
 		return ContainerTypes;
@@ -178,6 +195,11 @@ namespace UE::StructUtils::Private
 			return GetValueTypeFromProperty(ArrayProperty->Inner);	
 		}
 
+		if (const FSetProperty* SetProperty = CastField<FSetProperty>(InSourceProperty))
+		{
+			return GetValueTypeFromProperty(SetProperty->ElementProp);
+		}
+
 		return EPropertyBagPropertyType::None;
 	}
 
@@ -223,6 +245,11 @@ namespace UE::StructUtils::Private
 			return GetValueTypeObjectFromProperty(ArrayProperty->Inner);
 		}
 
+		if (const FSetProperty* SetProperty = CastField<FSetProperty>(InSourceProperty))
+		{
+			return GetValueTypeObjectFromProperty(SetProperty->ElementProp);
+		}
+
 		return nullptr;
 	}
 
