[ZenLoader]
- Fix case where the loader was not properly loading a package if that UPackage was created outside the loader

diff --git a/Engine/Source/Runtime/CoreUObject/Private/Serialization/AsyncLoading2.cpp b/Engine/Source/Runtime/CoreUObject/Private/Serialization/AsyncLoading2.cpp
--- a/Engine/Source/Runtime/CoreUObject/Private/Serialization/AsyncLoading2.cpp
+++ b/Engine/Source/Runtime/CoreUObject/Private/Serialization/AsyncLoading2.cpp
@@ -1312,7 +1312,14 @@ public:
 		{
 			return;
 		}
-		bAreAllPublicExportsLoaded = PublicExportMap.PinForGC(OutUnreachableObjectIndices);
+		// If we're missing export table objects, reset the export loaded flags
+		// but never overwrite it to true because if the UPackage was created outside
+		// of the loader it wont have any exports in its map yet and we don't want
+		// to make the loader think that the package was already loaded.
+		if (!PublicExportMap.PinForGC(OutUnreachableObjectIndices))
+		{
+			bAreAllPublicExportsLoaded = false;
+		}
 #if WITH_EDITOR
 		// Only reload if we ever got snapshotted after a proper load and the export count has shrinked since the package was last loaded
 		if (GReloadPackagesWithGCedExports && bAreAllPublicExportsLoaded && ExportCount > PublicExportMap.Num())
