Empty FAsyncLoadingThread2::UnreachableObjects TArray to prevent its slack memory to stay alive for nothing (saves ~2MB)

#rnx
#tests ran a bunch of replays on multiple platforms

diff --git a/Engine/Source/Runtime/CoreUObject/Private/Serialization/AsyncLoading2.cpp b/Engine/Source/Runtime/CoreUObject/Private/Serialization/AsyncLoading2.cpp
--- a/Engine/Source/Runtime/CoreUObject/Private/Serialization/AsyncLoading2.cpp
+++ b/Engine/Source/Runtime/CoreUObject/Private/Serialization/AsyncLoading2.cpp
@@ -9221,7 +9221,9 @@ void FAsyncLoadingThread2::RemoveUnreachableObjects(FUnreachableObjects& Objects
 
 	GlobalImportStore.RemovePackages(ObjectsToRemove);
 	GlobalImportStore.RemovePublicExports(ObjectsToRemove);
-	ObjectsToRemove.Reset();
+
+	// No need to keep Array slack memory as it may not be entirely filled over the time
+	ObjectsToRemove.Empty();
 
 	const int32 NewLoadedPackageCount = GlobalImportStore.GetStoredPackagesCount();
 	const int32 NewPublicExportCount = GlobalImportStore.GetStoredPublicExportsCount();
