Add Id number accessor to FRemoteServerId, FRemoteObjectId and FRemoteTransactionId

Added GetIdNumber() accessor to these structures.  This is to help with profiling and debugging.

diff --git a/Engine/Source/Runtime/CoreUObject/Public/UObject/RemoteExecutor.h b/Engine/Source/Runtime/CoreUObject/Public/UObject/RemoteExecutor.h
--- a/Engine/Source/Runtime/CoreUObject/Public/UObject/RemoteExecutor.h
+++ b/Engine/Source/Runtime/CoreUObject/Public/UObject/RemoteExecutor.h
@@ -26,6 +26,10 @@ struct FRemoteTransactionId
 	bool operator==(const FRemoteTransactionId& Rhs) const { return Id == Rhs.Id; }
 	bool operator!=(const FRemoteTransactionId& Rhs) const { return Id != Rhs.Id; }
 
+	uint32 GetIdNumber() const
+	{
+		return Id;
+	}
 	bool IsValid() const
 	{
 		return Id != 0;
diff --git a/Engine/Source/Runtime/CoreUObject/Public/UObject/RemoteObjectTypes.h b/Engine/Source/Runtime/CoreUObject/Public/UObject/RemoteObjectTypes.h
--- a/Engine/Source/Runtime/CoreUObject/Public/UObject/RemoteObjectTypes.h
+++ b/Engine/Source/Runtime/CoreUObject/Public/UObject/RemoteObjectTypes.h
@@ -20,6 +20,10 @@ struct FRemoteServerId
 	{
 		checkf(Id < (1 << 10), TEXT("Remote server id can only be a 10-bit value, got: %u"), Id);
 	}
+	uint32 GetIdNumber() const
+	{
+		return Id;
+	}
 	FString ToString() const
 	{
 		return FString::FromInt(Id);
@@ -114,6 +118,11 @@ public:
 		return *this != FRemoteObjectId();
 	}
 
+	FORCEINLINE uint64 GetIdNumber() const
+	{
+		return Id;
+	}
+	
 	FORCEINLINE FRemoteServerId GetServerId() const
 	{
 		return FRemoteServerId(ServerId);
