CoreRedirects: Add FCoreRedirectObjectName paths functions to get parent and child paths.
#rnx

diff --git a/Engine/Source/Runtime/CoreUObject/Private/UObject/CoreRedirects.cpp b/Engine/Source/Runtime/CoreUObject/Private/UObject/CoreRedirects.cpp
--- a/Engine/Source/Runtime/CoreUObject/Private/UObject/CoreRedirects.cpp
+++ b/Engine/Source/Runtime/CoreUObject/Private/UObject/CoreRedirects.cpp
@@ -908,6 +908,242 @@ FString FCoreRedirectObjectName::CombineNames(FName NewName, FName NewOuter, FNa
 	return NewName.ToString();
 }
 
+namespace UE::CoreRedirects
+{
+
+static void SplitFirstComponent(FName Text, FStringBuilderBase& Buffer, FStringView& OutFirst, FStringView& OutRemainder)
+{
+	if (Text.IsNone())
+	{
+		OutFirst = FStringView();
+		OutRemainder = FStringView();
+		return;
+	}
+
+	Buffer.Reset();
+	Buffer << Text;
+	// Normalize the Text: treat SUBOBJECT_DELIMITER_CHAR as '.'
+	for (TCHAR& C : Buffer)
+	{
+		if (C == SUBOBJECT_DELIMITER_CHAR)
+		{
+			C = '.';
+		}
+	}
+	FStringView TextStr(Buffer);
+	int32 FirstDotIndex;
+	if (TextStr.FindChar('.', FirstDotIndex))
+	{
+		OutFirst = TextStr.Left(FirstDotIndex);
+		OutRemainder = TextStr.RightChop(FirstDotIndex + 1);
+	}
+	else
+	{
+		OutFirst = TextStr;
+		OutRemainder = FStringView();
+	}
+}
+
+static void SplitLastComponent(FName Text, FStringBuilderBase& Buffer, FStringView& OutRemainder, FStringView& OutLast)
+{
+	if (Text.IsNone())
+	{
+		OutRemainder = FStringView();
+		OutLast = FStringView();
+		return;
+	}
+
+	Buffer.Reset();
+	Buffer << Text;
+	// Normalize the Text: treat SUBOBJECT_DELIMITER_CHAR as '.'
+	for (TCHAR& C : Buffer)
+	{
+		if (C == SUBOBJECT_DELIMITER_CHAR)
+		{
+			C = '.';
+		}
+	}
+	FStringView TextStr(Buffer);
+	int32 LastDotIndex;
+	if (TextStr.FindLastChar('.', LastDotIndex))
+	{
+		OutRemainder = TextStr.Left(LastDotIndex);
+		OutLast = TextStr.RightChop(LastDotIndex + 1);
+	}
+	else
+	{
+		OutRemainder = FStringView();
+		OutLast = TextStr;
+	}
+}
+
+}
+
+FCoreRedirectObjectName FCoreRedirectObjectName::AppendObjectName(const FCoreRedirectObjectName& Parent,
+	FName ObjectName)
+{
+	using namespace UE::CoreRedirects;
+
+	if (Parent.ObjectName.IsNone())
+	{
+		if (Parent.OuterName.IsNone())
+		{
+			if (Parent.PackageName.IsNone())
+			{
+				// Empty parent, return a packagename with the given ObjectName
+				return FCoreRedirectObjectName(NAME_None, NAME_None, ObjectName);
+			}
+			else
+			{
+				// Child of a package
+				return FCoreRedirectObjectName(ObjectName, NAME_None, Parent.PackageName);
+			}
+		}
+		else
+		{
+			// An unexpected case; anything with an outer name should also had an objectname.
+			// But if this case were valid, then the proper child would be to just set the objectname
+			return FCoreRedirectObjectName(ObjectName, Parent.OuterName, Parent.PackageName);
+		}
+	}
+	else
+	{
+		if (Parent.OuterName.IsNone())
+		{
+			if (Parent.PackageName.IsNone())
+			{
+				// An unexpected case; there should not be an objectname with no outer or packagename.
+				// Treat this case is if the ObjectName were the packagename
+				return FCoreRedirectObjectName(ObjectName, NAME_None, Parent.ObjectName);
+			}
+			else
+			{
+				// Parent is a toplevel child of a package, so the child has just the parent's objectname
+				// as its outername.
+				return FCoreRedirectObjectName(ObjectName, Parent.ObjectName, Parent.PackageName);
+			}
+		}
+		else
+		{
+			if (Parent.PackageName.IsNone())
+			{
+				// An unexpected case; there should not be an objectname and outername with no packagename.
+				// Treat this case is if the outermost of the outername were the packagename.
+				TStringBuilder<NAME_SIZE> OuterNameBuffer;
+				FStringView PackageName;
+				FStringView OuterNameStr;
+				SplitFirstComponent(Parent.OuterName, OuterNameBuffer, PackageName, OuterNameStr);
+				if (PackageName.Len() == 0)
+				{
+					// Just a period by itself, treat this as OuterName == None.
+					return AppendObjectName(FCoreRedirectObjectName(Parent.ObjectName, NAME_None, NAME_None), ObjectName);
+				}
+				FName OuterName = OuterNameStr.Len() > 0 ? FName(OuterNameStr) : NAME_None;
+				return AppendObjectName(FCoreRedirectObjectName(Parent.ObjectName, OuterName, FName(PackageName)), ObjectName);
+			}
+			else
+			{
+				// Parent is Package.Outer.Object. Append Parent's Object to its Outer and set the child's ObjectName.
+				TStringBuilder<NAME_SIZE> NewOuterNameBuffer(InPlace, Parent.OuterName);
+				NewOuterNameBuffer << '.' << Parent.ObjectName;
+				return FCoreRedirectObjectName(ObjectName, FName(NewOuterNameBuffer.ToView()), Parent.PackageName);
+			}
+		}
+	}
+}
+
+FCoreRedirectObjectName FCoreRedirectObjectName::GetParent(const FCoreRedirectObjectName& Child)
+{
+	using namespace UE::CoreRedirects;
+
+	if (Child.ObjectName.IsNone())
+	{
+		if (Child.OuterName.IsNone())
+		{
+			// Either An empty FCoreRedirectObjectName (Child.PackageName.IsNone()) or a package
+			// Return an empty FCoreRedirectObjectName() for both of these cases.
+			return FCoreRedirectObjectName();
+		}
+		else
+		{
+			if (Child.PackageName.IsNone())
+			{
+				// An unexpected case; anything with an outer name should also have an objectname and a packagename.
+				// Treat this case is if the outermost of the outername were the packagename.
+				TStringBuilder<NAME_SIZE> OuterNameBuffer;
+				FStringView PackageName;
+				FStringView OuterName;
+				SplitFirstComponent(Child.OuterName, OuterNameBuffer, PackageName, OuterName);
+				if (PackageName.Len() == 0)
+				{
+					// Just a period by itself, treat this as no outer name.
+					return GetParent(FCoreRedirectObjectName(NAME_None, NAME_None, NAME_None));
+				}
+				else
+				{
+					FName NewOuter = OuterName.Len() > 0 ? FName(OuterName) : NAME_None;
+					return GetParent(FCoreRedirectObjectName(NAME_None, NewOuter, FName(PackageName)));
+				}
+			}
+			else
+			{
+				// An unexpected case; anything with an outer name should also have an objectname.
+				// Treat this case is if the innermost of the outername were the objectname.
+				TStringBuilder<NAME_SIZE> OuterNameBuffer;
+				FStringView NewOuterNameStr;
+				FStringView NewObjectName;
+				SplitLastComponent(Child.OuterName, OuterNameBuffer, NewOuterNameStr, NewObjectName);
+				if (NewObjectName.Len() == 0)
+				{
+					// Just a period by itself, treat this as no outer name.
+					return GetParent(FCoreRedirectObjectName(NAME_None, NAME_None, Child.PackageName));
+				}
+				else
+				{
+					FName NewOuterName = NewOuterNameStr.Len() > 0 ? FName(NewOuterNameStr) : NAME_None;
+					return GetParent(FCoreRedirectObjectName(FName(NewObjectName), NewOuterName, Child.PackageName));
+				}
+			}
+		}
+	}
+	else
+	{
+		if (Child.OuterName.IsNone())
+		{
+			if (Child.PackageName.IsNone())
+			{
+				// An unexpected case; there should not be an objectname with no outer or packagename.
+				// Treat this case is if the ObjectName were the packagename
+				return GetParent(FCoreRedirectObjectName(NAME_None, NAME_None, Child.ObjectName));
+			}
+			else
+			{
+				// A toplevel object, parent is its package
+				return FCoreRedirectObjectName(NAME_None, NAME_None, Child.PackageName);
+			}
+		}
+		else
+		{
+			// A subobject, aka a child of an object that is not the package. Get the
+			// FCoreRedirectObjectName of the outer by peeling the last dot-delimited name from the OuterName.
+			TStringBuilder<NAME_SIZE> OuterNameBuffer;
+			FStringView NewOuterNameStr;
+			FStringView NewObjectName;
+			SplitLastComponent(Child.OuterName, OuterNameBuffer, NewOuterNameStr, NewObjectName);
+			if (NewObjectName.Len() == 0)
+			{
+				// Just a period by itself, treat this as no outer name.
+				return GetParent(FCoreRedirectObjectName(Child.ObjectName, NAME_None, Child.PackageName));
+			}
+			else
+			{
+				FName NewOuterName = NewOuterNameStr.Len() > 0 ? FName(NewOuterNameStr) : NAME_None;
+				return FCoreRedirectObjectName(FName(NewObjectName), NewOuterName, Child.PackageName);
+			}
+		}
+	}
+}
+
 void FCoreRedirect::NormalizeNewName()
 {
 	// Fill in NewName as needed
@@ -1921,6 +2157,122 @@ bool FCoreRedirects::RunTests()
 	return bSuccess;
 }
 
+#if WITH_DEV_AUTOMATION_TESTS
+IMPLEMENT_SIMPLE_AUTOMATION_TEST(FCoreRedirectObjectNameTest, "System.Core.Misc.CoreRedirects.FCoreRedirectObjectName", EAutomationTestFlags_ApplicationContextMask | EAutomationTestFlags::EngineFilter)
+bool FCoreRedirectObjectNameTest::RunTest(const FString& Parameters)
+{
+	FCoreRedirectObjectName Parent;
+	FCoreRedirectObjectName Child;
+	FCoreRedirectObjectName ExpectedChild;
+	FCoreRedirectObjectName ExpectedParent;
+	FName NamePackage(TEXT("/Root/Package"));
+	FName NameA(TEXT("A"));
+	FName NameB(TEXT("B"));
+	FName NameC(TEXT("C"));
+	FName NameD(TEXT("D"));
+	FName NameADotB(TEXT("A.B"));
+	FName NameADotBDotC(TEXT("A.B.C"));
+
+	// AppendObjectName
+	Parent = FCoreRedirectObjectName(NAME_None, NAME_None, NAME_None);
+	ExpectedChild = FCoreRedirectObjectName(NAME_None, NAME_None, NamePackage);
+	Child = FCoreRedirectObjectName::AppendObjectName(Parent, NamePackage);
+	TestTrue(TEXT("AppendObjectName(Empty + Package)"), Child == ExpectedChild);
+
+	Parent = FCoreRedirectObjectName(NAME_None, NAME_None, NamePackage);
+	ExpectedChild = FCoreRedirectObjectName(NameA, NAME_None, NamePackage);
+	Child = FCoreRedirectObjectName::AppendObjectName(Parent, NameA);
+	TestTrue(TEXT("AppendObjectName(Package + A)"), Child == ExpectedChild);
+
+	Parent = FCoreRedirectObjectName(NameA, NAME_None, NamePackage);
+	ExpectedChild = FCoreRedirectObjectName(NameB, NameA, NamePackage);
+	Child = FCoreRedirectObjectName::AppendObjectName(Parent, NameB);
+	TestTrue(TEXT("AppendObjectName(Package.A + B)"), Child == ExpectedChild);
+
+	Parent = FCoreRedirectObjectName(NameB, NameA, NamePackage);
+	ExpectedChild = FCoreRedirectObjectName(NameC, NameADotB, NamePackage);
+	Child = FCoreRedirectObjectName::AppendObjectName(Parent, NameC);
+	TestTrue(TEXT("AppendObjectName(Package.A.B + C)"), Child == ExpectedChild);
+
+	Parent = FCoreRedirectObjectName(NameC, NameADotB, NamePackage);
+	ExpectedChild = FCoreRedirectObjectName(NameD, NameADotBDotC, NamePackage);
+	Child = FCoreRedirectObjectName::AppendObjectName(Parent, NameD);
+	TestTrue(TEXT("AppendObjectName(Package.A.B.C + D)"), Child == ExpectedChild);
+
+	// Verify these edge cases do not crash, but we do not provide a contract for what they mean.
+	// ObjectName with no PackageName
+	Parent = FCoreRedirectObjectName(NameA, NAME_None, NAME_None);
+	FCoreRedirectObjectName::AppendObjectName(Parent, NameB);
+
+	// ObjectName and OuterName with no PackageName
+	Parent = FCoreRedirectObjectName(NameB, NameA, NAME_None);
+	FCoreRedirectObjectName::AppendObjectName(Parent, NameC);
+
+	// OuterName with no ObjectName and no PackageName
+	Parent = FCoreRedirectObjectName(NAME_None, NameA, NAME_None);
+	FCoreRedirectObjectName::AppendObjectName(Parent, NameB);
+
+	// OuterName and ObjectName with no PackageName
+	Parent = FCoreRedirectObjectName(NameB, NameA, NAME_None);
+	FCoreRedirectObjectName::AppendObjectName(Parent, NameC);
+
+	// OuterName and PackageName with no ObjectName
+	Parent = FCoreRedirectObjectName(NAME_None, NameA, NamePackage);
+	FCoreRedirectObjectName::AppendObjectName(Parent, NameB);
+
+
+	// GetParent
+	Child = FCoreRedirectObjectName(NAME_None, NAME_None, NAME_None);
+	ExpectedParent = FCoreRedirectObjectName(NAME_None, NAME_None, NAME_None);
+	Parent = FCoreRedirectObjectName::GetParent(Child);
+	TestTrue(TEXT("GetParent(Empty)"), Parent == ExpectedParent);
+
+	Child = FCoreRedirectObjectName(NAME_None, NAME_None, NamePackage);
+	ExpectedParent = FCoreRedirectObjectName(NAME_None, NAME_None, NAME_None);
+	Parent = FCoreRedirectObjectName::GetParent(Child);
+	TestTrue(TEXT("GetParent(Package)"), Parent == ExpectedParent);
+
+	Child = FCoreRedirectObjectName(NameA, NAME_None, NamePackage);
+	ExpectedParent = FCoreRedirectObjectName(NAME_None, NAME_None, NamePackage);
+	Parent = FCoreRedirectObjectName::GetParent(Child);
+	TestTrue(TEXT("GetParent(Package.A)"), Parent == ExpectedParent);
+
+	Child = FCoreRedirectObjectName(NameC, NameADotB, NamePackage);
+	ExpectedParent = FCoreRedirectObjectName(NameB, NameA, NamePackage);
+	Parent = FCoreRedirectObjectName::GetParent(Child);
+	TestTrue(TEXT("GetParent(Package.A.B)"), Parent == ExpectedParent);
+
+	Child = FCoreRedirectObjectName(NameD, NameADotBDotC, NamePackage);
+	ExpectedParent = FCoreRedirectObjectName(NameC, NameADotB, NamePackage);
+	Parent = FCoreRedirectObjectName::GetParent(Child);
+	TestTrue(TEXT("GetParent(Package.A.B.C)"), Parent == ExpectedParent);
+
+	// Verify these edge cases do not crash, but we do not provide a contract for what they mean.
+	// ObjectName with no PackageName
+	Child = FCoreRedirectObjectName(NameA, NAME_None, NAME_None);
+	FCoreRedirectObjectName::GetParent(Child);
+
+	// ObjectName and OuterName with no PackageName
+	Child = FCoreRedirectObjectName(NameB, NameA, NAME_None);
+	FCoreRedirectObjectName::GetParent(Child);
+
+	// OuterName with no ObjectName and no PackageName
+	Child = FCoreRedirectObjectName(NAME_None, NameA, NAME_None);
+	FCoreRedirectObjectName::GetParent(Child);
+
+	// OuterName and ObjectName with no PackageName
+	Child = FCoreRedirectObjectName(NameB, NameA, NAME_None);
+	FCoreRedirectObjectName::GetParent(Child);
+
+	// OuterName and PackageName with no ObjectName
+	Child = FCoreRedirectObjectName(NAME_None, NameA, NamePackage);
+	FCoreRedirectObjectName::GetParent(Child);
+
+	return true;
+}
+#endif
+
+
 void FCoreRedirects::AddAssetRedirects(const TMap<FSoftObjectPath, FSoftObjectPath>& InRedirects)
 {
 	if (InRedirects.IsEmpty())
diff --git a/Engine/Source/Runtime/CoreUObject/Public/UObject/CoreRedirects.h b/Engine/Source/Runtime/CoreUObject/Public/UObject/CoreRedirects.h
--- a/Engine/Source/Runtime/CoreUObject/Public/UObject/CoreRedirects.h
+++ b/Engine/Source/Runtime/CoreUObject/Public/UObject/CoreRedirects.h
@@ -198,6 +198,15 @@ struct FCoreRedirectObjectName
 
 	/** Turn it back into an FString */
 	static COREUOBJECT_API FString CombineNames(FName NewName, FName NewOuter, FName NewPackage);
+
+	/** Given parent FCoreRedirectObjectName and FName of a child under it, return child's FCoreRedirectObjectName. */
+	static COREUOBJECT_API FCoreRedirectObjectName
+		AppendObjectName(const FCoreRedirectObjectName& Parent, FName ObjectName);
+	/**
+	 * Given a child FCoreRedirectObjectName, return its parent's FCoreRedirectObjectName.
+	 * If the input has no parent (empty or is package with no outer), returns an empty FCoreRedirectObjectName.
+	 */
+	static COREUOBJECT_API FCoreRedirectObjectName GetParent(const FCoreRedirectObjectName& Child);
 };
 ENUM_CLASS_FLAGS(FCoreRedirectObjectName::EMatchFlags);
 
