[AsyncLoading2]

- Disables the loose file loading path for non-editor platforms unless "s.LooseFileLoadingEnabled" is enabled in the Engine.StreamingSettings or -loosefileloading is passed on commandline. This restores the previous performance characteristics for zenloader on 5.4 in runtime configs which avoided checking the filesystem for loose files to override cooked data in the iostore.

- "s.LooseFileLoadingEnabled" replaces the shortlived "s.AllowLoadingPackageFromFileSystemInShipping" as the former is more encompassing and the latter also poses a initialization timing issue since cvars are initialized to late to handle all potential loads during boot.

- Corrects CoreAssetLoading test to correctly load BPs in blueprint test cases (loading class vs loading actor for test case comparisons)

[RN] Zenloader can be enabled to read loose files when loading in non-editor configs by passing "-loosefileloading" on the commandline or specifying "s.LooseFileLoadingEnabled=true" in Engine.StreamingSettings INI. `WITH_EDITOR` configs continue to load loose files by default. Note, loose files will override IOStore content if there is an overlap.

#virtualized

diff --git a/Engine/Source/Developer/LauncherServices/Private/Launcher/LauncherWorker.cpp b/Engine/Source/Developer/LauncherServices/Private/Launcher/LauncherWorker.cpp
--- a/Engine/Source/Developer/LauncherServices/Private/Launcher/LauncherWorker.cpp
+++ b/Engine/Source/Developer/LauncherServices/Private/Launcher/LauncherWorker.cpp
@@ -593,10 +593,20 @@ FString FLauncherWorker::CreateUATCommand( const ILauncherProfileRef& InProfile,
 		*SessionOwner,
 		*SessionName);
 	}
-	FString AdditionalCommandLine = FString::Printf(TEXT(" -addcmdline=\"%s%s %s %s\""),
+
+	// Provide reasonable defaults for loading depending on the user config
+	FString LoaderCommands;
+	if (!InProfile->IsUsingZenStore())
+	{
+		// If not streaming from zenstore, enable loose file reads (only enabled in editor by default)
+		LoaderCommands = TEXT("-loosefileloading");
+	}
+
+	FString AdditionalCommandLine = FString::Printf(TEXT(" -addcmdline=\"%s%s %s %s %s\""),
 		*SessionCommands,
 		*RoleCommands,
 		*LocalizationCommands,
+		*LoaderCommands,
 		*InProfile->GetAdditionalCommandLineParameters());
 	if (AdditionalCommandLine.EndsWith(TEXT("\"  \"")))
 	{
diff --git a/Engine/Source/Runtime/CoreUObject/Private/Serialization/AsyncLoading2.cpp b/Engine/Source/Runtime/CoreUObject/Private/Serialization/AsyncLoading2.cpp
--- a/Engine/Source/Runtime/CoreUObject/Private/Serialization/AsyncLoading2.cpp
+++ b/Engine/Source/Runtime/CoreUObject/Private/Serialization/AsyncLoading2.cpp
@@ -276,7 +276,7 @@ void FindAllRuntimeScriptPackages(FRuntimeScriptPackages& OutPackages)
 #endif
 
 #ifndef ALT2_ENABLE_LINKERLOAD_SUPPORT
-#define ALT2_ENABLE_LINKERLOAD_SUPPORT WITH_EDITOR
+#define ALT2_ENABLE_LINKERLOAD_SUPPORT 1
 #endif
 
 #ifndef ALT2_ENABLE_NEW_ARCHIVE_FOR_LINKERLOAD
@@ -442,15 +442,6 @@ static FAutoConsoleVariableRef CVarGOnlyProcessRequiredPackagesWhenSyncLoading(
 	ECVF_Default
 );
 
-// If allowed in Shipping builds to load from the file system (Editor will always be allowed for this)
-static bool GAllowLoadingPackageFromFileSystemInShipping = false;
-static FAutoConsoleVariableRef CVarAllowLoadingPackagesFromFileSystem(
-	TEXT("s.AllowLoadingPackageFromFileSystemInShipping"),
-	GAllowLoadingPackageFromFileSystemInShipping,
-	TEXT("Allow the loader to try and load packages from the filesystem in Shipping configs"),
-	ECVF_Default
-	);
-
 // In editor builds, for packages that are both local and in zenstore, do we prefer uncooked (local) or cooked (zenstore) packages
 static bool GPreferCookedPackagesInEditor = false;
 static FAutoConsoleVariableRef CVarPreferCookedPackagesInEditor(
@@ -4046,16 +4037,8 @@ public:
 	}
 
 private:
-#if ALT2_ENABLE_LINKERLOAD_SUPPORT || WITH_EDITOR
 	bool TryGetPackagePathFromFileSystem(FName& InOutPackageNameToLoad, FName& InOutUPackageName, FPackagePath& InOutPackagePath)
 	{
-		// If disabling loading from the file system in a shipping/test build and we are not the Editor avoid doing this look-up
-		// As this can be expensive to call and can cause contention on the GameThread and ASyncLoadingThread
-		if (!GAllowLoadingPackageFromFileSystemInShipping && !WITH_EDITOR && (UE_BUILD_SHIPPING || UE_BUILD_TEST))
-		{
-			return false;
-		}
-
 		// if we prefer cooked packages in the editor, then we don't want to use the linkerload version
 		// NOTE that we won't have fixed up the FNames, but if we aren't using linkerloads, it shouldn't matter
 		if (GPreferCookedPackagesInEditor)
@@ -4079,8 +4062,7 @@ private:
 #endif
 
 		// now, check to see if the package exists in the local file system, and if it does, correct the case as needed
-		if (!InOutPackageNameToLoad.IsNone() &&
-			FPackageName::DoesPackageExistEx(InOutPackagePath, FPackageName::EPackageLocationFilter::FileSystem,
+		if (FPackageName::DoesPackageExistEx(InOutPackagePath, FPackageName::EPackageLocationFilter::FileSystem,
 				bMatchCaseOnDisk, &InOutPackagePath) != FPackageName::EPackageLocationFilter::None)
 		{
 			FName CaseCorrectedPackageName = InOutPackagePath.GetPackageFName();
@@ -4094,7 +4076,6 @@ private:
 		}
 		return false;
 	}
-#endif
 
 #if WITH_EDITOR
 	void ConditionalProcessEditorCallbacks();
@@ -4916,25 +4897,26 @@ bool FAsyncLoadingThread2::CreateAsyncPackagesFromQueue(FAsyncLoadingThreadState
 			}
 
 			bool bIsZenPackage = true;
-#if ALT2_ENABLE_LINKERLOAD_SUPPORT
-			FPackagePath::TryFromPackageName(PackageNameToLoad, Request.PackagePath);
-
-			// Lookup if we have the package in the importstore already to avoid expensive calls in TryGetPackagePathFromFileSystem
-			if (FLoadedPackageRef* PackageRef = GlobalImportStore.FindPackageRef(PackageIdToLoad))
+			if (!PackageNameToLoad.IsNone() && FAsyncLoadingThreadSettings::Get().bLooseFileLoadingEnabled)
 			{
-				if (!PackageRef->HasErrors())
+				FPackagePath::TryFromPackageName(PackageNameToLoad, Request.PackagePath);
+
+				// Lookup if we have the package in the importstore already to avoid expensive calls in TryGetPackagePathFromFileSystem
+				if (FLoadedPackageRef* PackageRef = GlobalImportStore.FindPackageRef(PackageIdToLoad))
+				{
+					if (!PackageRef->HasErrors())
+					{
+						bIsZenPackage = PackageRef->IsZenPackage();
+						Request.PackagePath.SetHeaderExtension(PackageRef->GetPackageHeaderExtension());
+						PackageStatus = EPackageStoreEntryStatus::Ok;
+					}
+				}
+				else if (TryGetPackagePathFromFileSystem(PackageNameToLoad, UPackageName, Request.PackagePath))
 				{
-					bIsZenPackage = PackageRef->IsZenPackage();
-					Request.PackagePath.SetHeaderExtension(PackageRef->GetPackageHeaderExtension());
+					bIsZenPackage = false;
 					PackageStatus = EPackageStoreEntryStatus::Ok;
 				}
 			}
-			else if (TryGetPackagePathFromFileSystem(PackageNameToLoad, UPackageName, Request.PackagePath))
-			{
-				bIsZenPackage = false;
-				PackageStatus = EPackageStoreEntryStatus::Ok;
-			}
-#endif
 
 			// Fixup CustomName to handle any input string that can be converted to a long package name.
 			if (!Request.CustomName.IsNone())
@@ -5571,31 +5553,29 @@ void FAsyncPackage2::ImportPackagesRecursiveInner(FAsyncLoadingThreadState2& Thr
 		ImportedPackageStatus = PackageStore.GetPackageStoreEntry(ImportedPackageIdToLoad, ImportedPackageUPackageName, ImportedPackageEntry);
 
 		FPackagePath ImportedPackagePath;
-#if ALT2_ENABLE_LINKERLOAD_SUPPORT || WITH_EDITOR
 		bool bIsZenPackage = true;
 		bool bIsZenPackageImport = true;
-		FPackagePath::TryFromPackageName(ImportedPackageNameToLoad, ImportedPackagePath);
-		bIsZenPackage = !LinkerLoadState.IsSet();
-
-		// Lookup if we have the package in the importstore already to avoid expensive calls in TryGetPackagePathFromFileSystem
-		if (FLoadedPackageRef* PackageRef = AsyncLoadingThread.GlobalImportStore.FindPackageRef(ImportedPackageId))
+		if (!ImportedPackageNameToLoad.IsNone() && FAsyncLoadingThreadSettings::Get().bLooseFileLoadingEnabled)
 		{
-			if (!PackageRef->HasErrors())
+			FPackagePath::TryFromPackageName(ImportedPackageNameToLoad, ImportedPackagePath);
+			bIsZenPackage = !LinkerLoadState.IsSet();
+
+			// Lookup if we have the package in the importstore already to avoid expensive calls in TryGetPackagePathFromFileSystem
+			if (FLoadedPackageRef* PackageRef = AsyncLoadingThread.GlobalImportStore.FindPackageRef(ImportedPackageId))
+			{
+				if (!PackageRef->HasErrors())
+				{
+					bIsZenPackageImport = PackageRef->IsZenPackage();
+					ImportedPackagePath.SetHeaderExtension(PackageRef->GetPackageHeaderExtension());
+					ImportedPackageStatus = EPackageStoreEntryStatus::Ok;
+				}
+			}
+			else if (AsyncLoadingThread.TryGetPackagePathFromFileSystem(ImportedPackageNameToLoad, ImportedPackageUPackageName, ImportedPackagePath))
 			{
-				bIsZenPackageImport = PackageRef->IsZenPackage();
-				ImportedPackagePath.SetHeaderExtension(PackageRef->GetPackageHeaderExtension());
+				bIsZenPackageImport = false;
 				ImportedPackageStatus = EPackageStoreEntryStatus::Ok;
 			}
 		}
-		else if (AsyncLoadingThread.TryGetPackagePathFromFileSystem(ImportedPackageNameToLoad, ImportedPackageUPackageName, ImportedPackagePath))
-		{
-			bIsZenPackageImport = false;
-			ImportedPackageStatus = EPackageStoreEntryStatus::Ok;
-		}
-#else
-		constexpr bool bIsZenPackage = true;
-		constexpr bool bIsZenPackageImport = true;
-#endif
 
 		FLoadedPackageRef& ImportedPackageRef = ImportStore.AddImportedPackageReference(ImportedPackageId, ImportedPackageUPackageName, bIsZenPackageImport, ImportedPackagePath.GetHeaderExtension());
 		checkf(ImportedPackageStatus == EPackageStoreEntryStatus::Missing || ImportedPackageRef.IsZenPackage() == bIsZenPackageImport, TEXT("ImportPackageStatus(%d) != EPackageStoreEntryStatus::Missing(%d) || IsZenPackage(%d) != bIsZenPackageImport(%d)"), ImportedPackageStatus, EPackageStoreEntryStatus::Missing, ImportedPackageRef.IsZenPackage(), bIsZenPackageImport);
diff --git a/Engine/Source/Runtime/CoreUObject/Private/Serialization/AsyncPackageLoader.cpp b/Engine/Source/Runtime/CoreUObject/Private/Serialization/AsyncPackageLoader.cpp
--- a/Engine/Source/Runtime/CoreUObject/Private/Serialization/AsyncPackageLoader.cpp
+++ b/Engine/Source/Runtime/CoreUObject/Private/Serialization/AsyncPackageLoader.cpp
@@ -132,6 +132,7 @@ const FName PrestreamPackageClassNameLoad = FName("PrestreamPackage");
 FAsyncLoadingThreadSettings::FAsyncLoadingThreadSettings()
 	: bAsyncLoadingThreadEnabled(false)
 	, bAsyncPostLoadEnabled(false)
+	, bLooseFileLoadingEnabled(!!WITH_EDITOR)
 {
 #if THREADSAFE_UOBJECTS
 	check(GConfig);
@@ -166,6 +167,12 @@ FAsyncLoadingThreadSettings::FAsyncLoadingThreadSettings()
 	bCommandLineDisable = FParse::Param(FCommandLine::Get(), TEXT("NoAsyncPostLoad"));
 	bCommandLineEnable = FParse::Param(FCommandLine::Get(), TEXT("AsyncPostLoad"));
 	bAsyncPostLoadEnabled = bCommandLineEnable || (bConfigValue && FApp::ShouldUseThreadingForPerformance() && !bCommandLineDisable);
+
+	// By default the editor will look for uncooked files on disk and attempt to load them first
+	bConfigValue = !!WITH_EDITOR;
+	GConfig->GetBool(SectionName, TEXT("s.LooseFileLoadingEnabled"), bConfigValue, GEngineIni);
+	bCommandLineEnable = FParse::Param(FCommandLine::Get(), TEXT("loosefileloading"));
+	bLooseFileLoadingEnabled = bConfigValue || bCommandLineEnable;
 #endif
 }
 
diff --git a/Engine/Source/Runtime/CoreUObject/Public/Serialization/AsyncPackageLoader.h b/Engine/Source/Runtime/CoreUObject/Public/Serialization/AsyncPackageLoader.h
--- a/Engine/Source/Runtime/CoreUObject/Public/Serialization/AsyncPackageLoader.h
+++ b/Engine/Source/Runtime/CoreUObject/Public/Serialization/AsyncPackageLoader.h
@@ -72,6 +72,7 @@ struct FAsyncLoadingThreadSettings
 {
 	bool bAsyncLoadingThreadEnabled;
 	bool bAsyncPostLoadEnabled;
+	bool bLooseFileLoadingEnabled;
 
 	FAsyncLoadingThreadSettings();
 
