[AsyncLoading2]
- When looking at the FLoadedPackageRef cache for loaded packages to determine if using linkerload vs iodispatcher is necessary, only do so when the loaded package ref found is valid. We previously did this by checking if the package had errors or not but that's not enough. We defer removal of packages that have unloaded (and may have been deleted from disk) and we create in-memory packages for things that haven't been saved to disk yet. In those cases we cannot trust the cache results and should confirm the package location with the loader instead.
- Restores CL 39827287 to "Fix case where the loader was not properly loading a package if that UPackage was created outside the loader" since the changes are reasonable but were backed out due to bad interactions with the cache lookup issue fixed with this change.

#rnx

diff --git a/Engine/Source/Runtime/CoreUObject/Private/Serialization/AsyncLoading2.cpp b/Engine/Source/Runtime/CoreUObject/Private/Serialization/AsyncLoading2.cpp
--- a/Engine/Source/Runtime/CoreUObject/Private/Serialization/AsyncLoading2.cpp
+++ b/Engine/Source/Runtime/CoreUObject/Private/Serialization/AsyncLoading2.cpp
@@ -1303,7 +1303,14 @@ public:
 		{
 			return;
 		}
-		bAreAllPublicExportsLoaded = PublicExportMap.PinForGC(OutUnreachableObjectIndices);
+		// If we're missing export table objects, reset the export loaded flags
+		// but never overwrite it to true because if the UPackage was created outside
+		// of the loader it wont have any exports in its map yet and we don't want
+		// to make the loader think that the package was already loaded.
+		if (!PublicExportMap.PinForGC(OutUnreachableObjectIndices))
+		{
+			bAreAllPublicExportsLoaded = false;
+		}
 #if WITH_EDITOR
 		// Only reload if we ever got snapshotted after a proper load and the export count has shrinked since the package was last loaded
 		if (GReloadPackagesWithGCedExports && bAreAllPublicExportsLoaded && ExportCount > PublicExportMap.Num())
@@ -4030,6 +4037,23 @@ public:
 	}
 
 private:
+	bool TryGetExistingLoadedPackagePath(FPackageId InPackageIdToLoad, FPackagePath& OutPackagePath, bool& OutbIsZenPackage)
+	{
+		// Lookup if we have the package in the importstore already to avoid expensive calls in TryGetPackagePathFromFileSystem
+		// Note we must only look at valid, non-zero refcount, PackageRefs since we might not have yet cleared unloaded packages 
+		// or have yet to save and only reside in-memory currently.
+		if (FLoadedPackageRef* PackageRef = GlobalImportStore.FindPackageRef(InPackageIdToLoad); PackageRef && PackageRef->GetRefCount())
+		{
+			if (!PackageRef->HasErrors())
+			{
+				OutbIsZenPackage = PackageRef->IsZenPackage();
+				OutPackagePath.SetHeaderExtension(PackageRef->GetPackageHeaderExtension());
+				return true;
+			}
+		}
+		return false;
+	}
+
 	bool TryGetPackagePathFromFileSystem(FName& InOutPackageNameToLoad, FName& InOutUPackageName, FPackagePath& InOutPackagePath)
 	{
 		// if we prefer cooked packages in the editor, then we don't want to use the linkerload version
@@ -4894,15 +4918,9 @@ bool FAsyncLoadingThread2::CreateAsyncPackagesFromQueue(FAsyncLoadingThreadState
 			{
 				FPackagePath::TryFromPackageName(PackageNameToLoad, Request.PackagePath);
 
-				// Lookup if we have the package in the importstore already to avoid expensive calls in TryGetPackagePathFromFileSystem
-				if (FLoadedPackageRef* PackageRef = GlobalImportStore.FindPackageRef(PackageIdToLoad))
+				if (TryGetExistingLoadedPackagePath(PackageIdToLoad, Request.PackagePath, bIsZenPackage))
 				{
-					if (!PackageRef->HasErrors())
-					{
-						bIsZenPackage = PackageRef->IsZenPackage();
-						Request.PackagePath.SetHeaderExtension(PackageRef->GetPackageHeaderExtension());
-						PackageStatus = EPackageStoreEntryStatus::Ok;
-					}
+					PackageStatus = EPackageStoreEntryStatus::Ok;
 				}
 				else if (TryGetPackagePathFromFileSystem(PackageNameToLoad, UPackageName, Request.PackagePath))
 				{
@@ -5553,15 +5571,9 @@ void FAsyncPackage2::ImportPackagesRecursiveInner(FAsyncLoadingThreadState2& Thr
 			FPackagePath::TryFromPackageName(ImportedPackageNameToLoad, ImportedPackagePath);
 			bIsZenPackage = !LinkerLoadState.IsSet();
 
-			// Lookup if we have the package in the importstore already to avoid expensive calls in TryGetPackagePathFromFileSystem
-			if (FLoadedPackageRef* PackageRef = AsyncLoadingThread.GlobalImportStore.FindPackageRef(ImportedPackageId))
+			if (AsyncLoadingThread.TryGetExistingLoadedPackagePath(ImportedPackageId, ImportedPackagePath, bIsZenPackageImport))
 			{
-				if (!PackageRef->HasErrors())
-				{
-					bIsZenPackageImport = PackageRef->IsZenPackage();
-					ImportedPackagePath.SetHeaderExtension(PackageRef->GetPackageHeaderExtension());
-					ImportedPackageStatus = EPackageStoreEntryStatus::Ok;
-				}
+				ImportedPackageStatus = EPackageStoreEntryStatus::Ok;
 			}
 			else if (AsyncLoadingThread.TryGetPackagePathFromFileSystem(ImportedPackageNameToLoad, ImportedPackageUPackageName, ImportedPackagePath))
 			{
