Handle instances of verse dead classes in overridable manager, scenegraph and prefabs

#rnx

diff --git a/Engine/Source/Runtime/CoreUObject/Private/UObject/Obj.cpp b/Engine/Source/Runtime/CoreUObject/Private/UObject/Obj.cpp
--- a/Engine/Source/Runtime/CoreUObject/Private/UObject/Obj.cpp
+++ b/Engine/Source/Runtime/CoreUObject/Private/UObject/Obj.cpp
@@ -5513,6 +5513,7 @@ void InitUObject()
 	// Object initialization.
 	StaticUObjectInit();
 
+	FOverridableManager::Create();
 #if WITH_EDITORONLY_DATA
 	UE::FPropertyBagRepository::Get();
 #endif
diff --git a/Engine/Source/Runtime/CoreUObject/Private/UObject/OverridableManager.cpp b/Engine/Source/Runtime/CoreUObject/Private/UObject/OverridableManager.cpp
--- a/Engine/Source/Runtime/CoreUObject/Private/UObject/OverridableManager.cpp
+++ b/Engine/Source/Runtime/CoreUObject/Private/UObject/OverridableManager.cpp
@@ -9,11 +9,14 @@
  * Overridable serialization is experimental, not supported and use at your own risk *
  *************************************************************************************
  */
- 
-FOverridableManager& FOverridableManager::Get()
+ FOverridableManager* FOverridableManager::OverridableManager = nullptr;
+
+void FOverridableManager::Create()
 {
-	static FOverridableManager OverridableManager;
-	return OverridableManager;
+	if (!OverridableManager)
+	{
+		OverridableManager = new FOverridableManager();
+	}
 }
 
 bool FOverridableManager::IsEnabled(const UObject& Object)
@@ -390,6 +393,7 @@ FOverridableManager::FOverridableManager()
 {
 #if WITH_EDITOR
 	FCoreUObjectDelegates::OnObjectsReinstanced.AddRaw(this, &FOverridableManager::HandleObjectsReInstantiated);
+	FCoreUObjectDelegates::OnVerseDeadObjectReferences.AddRaw(this, &FOverridableManager::HandleDeadObjectReferences);
 #endif
 }
 
@@ -406,3 +410,31 @@ void FOverridableManager::HandleObjectsReInstantiated(const TMap<UObject*, UObje
 	}
 #endif // WITH_EDITORONLY_DATA
 }
+
+void FOverridableManager::AddReferencedObjects(FReferenceCollector& Collector)
+{
+#if WITH_EDITORONLY_DATA
+	const TMap<const UObjectBase *, FOverriddenPropertyAnnotation>& AnnotationMap = OverriddenObjectAnnotations.GetAnnotationMap();
+	for (const auto& Pair : AnnotationMap)
+	{
+		if( FOverriddenPropertySet* OverriddenProperties = Pair.Value.OverriddenProperties.Get())
+		{
+			OverriddenProperties->AddReferencedObjects(Collector);
+		}
+	}
+#endif // WITH_EDITORONLY_DATA
+}
+
+void FOverridableManager::HandleDeadObjectReferences(const TSet<UClass*>& DeadClasses, const TSet<UObject*>& ActiveInstances, const TSet<UObject*>& TemplateInstances)
+{
+#if WITH_EDITORONLY_DATA
+	const TMap<const UObjectBase *, FOverriddenPropertyAnnotation>& AnnotationMap = OverriddenObjectAnnotations.GetAnnotationMap();
+	for (const auto& Pair : AnnotationMap)
+	{
+		if( FOverriddenPropertySet* OverriddenProperties = Pair.Value.OverriddenProperties.Get())
+		{
+			OverriddenProperties->HandleDeadObjectReferences(ActiveInstances, TemplateInstances);
+		}
+	}
+#endif // WITH_EDITORONLY_DATA
+}
\ No newline at end of file
diff --git a/Engine/Source/Runtime/CoreUObject/Private/UObject/OverriddenPropertySet.cpp b/Engine/Source/Runtime/CoreUObject/Private/UObject/OverriddenPropertySet.cpp
--- a/Engine/Source/Runtime/CoreUObject/Private/UObject/OverriddenPropertySet.cpp
+++ b/Engine/Source/Runtime/CoreUObject/Private/UObject/OverriddenPropertySet.cpp
@@ -353,6 +353,23 @@ void FOverriddenPropertyNodeID::HandleObjectsReInstantiated(const TMap<UObject*,
 	}
 }
 
+void FOverriddenPropertyNodeID::AddReferencedObjects(FReferenceCollector& Collector)
+{
+	Collector.AddReferencedObject(Object);
+}
+
+void FOverriddenPropertyNodeID::HandleDeadObjectReferences(const TSet<UObject*>& ActiveInstances, const TSet<UObject*>& TemplateInstances)
+{
+	if (!Object)
+	{
+		return;
+	}
+
+	if (ActiveInstances.Find(Object) || TemplateInstances.Find(Object) )
+	{
+		Object = nullptr;
+	}
+}
 
 void FOverriddenPropertySet::RestoreOverriddenState(const FOverriddenPropertySet& FromOverriddenProperties)
 {
@@ -1507,6 +1524,32 @@ void FOverriddenPropertySet::HandleObjectsReInstantiated(const TMap<UObject*, UO
 	}
 }
 
+void FOverriddenPropertySet::AddReferencedObjects(FReferenceCollector& Collector)
+{
+	for (FOverriddenPropertyNode& Node : OverriddenPropertyNodes)
+	{
+		Node.NodeID.AddReferencedObjects(Collector);
+		for (auto& Pair : Node.SubPropertyNodeKeys)
+		{
+			Pair.Key.AddReferencedObjects(Collector);
+			Pair.Value.AddReferencedObjects(Collector);
+		}
+	}
+}
+
+void FOverriddenPropertySet::HandleDeadObjectReferences(const TSet<UObject*>& ActiveInstances, const TSet<UObject*>& TemplateInstances)
+{
+	for (FOverriddenPropertyNode& Node : OverriddenPropertyNodes)
+	{
+		Node.NodeID.HandleDeadObjectReferences(ActiveInstances, TemplateInstances);
+		for (auto& Pair : Node.SubPropertyNodeKeys)
+		{
+			Pair.Key.HandleDeadObjectReferences(ActiveInstances, TemplateInstances);
+			Pair.Value.HandleDeadObjectReferences(ActiveInstances, TemplateInstances);
+		}
+	}
+}
+
 const FOverriddenPropertyNode* FOverriddenPropertySet::GetOverriddenPropertyNode(const FOverriddenPropertyNode& ParentPropertyNode, const FArchiveSerializedPropertyChain* CurrentPropertyChain) const
 {
 	if (!CurrentPropertyChain)
diff --git a/Engine/Source/Runtime/CoreUObject/Private/UObject/UObjectGlobals.cpp b/Engine/Source/Runtime/CoreUObject/Private/UObject/UObjectGlobals.cpp
--- a/Engine/Source/Runtime/CoreUObject/Private/UObject/UObjectGlobals.cpp
+++ b/Engine/Source/Runtime/CoreUObject/Private/UObject/UObjectGlobals.cpp
@@ -205,6 +205,7 @@ FCoreUObjectDelegates::FOnObjectModified FCoreUObjectDelegates::OnObjectModified
 FCoreUObjectDelegates::FOnObjectTransacted FCoreUObjectDelegates::OnObjectTransacted;
 FCoreUObjectDelegates::FOnObjectsReplaced FCoreUObjectDelegates::OnObjectsReplaced;
 FCoreUObjectDelegates::FOnObjectsReplaced FCoreUObjectDelegates::OnObjectsReinstanced;
+FCoreUObjectDelegates::FOnVerseDeadReferences FCoreUObjectDelegates::OnVerseDeadObjectReferences;
 FCoreUObjectDelegates::FOnObjectPostCDOCompiled FCoreUObjectDelegates::OnObjectPostCDOCompiled;
 FCoreUObjectDelegates::FOnAssetLoaded FCoreUObjectDelegates::OnAssetLoaded;
 FCoreUObjectDelegates::FOnObjectConstructed FCoreUObjectDelegates::OnObjectConstructed;
@@ -3924,7 +3925,10 @@ void UObject::PostInitProperties()
 #endif
 
 	GetClass()->CreatePersistentUberGraphFrame(this, true);
-	FOverridableManager::Get().ClearOverrides(*this);
+	if (FOverridableManager* Manager = FOverridableManager::TryGet())
+	{
+		Manager->ClearOverrides(*this);
+	}
 }
 
 // Construct the UObject in the open with memory validation disabled as UObjects can be 
@@ -4494,7 +4498,10 @@ void FObjectInitializer::InitProperties(UObject* Obj, UClass* DefaultsClass, UOb
 
 	check(DefaultsClass && Obj);
 
-	FOverridableManager::Get().InheritEnabledFrom(*Obj, DefaultData);
+	if (FOverridableManager* Manager = FOverridableManager::TryGet())
+	{
+		Manager->InheritEnabledFrom(*Obj, DefaultData);
+	}
 
 	UClass* Class = Obj->GetClass();
 
diff --git a/Engine/Source/Runtime/CoreUObject/Public/UObject/OverridableManager.h b/Engine/Source/Runtime/CoreUObject/Public/UObject/OverridableManager.h
--- a/Engine/Source/Runtime/CoreUObject/Public/UObject/OverridableManager.h
+++ b/Engine/Source/Runtime/CoreUObject/Public/UObject/OverridableManager.h
@@ -104,12 +104,28 @@ struct FOverrideManagerCapabilityInterface;
  * Overridable serialization is experimental, not supported and use at your own risk *
  *************************************************************************************
  */
-class FOverridableManager
+class FOverridableManager : public FGCObject
 {
 public:
+	/**
+	 * @return the instance managing the overridability */
+	static FOverridableManager& Get()
+	{
+		checkf(OverridableManager, TEXT("Expected the overridable manager to be created."))
+		return *OverridableManager;
+	}
+
+	/**
+	 * @return the instance managing the overridability if available
+	 * @note it is not available only during initialization time */
+	static FOverridableManager* TryGet()
+	{
+		return OverridableManager;
+	}
+
 	/**
 	 * @return the static instance managing the overridability */
-	COREUOBJECT_API static FOverridableManager& Get();
+	COREUOBJECT_API static void Create();
 
 	/**
 	 * Lookup if the for the specified object has overridable serialization enabled
@@ -318,10 +334,22 @@ public:
 
 
 	void HandleObjectsReInstantiated(const TMap<UObject*, UObject*>& OldToNewInstanceMap);
+	void HandleDeadObjectReferences(const TSet<UClass*>& DeadClasses, const TSet<UObject*>& ActiveInstances, const TSet<UObject*>& TemplateInstances);
+
+	// FSerializableObject interface
+	virtual void AddReferencedObjects(FReferenceCollector& Collector) override;
+	virtual FString GetReferencerName() const override
+	{
+		return TEXT("FOverridableManager");
+	}
+	// End of FSerializableObject interface
+
 
 protected:
 	FOverridableManager();
 
+	COREUOBJECT_API static FOverridableManager* OverridableManager;
+
 #if WITH_EDITORONLY_DATA
 	FOverriddenPropertyAnnotations OverriddenObjectAnnotations;
 #else
diff --git a/Engine/Source/Runtime/CoreUObject/Public/UObject/OverriddenPropertySet.h b/Engine/Source/Runtime/CoreUObject/Public/UObject/OverriddenPropertySet.h
--- a/Engine/Source/Runtime/CoreUObject/Public/UObject/OverriddenPropertySet.h
+++ b/Engine/Source/Runtime/CoreUObject/Public/UObject/OverriddenPropertySet.h
@@ -236,6 +236,8 @@ struct FOverriddenPropertyNodeID
 	}
 
 	void HandleObjectsReInstantiated(const TMap<UObject*, UObject*>& Map);
+	void HandleDeadObjectReferences(const TSet<UObject*>& ActiveInstances, const TSet<UObject*>& TemplateInstances);
+	void AddReferencedObjects(FReferenceCollector& Collector);
 
 	friend uint32 GetTypeHash(const FOverriddenPropertyNodeID& NodeID)
 	{
@@ -411,6 +413,8 @@ public:
 	 * Resets all overrides of the object */
 	void Reset();
 	void HandleObjectsReInstantiated(const TMap<UObject*, UObject*>& Map);
+	void HandleDeadObjectReferences(const TSet<UObject*>& ActiveInstances, const TSet<UObject*>& TemplateInstances);
+	void AddReferencedObjects(FReferenceCollector& Collector);
 
 	/* return whether this object is considered added or not */
 	bool WasAdded() const
diff --git a/Engine/Source/Runtime/CoreUObject/Public/UObject/UObjectGlobals.h b/Engine/Source/Runtime/CoreUObject/Public/UObject/UObjectGlobals.h
--- a/Engine/Source/Runtime/CoreUObject/Public/UObject/UObjectGlobals.h
+++ b/Engine/Source/Runtime/CoreUObject/Public/UObject/UObjectGlobals.h
@@ -3174,6 +3174,12 @@ struct FCoreUObjectDelegates
 	DECLARE_MULTICAST_DELEGATE_OneParam(FOnObjectsReinstanced, const FReplacementObjectMap&);
 	static COREUOBJECT_API FOnObjectsReinstanced OnObjectsReinstanced;
 
+	/**
+	 * Called when verse found instances of object that their classes is now dead (renamed or deleted)
+	 */
+	DECLARE_MULTICAST_DELEGATE_ThreeParams(FOnVerseDeadReferences, const TSet<UClass*>& /*DeadClasses*/, const TSet<UObject*>& /*ActiveInstances*/, const TSet<UObject*>& /*TemplateInstances*/);
+	static COREUOBJECT_API FOnVerseDeadReferences OnVerseDeadObjectReferences;
+
 	/**
 	 * Called after the Blueprint compiler has finished generating the Class Default Object (CDO) for a class. This can only happen in the editor.
 	 * This is called when the CDO and its associated class structure have been fully generated and populated, and allows the assignment of cached/derived data,
