[Feature] IDO CVar support: Scene Graph specific cvars

IDO.EnablePrefabArchetypes - Allow creation of IDOs for prefabs
IDO.EnablePrefabInstances - Allow creation of IDOs for each in-world prefab instance

diff --git a/Engine/Source/Runtime/CoreUObject/Private/UObject/InstanceDataObjectUtils.cpp b/Engine/Source/Runtime/CoreUObject/Private/UObject/InstanceDataObjectUtils.cpp
--- a/Engine/Source/Runtime/CoreUObject/Private/UObject/InstanceDataObjectUtils.cpp
+++ b/Engine/Source/Runtime/CoreUObject/Private/UObject/InstanceDataObjectUtils.cpp
@@ -296,6 +296,20 @@ namespace UE
 		TEXT("When enabled, IDOs will be saved instead of instances. Disabling this will stop data retention on save.")
 	);
 
+	bool bEnableIDOsForPrefabArchetypes = true;
+	FAutoConsoleVariableRef EnableIDOsForPrefabArchetypes(
+		TEXT("IDO.EnablePrefabArchetypes"),
+		bEnableIDOsForPrefabArchetypes,
+		TEXT("When enabled, prefab archetypes can have IDOs generated for them")
+	);
+
+	bool bEnableIDOsForPrefabInstances = true;
+	FAutoConsoleVariableRef EnableIDOsForPrefabInstances(
+		TEXT("IDO.EnablePrefabInstances"),
+		bEnableIDOsForPrefabInstances,
+		TEXT("When enabled, prefab instances can have IDOs generated for them")
+	);
+
 	static TSet<FString> GetExcludedLoosePropertyTypes()
 	{
 		TArray<FString> Result;
@@ -318,6 +332,21 @@ namespace UE
 		return bEnableIDOImpersonationOnSave;
 	}
 
+	static const UObject* GetPrefabObject(const UObject* InObject)
+	{
+		static const FName NAME_BlueprintGeneratedClass(TEXT("BlueprintGeneratedClass"));
+		const UObject* Current = InObject;
+		while (Current && !Current->IsA<UPackage>())
+		{
+			if (Current->GetClass()->GetClass()->GetFName() == NAME_BlueprintGeneratedClass)
+			{
+				return Current;
+			}
+			Current = Current->GetOuter();
+		}
+		return nullptr;
+	}
+
 	bool IsInstanceDataObjectSupportEnabled(const UObject* InObject)
 	{
 		if (!IsInstanceDataObjectSupportEnabled() || !InObject)
@@ -350,6 +379,23 @@ namespace UE
 			return false;
 		}
 
+		// respect flags for disabling the generation of prefab archetypes and/or their instances
+		if (!bEnableIDOsForPrefabArchetypes || !bEnableIDOsForPrefabInstances)
+		{
+			if (const UObject* PrefabObject = GetPrefabObject(InObject))
+			{
+				const bool bCdoMatchesPrefab = PrefabObject->GetClass()->ClassDefaultObject == PrefabObject;
+				if (!bEnableIDOsForPrefabArchetypes && bCdoMatchesPrefab)
+				{
+					return false;
+				}
+				if (!bEnableIDOsForPrefabInstances && !bCdoMatchesPrefab)
+				{
+					return false;
+				}
+			}
+		}
+
 		// TODO: Temp! Remove with the conditions below.
 		while (ObjClass && ObjClass->GetClass()->GetFName() != NAME_VerseClass)
 		{
