Fix Scenegraph tests GC problem

[RN]

diff --git a/Engine/Source/Runtime/CoreUObject/Private/UObject/OverridableManager.cpp b/Engine/Source/Runtime/CoreUObject/Private/UObject/OverridableManager.cpp
--- a/Engine/Source/Runtime/CoreUObject/Private/UObject/OverridableManager.cpp
+++ b/Engine/Source/Runtime/CoreUObject/Private/UObject/OverridableManager.cpp
@@ -414,12 +414,17 @@ void FOverridableManager::HandleObjectsReInstantiated(const TMap<UObject*, UObje
 void FOverridableManager::AddReferencedObjects(FReferenceCollector& Collector)
 {
 #if WITH_EDITORONLY_DATA
-	const TMap<const UObjectBase *, FOverriddenPropertyAnnotation>& AnnotationMap = OverriddenObjectAnnotations.GetAnnotationMap();
-	for (const auto& Pair : AnnotationMap)
+	// This isn't great but this object implements this FGCObject interface just to so 
+	// that the replace references archives go through this object
+	if (!IsGarbageCollecting())
 	{
-		if( FOverriddenPropertySet* OverriddenProperties = Pair.Value.OverriddenProperties.Get())
+		const TMap<const UObjectBase *, FOverriddenPropertyAnnotation>& AnnotationMap = OverriddenObjectAnnotations.GetAnnotationMap();
+		for (const auto& Pair : AnnotationMap)
 		{
-			OverriddenProperties->AddReferencedObjects(Collector);
+			if( FOverriddenPropertySet* OverriddenProperties = Pair.Value.OverriddenProperties.Get())
+			{
+				OverriddenProperties->AddReferencedObjects(Collector);
+			}
 		}
 	}
 #endif // WITH_EDITORONLY_DATA
