Fixed LoadFromType to use paths instead of short names

Also fixed IDO registration to occur before serializing properties into the IDO.

diff --git a/Engine/Source/Runtime/CoreUObject/Private/UObject/EnumProperty.cpp b/Engine/Source/Runtime/CoreUObject/Private/UObject/EnumProperty.cpp
--- a/Engine/Source/Runtime/CoreUObject/Private/UObject/EnumProperty.cpp
+++ b/Engine/Source/Runtime/CoreUObject/Private/UObject/EnumProperty.cpp
@@ -642,13 +642,13 @@ bool FEnumProperty::LoadTypeName(UE::FPropertyTypeName Type, const FPropertyTag*
 		return false;
 	}
 
-	const FName EnumName = Type.GetParameterName(0);
-	UEnum* LocalEnum = FindFirstObject<UEnum>(*WriteToString<256>(EnumName), EFindFirstObjectOptions::NativeFirst);
+	const UE::FPropertyTypeName TypePath = Type.GetParameter(0);
+	UEnum* LocalEnum = UE::FindObjectByTypePath<UEnum>(TypePath);
 	if (!LocalEnum)
 	{
 	#if WITH_EDITORONLY_DATA
 		LocalEnum = StaticEnum<EFallbackEnum>();
-		SetMetaData(UE::NAME_OriginalType, *WriteToString<256>(Type.GetParameter(0)));
+		SetMetaData(UE::NAME_OriginalType, *WriteToString<256>(TypePath));
 	#else
 		return false;
 	#endif
diff --git a/Engine/Source/Runtime/CoreUObject/Private/UObject/PropertyByte.cpp b/Engine/Source/Runtime/CoreUObject/Private/UObject/PropertyByte.cpp
--- a/Engine/Source/Runtime/CoreUObject/Private/UObject/PropertyByte.cpp
+++ b/Engine/Source/Runtime/CoreUObject/Private/UObject/PropertyByte.cpp
@@ -483,21 +483,20 @@ bool FByteProperty::LoadTypeName(UE::FPropertyTypeName Type, const FPropertyTag*
 		return false;
 	}
 
-	const FName EnumName = Type.GetParameterName(0);
-	if (EnumName.IsNone())
+	const UE::FPropertyTypeName TypePath = Type.GetParameter(0);
+	if (TypePath.IsEmpty())
 	{
 		return true;
 	}
 
-	if (UEnum* LocalEnum = FindFirstObject<UEnum>(*WriteToString<256>(EnumName), EFindFirstObjectOptions::NativeFirst))
+	if ((Enum = UE::FindObjectByTypePath<UEnum>(TypePath)))
 	{
-		Enum = LocalEnum;
 		return true;
 	}
 
 #if WITH_EDITORONLY_DATA
 	Enum = StaticEnum<EFallbackEnum>();
-	SetMetaData(UE::NAME_OriginalType, *WriteToString<256>(Type.GetParameter(0)));
+	SetMetaData(UE::NAME_OriginalType, *WriteToString<256>(TypePath));
 	return true;
 #else
 	return false;
diff --git a/Engine/Source/Runtime/CoreUObject/Private/UObject/PropertyHelper.cpp b/Engine/Source/Runtime/CoreUObject/Private/UObject/PropertyHelper.cpp
--- a/Engine/Source/Runtime/CoreUObject/Private/UObject/PropertyHelper.cpp
+++ b/Engine/Source/Runtime/CoreUObject/Private/UObject/PropertyHelper.cpp
@@ -2,6 +2,7 @@
 
 #include "UObject/PropertyHelper.h"
 
+#include "Logging/StructuredLog.h"
 #include "String/Find.h"
 #include "String/ParseTokens.h"
 #include "UObject/Class.h"
@@ -292,6 +293,40 @@ const TCHAR* DelegatePropertyTools::ImportDelegateFromText( FScriptDelegate& Del
 namespace UE
 {
 
+UObject* FindObjectByTypePath(UClass* Class, FPropertyTypeName TypePath)
+{
+	if (const FName Name = TypePath.GetName(); !Name.IsNone())
+	{
+		TStringBuilder<256> Path;
+		if (int32 ParamCount = TypePath.GetParameterCount())
+		{
+			Path << TypePath.GetParameterName(0) << TEXT('.');
+			for (int32 OuterIndex = 1; OuterIndex < ParamCount; ++OuterIndex)
+			{
+				Path << TypePath.GetParameterName(OuterIndex) << (OuterIndex == 1 ? TEXT(':') : TEXT('.'));
+			}
+		}
+
+		const int32 NameIndex = Path.Len();
+		Path << Name;
+
+		if (NameIndex)
+		{
+			if (UObject* Object = StaticFindObject(Class, (UObject*)nullptr, *Path))
+			{
+				return Object;
+			}
+		}
+
+		if (UObject* Object = StaticFindFirstObject(Class, *Path + NameIndex, EFindFirstObjectOptions::NativeFirst))
+		{
+			UE_LOGFMT(LogProperty, Warning, "{Class} '{Path}' was found by its short name but not its path.", Class->GetFName(), Path);
+			return Object;
+		}
+	}
+	return nullptr;
+}
+
 #if WITH_EDITORONLY_DATA
 const FString* FindOriginalTypeName(const UField* Field)
 {
diff --git a/Engine/Source/Runtime/CoreUObject/Private/UObject/PropertyHelper.h b/Engine/Source/Runtime/CoreUObject/Private/UObject/PropertyHelper.h
--- a/Engine/Source/Runtime/CoreUObject/Private/UObject/PropertyHelper.h
+++ b/Engine/Source/Runtime/CoreUObject/Private/UObject/PropertyHelper.h
@@ -90,6 +90,14 @@ namespace DelegatePropertyTools
 namespace UE
 {
 
+UObject* FindObjectByTypePath(UClass* Class, FPropertyTypeName TypePath);
+
+template <typename T>
+T* FindObjectByTypePath(FPropertyTypeName TypePath)
+{
+	return (T*)FindObjectByTypePath(T::StaticClass(), TypePath);
+}
+
 #if WITH_EDITORONLY_DATA
 inline static const FName NAME_OriginalType(ANSITEXTVIEW("OriginalType"));
 const FString* FindOriginalTypeName(const UField* Field);
diff --git a/Engine/Source/Runtime/CoreUObject/Private/UObject/PropertyStruct.cpp b/Engine/Source/Runtime/CoreUObject/Private/UObject/PropertyStruct.cpp
--- a/Engine/Source/Runtime/CoreUObject/Private/UObject/PropertyStruct.cpp
+++ b/Engine/Source/Runtime/CoreUObject/Private/UObject/PropertyStruct.cpp
@@ -549,13 +549,10 @@ bool FStructProperty::LoadTypeName(UE::FPropertyTypeName Type, const FPropertyTa
 		return false;
 	}
 
-	if (const FName Name = Type.GetParameterName(0); !Name.IsNone())
+	const UE::FPropertyTypeName TypePath = Type.GetParameter(0);
+	if ((Struct = UE::FindObjectByTypePath<UScriptStruct>(TypePath)))
 	{
-		TStringBuilder<256> NameString(InPlace, Name);
-		if (Struct = FindFirstObject<UScriptStruct>(*NameString, EFindFirstObjectOptions::NativeFirst); Struct)
-		{
-			return true;
-		}
+		return true;
 	}
 
 	// TODO: Look up the struct based on the guid.
@@ -568,7 +565,7 @@ bool FStructProperty::LoadTypeName(UE::FPropertyTypeName Type, const FPropertyTa
 	if (Tag && Tag->SerializeType == EPropertyTagSerializeType::Property)
 	{
 		Struct = GetFallbackStruct();
-		SetMetaData(UE::NAME_OriginalType, *WriteToString<256>(Type.GetParameter(0)));
+		SetMetaData(UE::NAME_OriginalType, *WriteToString<256>(TypePath));
 		return true;
 	}
 #endif // WITH_EDITORONLY_DATA
