Better handling of Hidden component for the EntityActorProxyComponent.
Also do not track removed items if they are not removal of parent elements in the overridden properties.
Added Unittest to test this last case in OverridableSerialization

#rnx

diff --git a/Engine/Source/Runtime/CoreUObject/Private/UObject/OverriddenPropertySet.cpp b/Engine/Source/Runtime/CoreUObject/Private/UObject/OverriddenPropertySet.cpp
--- a/Engine/Source/Runtime/CoreUObject/Private/UObject/OverriddenPropertySet.cpp
+++ b/Engine/Source/Runtime/CoreUObject/Private/UObject/OverriddenPropertySet.cpp
@@ -883,23 +883,52 @@ void FOverriddenPropertySet::NotifyPropertyChange(FOverriddenPropertyNode* Paren
 						{
 							if(SubPropertyNode)
 							{
+								// Check if there is a matching archetype for this object
 								UObject* RemovedSubObjectArchetype = RemovedSubObject->GetArchetype();
-								const FOverriddenPropertyNodeID RemovedSubObjectID (!RemovedSubObjectArchetype || RemovedSubObjectArchetype->HasAnyFlags(RF_ClassDefaultObject) ? *RemovedSubObject : *RemovedSubObjectArchetype);
-								FOverriddenPropertyNode& RemovedSubObjectNode = FindOrAddNode(*SubPropertyNode, RemovedSubObjectID);
-
-								if (RemovedSubObjectNode.Operation == EOverriddenPropertyOperation::Add)
+								if (RemovedSubObjectArchetype && !RemovedSubObjectArchetype->HasAnyFlags(RF_ClassDefaultObject))
 								{
-									// An add then a remove becomes no opt
-									FOverriddenPropertyNodeID RemovedNodeID;
-									if (SubPropertyNode->SubPropertyNodeKeys.RemoveAndCopyValue(RemovedSubObjectID, RemovedNodeID))
+									const FOverriddenPropertyNodeID RemovedSubObjectID (*RemovedSubObjectArchetype);
+									FOverriddenPropertyNode& RemovedSubObjectNode = FindOrAddNode(*SubPropertyNode, RemovedSubObjectID);
+									if (RemovedSubObjectNode.Operation == EOverriddenPropertyOperation::Add)
+									{
+										// An add then a remove becomes no opt
+										FOverriddenPropertyNodeID RemovedNodeID;
+										if (SubPropertyNode->SubPropertyNodeKeys.RemoveAndCopyValue(RemovedSubObjectID, RemovedNodeID))
+										{
+											verifyf(OverriddenPropertyNodes.Remove(RemovedNodeID), TEXT("Expecting the node to be removed"));
+											bNeedsCleanup = true;
+										}
+									}
+									else
 									{
-										verifyf(OverriddenPropertyNodes.Remove(RemovedNodeID), TEXT("Expecting the node to be removed"));
-										bNeedsCleanup = true;
+										RemovedSubObjectNode.Operation = EOverriddenPropertyOperation::Remove;
 									}
 								}
 								else
 								{
-									RemovedSubObjectNode.Operation = EOverriddenPropertyOperation::Remove;
+									// Figure out if it is a remove of a previously added element
+									const FOverriddenPropertyNodeID RemovedSubObjectID (*RemovedSubObject);
+									if (const FOverriddenPropertyNodeID* AddedSubObjectID = SubPropertyNode->SubPropertyNodeKeys.Find(RemovedSubObjectID))
+									{
+										FOverriddenPropertyNode* AddedSubObjectNode = OverriddenPropertyNodes.Find(*AddedSubObjectID);
+										checkf(AddedSubObjectNode, TEXT("Expecting a node"));
+										if (AddedSubObjectNode->Operation != EOverriddenPropertyOperation::Add)
+										{
+											UE_LOG(LogOverridableObject, Warning, TEXT("This removed object:%s(0x%p) was not tracked as an add in the overridden properties"), *GetNameSafe(RemovedSubObject), RemovedSubObject);
+										}
+
+										// An add then a remove becomes no opt
+										FOverriddenPropertyNodeID RemovedNodeID;
+										if (SubPropertyNode->SubPropertyNodeKeys.RemoveAndCopyValue(RemovedSubObjectID, RemovedNodeID))
+										{
+											verifyf(OverriddenPropertyNodes.Remove(RemovedNodeID), TEXT("Expecting the node to be removed"));
+											bNeedsCleanup = true;
+										}
+									}
+									else
+									{
+										UE_LOG(LogOverridableObject, Log, TEXT("This removed object:%s(0x%p) was not tracked in the overridden properties"), *GetNameSafe(RemovedSubObject), RemovedSubObject);
+									}
 								}
 							}
 						}
