Preparation for removing null checks with `this` pointer, to allow taking advantage of -fdelete-null-pointer-checks optimisations.

diff --git a/Engine/Plugins/Enterprise/VariantManagerContent/Source/VariantManagerContent/Private/LevelVariantSets.cpp b/Engine/Plugins/Enterprise/VariantManagerContent/Source/VariantManagerContent/Private/LevelVariantSets.cpp
--- a/Engine/Plugins/Enterprise/VariantManagerContent/Source/VariantManagerContent/Private/LevelVariantSets.cpp
+++ b/Engine/Plugins/Enterprise/VariantManagerContent/Source/VariantManagerContent/Private/LevelVariantSets.cpp
@@ -325,7 +325,7 @@ UObject* ULevelVariantSets::GetDirectorInstance(UObject* WorldContext)
 
 	TargetDirector->GetOnDestroy().AddLambda([this](ULevelVariantSetsFunctionDirector* Director)
 	{
-		if (this != nullptr && this->IsValidLowLevel() && IsValidChecked(this) && !this->IsUnreachable())
+		if (this->IsValidLowLevel() && IsValidChecked(this) && !this->IsUnreachable())
 		{
 			HandleDirectorDestroyed(Director);
 		}
diff --git a/Engine/Plugins/Enterprise/VariantManagerContent/Source/VariantManagerContent/Private/PropertyValue.cpp b/Engine/Plugins/Enterprise/VariantManagerContent/Source/VariantManagerContent/Private/PropertyValue.cpp
--- a/Engine/Plugins/Enterprise/VariantManagerContent/Source/VariantManagerContent/Private/PropertyValue.cpp
+++ b/Engine/Plugins/Enterprise/VariantManagerContent/Source/VariantManagerContent/Private/PropertyValue.cpp
@@ -600,8 +600,7 @@ bool UPropertyValue::HasValidResolve() const
 		{
 			// Replicate IsValidLowLevel without the warning, plus a couple of extra
 			// checks because this random address might be pointing at a destroyed component
-			if (this == nullptr ||
-				!GUObjectArray.IsValid(this) ||
+			if (!GUObjectArray.IsValid(this) ||
 				!Container->GetClass() ||
 				!IsValidChecked(Container) ||
 				Container->IsUnreachable() ||
diff --git a/Engine/Plugins/Experimental/DynamicMaterial/Source/DynamicMaterial/Private/Components/DMMaterialComponent.cpp b/Engine/Plugins/Experimental/DynamicMaterial/Source/DynamicMaterial/Private/Components/DMMaterialComponent.cpp
--- a/Engine/Plugins/Experimental/DynamicMaterial/Source/DynamicMaterial/Private/Components/DMMaterialComponent.cpp
+++ b/Engine/Plugins/Experimental/DynamicMaterial/Source/DynamicMaterial/Private/Components/DMMaterialComponent.cpp
@@ -42,7 +42,10 @@ UObject* UDMMaterialComponent::GetOuterSafe() const
 
 bool UDMMaterialComponent::IsComponentValid() const
 {
-	const bool bValid = this != nullptr && IsValid(this)
+	PRAGMA_DISABLE_DEPRECATION_WARNINGS 
+	const bool bValid = IsThisNull(this,"UDMMaterialComponent::IsComponentValid")
+	PRAGMA_ENABLE_DEPRECATION_WARNINGS
+		&& IsValidChecked(this)
 		&& !HasAnyFlags(RF_BeginDestroyed | RF_FinishDestroyed);
 
 	return bValid;
diff --git a/Engine/Plugins/FX/Niagara/Source/Niagara/Private/NiagaraDataSet.cpp b/Engine/Plugins/FX/Niagara/Source/Niagara/Private/NiagaraDataSet.cpp
--- a/Engine/Plugins/FX/Niagara/Source/Niagara/Private/NiagaraDataSet.cpp
+++ b/Engine/Plugins/FX/Niagara/Source/Niagara/Private/NiagaraDataSet.cpp
@@ -19,7 +19,9 @@ TArray<FNiagaraSharedObject*> FNiagaraSharedObject::DeferredDeletionList;
 void FNiagaraSharedObject::Destroy()
 {
 	FScopeLock Lock(&CritSec);
-	check(this != nullptr);
+	PRAGMA_DISABLE_DEPRECATION_WARNINGS
+	check(IsThisNull(this, "FNiagaraSharedObject::Destroy"));
+	PRAGMA_ENABLE_DEPRECATION_WARNINGS
 	check(DeferredDeletionList.Contains(this) == false);
 	DeferredDeletionList.Add(this);
 }
diff --git a/Engine/Source/Editor/UnrealEd/Classes/Editor/EditorEngine.h b/Engine/Source/Editor/UnrealEd/Classes/Editor/EditorEngine.h
--- a/Engine/Source/Editor/UnrealEd/Classes/Editor/EditorEngine.h
+++ b/Engine/Source/Editor/UnrealEd/Classes/Editor/EditorEngine.h
@@ -3192,7 +3192,9 @@ public:
 	 */
 	UEditorSubsystem* GetEditorSubsystemBase(TSubclassOf<UEditorSubsystem> SubsystemClass) const
 	{
-		checkSlow(this != nullptr);
+		PRAGMA_DISABLE_DEPRECATION_WARNINGS
+		checkSlow(IsThisNull(this, "UEditorEngine::GetEditorSubsystemBase"));
+		PRAGMA_ENABLE_DEPRECATION_WARNINGS
 		return EditorSubsystemCollection.GetSubsystem<UEditorSubsystem>(SubsystemClass);
 	}
 
@@ -3202,7 +3204,9 @@ public:
 	template <typename TSubsystemClass>
 	TSubsystemClass* GetEditorSubsystem() const
 	{
-		checkSlow(this != nullptr);
+		PRAGMA_DISABLE_DEPRECATION_WARNINGS
+		checkSlow(IsThisNull(this, "UEditorEngine::GetEditorSubsystem"));
+		PRAGMA_ENABLE_DEPRECATION_WARNINGS
 		return EditorSubsystemCollection.GetSubsystem<TSubsystemClass>(TSubsystemClass::StaticClass());
 	}
 
diff --git a/Engine/Source/Runtime/AssetRegistry/Private/DependsNode.cpp b/Engine/Source/Runtime/AssetRegistry/Private/DependsNode.cpp
--- a/Engine/Source/Runtime/AssetRegistry/Private/DependsNode.cpp
+++ b/Engine/Source/Runtime/AssetRegistry/Private/DependsNode.cpp
@@ -666,7 +666,9 @@ bool FDependsNode::ContainsDependency(FDependsNode* InDependency,
 
 void FDependsNode::PrintDependenciesRecursive(const FString& Indent, TSet<const FDependsNode*>& VisitedNodes) const
 {
-	if ( this == NULL )
+	PRAGMA_DISABLE_DEPRECATION_WARNINGS
+	if (!IsThisNull(this, "FDependsNode::PrintDependenciesRecursive"))
+	PRAGMA_ENABLE_DEPRECATION_WARNINGS
 	{
 		UE_LOG(LogAssetRegistry, Log, TEXT("%sNULL"), *Indent);
 	}
@@ -694,7 +696,9 @@ void FDependsNode::PrintDependenciesRecursive(const FString& Indent, TSet<const
 
 void FDependsNode::PrintReferencersRecursive(const FString& Indent, TSet<const FDependsNode*>& VisitedNodes) const
 {
-	if ( this == NULL )
+	PRAGMA_DISABLE_DEPRECATION_WARNINGS
+	if (!IsThisNull(this, "FDependsNode::PrintReferencersRecursive"))
+	PRAGMA_ENABLE_DEPRECATION_WARNINGS
 	{
 		UE_LOG(LogAssetRegistry, Log, TEXT("%sNULL"), *Indent);
 	}
diff --git a/Engine/Source/Runtime/Core/Public/CoreGlobals.h b/Engine/Source/Runtime/Core/Public/CoreGlobals.h
--- a/Engine/Source/Runtime/Core/Public/CoreGlobals.h
+++ b/Engine/Source/Runtime/Core/Public/CoreGlobals.h
@@ -827,4 +827,20 @@ bool CORE_API IsRunningCookOnTheFly();
 
 namespace UE {
 	bool CORE_API IsUsingZenPakFileStreaming();
-}
\ No newline at end of file
+}
+
+/* Used to wrap around null checks with `this` which should be optimized out when not using -fno-delete-null-pointer-checks.
+ * The aim is to give the opportunity to the client code to fix any potential causes of `this` being null and stop relying
+ * on undefined behavior.
+ */
+UE_DEPRECATED(5.6, "IsThisNull has been deprecated.")
+FORCENOINLINE static bool IsThisNull(const void* This, const ANSICHAR* Function)
+{
+	if (This)
+	{
+		return true;
+	}
+
+	UE_LOG(LogCore, Warning, TEXT("'%hs' was called with a null 'this' pointer - this is no longer supported."), Function);
+	return false;
+}
diff --git a/Engine/Source/Runtime/CoreUObject/Private/UObject/Class.cpp b/Engine/Source/Runtime/CoreUObject/Private/UObject/Class.cpp
--- a/Engine/Source/Runtime/CoreUObject/Private/UObject/Class.cpp
+++ b/Engine/Source/Runtime/CoreUObject/Private/UObject/Class.cpp
@@ -2749,7 +2749,9 @@ bool UStruct::IsChildOf( const UStruct* SomeBase ) const
 {
 	// If you're looking at this check it is due to calling IsChildOf with a this nullptr. *MAKE* sure you do not call this function
 	// with a this nullptr. It is undefined behavior, and some compilers, clang13 have started to optimize out this == nullptr checks.
-	check(this);
+	PRAGMA_DISABLE_DEPRECATION_WARNINGS
+	check(IsThisNull(this, "UStruct::IsChildOf"));
+	PRAGMA_ENABLE_DEPRECATION_WARNINGS
 
 	if (SomeBase == nullptr)
 	{
diff --git a/Engine/Source/Runtime/CoreUObject/Private/UObject/Field.cpp b/Engine/Source/Runtime/CoreUObject/Private/UObject/Field.cpp
--- a/Engine/Source/Runtime/CoreUObject/Private/UObject/Field.cpp
+++ b/Engine/Source/Runtime/CoreUObject/Private/UObject/Field.cpp
@@ -498,7 +498,9 @@ bool FField::IsNative() const
 
 bool FField::IsValidLowLevel() const
 {
-	return this != nullptr;
+	PRAGMA_DISABLE_DEPRECATION_WARNINGS
+	return IsThisNull(this, "FField::IsValidLowLevel");
+	PRAGMA_ENABLE_DEPRECATION_WARNINGS
 }
 
 bool FField::IsIn(const UObject* InOwner) const
diff --git a/Engine/Source/Runtime/CoreUObject/Private/UObject/GCObjectInfo.cpp b/Engine/Source/Runtime/CoreUObject/Private/UObject/GCObjectInfo.cpp
--- a/Engine/Source/Runtime/CoreUObject/Private/UObject/GCObjectInfo.cpp
+++ b/Engine/Source/Runtime/CoreUObject/Private/UObject/GCObjectInfo.cpp
@@ -37,7 +37,9 @@ UObject* FGCObjectInfo::TryResolveObject()
 
 void FGCObjectInfo::GetPathName(FStringBuilderBase& ResultString) const
 {
-	if (this != nullptr)
+	PRAGMA_DISABLE_DEPRECATION_WARNINGS
+	if (IsThisNull(this, "FGCObjectInfo::GetPathName"))
+	PRAGMA_ENABLE_DEPRECATION_WARNINGS
 	{
 		if (Outer)
 		{
diff --git a/Engine/Source/Runtime/CoreUObject/Private/UObject/Obj.cpp b/Engine/Source/Runtime/CoreUObject/Private/UObject/Obj.cpp
--- a/Engine/Source/Runtime/CoreUObject/Private/UObject/Obj.cpp
+++ b/Engine/Source/Runtime/CoreUObject/Private/UObject/Obj.cpp
@@ -1076,7 +1076,9 @@ void UObject::FinishDestroy()
 FString UObject::GetDetailedInfo() const
 {
 	FString Result;  
-	if( this != nullptr )
+	PRAGMA_DISABLE_DEPRECATION_WARNINGS
+	if (IsThisNull(this, "UObject::GetDetailedInfo"))
+	PRAGMA_ENABLE_DEPRECATION_WARNINGS
 	{
 		Result = GetDetailedInfoInternal();
 	}
@@ -2123,7 +2125,10 @@ bool UObject::CanCheckDefaultSubObjects(bool bForceCheck, bool& bResult) const
 {
 	bool bCanCheck = true;
 	bResult = true;
-	if (this == nullptr)
+	
+	PRAGMA_DISABLE_DEPRECATION_WARNINGS
+	if (!IsThisNull(this, "UObject::CanCheckDefaultSubObjects"))
+	PRAGMA_ENABLE_DEPRECATION_WARNINGS
 	{
 		bResult = false; // these aren't in a suitable spot in their lifetime for testing
 		bCanCheck = false;
diff --git a/Engine/Source/Runtime/CoreUObject/Private/UObject/UObjectBase.cpp b/Engine/Source/Runtime/CoreUObject/Private/UObject/UObjectBase.cpp
--- a/Engine/Source/Runtime/CoreUObject/Private/UObject/UObjectBase.cpp
+++ b/Engine/Source/Runtime/CoreUObject/Private/UObject/UObjectBase.cpp
@@ -364,7 +364,9 @@ void UObjectBase::SetClass(UClass* NewClass)
 
 bool UObjectBase::IsValidLowLevelForDestruction() const
 {
-	if (this == nullptr)
+	PRAGMA_DISABLE_DEPRECATION_WARNINGS
+	if (!IsThisNull(this, "UObjectBase::IsValidLowLevelForDestruction"))
+	PRAGMA_ENABLE_DEPRECATION_WARNINGS
 	{
 		UE_LOG(LogUObjectBase, Warning, TEXT("NULL object"));
 		return false;
@@ -392,7 +394,9 @@ bool UObjectBase::IsValidLowLevelFast(bool bRecursive /*= true*/) const
 	const int32 AlignmentCheck = MIN_ALIGNMENT - 1;
 
 	// Check 'this' pointer before trying to access any of the Object's members
-	if ((this == nullptr) || (UPTRINT)this < 0x100)
+	PRAGMA_DISABLE_DEPRECATION_WARNINGS
+	if ((!IsThisNull(this, "UObjectBase::IsValidLowLevelFast")) || (UPTRINT)this < 0x100)
+	PRAGMA_ENABLE_DEPRECATION_WARNINGS
 	{
 		UE_LOG(LogUObjectBase, Error, TEXT("\'this\' pointer is invalid."));
 		return false;
diff --git a/Engine/Source/Runtime/CoreUObject/Private/UObject/UObjectBaseUtility.cpp b/Engine/Source/Runtime/CoreUObject/Private/UObject/UObjectBaseUtility.cpp
--- a/Engine/Source/Runtime/CoreUObject/Private/UObject/UObjectBaseUtility.cpp
+++ b/Engine/Source/Runtime/CoreUObject/Private/UObject/UObjectBaseUtility.cpp
@@ -54,7 +54,9 @@ void UObjectBaseUtility::GetPathName(const UObject* StopOuter, FString& ResultSt
 
 void UObjectBaseUtility::GetPathName(const UObject* StopOuter, FStringBuilderBase& ResultString) const
 {
-	if(this != StopOuter && this != NULL)
+	PRAGMA_DISABLE_DEPRECATION_WARNINGS
+	if(this != StopOuter && IsThisNull(this, "UObjectBaseUtility::GetPathName"))
+	PRAGMA_ENABLE_DEPRECATION_WARNINGS
 	{
 		UObject* ObjOuter = GetOuter();
 		if (ObjOuter && ObjOuter != StopOuter)
@@ -109,7 +111,9 @@ void UObjectBaseUtility::GetFullName(const UObject* StopOuter, FString& ResultSt
 
 void UObjectBaseUtility::GetFullName(FStringBuilderBase& ResultString, const UObject* StopOuter, EObjectFullNameFlags Flags) const
 {
-	if (this != nullptr)
+	PRAGMA_DISABLE_DEPRECATION_WARNINGS
+	if (IsThisNull(this, "UObjectBaseUtility::GetFullName"))
+	PRAGMA_ENABLE_DEPRECATION_WARNINGS
 	{
 		if (EnumHasAllFlags(Flags, EObjectFullNameFlags::IncludeClassPackage))
 		{
diff --git a/Engine/Source/Runtime/CoreUObject/Public/UObject/Field.h b/Engine/Source/Runtime/CoreUObject/Public/UObject/Field.h
--- a/Engine/Source/Runtime/CoreUObject/Public/UObject/Field.h
+++ b/Engine/Source/Runtime/CoreUObject/Public/UObject/Field.h
@@ -696,7 +696,9 @@ public:
 
 	FORCEINLINE FName GetFName() const
 	{
-		if (this != nullptr)
+		PRAGMA_DISABLE_DEPRECATION_WARNINGS
+		if (IsThisNull(this, "FField::GetFName"))
+		PRAGMA_ENABLE_DEPRECATION_WARNINGS
 		{
 			return NamePrivate;
 		}
@@ -708,7 +710,9 @@ public:
 
 	FORCEINLINE FString GetName() const
 	{
-		if (this != nullptr)
+		PRAGMA_DISABLE_DEPRECATION_WARNINGS
+		if (IsThisNull(this, "FField::GetName"))
+		PRAGMA_ENABLE_DEPRECATION_WARNINGS
 		{
 			return NamePrivate.ToString();
 		}
@@ -720,7 +724,9 @@ public:
 
 	FORCEINLINE void GetName(FString& OutName) const
 	{
-		if (this != nullptr)
+		PRAGMA_DISABLE_DEPRECATION_WARNINGS
+		if (IsThisNull(this, "FField::GetName"))
+		PRAGMA_ENABLE_DEPRECATION_WARNINGS
 		{
 			return NamePrivate.ToString(OutName);
 		}
diff --git a/Engine/Source/Runtime/Engine/Classes/Engine/Engine.h b/Engine/Source/Runtime/Engine/Classes/Engine/Engine.h
--- a/Engine/Source/Runtime/Engine/Classes/Engine/Engine.h
+++ b/Engine/Source/Runtime/Engine/Classes/Engine/Engine.h
@@ -3695,7 +3695,9 @@ public:
 	 */
 	UEngineSubsystem* GetEngineSubsystemBase(TSubclassOf<UEngineSubsystem> SubsystemClass) const
 	{
-		checkSlow(this != nullptr);
+		PRAGMA_DISABLE_DEPRECATION_WARNINGS
+		checkSlow(IsThisNull(this, "UEngine::GetEditorSubsystemBase"));
+		PRAGMA_ENABLE_DEPRECATION_WARNINGS
 		return EngineSubsystemCollection.GetSubsystem<UEngineSubsystem>(SubsystemClass);
 	}
 
@@ -3705,7 +3707,9 @@ public:
 	template <typename TSubsystemClass>
 	TSubsystemClass* GetEngineSubsystem() const
 	{
-		checkSlow(this != nullptr);
+		PRAGMA_DISABLE_DEPRECATION_WARNINGS
+		checkSlow(IsThisNull(this, "UEngine::GetEditorSubsystem"));
+		PRAGMA_ENABLE_DEPRECATION_WARNINGS
 		return EngineSubsystemCollection.GetSubsystem<TSubsystemClass>(TSubsystemClass::StaticClass());
 	}
 
