[AutoRTFM] Improve transactional safety of UObject annotations.

We avoid updating the annotation cache within a transaction because it was causing mixed open/closed writes to the same memory. However, after a transaction was complete, we had no cleanup mechanism and left the cache in a potentially stale/unsynchronized state. This would cause the level streaming system to hit the cache, read a stale pointer, and check-fail.

Additionally, the OnCommit from UObject::RemoveAnnotation had the potential to change the visible behavior of code which reads UObject annotations from within a transaction, so that has been removed. This OnCommit was added quite a long time ago, at https://p4-swarm.epicgames.net/changes/32943285 , when our strategy for dealing with UObjects was a little less rigorous.

I haven't added a test for the stale-cache behavior yet but this definitively fixes a server crash in at least one case (Bed Wars).

diff --git a/Engine/Source/Programs/AutoRTFMTests/Private/UObjectTests.cpp b/Engine/Source/Programs/AutoRTFMTests/Private/UObjectTests.cpp
--- a/Engine/Source/Programs/AutoRTFMTests/Private/UObjectTests.cpp
+++ b/Engine/Source/Programs/AutoRTFMTests/Private/UObjectTests.cpp
@@ -390,7 +390,7 @@ TEST_CASE("FUObjectAnnotationSparse.AddAnnotation")
 	struct FTestAnnotation
 	{
 		FTestAnnotation()
-			: TestAnnotationNumber(32)
+			: TestAnnotationNumber(42)
 		{
 		}
 
@@ -398,7 +398,7 @@ TEST_CASE("FUObjectAnnotationSparse.AddAnnotation")
 
 		bool IsDefault() const
 		{
-			return TestAnnotationNumber == 32;
+			return TestAnnotationNumber == 42;
 		}
 
 		bool operator==(const FTestAnnotation& Other) const
@@ -410,13 +410,17 @@ TEST_CASE("FUObjectAnnotationSparse.AddAnnotation")
 	FUObjectAnnotationSparse<FTestAnnotation, true> AnnotationMap;
 	
 	UMyAutoRTFMTestObject* Key = NewObject<UMyAutoRTFMTestObject>();
+	UMyAutoRTFMTestObject* Key2 = NewObject<UMyAutoRTFMTestObject>();
 
 	FTestAnnotation ValueA;
 	ValueA.TestAnnotationNumber = 10;
-	
+
 	FTestAnnotation ValueB;
 	ValueB.TestAnnotationNumber = 20;
 
+	FTestAnnotation ValueC;
+	ValueC.TestAnnotationNumber = 30;
+
 	SECTION("Add")
 	{
 		SECTION("Commit")
@@ -470,6 +474,76 @@ TEST_CASE("FUObjectAnnotationSparse.AddAnnotation")
 			REQUIRE(ValueB == AnnotationMap.GetAnnotation(Key));
 		}
 	}
+
+	SECTION("Add, Commit(Remove), Get")
+	{
+		AnnotationMap.AddAnnotation(Key, ValueC);
+
+		AutoRTFM::Testing::Commit([&]
+		{
+			AnnotationMap.RemoveAnnotation(Key);
+			REQUIRE(AnnotationMap.GetAnnotation(Key) == FTestAnnotation{});
+		});
+
+		REQUIRE(AnnotationMap.GetAnnotation(Key) == FTestAnnotation{});
+	}
+
+	SECTION("Add, Abort(Remove), Get")
+	{
+		AnnotationMap.AddAnnotation(Key, ValueC);
+
+		AutoRTFM::Testing::Abort([&]
+		{
+			AnnotationMap.RemoveAnnotation(Key);
+			REQUIRE(AnnotationMap.GetAnnotation(Key) == FTestAnnotation{});
+			AutoRTFM::AbortTransaction();
+		});
+
+		REQUIRE(AnnotationMap.GetAnnotation(Key) == ValueC);
+	}
+
+	SECTION("Add 1, Add 2, Commit(Get 1), Get 2")
+	{
+		AnnotationMap.AddAnnotation(Key,  ValueA);
+		AnnotationMap.AddAnnotation(Key2, ValueB);
+
+		AutoRTFM::Testing::Commit([&]
+		{
+			REQUIRE(AnnotationMap.GetAnnotation(Key) == ValueA);
+		});
+
+		REQUIRE(AnnotationMap.GetAnnotation(Key2) == ValueB);
+	}
+
+	SECTION("Add 1, Add 2, Abort(Get 1), Get 2")
+	{
+		AnnotationMap.AddAnnotation(Key, ValueA);
+		AnnotationMap.AddAnnotation(Key2, ValueB);
+
+		AutoRTFM::Testing::Abort([&]
+		{
+			REQUIRE(AnnotationMap.GetAnnotation(Key) == ValueA);
+			AutoRTFM::AbortTransaction();
+		});
+
+		REQUIRE(AnnotationMap.GetAnnotation(Key2) == ValueB);
+	}
+
+	SECTION("Add 1, Add 2, Open(Get 1), Get 2")
+	{
+		AnnotationMap.AddAnnotation(Key, ValueA);
+		AnnotationMap.AddAnnotation(Key2, ValueB);
+
+		AutoRTFM::Testing::Commit([&]
+		{
+			AutoRTFM::Open([&]
+			{
+				REQUIRE(AnnotationMap.GetAnnotation(Key) == ValueA);
+				});
+		});
+
+		REQUIRE(AnnotationMap.GetAnnotation(Key2) == ValueB);
+	}
 }
 
 struct FAnnotationObject
diff --git a/Engine/Source/Runtime/CoreUObject/Public/UObject/UObjectAnnotation.h b/Engine/Source/Runtime/CoreUObject/Public/UObject/UObjectAnnotation.h
--- a/Engine/Source/Runtime/CoreUObject/Public/UObject/UObjectAnnotation.h
+++ b/Engine/Source/Runtime/CoreUObject/Public/UObject/UObjectAnnotation.h
@@ -63,8 +63,7 @@ public:
 	/**
 	 * Constructor, initializes to nothing
 	 */
-	FUObjectAnnotationSparse() :
-		AnnotationCacheKey(NULL)
+	FUObjectAnnotationSparse()
 	{
 		// default constructor is required to be default annotation
 		check(AnnotationCacheValue.IsDefault());
@@ -75,7 +74,16 @@ public:
 	 */
 	virtual ~FUObjectAnnotationSparse()
 	{
-		RemoveAllAnnotations();		
+#if UE_AUTORTFM
+		if (bRegisteredAutoRTFMHandlers)
+		{
+			AutoRTFM::Close([this]
+			{
+				AutoRTFM::PopOnCommitHandler(this);
+			});				
+		}
+#endif
+		RemoveAllAnnotations();
 	}
 
 private:
@@ -143,7 +151,7 @@ public:
 			UE::TScopeLock AnnotationMapLock(AnnotationMapCritical);
 			SetAnnotationCacheKeyAndValue(Object, TAnnotation());
 			bHadElements = (AnnotationMap.Num() > 0);
-			AnnotationMap.RemoveAndCopyValue(AnnotationCacheKey, Result);
+			AnnotationMap.RemoveAndCopyValue(Object, Result);
 			bIsNowEmpty = (AnnotationMap.Num() == 0);
 		}
 
@@ -166,30 +174,27 @@ public:
 	 */
 	void RemoveAnnotation(const UObjectBase *Object)
 	{
-		UE_AUTORTFM_ONCOMMIT(this, Object)
+		check(Object);
+		bool bHadElements = false;
+		bool bIsNowEmpty = false;
+		// Avoid holding the lock while we call GUObjectArray.RemoveUObjectDeleteListener as it could deadlock
 		{
-			check(Object);
-			bool bHadElements = false;
-			bool bIsNowEmpty = false;
-			// Avoid holding the lock while we call GUObjectArray.RemoveUObjectDeleteListener as it could deadlock
-			{
-				UE::TScopeLock AnnotationMapLock(AnnotationMapCritical);
-				SetAnnotationCacheKeyAndValue(Object, TAnnotation());
-				bHadElements = (AnnotationMap.Num() > 0);
-				AnnotationMap.Remove(AnnotationCacheKey);
-				bIsNowEmpty = (AnnotationMap.Num() == 0);
-			}
-			if (bHadElements && bIsNowEmpty)
-			{
-				// we are removing the last one, so if we are auto removing or verifying removal, unregister now
+			UE::TScopeLock AnnotationMapLock(AnnotationMapCritical);
+			SetAnnotationCacheKeyAndValue(Object, TAnnotation());
+			bHadElements = (AnnotationMap.Num() > 0);
+			AnnotationMap.Remove(Object);
+			bIsNowEmpty = (AnnotationMap.Num() == 0);
+		}
+		if (bHadElements && bIsNowEmpty)
+		{
+			// we are removing the last one, so if we are auto removing or verifying removal, unregister now
 #if (UE_BUILD_SHIPPING || UE_BUILD_TEST)
-				if (bAutoRemove)
+			if (bAutoRemove)
 #endif
-				{
-					GUObjectArray.RemoveUObjectDeleteListener(this);
-				}
+			{
+				GUObjectArray.RemoveUObjectDeleteListener(this);
 			}
-		};
+		}
 	}
 	/**
 	 * Removes all annotation from the annotation list. 
@@ -228,18 +233,20 @@ public:
 		check(Object);
 
 		UE::TScopeLock AnnotationMapLock(AnnotationMapCritical);
-		if (AutoRTFM::IsTransactional() || Object != AnnotationCacheKey)
+
+		// If we are within a transaction, we don't touch the annotation cache, as it can be altered from the open.
+		// Look up the annotation directly from the map.
+		if (AutoRTFM::IsTransactional())
+		{
+			TAnnotation* Entry = AnnotationMap.Find(Object);
+			return Entry ? *Entry : TAnnotation();
+		}
+
+		// Update the cache if it isn't a match.
+		if (Object != AnnotationCacheKey)
 		{			
 			TAnnotation* Entry = AnnotationMap.Find(Object);
 			SetAnnotationCacheKeyAndValue(Object, Entry ? *Entry : TAnnotation());
-			if (Entry)
-			{
-				AnnotationCacheValue = *Entry;
-			}
-			else
-			{
-				AnnotationCacheValue = TAnnotation();
-			}
 		}
 
 		return AnnotationCacheValue;
@@ -271,42 +278,86 @@ public:
 private:
 	// Sets the AnnotationCacheKey and AnnotationCacheValue if called outside of
 	// an AutoRTFM transaction.
-	// This is a no-op in a transaction as we do not want to revert the cache on
-	// abort, nor do we want to have writes to the same memory locations from
-	// open and closed code.
+	// This is a no-op in a transaction to avoid writes to the same memory locations 
+	// from open and closed code. If a transaction is committed or aborted, we reset
+	// the cache to avoid reading back a stale value.
 	void SetAnnotationCacheKeyAndValue(const UObjectBase* Key, const TAnnotation& Value)
 	{
-		if (!AutoRTFM::IsTransactional())
+		if (AutoRTFM::IsTransactional())
+		{
+			MaybeRegisterAutoRTFMHandlers();
+		}
+		else
 		{
 			AnnotationCacheKey = Key;
 			AnnotationCacheValue = Value;
 		}
 	}
+
 	void SetAnnotationCacheKeyAndValue(const UObjectBase* Key, TAnnotation&& Value)
 	{
-		if (!AutoRTFM::IsTransactional())
+		if (AutoRTFM::IsTransactional())
+		{
+			MaybeRegisterAutoRTFMHandlers();
+		}
+		else
 		{
 			AnnotationCacheKey = Key;
 			AnnotationCacheValue = MoveTemp(Value);
 		}
 	}
 
+	void MaybeRegisterAutoRTFMHandlers()
+	{
+#if UE_AUTORTFM
+		if (!bRegisteredAutoRTFMHandlers)
+		{
+			checkSlow(AutoRTFM::IsTransactional());  // Precondition: our caller should have checked this
+
+			AutoRTFM::Close([this]
+			{
+				bRegisteredAutoRTFMHandlers = true;
+
+				// We avoid touching the cache key and value within a transaction, so after a transaction completes,
+				// the cache value may be outdated. To avoid stale data, we null out the cache key at the end of a
+				// transaction. It will be automatically repopulated the next time GetAnnotation is called.
+				AutoRTFM::PushOnCommitHandler(this, [this] 
+				{ 
+					AnnotationCacheKey = nullptr; 
+					bRegisteredAutoRTFMHandlers = false;
+				});
+			});
+		}
+#endif
+	}
+
 	/**
-	 * Map from live objects to an annotation
+	 * Guards the annotation map, cache key and cache value.
 	 */
-	TMap<const UObjectBase *,TAnnotation> AnnotationMap;
 	FTransactionallySafeCriticalSection AnnotationMapCritical;
 
 	/**
-	 * Key for a one-item cache of the last lookup into AnnotationMap.
+	 * Map from live objects to an annotation. Guarded by AnnotationMapCritical.
+	 */
+	TMap<const UObjectBase *, TAnnotation> AnnotationMap;
+
+	/**
+	 * Key for a one-item cache of the last lookup into AnnotationMap. Guarded by AnnotationMapCritical.
 	 * Annotation are often called back-to-back so this is a performance optimization for that.
 	 */
-	const UObjectBase* AnnotationCacheKey;
+	const UObjectBase* AnnotationCacheKey = nullptr;
+
 	/**
-	 * Value for a one-item cache of the last lookup into AnnotationMap.
+	 * Value for a one-item cache of the last lookup into AnnotationMap. Guarded by AnnotationMapCritical.
 	 */
 	TAnnotation AnnotationCacheValue;
 
+#if UE_AUTORTFM
+	/**
+	 * Tracks the presence of a AutoRTFM handlers for cleaning up the annotation cache.
+	 */
+	bool bRegisteredAutoRTFMHandlers = false;
+#endif
 };
 
 
@@ -1402,5 +1453,4 @@ private:
 	 * Map from live objects to an annotation
 	 */
 	TArray<TBitType> AnnotationArray;
-
 };
