[CoreRedirects]
- Clarifies that all wildcard matches are performed with PartialLHS matching enabled (already on by default but makes it clear DisallowPartialLHSMatching will be ignored)
- Corrects the MatchScore sorting to handle wildcard redirects so they can be applied in the appropriate order. Direct matches are explicitly higher priority than wildcard matches.

#rnx

diff --git a/Engine/Source/Runtime/CoreUObject/Private/UObject/CoreRedirects.cpp b/Engine/Source/Runtime/CoreUObject/Private/UObject/CoreRedirects.cpp
--- a/Engine/Source/Runtime/CoreUObject/Private/UObject/CoreRedirects.cpp
+++ b/Engine/Source/Runtime/CoreUObject/Private/UObject/CoreRedirects.cpp
@@ -148,15 +148,15 @@ namespace UE::CoreRedirects::Private
 		FNameUtf8String PackageName;
 	};
 
-	struct SubstringMatcher
+	struct FSubstringMatcher
 	{
-		static bool Matches(const FName LHS, const FNameUtf8String& RHSString, bool bPartialLHS, bool bPartialRHS)
+		static bool Matches(const FName LHS, const FNameUtf8String& RHSString, bool bPartialRHS)
 		{
 			const bool bLHSIsNone = LHS.IsNone();
 			const bool bRHSIsNone = RHSString.IsNone();
 			if (bLHSIsNone || bRHSIsNone)
 			{
-				return bLHSIsNone == bRHSIsNone || ((!bLHSIsNone || bPartialLHS) && (!bRHSIsNone || bPartialRHS));
+				return bLHSIsNone == bRHSIsNone || (!bRHSIsNone || bPartialRHS);
 			}
 
 			FNameUtf8String LHSString(LHS);
@@ -164,15 +164,15 @@ namespace UE::CoreRedirects::Private
 		}
 	};
 
-	struct PrefixMatcher
+	struct FPrefixMatcher
 	{
-		static bool Matches(const FName LHS, const FNameUtf8String& RHSString, bool bPartialLHS, bool bPartialRHS)
+		static bool Matches(const FName LHS, const FNameUtf8String& RHSString, bool bPartialRHS)
 		{
 			const bool bLHSIsNone = LHS.IsNone();
 			const bool bRHSIsNone = RHSString.IsNone();
 			if (bLHSIsNone || bRHSIsNone)
 			{
-				return bLHSIsNone == bRHSIsNone || ((!bLHSIsNone || bPartialLHS) && (!bRHSIsNone || bPartialRHS));
+				return bLHSIsNone == bRHSIsNone || (!bRHSIsNone || bPartialRHS);
 			}
 
 			FNameUtf8String LHSString(LHS);
@@ -180,15 +180,15 @@ namespace UE::CoreRedirects::Private
 		}
 	};
 
-	struct SuffixMatcher
+	struct FSuffixMatcher
 	{
-		static bool Matches(const FName LHS, const FNameUtf8String& RHSString, bool bPartialLHS, bool bPartialRHS)
+		static bool Matches(const FName LHS, const FNameUtf8String& RHSString, bool bPartialRHS)
 		{
 			const bool bLHSIsNone = LHS.IsNone();
 			const bool bRHSIsNone = RHSString.IsNone();
 			if (bLHSIsNone || bRHSIsNone)
 			{
-				return bLHSIsNone == bRHSIsNone || ((!bLHSIsNone || bPartialLHS) && (!bRHSIsNone || bPartialRHS));
+				return bLHSIsNone == bRHSIsNone || (!bRHSIsNone || bPartialRHS);
 			}
 
 			FNameUtf8String LHSString(LHS);
@@ -198,39 +198,36 @@ namespace UE::CoreRedirects::Private
 
 	template <typename Matcher>
 	static bool MatchWildcardRedirect(const FCoreRedirectObjectName& InRedirect, 
-		const FCoreRedirectObjectUtf8Name& InUtf8Name, bool bPartialLHS, bool bPartialRHS)
+		const FCoreRedirectObjectUtf8Name& InUtf8Name, bool bPartialRHS)
 	{
-		if (!Matcher::Matches(InRedirect.ObjectName, InUtf8Name.ObjectName, bPartialLHS, bPartialRHS))
+		if (!Matcher::Matches(InRedirect.ObjectName, InUtf8Name.ObjectName, bPartialRHS))
 		{
 			return false;
 		}
-		if (!Matcher::Matches(InRedirect.OuterName, InUtf8Name.OuterName, bPartialLHS, bPartialRHS))
+		if (!Matcher::Matches(InRedirect.OuterName, InUtf8Name.OuterName, bPartialRHS))
 		{
 			return false;
 		}
-		if (!Matcher::Matches(InRedirect.PackageName, InUtf8Name.PackageName, bPartialLHS, bPartialRHS))
+		if (!Matcher::Matches(InRedirect.PackageName, InUtf8Name.PackageName, bPartialRHS))
 		{
 			return false;
 		}
 		return true;
 	}
 
-	static bool MatchSubstring(const FCoreRedirectObjectName& InSubstringRedirect, const FCoreRedirectObjectUtf8Name& InUtf8Name,
-		bool bPartialLHS, bool bPartialRHS)
+	static bool MatchSubstring(const FCoreRedirectObjectName& InSubstringRedirect, const FCoreRedirectObjectUtf8Name& InUtf8Name, bool bPartialRHS)
 	{
-		return MatchWildcardRedirect<SubstringMatcher>(InSubstringRedirect, InUtf8Name, bPartialLHS, bPartialRHS);
+		return MatchWildcardRedirect<FSubstringMatcher>(InSubstringRedirect, InUtf8Name, bPartialRHS);
 	}
 
-	static bool MatchPrefix(const FCoreRedirectObjectName& InPrefixRedirect, const FCoreRedirectObjectUtf8Name& InUtf8Name,
-		bool bPartialLHS, bool bPartialRHS)
+	static bool MatchPrefix(const FCoreRedirectObjectName& InPrefixRedirect, const FCoreRedirectObjectUtf8Name& InUtf8Name, bool bPartialRHS)
 	{
-		return MatchWildcardRedirect<PrefixMatcher>(InPrefixRedirect, InUtf8Name, bPartialLHS, bPartialRHS);
+		return MatchWildcardRedirect<FPrefixMatcher>(InPrefixRedirect, InUtf8Name, bPartialRHS);
 	}
 
-	static bool MatchSuffix(const FCoreRedirectObjectName& InSuffixRedirect, const FCoreRedirectObjectUtf8Name& InUtf8Name,
-		bool bPartialLHS, bool bPartialRHS)
+	static bool MatchSuffix(const FCoreRedirectObjectName& InSuffixRedirect, const FCoreRedirectObjectUtf8Name& InUtf8Name, bool bPartialRHS)
 	{
-		return MatchWildcardRedirect<SuffixMatcher>(InSuffixRedirect, InUtf8Name, bPartialLHS, bPartialRHS);
+		return MatchWildcardRedirect<FSuffixMatcher>(InSuffixRedirect, InUtf8Name, bPartialRHS);
 	}
 
 
@@ -620,7 +617,7 @@ bool FCoreRedirects::FWildcardData::Matches(ECoreRedirectFlags InFlags, const FC
 	{
 		for (const FCoreRedirect& CheckRedirect : Prefixes)
 		{
-			if (MatchPrefix(CheckRedirect.OldName, ObjectUtf8Name, true /*bPartialLHS*/, bPartialRHS))
+			if (MatchPrefix(CheckRedirect.OldName, ObjectUtf8Name, bPartialRHS))
 			{
 				bFound = true;
 				OutFoundRedirects.Add(&CheckRedirect);
@@ -632,7 +629,7 @@ bool FCoreRedirects::FWildcardData::Matches(ECoreRedirectFlags InFlags, const FC
 	{
 		for (const FCoreRedirect& CheckRedirect : Suffixes)
 		{
-			if (MatchSuffix(CheckRedirect.OldName, ObjectUtf8Name, true /*bPartialLHS*/, bPartialRHS))
+			if (MatchSuffix(CheckRedirect.OldName, ObjectUtf8Name, bPartialRHS))
 			{
 				bFound = true;
 				OutFoundRedirects.Add(&CheckRedirect);
@@ -651,7 +648,7 @@ bool FCoreRedirects::FWildcardData::Matches(ECoreRedirectFlags InFlags, const FC
 			// We found a possible substring match, so now do a slow match to find if we really match
 			for (const FCoreRedirect& CheckRedirect : Substrings)
 			{
-				if (MatchSubstring(CheckRedirect.OldName, ObjectUtf8Name, true /*bPartialLHS*/, bPartialRHS))
+				if (MatchSubstring(CheckRedirect.OldName, ObjectUtf8Name, bPartialRHS))
 				{
 					bFound = true;
 					OutFoundRedirects.Add(&CheckRedirect);
@@ -677,15 +674,15 @@ bool FCoreRedirectObjectName::Matches(const FCoreRedirectObjectName& Other, EMat
 	// Substring implies prefix and suffix so we must check it first
 	if (bSubstring)
 	{
-		return MatchSubstring(*this, Other, bPartialLHS, bPartialRHS);
+		return MatchSubstring(*this, Other, bPartialRHS);
 	}
 	else if(bPrefix)
 	{
-		return MatchPrefix(*this, Other, bPartialLHS, bPartialRHS);
+		return MatchPrefix(*this, Other, bPartialRHS);
 	}
 	else if (bSuffix)
 	{
-		return MatchSuffix(*this, Other, bPartialLHS, bPartialRHS);
+		return MatchSuffix(*this, Other, bPartialRHS);
 	}
 
 	auto FieldMatches = [bPartialLHS, bPartialRHS, bSubstring, bPrefix, bSuffix](FName LHS, FName RHS)
@@ -718,15 +715,32 @@ bool FCoreRedirectObjectName::Matches(const FCoreRedirectObjectName& Other, EMat
 	return true;
 }
 
-int32 FCoreRedirectObjectName::MatchScore(const FCoreRedirectObjectName& Other) const
+template<typename Matcher>
+static int32 WildcardMatchScore(const FCoreRedirectObjectName Redirect, const FCoreRedirectObjectName& Other, bool bPartialRHS)
 {
-	int32 MatchScore = 1;
+	using namespace UE::CoreRedirects::Private;
 
-	if (ObjectName != NAME_None)
+	// Note, WildcardMatches default lower (1 vs 2) than direct matches intentionally to give priority to direct matches
+	int32 MatchScore = 1;
+	if (Redirect.ObjectName != NAME_None)
 	{
-		if (ObjectName == Other.ObjectName)
+		FNameUtf8String OtherName(Other.ObjectName);
+		if (Matcher::Matches(Redirect.ObjectName, OtherName, bPartialRHS))
 		{
 			// Object name most important
+			MatchScore += 16;
+		}
+		else
+		{
+			return 0;
+		}
+	}
+
+	if (Redirect.OuterName != NAME_None)
+	{
+		FNameUtf8String OtherName(Other.OuterName);
+		if (Matcher::Matches(Redirect.OuterName, OtherName, bPartialRHS))
+		{
 			MatchScore += 8;
 		}
 		else
@@ -735,9 +749,10 @@ int32 FCoreRedirectObjectName::MatchScore(const FCoreRedirectObjectName& Other)
 		}
 	}
 
-	if (OuterName != NAME_None)
+	if (Redirect.PackageName != NAME_None)
 	{
-		if (OuterName == Other.OuterName)
+		FNameUtf8String OtherName(Other.PackageName);
+		if (Matcher::Matches(Redirect.PackageName, OtherName, bPartialRHS))
 		{
 			MatchScore += 4;
 		}
@@ -746,16 +761,70 @@ int32 FCoreRedirectObjectName::MatchScore(const FCoreRedirectObjectName& Other)
 			return 0;
 		}
 	}
+	return MatchScore;
+}
+
+int32 FCoreRedirectObjectName::MatchScore(const FCoreRedirectObjectName& Other, ECoreRedirectFlags RedirectFlags, ECoreRedirectMatchFlags MatchFlags) const
+{
+	int32 MatchScore = 2;
+	if (!EnumHasAnyFlags(RedirectFlags, ECoreRedirectFlags::Option_MatchWildcardMask))
+	{
+		if (ObjectName != NAME_None)
+		{
+			if (ObjectName == Other.ObjectName)
+			{
+				// Object name most important
+				MatchScore += 16;
+			}
+			else
+			{
+				return 0;
+			}
+		}
 
-	if (PackageName != NAME_None)
+		if (OuterName != NAME_None)
+		{
+			if (OuterName == Other.OuterName)
+			{
+				MatchScore += 8;
+			}
+			else
+			{
+				return 0;
+			}
+		}
+
+		if (PackageName != NAME_None)
+		{
+			if (PackageName == Other.PackageName)
+			{
+				MatchScore += 4;
+			}
+			else
+			{
+				return 0;
+			}
+		}
+	}
+	else
 	{
-		if (PackageName == Other.PackageName)
+		using namespace UE::CoreRedirects::Private;
+		const bool bPartialRHS = EnumHasAnyFlags(MatchFlags, ECoreRedirectMatchFlags::AllowPartialMatch);
+		if (EnumHasAnyFlags(RedirectFlags, ECoreRedirectFlags::Option_MatchPrefix))
+		{
+			return WildcardMatchScore<FPrefixMatcher>(*this, Other, bPartialRHS);
+		}
+		else if (EnumHasAnyFlags(RedirectFlags, ECoreRedirectFlags::Option_MatchSuffix))
+		{
+			return WildcardMatchScore<FSuffixMatcher>(*this, Other, bPartialRHS);
+		}
+		else if (EnumHasAnyFlags(RedirectFlags, ECoreRedirectFlags::Option_MatchSubstring))
 		{
-			MatchScore += 2;
+			return WildcardMatchScore<FSubstringMatcher>(*this, Other, bPartialRHS);
 		}
 		else
 		{
-			return 0;
+			checkNoEntry();
 		}
 	}
 
@@ -1588,7 +1657,7 @@ void FCoreRedirects::Initialize()
 	// Enable to run startup tests
 #if 0
 	{
-		FScopeCoreRedirectsContext ScopeContext(TEXT("FCoreRedirects::RunTests"));
+		FScopeCoreRedirectsContext ScopeContext;
 		ensure(RunTests());
 	}
 #endif
@@ -1634,8 +1703,8 @@ bool FCoreRedirects::RedirectNameAndValuesUnderReadLock(ECoreRedirectFlags Type,
 		if (FoundRedirects.Num() > 1)
 		{
 			// Sort them based on match
-			FoundRedirects.Sort([&OldObjectName](const FCoreRedirect& A, const FCoreRedirect& B) 
-				{ return A.OldName.MatchScore(OldObjectName) > B.OldName.MatchScore(OldObjectName); });
+			FoundRedirects.Sort([&OldObjectName, MatchFlags](const FCoreRedirect& A, const FCoreRedirect& B)
+				{ return A.OldName.MatchScore(OldObjectName, A.RedirectFlags, MatchFlags) > B.OldName.MatchScore(OldObjectName, B.RedirectFlags, MatchFlags); });
 			NewObjectName = ProcessRedirect(FoundRedirects[0], NewObjectName);
 
 			// Apply in order
@@ -1980,6 +2049,9 @@ bool FCoreRedirects::RunTests()
 	NewRedirects.Emplace(ECoreRedirectFlags::Type_Package | ECoreRedirectFlags::Option_MatchPrefix, TEXT("/oldprefix"), TEXT("/newprefix"));
 	NewRedirects.Emplace(ECoreRedirectFlags::Type_Package | ECoreRedirectFlags::Option_MatchPrefix, TEXT("/古いプレフィックス"), TEXT("/新しいプレフィックス"));
 	NewRedirects.Emplace(ECoreRedirectFlags::Type_Package | ECoreRedirectFlags::Option_MatchPrefix, TEXT("/混合接頭辞"), TEXT("/mixed_prefix"));
+	NewRedirects.Emplace(ECoreRedirectFlags::Type_Object  | ECoreRedirectFlags::Option_MatchPrefix, TEXT("/old/object.prefix."), TEXT("/new/superobject.prefix2."));
+	NewRedirects.Emplace(ECoreRedirectFlags::Type_Package | ECoreRedirectFlags::Option_MatchPrefix, TEXT("/PrefixOverlappingDirectMatch"), TEXT("/ShouldNeverHappen"));
+	NewRedirects.Emplace(ECoreRedirectFlags::Type_Package, TEXT("/PrefixOverlappingDirectMatch/Path"), TEXT("/DirectMatchIsPreferred/NewPath"));
 	NewRedirects.Emplace(ECoreRedirectFlags::Type_Package | ECoreRedirectFlags::Option_MatchSuffix, TEXT("/oldsuffix"), TEXT("/newsuffix"));
 	NewRedirects.Emplace(ECoreRedirectFlags::Type_Package | ECoreRedirectFlags::Option_MatchSuffix, TEXT("/古い接尾辞"), TEXT("/新しい接尾辞"));
 	NewRedirects.Emplace(ECoreRedirectFlags::Type_Package | ECoreRedirectFlags::Option_MatchSuffix, TEXT("/混合接尾辞"), TEXT("/mixed_suffix"));
@@ -2027,10 +2099,14 @@ bool FCoreRedirects::RunTests()
 	Tests.Emplace(TEXT("/oldgame/Package.DefaultClass"), TEXT("/newgame/Package.DefaultClass"), ECoreRedirectFlags::Type_Package);
 	Tests.Emplace(TEXT("/古いゲーム/Package.DefaultClass"), TEXT("/新しいゲーム/Package.DefaultClass"), ECoreRedirectFlags::Type_Package);
 	Tests.Emplace(TEXT("/混合部分文字列/Package.DefaultClass"), TEXT("/mixed_substring/Package.DefaultClass"), ECoreRedirectFlags::Type_Package);
+
 	// Prefix tests
 	Tests.Emplace(TEXT("/oldprefix_SomeGame/Package.DefaultClass"), TEXT("/newprefix_SomeGame/Package.DefaultClass"), ECoreRedirectFlags::Type_Package);
 	Tests.Emplace(TEXT("/古いプレフィックス_SomeGame/Package.DefaultClass"), TEXT("/新しいプレフィックス_SomeGame/Package.DefaultClass"), ECoreRedirectFlags::Type_Package);
 	Tests.Emplace(TEXT("/混合接頭辞_SomeGame/Package.DefaultClass"), TEXT("/mixed_prefix_SomeGame/Package.DefaultClass"), ECoreRedirectFlags::Type_Package);
+	Tests.Emplace(TEXT("/old/object.prefix.subobjects.do.not.change"), TEXT("/new/superobject.prefix2.subobjects.do.not.change"), ECoreRedirectFlags::Type_Object);
+	Tests.Emplace(TEXT("/PrefixOverlappingDirectMatch/Path.Remain.Unchanged"), TEXT("/DirectMatchIsPreferred/NewPath.Remain:Unchanged"), ECoreRedirectFlags::Type_Package);
+
 	// Suffix tests
 	Tests.Emplace(TEXT("/Game/Package/oldsuffix"), TEXT("/Game/Package/newsuffix"), ECoreRedirectFlags::Type_Package);
 	Tests.Emplace(TEXT("/Game/Package/古い接尾辞"), TEXT("/Game/Package/新しい接尾辞"), ECoreRedirectFlags::Type_Package);
diff --git a/Engine/Source/Runtime/CoreUObject/Public/UObject/CoreRedirects.h b/Engine/Source/Runtime/CoreUObject/Public/UObject/CoreRedirects.h
--- a/Engine/Source/Runtime/CoreUObject/Public/UObject/CoreRedirects.h
+++ b/Engine/Source/Runtime/CoreUObject/Public/UObject/CoreRedirects.h
@@ -147,6 +147,7 @@ struct FCoreRedirectObjectName
 		 * Without this flag a Match returns true if and only if the complete string matches: LHS == RHS.
 		 * With this flag a Match returns true if and only if RHS.Contains(LHS).
 		 * This flag makes the match more expensive and should be avoided when possible.
+		 * This flag forces partial LHS matches, and will ignore the DisallowPartialLHSMatch flag when matching.
 		 */
 		CheckSubString = (1 << 2),
 		/**
@@ -154,6 +155,7 @@ struct FCoreRedirectObjectName
 		 * Without this flag a Match returns true if and only if the complete string matches: LHS == RHS.
 		 * With this flag a Match returns true if and only if RHS.StartsWith(LHS).
 		 * This flag makes the match more expensive and should be avoided when possible.
+		 * This flag forces partial LHS matches, and will ignore the DisallowPartialLHSMatch flag when matching.
 		 */
 		CheckPrefix = (1 << 3),
 		/**
@@ -161,6 +163,7 @@ struct FCoreRedirectObjectName
 		 * Without this flag a Match returns true if and only if the complete string matches: LHS == RHS.
 		 * With this flag a Match returns true if and only if RHS.EndsWith(LHS).
 		 * This flag makes the match more expensive and should be avoided when possible.
+		 * This flag forces partial LHS matches, and will ignore the DisallowPartialLHSMatch flag when matching.
 		 */
 		CheckSuffix = (1 << 4),
 	};
@@ -168,7 +171,7 @@ struct FCoreRedirectObjectName
 	COREUOBJECT_API bool Matches(const FCoreRedirectObjectName& Other, EMatchFlags MatchFlags = EMatchFlags::None) const;
 
 	/** Returns integer of degree of match. 0 if doesn't match at all, higher integer for better matches */
-	COREUOBJECT_API int32 MatchScore(const FCoreRedirectObjectName& Other) const;
+	COREUOBJECT_API int32 MatchScore(const FCoreRedirectObjectName& Other, ECoreRedirectFlags RedirectFlags, ECoreRedirectMatchFlags MatchFlags) const;
 
 	/** Fills in any empty fields on this with the corresponding fields from Other. */
 	COREUOBJECT_API void UnionFieldsInline(const FCoreRedirectObjectName& Other);
