Workaround for RPC array size limit when sending remot object data

diff --git a/Engine/Source/Runtime/CoreUObject/Private/UObject/RemoteObjectSerialization.cpp b/Engine/Source/Runtime/CoreUObject/Private/UObject/RemoteObjectSerialization.cpp
--- a/Engine/Source/Runtime/CoreUObject/Private/UObject/RemoteObjectSerialization.cpp
+++ b/Engine/Source/Runtime/CoreUObject/Private/UObject/RemoteObjectSerialization.cpp
@@ -8,6 +8,7 @@
 #include "UObject/Class.h"
 #include "UObject/PropertyPortFlags.h"
 #include "UObject/UObjectHash.h"
+#include "UObject/UnrealType.h"
 #include "Serialization/ArchiveUObject.h"
 #include "Serialization/MemoryWriter.h"
 #include "Serialization/MemoryReader.h"
@@ -32,6 +33,12 @@ FArchive& operator<<(FArchive& Ar, FRemoteObjectPathName& PathName)
 	return Ar;
 }
 
+FArchive& operator<<(FArchive& Ar, FRemoteObjectBytes& Chunk)
+{
+	Ar << Chunk.Bytes;
+	return Ar;
+}
+
 namespace UE::RemoteObject::Serialization::Disk
 {
 
@@ -132,6 +139,7 @@ protected:
 
 	UObject* RootObject = nullptr;
 	FRemoteObjectData& ObjectData;
+	TArray<uint8> SerializedBytes;
 	FString ArchiveName;
 
 #if UE_WITH_REMOTEOBJECTARCHIVE_DEBUGGING
@@ -177,9 +185,9 @@ protected:
 			NameTableSize = NameTableData.Num();
 		}
 
-		const int64 TotalSize = NameTableSize + ObjectData.Bytes.Num();
+		const int64 TotalSize = NameTableSize + ObjectData.GetNumBytes();
 
-		UE_LOG(LogRemoteObject, Log, TEXT("%s Object Data stats for %s %s (Object Data toal: %d, total: %lld):"), *GetArchiveName(), *FRemoteObjectId(RootObject).ToString(), *GetFullNameSafe(RootObject), ObjectData.Bytes.Num(), TotalSize);
+		UE_LOG(LogRemoteObject, Log, TEXT("%s Object Data stats for %s %s (Object Data toal: %d, total: %lld):"), *GetArchiveName(), *FRemoteObjectId(RootObject).ToString(), *GetFullNameSafe(RootObject), ObjectData.GetNumBytes(), TotalSize);
 		UE_LOG(LogRemoteObject, Log, TEXT("  Name Table total size: %lld (Names: %d, Paths: %d)"), NameTableSize, ObjectData.Names.Num(), ObjectData.PathNames.Num());
 
 		TArray<FString> SortedPathNames;
@@ -220,7 +228,7 @@ protected:
 public:
 
 	TArchiveRemoteObjectBase(FRemoteObjectData& InObjectData, const TCHAR* InArchiveName)
-		: T(InObjectData.Bytes)
+		: T(SerializedBytes)
 		, ObjectData(InObjectData)
 		, ArchiveName(InArchiveName)
 	{
@@ -228,10 +236,42 @@ public:
 		T::SetUseUnversionedPropertySerialization(true);
 		T::SetMigratingRemoteObjects(true);
 		T::SetPortFlags(PPF_MigratingRemoteObjects);
+
+		if (T::IsLoading())
+		{
+			int32 NumSerializedBytes = InObjectData.GetNumBytes();
+			SerializedBytes.Reserve(NumSerializedBytes);
+			for (const FRemoteObjectBytes& Chunk : ObjectData.Bytes)
+			{
+				SerializedBytes.Append(Chunk.Bytes);
+			}			
+		}
 	}
 
 	virtual ~TArchiveRemoteObjectBase()
 	{
+		if (SerializedBytes.Num() > ObjectData.GetNumBytes())
+		{
+			ObjectData.Bytes.Empty();
+
+			constexpr int32 MaxChunkSize = int32(TNumericLimits<uint16>::Max()) - 1;
+			const uint8* RawBytes = SerializedBytes.GetData();
+			int32 NumBytes = SerializedBytes.Num();
+			int32 NumChunks = (NumBytes + MaxChunkSize - 1) / MaxChunkSize;
+			ObjectData.Bytes.AddDefaulted(NumChunks);
+
+			for (FRemoteObjectBytes& Chunk : ObjectData.Bytes)
+			{
+				int32 ChunkSize = FMath::Min(NumBytes, MaxChunkSize);
+				if (ChunkSize > 0)
+				{
+					Chunk.Bytes.AddZeroed(ChunkSize);
+					FMemory::Memcpy(Chunk.Bytes.GetData(), RawBytes, ChunkSize);
+					RawBytes += ChunkSize;
+					NumBytes -= ChunkSize;
+				}
+			}
+		}
 #if UE_WITH_REMOTEOBJECTARCHIVE_DEBUGGING
 		DumpStatsToLog();
 #endif
diff --git a/Engine/Source/Runtime/CoreUObject/Public/UObject/NoExportTypes.h b/Engine/Source/Runtime/CoreUObject/Public/UObject/NoExportTypes.h
--- a/Engine/Source/Runtime/CoreUObject/Public/UObject/NoExportTypes.h
+++ b/Engine/Source/Runtime/CoreUObject/Public/UObject/NoExportTypes.h
@@ -3417,6 +3417,13 @@ struct FRemoteObjectPathName
 	TArray<uint16> Names;
 };
 
+USTRUCT(noexport, BlueprintType)
+struct FRemoteObjectBytes
+{
+	UPROPERTY()
+	TArray<uint8> Bytes;
+};
+
 USTRUCT(noexport, BlueprintType)
 struct FRemoteObjectData
 {
@@ -3429,7 +3436,7 @@ struct FRemoteObjectData
 	TArray<FRemoteObjectPathName> PathNames;
 
 	UPROPERTY()
-	TArray<uint8> Bytes;
+	TArray<FRemoteObjectBytes> Bytes;
 };
 
 USTRUCT(noexport, BlueprintType)
diff --git a/Engine/Source/Runtime/CoreUObject/Public/UObject/RemoteObjectTypes.h b/Engine/Source/Runtime/CoreUObject/Public/UObject/RemoteObjectTypes.h
--- a/Engine/Source/Runtime/CoreUObject/Public/UObject/RemoteObjectTypes.h
+++ b/Engine/Source/Runtime/CoreUObject/Public/UObject/RemoteObjectTypes.h
@@ -152,6 +152,14 @@ struct FRemoteObjectPathName
 	TArray<FNameIndexType> Names;
 };
 
+/**
+* Structure that holds serialized remote object data chunk (< 64KB of data) (noexport type)
+*/
+struct FRemoteObjectBytes
+{
+	TArray<uint8> Bytes;
+};
+
 /**
 * Structure that holds remote object memory (noexport type)
 */
@@ -159,5 +167,15 @@ struct FRemoteObjectData
 {
 	TArray<FName> Names;
 	TArray<FRemoteObjectPathName> PathNames;
-	TArray<uint8> Bytes;
+	TArray<FRemoteObjectBytes> Bytes;
+
+	inline int32 GetNumBytes() const
+	{
+		int32 Num = 0;
+		for (const FRemoteObjectBytes& Chunk : Bytes)
+		{
+			Num += Chunk.Bytes.Num();
+		}
+		return Num;
+	}
 };
