Renamed IsThisNull to IsThisNotNull

diff --git a/Engine/Plugins/Experimental/DynamicMaterial/Source/DynamicMaterial/Private/Components/DMMaterialComponent.cpp b/Engine/Plugins/Experimental/DynamicMaterial/Source/DynamicMaterial/Private/Components/DMMaterialComponent.cpp
--- a/Engine/Plugins/Experimental/DynamicMaterial/Source/DynamicMaterial/Private/Components/DMMaterialComponent.cpp
+++ b/Engine/Plugins/Experimental/DynamicMaterial/Source/DynamicMaterial/Private/Components/DMMaterialComponent.cpp
@@ -43,7 +43,7 @@ UObject* UDMMaterialComponent::GetOuterSafe() const
 bool UDMMaterialComponent::IsComponentValid() const
 {
 	PRAGMA_DISABLE_DEPRECATION_WARNINGS 
-	const bool bValid = IsThisNull(this,"UDMMaterialComponent::IsComponentValid")
+	const bool bValid = IsThisNotNull(this,"UDMMaterialComponent::IsComponentValid")
 	PRAGMA_ENABLE_DEPRECATION_WARNINGS
 		&& IsValidChecked(this)
 		&& !HasAnyFlags(RF_BeginDestroyed | RF_FinishDestroyed);
diff --git a/Engine/Plugins/FX/Niagara/Source/Niagara/Private/NiagaraDataSet.cpp b/Engine/Plugins/FX/Niagara/Source/Niagara/Private/NiagaraDataSet.cpp
--- a/Engine/Plugins/FX/Niagara/Source/Niagara/Private/NiagaraDataSet.cpp
+++ b/Engine/Plugins/FX/Niagara/Source/Niagara/Private/NiagaraDataSet.cpp
@@ -20,7 +20,7 @@ void FNiagaraSharedObject::Destroy()
 {
 	FScopeLock Lock(&CritSec);
 	PRAGMA_DISABLE_DEPRECATION_WARNINGS
-	check(IsThisNull(this, "FNiagaraSharedObject::Destroy"));
+	check(IsThisNotNull(this, "FNiagaraSharedObject::Destroy"));
 	PRAGMA_ENABLE_DEPRECATION_WARNINGS
 	check(DeferredDeletionList.Contains(this) == false);
 	DeferredDeletionList.Add(this);
diff --git a/Engine/Source/Editor/UnrealEd/Classes/Editor/EditorEngine.h b/Engine/Source/Editor/UnrealEd/Classes/Editor/EditorEngine.h
--- a/Engine/Source/Editor/UnrealEd/Classes/Editor/EditorEngine.h
+++ b/Engine/Source/Editor/UnrealEd/Classes/Editor/EditorEngine.h
@@ -3192,7 +3192,7 @@ public:
 	UEditorSubsystem* GetEditorSubsystemBase(TSubclassOf<UEditorSubsystem> SubsystemClass) const
 	{
 		PRAGMA_DISABLE_DEPRECATION_WARNINGS
-		checkSlow(IsThisNull(this, "UEditorEngine::GetEditorSubsystemBase"));
+		checkSlow(IsThisNotNull(this, "UEditorEngine::GetEditorSubsystemBase"));
 		PRAGMA_ENABLE_DEPRECATION_WARNINGS
 		return EditorSubsystemCollection.GetSubsystem<UEditorSubsystem>(SubsystemClass);
 	}
@@ -3204,7 +3204,7 @@ public:
 	TSubsystemClass* GetEditorSubsystem() const
 	{
 		PRAGMA_DISABLE_DEPRECATION_WARNINGS
-		checkSlow(IsThisNull(this, "UEditorEngine::GetEditorSubsystem"));
+		checkSlow(IsThisNotNull(this, "UEditorEngine::GetEditorSubsystem"));
 		PRAGMA_ENABLE_DEPRECATION_WARNINGS
 		return EditorSubsystemCollection.GetSubsystem<TSubsystemClass>(TSubsystemClass::StaticClass());
 	}
diff --git a/Engine/Source/Runtime/AssetRegistry/Private/DependsNode.cpp b/Engine/Source/Runtime/AssetRegistry/Private/DependsNode.cpp
--- a/Engine/Source/Runtime/AssetRegistry/Private/DependsNode.cpp
+++ b/Engine/Source/Runtime/AssetRegistry/Private/DependsNode.cpp
@@ -667,7 +667,7 @@ bool FDependsNode::ContainsDependency(FDependsNode* InDependency,
 void FDependsNode::PrintDependenciesRecursive(const FString& Indent, TSet<const FDependsNode*>& VisitedNodes) const
 {
 	PRAGMA_DISABLE_DEPRECATION_WARNINGS
-	if (!IsThisNull(this, "FDependsNode::PrintDependenciesRecursive"))
+	if (!IsThisNotNull(this, "FDependsNode::PrintDependenciesRecursive"))
 	PRAGMA_ENABLE_DEPRECATION_WARNINGS
 	{
 		UE_LOG(LogAssetRegistry, Log, TEXT("%sNULL"), *Indent);
@@ -697,7 +697,7 @@ void FDependsNode::PrintDependenciesRecursive(const FString& Indent, TSet<const
 void FDependsNode::PrintReferencersRecursive(const FString& Indent, TSet<const FDependsNode*>& VisitedNodes) const
 {
 	PRAGMA_DISABLE_DEPRECATION_WARNINGS
-	if (!IsThisNull(this, "FDependsNode::PrintReferencersRecursive"))
+	if (!IsThisNotNull(this, "FDependsNode::PrintReferencersRecursive"))
 	PRAGMA_ENABLE_DEPRECATION_WARNINGS
 	{
 		UE_LOG(LogAssetRegistry, Log, TEXT("%sNULL"), *Indent);
diff --git a/Engine/Source/Runtime/Core/Public/CoreGlobals.h b/Engine/Source/Runtime/Core/Public/CoreGlobals.h
--- a/Engine/Source/Runtime/Core/Public/CoreGlobals.h
+++ b/Engine/Source/Runtime/Core/Public/CoreGlobals.h
@@ -833,8 +833,8 @@ namespace UE {
  * The aim is to give the opportunity to the client code to fix any potential causes of `this` being null and stop relying
  * on undefined behavior.
  */
-UE_DEPRECATED(5.6, "IsThisNull has been deprecated.")
-FORCENOINLINE static bool IsThisNull(const void* This, const ANSICHAR* Function)
+UE_DEPRECATED(5.6, "IsThisNotNull has been deprecated.")
+FORCENOINLINE static bool IsThisNotNull(const void* This, const ANSICHAR* Function)
 {
 	if (This)
 	{
diff --git a/Engine/Source/Runtime/CoreUObject/Private/UObject/Class.cpp b/Engine/Source/Runtime/CoreUObject/Private/UObject/Class.cpp
--- a/Engine/Source/Runtime/CoreUObject/Private/UObject/Class.cpp
+++ b/Engine/Source/Runtime/CoreUObject/Private/UObject/Class.cpp
@@ -2750,7 +2750,7 @@ bool UStruct::IsChildOf( const UStruct* SomeBase ) const
 	// If you're looking at this check it is due to calling IsChildOf with a this nullptr. *MAKE* sure you do not call this function
 	// with a this nullptr. It is undefined behavior, and some compilers, clang13 have started to optimize out this == nullptr checks.
 	PRAGMA_DISABLE_DEPRECATION_WARNINGS
-	check(IsThisNull(this, "UStruct::IsChildOf"));
+	check(IsThisNotNull(this, "UStruct::IsChildOf"));
 	PRAGMA_ENABLE_DEPRECATION_WARNINGS
 
 	if (SomeBase == nullptr)
diff --git a/Engine/Source/Runtime/CoreUObject/Private/UObject/Field.cpp b/Engine/Source/Runtime/CoreUObject/Private/UObject/Field.cpp
--- a/Engine/Source/Runtime/CoreUObject/Private/UObject/Field.cpp
+++ b/Engine/Source/Runtime/CoreUObject/Private/UObject/Field.cpp
@@ -499,7 +499,7 @@ bool FField::IsNative() const
 bool FField::IsValidLowLevel() const
 {
 	PRAGMA_DISABLE_DEPRECATION_WARNINGS
-	return IsThisNull(this, "FField::IsValidLowLevel");
+	return IsThisNotNull(this, "FField::IsValidLowLevel");
 	PRAGMA_ENABLE_DEPRECATION_WARNINGS
 }
 
diff --git a/Engine/Source/Runtime/CoreUObject/Private/UObject/GCObjectInfo.cpp b/Engine/Source/Runtime/CoreUObject/Private/UObject/GCObjectInfo.cpp
--- a/Engine/Source/Runtime/CoreUObject/Private/UObject/GCObjectInfo.cpp
+++ b/Engine/Source/Runtime/CoreUObject/Private/UObject/GCObjectInfo.cpp
@@ -38,7 +38,7 @@ UObject* FGCObjectInfo::TryResolveObject()
 void FGCObjectInfo::GetPathName(FStringBuilderBase& ResultString) const
 {
 	PRAGMA_DISABLE_DEPRECATION_WARNINGS
-	if (IsThisNull(this, "FGCObjectInfo::GetPathName"))
+	if (IsThisNotNull(this, "FGCObjectInfo::GetPathName"))
 	PRAGMA_ENABLE_DEPRECATION_WARNINGS
 	{
 		if (Outer)
diff --git a/Engine/Source/Runtime/CoreUObject/Private/UObject/Obj.cpp b/Engine/Source/Runtime/CoreUObject/Private/UObject/Obj.cpp
--- a/Engine/Source/Runtime/CoreUObject/Private/UObject/Obj.cpp
+++ b/Engine/Source/Runtime/CoreUObject/Private/UObject/Obj.cpp
@@ -1077,7 +1077,7 @@ FString UObject::GetDetailedInfo() const
 {
 	FString Result;  
 	PRAGMA_DISABLE_DEPRECATION_WARNINGS
-	if (IsThisNull(this, "UObject::GetDetailedInfo"))
+	if (IsThisNotNull(this, "UObject::GetDetailedInfo"))
 	PRAGMA_ENABLE_DEPRECATION_WARNINGS
 	{
 		Result = GetDetailedInfoInternal();
@@ -2127,7 +2127,7 @@ bool UObject::CanCheckDefaultSubObjects(bool bForceCheck, bool& bResult) const
 	bResult = true;
 	
 	PRAGMA_DISABLE_DEPRECATION_WARNINGS
-	if (!IsThisNull(this, "UObject::CanCheckDefaultSubObjects"))
+	if (!IsThisNotNull(this, "UObject::CanCheckDefaultSubObjects"))
 	PRAGMA_ENABLE_DEPRECATION_WARNINGS
 	{
 		bResult = false; // these aren't in a suitable spot in their lifetime for testing
diff --git a/Engine/Source/Runtime/CoreUObject/Private/UObject/UObjectBase.cpp b/Engine/Source/Runtime/CoreUObject/Private/UObject/UObjectBase.cpp
--- a/Engine/Source/Runtime/CoreUObject/Private/UObject/UObjectBase.cpp
+++ b/Engine/Source/Runtime/CoreUObject/Private/UObject/UObjectBase.cpp
@@ -365,7 +365,7 @@ void UObjectBase::SetClass(UClass* NewClass)
 bool UObjectBase::IsValidLowLevelForDestruction() const
 {
 	PRAGMA_DISABLE_DEPRECATION_WARNINGS
-	if (!IsThisNull(this, "UObjectBase::IsValidLowLevelForDestruction"))
+	if (!IsThisNotNull(this, "UObjectBase::IsValidLowLevelForDestruction"))
 	PRAGMA_ENABLE_DEPRECATION_WARNINGS
 	{
 		UE_LOG(LogUObjectBase, Warning, TEXT("NULL object"));
@@ -395,7 +395,7 @@ bool UObjectBase::IsValidLowLevelFast(bool bRecursive /*= true*/) const
 
 	// Check 'this' pointer before trying to access any of the Object's members
 	PRAGMA_DISABLE_DEPRECATION_WARNINGS
-	if ((!IsThisNull(this, "UObjectBase::IsValidLowLevelFast")) || (UPTRINT)this < 0x100)
+	if ((!IsThisNotNull(this, "UObjectBase::IsValidLowLevelFast")) || (UPTRINT)this < 0x100)
 	PRAGMA_ENABLE_DEPRECATION_WARNINGS
 	{
 		UE_LOG(LogUObjectBase, Error, TEXT("\'this\' pointer is invalid."));
diff --git a/Engine/Source/Runtime/CoreUObject/Private/UObject/UObjectBaseUtility.cpp b/Engine/Source/Runtime/CoreUObject/Private/UObject/UObjectBaseUtility.cpp
--- a/Engine/Source/Runtime/CoreUObject/Private/UObject/UObjectBaseUtility.cpp
+++ b/Engine/Source/Runtime/CoreUObject/Private/UObject/UObjectBaseUtility.cpp
@@ -55,7 +55,7 @@ void UObjectBaseUtility::GetPathName(const UObject* StopOuter, FString& ResultSt
 void UObjectBaseUtility::GetPathName(const UObject* StopOuter, FStringBuilderBase& ResultString) const
 {
 	PRAGMA_DISABLE_DEPRECATION_WARNINGS
-	if(this != StopOuter && IsThisNull(this, "UObjectBaseUtility::GetPathName"))
+	if(this != StopOuter && IsThisNotNull(this, "UObjectBaseUtility::GetPathName"))
 	PRAGMA_ENABLE_DEPRECATION_WARNINGS
 	{
 		UObject* ObjOuter = GetOuter();
@@ -112,7 +112,7 @@ void UObjectBaseUtility::GetFullName(const UObject* StopOuter, FString& ResultSt
 void UObjectBaseUtility::GetFullName(FStringBuilderBase& ResultString, const UObject* StopOuter, EObjectFullNameFlags Flags) const
 {
 	PRAGMA_DISABLE_DEPRECATION_WARNINGS
-	if (IsThisNull(this, "UObjectBaseUtility::GetFullName"))
+	if (IsThisNotNull(this, "UObjectBaseUtility::GetFullName"))
 	PRAGMA_ENABLE_DEPRECATION_WARNINGS
 	{
 		if (EnumHasAllFlags(Flags, EObjectFullNameFlags::IncludeClassPackage))
diff --git a/Engine/Source/Runtime/CoreUObject/Public/UObject/Field.h b/Engine/Source/Runtime/CoreUObject/Public/UObject/Field.h
--- a/Engine/Source/Runtime/CoreUObject/Public/UObject/Field.h
+++ b/Engine/Source/Runtime/CoreUObject/Public/UObject/Field.h
@@ -697,7 +697,7 @@ public:
 	FORCEINLINE FName GetFName() const
 	{
 		PRAGMA_DISABLE_DEPRECATION_WARNINGS
-		if (IsThisNull(this, "FField::GetFName"))
+		if (IsThisNotNull(this, "FField::GetFName"))
 		PRAGMA_ENABLE_DEPRECATION_WARNINGS
 		{
 			return NamePrivate;
@@ -711,7 +711,7 @@ public:
 	FORCEINLINE FString GetName() const
 	{
 		PRAGMA_DISABLE_DEPRECATION_WARNINGS
-		if (IsThisNull(this, "FField::GetName"))
+		if (IsThisNotNull(this, "FField::GetName"))
 		PRAGMA_ENABLE_DEPRECATION_WARNINGS
 		{
 			return NamePrivate.ToString();
@@ -725,7 +725,7 @@ public:
 	FORCEINLINE void GetName(FString& OutName) const
 	{
 		PRAGMA_DISABLE_DEPRECATION_WARNINGS
-		if (IsThisNull(this, "FField::GetName"))
+		if (IsThisNotNull(this, "FField::GetName"))
 		PRAGMA_ENABLE_DEPRECATION_WARNINGS
 		{
 			return NamePrivate.ToString(OutName);
diff --git a/Engine/Source/Runtime/Engine/Classes/Engine/Engine.h b/Engine/Source/Runtime/Engine/Classes/Engine/Engine.h
--- a/Engine/Source/Runtime/Engine/Classes/Engine/Engine.h
+++ b/Engine/Source/Runtime/Engine/Classes/Engine/Engine.h
@@ -3696,7 +3696,7 @@ public:
 	UEngineSubsystem* GetEngineSubsystemBase(TSubclassOf<UEngineSubsystem> SubsystemClass) const
 	{
 		PRAGMA_DISABLE_DEPRECATION_WARNINGS
-		checkSlow(IsThisNull(this, "UEngine::GetEditorSubsystemBase"));
+		checkSlow(IsThisNotNull(this, "UEngine::GetEditorSubsystemBase"));
 		PRAGMA_ENABLE_DEPRECATION_WARNINGS
 		return EngineSubsystemCollection.GetSubsystem<UEngineSubsystem>(SubsystemClass);
 	}
@@ -3708,7 +3708,7 @@ public:
 	TSubsystemClass* GetEngineSubsystem() const
 	{
 		PRAGMA_DISABLE_DEPRECATION_WARNINGS
-		checkSlow(IsThisNull(this, "UEngine::GetEditorSubsystem"));
+		checkSlow(IsThisNotNull(this, "UEngine::GetEditorSubsystem"));
 		PRAGMA_ENABLE_DEPRECATION_WARNINGS
 		return EngineSubsystemCollection.GetSubsystem<TSubsystemClass>(TSubsystemClass::StaticClass());
 	}
