Log the packages that are causing us to retry a transaction in the transactionally-safe async loader.

diff --git a/Engine/Source/Runtime/CoreUObject/Private/Serialization/TransactionallySafeAsyncLoading.cpp b/Engine/Source/Runtime/CoreUObject/Private/Serialization/TransactionallySafeAsyncLoading.cpp
--- a/Engine/Source/Runtime/CoreUObject/Private/Serialization/TransactionallySafeAsyncLoading.cpp
+++ b/Engine/Source/Runtime/CoreUObject/Private/Serialization/TransactionallySafeAsyncLoading.cpp
@@ -291,7 +291,7 @@ public:
 						AutoRTFM::CascadingRetryTransaction([this, RequestIdsClone = ::MoveTemp(RequestIdsClone)]() mutable
 						{
 							UE_LOG(LogStreaming, Display, TEXT("A call to `FlushLoading` that is flushing non-cached packages is causing Verse to abort and retry."));
-							DrainTransactionallyDeferredLoadPackages(false, &RequestIdsClone);
+							DrainTransactionallyDeferredLoadPackages(false, &RequestIdsClone, true);
 							UpdateMaxFlushedRequestId(RequestIdsClone);
 							WrappedPackageLoader->FlushLoading(RequestIdsClone);
 							UE_LOG(LogStreaming, Display, TEXT("`FlushLoading` has completed after Verse aborted, and we are now retrying."));
@@ -493,7 +493,7 @@ private:
 			InLoadFlags);
 	}
 
-	void DrainTransactionallyDeferredLoadPackages(const bool bCallOriginalCompletionDelegate = true, TArray<int32>* RequestIds = nullptr)
+	void DrainTransactionallyDeferredLoadPackages(const bool bCallOriginalCompletionDelegate = true, TArray<int32>* RequestIds = nullptr, const bool bLogPackagePaths = false)
 	{
 		AutoRTFM::UnreachableIfTransactional(UnreachableMessage);
 
@@ -514,6 +514,14 @@ private:
 			}
 		}
 
+		if (bLogPackagePaths)
+		{
+			for (const FTransactionallyDeferredLoadPackagePayload& Payload : TransactionallyDeferredLoadPackages)
+			{
+				UE_LOG(LogStreaming, Display, TEXT("Loading and caching '%s' in the transactionally-safe async loader."), *Payload.PackagePath.GetDebugNameWithExtension());
+			}
+		}
+
 		TransactionallyDeferredLoadPackages.Empty();
 	}
 };
