Updated core remote object transfer delegates to include migration flags

diff --git a/Engine/Source/Runtime/CoreUObject/Private/UObject/RemoteObjectTransfer.cpp b/Engine/Source/Runtime/CoreUObject/Private/UObject/RemoteObjectTransfer.cpp
--- a/Engine/Source/Runtime/CoreUObject/Private/UObject/RemoteObjectTransfer.cpp
+++ b/Engine/Source/Runtime/CoreUObject/Private/UObject/RemoteObjectTransfer.cpp
@@ -25,11 +25,10 @@ TDelegate<void(FRemoteTransactionId /*RequestId*/, FRemoteWorkPriority /*Request
 TDelegate<void(FRemoteTransactionId /*RequestId*/, FRemoteObjectId /*ObjectId*/, FRemoteServerId /*DestinationServerId*/, FRemoteServerId /*OwnerServerId*/, FRemoteObjectData& /*ObjectData*/)> RemoteObjectTransferDelegate;
 TDelegate<void(FRemoteTransactionId /*RequestId*/, FRemoteObjectId /*ObjectId*/, FRemoteServerId /*DestinationServerId*/)> RemoteObjectDeniedTransferDelegate;
 
-TMulticastDelegate<void(const FRemoteObjectData& /*ObjectData*/, FRemoteServerId /*RemoteServerId*/)> OnObjectDataSentDelegate;
-TMulticastDelegate<void(const TSet<UObject*>& /*MigratedObjects*/, FRemoteServerId /*RemoteServerId*/)> OnObjectsSentDelegate;
-TMulticastDelegate<void(const FRemoteObjectData& /*ObjectData*/, FRemoteServerId /*RemoteServerId*/)> OnObjectDataReceivedDelegate;
-TMulticastDelegate<void(FRemoteObjectId /*ObjectId*/, FRemoteServerId /*RemoteServerId*/)> OnObjectDataDeniedDelegate;
-TMulticastDelegate<void(const TArray<UObject*>& /*MigratedObjects*/, FRemoteServerId /*RemoteServerId*/)> OnObjectsReceivedDelegate;
+TMulticastDelegate<void(const FRemoteObjectData& /*ObjectData*/, FRemoteServerId /*RemoteServerId*/, EObjectMigrationFlags /*MigrationFlags*/)> OnObjectDataSentDelegate;
+TMulticastDelegate<void(const TSet<UObject*>& /*MigratedObjects*/, FRemoteServerId /*RemoteServerId*/, EObjectMigrationFlags /*MigrationFlags*/)> OnObjectsSentDelegate;
+TMulticastDelegate<void(const FRemoteObjectData& /*ObjectData*/, FRemoteServerId /*RemoteServerId*/, EObjectMigrationFlags /*MigrationFlags*/)> OnObjectDataReceivedDelegate;
+TMulticastDelegate<void(const TArray<UObject*>& /*MigratedObjects*/, FRemoteServerId /*RemoteServerId*/, EObjectMigrationFlags /*MigrationFlags*/)> OnObjectsReceivedDelegate;
 
 TDelegate<void(FRemoteTransactionId /*RequestId*/, FRemoteObjectId /*ObjectId*/, FRemoteServerId /*DatabaseId*/, FRemoteServerId /*OwnerServerId*/, FRemoteObjectData& /*ObjectData*/)> StoreRemoteObjectDataDelegate;
 TDelegate<void(FRemoteTransactionId /*RequestId*/, FRemoteWorkPriority /*RequestPriority*/, FRemoteObjectId /*ObjectId*/, FRemoteServerId /*DatabaseId*/, FRemoteServerId /*DestinationServerId*/)> RestoreRemoteObjectDataDelegate;
@@ -247,7 +246,7 @@ UObject* MigrateObjectFromObjectData(FRemoteServerId OwnerServerId, FRemoteObjec
 	check(!AutoRTFM::IsClosed());
 
 	// Notify that object data has been migrated from a remote server
-	OnObjectDataReceivedDelegate.Broadcast(ObjectData, RemoteServerId);
+	OnObjectDataReceivedDelegate.Broadcast(ObjectData, RemoteServerId, MigrationFlags);
 
 	TArray<FRemoteObjectId> ReceivedObjectRemoteIds;
 	TArray<UObject*> ReceivedObjects;
@@ -286,7 +285,7 @@ UObject* MigrateObjectFromObjectData(FRemoteServerId OwnerServerId, FRemoteObjec
 #endif
 
 	// Notify that objects have been migrated from a remote server
-	OnObjectsReceivedDelegate.Broadcast(ReceivedObjects, RemoteServerId);
+	OnObjectsReceivedDelegate.Broadcast(ReceivedObjects, RemoteServerId, MigrationFlags);
 
 	// do any existing requests require any of the objects in this hierarchy?
 	// if so, we can now update the canonical root id
@@ -490,8 +489,11 @@ void MigrateObjectToRemoteServerWithExplicitPriority(FRemoteTransactionId Reques
 			UE_LOG(LogRemoteObject, VeryVerbose, TEXT("         sent subobject : %s [%s]"), *SubObjectId.ToString(), *GetNameSafe(SO));
 		}
 
+		// Are we loaning it out or changing ownership?
+		const EObjectMigrationFlags MigrationFlags = bOwnershipChange ? EObjectMigrationFlags::Reassign : EObjectMigrationFlags::Loan;
+			
 		// Notify that object data has been migrated (actually will be but SendObjectData does a MoveTemp on the data)
-		OnObjectDataSentDelegate.Broadcast(*ObjectData, DestinationServerId);
+		OnObjectDataSentDelegate.Broadcast(*ObjectData, DestinationServerId, MigrationFlags);
 
 		GTransferQueue.SendRemoteObject(RequestId, ObjectId, DestinationServerId, OwnerServerId, ObjectData);
 
@@ -503,13 +505,12 @@ void MigrateObjectToRemoteServerWithExplicitPriority(FRemoteTransactionId Reques
 			FUnsafeToMigrateScope UnsafeToMigrateScope;
 
 			// Notify that objects have been migrated to a remote server
-			OnObjectsSentDelegate.Broadcast(SentObjects, DestinationServerId);
+			OnObjectsSentDelegate.Broadcast(SentObjects, DestinationServerId, MigrationFlags);
 
-			// Are we loaning it out or changing ownership?
 			FPostMigrateParams PostMigrateParams;
 			PostMigrateParams.RequestingServerId = RequestPriority.GetRootServerId();
 			PostMigrateParams.DestinationServerId = DestinationServerId;
-			PostMigrateParams.MigrationFlags = bOwnershipChange ? EObjectMigrationFlags::Reassign : EObjectMigrationFlags::Loan;
+			PostMigrateParams.MigrationFlags = MigrationFlags;
 			for (UObject* SentObject : SentObjects)
 			{
 #if UE_WITH_REMOTE_OBJECT_HANDLE
diff --git a/Engine/Source/Runtime/CoreUObject/Public/UObject/RemoteObjectTransfer.h b/Engine/Source/Runtime/CoreUObject/Public/UObject/RemoteObjectTransfer.h
--- a/Engine/Source/Runtime/CoreUObject/Public/UObject/RemoteObjectTransfer.h
+++ b/Engine/Source/Runtime/CoreUObject/Public/UObject/RemoteObjectTransfer.h
@@ -134,21 +134,21 @@ namespace UE::RemoteObject::Transfer
 	/** 
 	 * Delegate executed when objects have been migrated from another server. 
 	 */
-	extern COREUOBJECT_API TMulticastDelegate<void(const FRemoteObjectData& /*ObjectData*/, FRemoteServerId /*RemoteServerId*/)> OnObjectDataReceivedDelegate;
+	extern COREUOBJECT_API TMulticastDelegate<void(const FRemoteObjectData& /*ObjectData*/, FRemoteServerId /*RemoteServerId*/, EObjectMigrationFlags /*MigrationFlags*/)> OnObjectDataReceivedDelegate;
 
 	/**
 	 * Delegate executed when objects have been migrated from another server.
 	 */
-	extern COREUOBJECT_API TMulticastDelegate<void(const TArray<UObject*>& /*MigratedObjects*/, FRemoteServerId /*RemoteServerId*/)> OnObjectsReceivedDelegate;
+	extern COREUOBJECT_API TMulticastDelegate<void(const TArray<UObject*>& /*MigratedObjects*/, FRemoteServerId /*RemoteServerId*/, EObjectMigrationFlags /*MigrationFlags*/)> OnObjectsReceivedDelegate;
 
 	/** Delegate executed when object data has been migrated to another server */
-	extern COREUOBJECT_API TMulticastDelegate<void(const FRemoteObjectData& /*ObjectData*/, FRemoteServerId /*RemoteServerId*/)> OnObjectDataSentDelegate;
+	extern COREUOBJECT_API TMulticastDelegate<void(const FRemoteObjectData& /*ObjectData*/, FRemoteServerId /*RemoteServerId*/, EObjectMigrationFlags /*MigrationFlags*/)> OnObjectDataSentDelegate;
 
 	/** 
 	 * Delegate executed when objects have been migrated to another server. Migrated objects can still be resolved with weak pointers but at this point 
 	 * any changes to their internal state won't be migrated across to the remote server 
 	 */
-	extern COREUOBJECT_API TMulticastDelegate<void(const TSet<UObject*>& /*MigratedObjects*/, FRemoteServerId /*RemoteServerId*/)> OnObjectsSentDelegate;
+	extern COREUOBJECT_API TMulticastDelegate<void(const TSet<UObject*>& /*MigratedObjects*/, FRemoteServerId /*RemoteServerId*/, EObjectMigrationFlags /*MigrationFlags*/)> OnObjectsSentDelegate;
 
 	/** Delegate that stores locally unrachable object data into a database. */
 	extern COREUOBJECT_API TDelegate<void(FRemoteTransactionId /*RequestId*/, FRemoteObjectId /*ObjectId*/, FRemoteServerId /*DatabaseId*/, FRemoteServerId /*OwnerServerId*/, FRemoteObjectData& /*ObjectData*/)> StoreRemoteObjectDataDelegate;
