dstm: rework remote physics locks

- remote physics queries are now grouped by a 'lock token' per server. before sending queries to a remote server, we 'BeginLockToken'. this allows the remote server to notify us if any writers invalidate our query results such that we need to try them all again

- transactions can now be aborted and request abandonment. this happens when we try to execute a remote query that has had its lock token invalidated. the originating server was already told of the invalidation and it is no longer expecting a response from that query

- FPhysSceneLockRemoteObject lightly wraps FPhysSceneLockTransactionallySafe so we can call into the query queue logic via callbacks

diff --git a/Engine/Source/Runtime/CoreUObject/Private/UObject/RemoteExecutor.cpp b/Engine/Source/Runtime/CoreUObject/Private/UObject/RemoteExecutor.cpp
--- a/Engine/Source/Runtime/CoreUObject/Private/UObject/RemoteExecutor.cpp
+++ b/Engine/Source/Runtime/CoreUObject/Private/UObject/RemoteExecutor.cpp
@@ -17,6 +17,13 @@ TMulticastDelegate<void(FRemoteTransactionId)> OnTransactionCompletedDelegate;
 TMulticastDelegate<void(FRemoteTransactionId)> OnTransactionAbortedDelegate;
 }
 
+enum class ERemoteExecutorAbortReason : uint8
+{
+	Unspecified,
+	RequiresDependencies,
+	AbandonWork
+};
+
 FArchive& operator<<(FArchive& Ar, FRemoteTransactionId& Id)
 {
 	Ar << Id.Id;
@@ -29,6 +36,7 @@ struct FRemoteExecutorWork
 	FRemoteWorkPriority Priority;
 	TFunction<void(void)> Work;
 	TArray<void*> SubsystemRequests;
+	uint32 ExecutionAttempts = 0;
 };
 
 class FRemoteExecutor
@@ -40,7 +48,8 @@ class FRemoteExecutor
 	TArray<FRemoteExecutorWork> PendingWorks;
 
 public:
-	FString DebugLastAbortReason;
+	ERemoteExecutorAbortReason AbortReason = ERemoteExecutorAbortReason::Unspecified;
+	FString AbortReasonDescription;
 
 	void RegisterSubsystem(FRemoteSubsystemBase* Subsystem)
 	{
@@ -177,30 +186,54 @@ public:
 				}
 
 				// all of the subsystems are ready, try to perform the work
-				DebugLastAbortReason = TEXT("");
+				AbortReason = ERemoteExecutorAbortReason::Unspecified;
+				AbortReasonDescription = TEXT("");
+
+				ExecutingWork.ExecutionAttempts++;
 
 				FRemoteTransactionId RequestId = ExecutingWork.RequestId;
 				UE::RemoteExecutor::OnTransactionStartingDelegate.Broadcast(RequestId);
 				AutoRTFM::ETransactionResult WorkTransactResult = AutoRTFM::Transact([&ExecutingWork]() { ExecutingWork.Work(); });
+				bool bWorkComplete = false;
+				bool bWorkTransactionAborted = false;
 
 				if (WorkTransactResult == AutoRTFM::ETransactionResult::AbortedByRequest)
 				{
-					// the work aborted, it should have inserted a new requested object, retry
-					UE_LOG(LogRemoteExec, Verbose, TEXT("\twork ABORTED (%s)"), *DebugLastAbortReason);
+					bWorkTransactionAborted = true;
+
 					UE::RemoteExecutor::OnTransactionAbortedDelegate.Broadcast(RequestId);
 
-					PendingWorkIndex++;
+					if (AbortReason == ERemoteExecutorAbortReason::RequiresDependencies)
+					{
+						UE_LOG(LogRemoteExec, Verbose, TEXT("\twork ABORTED for required dependency '%s'"), *AbortReasonDescription);
+					}
+					else if (AbortReason == ERemoteExecutorAbortReason::AbandonWork)
+					{
+						UE_LOG(LogRemoteExec, Verbose, TEXT("\twork ABORTED for ABANDONMENT '%s'"), *AbortReasonDescription);
+						bWorkComplete = true;
+					}
+					else
+					{
+						UE_LOG(LogRemoteExec, Verbose, TEXT("\twork ABORTED for unknown reason '%s'"), *AbortReasonDescription);
+					}
 				}
 				else
 				{
-					UE_LOG(LogRemoteExec, VeryVerbose, TEXT("\twork COMPLETED!"));
+					check(AbortReason == ERemoteExecutorAbortReason::Unspecified);
+
+					UE_LOG(LogRemoteExec, VeryVerbose, TEXT("\twork COMPLETED after %u attempts"), ExecutingWork.ExecutionAttempts);
 					UE::RemoteExecutor::OnTransactionCompletedDelegate.Broadcast(RequestId);
 
+					bWorkComplete = true;
+				}
+
+				if (bWorkComplete)
+				{
 					for (int32 SubsystemIndex = 0; SubsystemIndex < Subsystems.Num(); SubsystemIndex++)
 					{
 						FRemoteSubsystemBase* Subsystem = Subsystems[SubsystemIndex];
 						UE_LOG(LogRemoteExec, VeryVerbose, TEXT("\tend request[%s]"), Subsystem->NameForDebug());
-						Subsystem->EndRequest();
+						Subsystem->EndRequest(!bWorkTransactionAborted);
 						Subsystem->ClearActiveRequest();
 					}
 
@@ -210,12 +243,13 @@ public:
 						Subsystem->DestroyRequest(ExecutingWork.RequestId);
 					}
 
-					UE_LOG(LogRemoteExec, VeryVerbose, TEXT("-- ExecuteTransactional %s DONE --"), *ExecutingWork.RequestId.ToString());
-
-					// remove this work from the list
+					// remove this work from the list and do not increment PendingWorkIndex
 					PendingWorks.RemoveAt(PendingWorkIndex, EAllowShrinking::No);
-
-					// do not increment PendingWorkIndex
+				}
+				else
+				{
+					// we will revisit this later, move on to the next work
+					PendingWorkIndex++;
 				}
 			}
 		}
@@ -249,12 +283,26 @@ void WaitForServersToSync()
 	WaitForServersToSyncDelegate.ExecuteIfBound();
 }
 
-void SetAbortReason(FStringView Reason)
+void AbortTransactionRequiresDependencies(FStringView Description)
 {
 	UE_AUTORTFM_OPEN
 	{
-		GRemoteExecutor.DebugLastAbortReason = Reason;
+		GRemoteExecutor.AbortReason = ERemoteExecutorAbortReason::RequiresDependencies;
+		GRemoteExecutor.AbortReasonDescription = Description;
 	};
+
+	AutoRTFM::AbortTransaction();
+}
+
+void AbortTransactionAndAbandonWork(FStringView Description)
+{
+	UE_AUTORTFM_OPEN
+	{
+		GRemoteExecutor.AbortReason = ERemoteExecutorAbortReason::AbandonWork;
+		GRemoteExecutor.AbortReasonDescription = Description;
+	};
+
+	AutoRTFM::AbortTransaction();
 }
 
 FRemoteWorkPriority CreateRootWorkPriority()
diff --git a/Engine/Source/Runtime/CoreUObject/Private/UObject/RemoteObjectTransfer.cpp b/Engine/Source/Runtime/CoreUObject/Private/UObject/RemoteObjectTransfer.cpp
--- a/Engine/Source/Runtime/CoreUObject/Private/UObject/RemoteObjectTransfer.cpp
+++ b/Engine/Source/Runtime/CoreUObject/Private/UObject/RemoteObjectTransfer.cpp
@@ -154,7 +154,7 @@ public:
 		return true;
 	}
 
-	void EndRequest() final
+	void EndRequest(bool bTransactionCommitted) final
 	{
 	}
 
@@ -358,8 +358,7 @@ void MigrateObjectFromRemoteServer(FRemoteObjectId ObjectId, FRemoteServerId Cur
 	// TODO: DO we want to log when we're pulling over an object we've NEVER known about?!
 	// 
 	// abort so the outer handler can renegotiate object transfers and retry
-	UE::RemoteExecutor::SetAbortReason(TEXT("object migration"));
-	AutoRTFM::AbortTransaction();
+	UE::RemoteExecutor::AbortTransactionRequiresDependencies(TEXT("object migration"));
 }
 
 void MigrateObjectToRemoteServer(FRemoteTransactionId RequestId, UObject* Object, FRemoteServerId DestinationServerId)
diff --git a/Engine/Source/Runtime/CoreUObject/Public/UObject/RemoteExecutor.h b/Engine/Source/Runtime/CoreUObject/Public/UObject/RemoteExecutor.h
--- a/Engine/Source/Runtime/CoreUObject/Public/UObject/RemoteExecutor.h
+++ b/Engine/Source/Runtime/CoreUObject/Public/UObject/RemoteExecutor.h
@@ -108,7 +108,7 @@ public:
 	virtual void BeginRequest() = 0;
 	virtual void Tick() = 0;
 	virtual bool AreDependenciesSatisfied() const = 0;
-	virtual void EndRequest() = 0;
+	virtual void EndRequest(bool bTransactionCommitted) = 0;
 };
 
 template<typename RequestType>
@@ -117,17 +117,17 @@ class FRemoteSubsystem : public FRemoteSubsystemBase
 public:
 	RequestType* ActiveRequest = nullptr;
 
-	RequestType* GetRequestByIndex(int32 RequestIndex)
+	RequestType* GetRequestByIndex(int32 RequestIndex) const
 	{
-		TTuple<FRemoteTransactionId, void*> RequestData = Requests[RequestIndex];
+		const TTuple<FRemoteTransactionId, void*> RequestData = Requests[RequestIndex];
 		return static_cast<RequestType*>(RequestData.Get<1>());
 	}
 
-	RequestType* TryGetRequest(FRemoteTransactionId RequestId)
+	RequestType* TryGetRequest(FRemoteTransactionId RequestId) const
 	{
 		for (int32 RequestIndex = 0; RequestIndex < Requests.Num(); RequestIndex++)
 		{
-			TTuple<FRemoteTransactionId, void*>& RequestData = Requests[RequestIndex];
+			const TTuple<FRemoteTransactionId, void*>& RequestData = Requests[RequestIndex];
 			if (RequestData.Get<0>() == RequestId)
 				return static_cast<RequestType*>(RequestData.Get<1>());
 		}
@@ -135,13 +135,13 @@ public:
 		return nullptr;
 	}
 
-	int32 TryGetRequestIndex(FRemoteTransactionId RequestId)
+	int32 TryGetRequestIndex(FRemoteTransactionId RequestId) const
 	{
 		int32 Result = INDEX_NONE;
 
 		for (int32 RequestIndex = 0; RequestIndex < Requests.Num(); RequestIndex++)
 		{
-			TTuple<FRemoteTransactionId, void*>& RequestData = Requests[RequestIndex];
+			const TTuple<FRemoteTransactionId, void*>& RequestData = Requests[RequestIndex];
 			if (RequestData.Get<0>() == RequestId)
 			{
 				Result = RequestIndex;
@@ -207,7 +207,8 @@ namespace UE::RemoteExecutor
 	COREUOBJECT_API void ExecuteTransactional(const TFunctionRef<void(void)>& Work);
 	COREUOBJECT_API void ExecuteTransactionalWithExplicitPriority(FRemoteWorkPriority WorkPriority, const TFunctionRef<void(void)>& Work);
 
-	COREUOBJECT_API void SetAbortReason(FStringView Reason);
+	COREUOBJECT_API void AbortTransactionRequiresDependencies(FStringView Description);
+	COREUOBJECT_API void AbortTransactionAndAbandonWork(FStringView Description);
 	
 	/** Establishes a synchronization barrier across all servers */
 	COREUOBJECT_API void WaitForServersToSync();
diff --git a/Engine/Source/Runtime/Experimental/Chaos/Private/Framework/Threading.cpp b/Engine/Source/Runtime/Experimental/Chaos/Private/Framework/Threading.cpp
--- a/Engine/Source/Runtime/Experimental/Chaos/Private/Framework/Threading.cpp
+++ b/Engine/Source/Runtime/Experimental/Chaos/Private/Framework/Threading.cpp
@@ -111,4 +111,8 @@ namespace Chaos
 		}
 
 	}
-}
\ No newline at end of file
+
+#if UE_AUTORTFM && UE_WITH_REMOTE_OBJECT_HANDLE
+	FPhysSceneLockCallbacks GPhysSceneLockRemoteObjectCallbacks;
+#endif
+}
diff --git a/Engine/Source/Runtime/Experimental/Chaos/Public/Framework/Threading.h b/Engine/Source/Runtime/Experimental/Chaos/Public/Framework/Threading.h
--- a/Engine/Source/Runtime/Experimental/Chaos/Public/Framework/Threading.h
+++ b/Engine/Source/Runtime/Experimental/Chaos/Public/Framework/Threading.h
@@ -1117,7 +1117,63 @@ FORCEINLINE void EnsureIsInGameThreadContext()
 		TSharedPtr<FState> State;
 	};
 
+#if UE_WITH_REMOTE_OBJECT_HANDLE
+	// With remote object support, we wrap the underlying FPhysSceneLockTransactionallySafe
+	// in an additional layer and expose callbacks to register additional logic to execute
+
+	struct FPhysSceneLockRemoteObject;
+
+	struct FPhysSceneLockCallbacks
+	{
+		void (*ReadLock)(FPhysSceneLockRemoteObject*) = nullptr;
+		void (*ReadUnlock)(FPhysSceneLockRemoteObject*) = nullptr;
+		void (*WriteLock)(FPhysSceneLockRemoteObject*) = nullptr;
+		void (*WriteUnlock)(FPhysSceneLockRemoteObject*) = nullptr;
+	};
+
+	extern FPhysSceneLockCallbacks GPhysSceneLockRemoteObjectCallbacks;
+
+	struct FPhysSceneLockRemoteObject
+	{
+		void ReadLock()
+		{
+			if (GPhysSceneLockRemoteObjectCallbacks.ReadLock)
+				GPhysSceneLockRemoteObjectCallbacks.ReadLock(this);
+			else
+				UnderlyingLock.ReadLock();
+		}
+
+		void ReadUnlock()
+		{
+			if (GPhysSceneLockRemoteObjectCallbacks.ReadUnlock)
+				GPhysSceneLockRemoteObjectCallbacks.ReadUnlock(this);
+			else
+				UnderlyingLock.ReadUnlock();
+		}
+
+		void WriteLock()
+		{
+			if (GPhysSceneLockRemoteObjectCallbacks.WriteLock)
+				GPhysSceneLockRemoteObjectCallbacks.WriteLock(this);
+			else
+				UnderlyingLock.WriteLock();
+		}
+
+		void WriteUnlock()
+		{
+			if (GPhysSceneLockRemoteObjectCallbacks.WriteUnlock)
+				GPhysSceneLockRemoteObjectCallbacks.WriteUnlock(this);
+			else
+				UnderlyingLock.WriteUnlock();
+		}
+
+		FPhysSceneLockTransactionallySafe UnderlyingLock;
+	};
+
+	using FPhysSceneLock = FPhysSceneLockRemoteObject;
+#else
 	using FPhysSceneLock = FPhysSceneLockTransactionallySafe;
+#endif
 #else
 	using FPhysSceneLock = FPhysSceneLockNonTransactional;
 #endif
