Don't assume ownership of a remote object before it's explicitly assigned with ChangeOwnerServerId

diff --git a/Engine/Source/Runtime/CoreUObject/Private/UObject/RemoteObject.cpp b/Engine/Source/Runtime/CoreUObject/Private/UObject/RemoteObject.cpp
--- a/Engine/Source/Runtime/CoreUObject/Private/UObject/RemoteObject.cpp
+++ b/Engine/Source/Runtime/CoreUObject/Private/UObject/RemoteObject.cpp
@@ -58,9 +58,10 @@ public:
 			Stub->Id = Id;
 			Stub->ResidentServerId = ResidentServerId.IsValid() ? ResidentServerId : Id.GetServerId();
 
-			// if we are creating the stub, then this object wasn't received by
-			// a remote server, therefore we are the owner
-			Stub->OwningServerId = GetGlobalServerId();
+			// if we are creating the stub, then this object's owner is deduced from its ID
+			// if the server ID is invalid then it's a local native object that was created before the local server had its ID assigned
+			FRemoteServerId ObjectServerId = Id.GetServerId();
+			Stub->OwningServerId = ObjectServerId.IsValid() ? ObjectServerId : GetGlobalServerId();
 		}		
 		return Stub;
 	}
@@ -75,9 +76,10 @@ public:
 		{
 			Stub = new FRemoteObjectStub(Object);
 
-			// if we are creating the stub, then this object wasn't received by
-			// a remote server, therefore we are the owner
-			Stub->OwningServerId = GetGlobalServerId();
+			// if we are creating the stub, then this object's owner is deduced from its ID
+			// if the server ID is invalid then it's a local native object that was created before the local server had its ID assigned
+			FRemoteServerId ObjectServerId = Stub->Id.GetServerId();
+			Stub->OwningServerId = ObjectServerId.IsValid() ? ObjectServerId : GetGlobalServerId();
 		}
 		Stub->ResidentServerId = DestinationServerId;
 		FUObjectItem* ObjectItem = GUObjectArray.ObjectToObjectItem(Object);
@@ -267,8 +269,9 @@ bool IsOwned(FRemoteObjectId ObjectId)
 	}
 	else
 	{
-		// if the object wasn't received or ever migrated, we own it locally
-		Result = true;
+		const FRemoteServerId ServerId = ObjectId.GetServerId();
+		// Invalid server Id means local native objects which were created before the local server had a chance to have an id assigned
+		return !ServerId.IsValid() || ServerId == UE::RemoteObject::GetGlobalServerId();
 	}
 #endif
 	return Result;
