Fix for IDOs stomping over owner property values resoled via SerializeFromMismatchedTag.
This was caused because of the coupling we had between an IDO and its owner when finding the FOverriddenPropertyAnnotations for an object (the default value of the renamed property would be marked as overriden).

Changes:
- Removed coupling between an IDO and its owner when finding the FOverriddenPropertyAnnotations for an object.
- No longer skipping NotifyPropertyChange on IDOs.
- Removed GetOwningObjectOfProperty as IDOs are no longer 'keyed' via their owners.
- Jordan's fix to FRemappedChainEvent for the EPropertyChangeType::ArrayAdd case.

#rnx

diff --git a/Engine/Source/Runtime/CoreUObject/Private/UObject/OverridableManager.cpp b/Engine/Source/Runtime/CoreUObject/Private/UObject/OverridableManager.cpp
--- a/Engine/Source/Runtime/CoreUObject/Private/UObject/OverridableManager.cpp
+++ b/Engine/Source/Runtime/CoreUObject/Private/UObject/OverridableManager.cpp
@@ -297,11 +297,6 @@ void FOverridableManager::PostOverrideProperty(UObject& Object, const FPropertyV
 void FOverridableManager::NotifyPropertyChange(const EPropertyNotificationType Notification, UObject& Object, FPropertyVisitorPath::Iterator PropertyIterator, const EPropertyChangeType::Type ChangeType)
 {
 #if WITH_EDITORONLY_DATA
-	if (UE::IsClassOfInstanceDataObjectClass(Object.GetClass()))
-	{
-		// IDOs share override keys with their associated instances so override events are handled by the instance
-		return;
-	}
 	if (FOverriddenPropertySet* ThisObjectOverriddenProperties = OverriddenObjectAnnotations.Find(Object))
 	{
 		ThisObjectOverriddenProperties->NotifyPropertyChange(Notification, PropertyIterator, ChangeType, &Object);
diff --git a/Engine/Source/Runtime/CoreUObject/Private/UObject/OverriddenPropertySet.cpp b/Engine/Source/Runtime/CoreUObject/Private/UObject/OverriddenPropertySet.cpp
--- a/Engine/Source/Runtime/CoreUObject/Private/UObject/OverriddenPropertySet.cpp
+++ b/Engine/Source/Runtime/CoreUObject/Private/UObject/OverriddenPropertySet.cpp
@@ -1313,33 +1313,11 @@ void FOverriddenPropertySet::RemoveOverriddenSubProperties(FOverriddenPropertyNo
 	PropertyNode.SubPropertyNodeKeys.Empty();
 }
 
-// IDO properties are keyed as if they're instance properties. This creates some ambiguity so we need to check what we're looking at here.
-static UObject* GetOwningObjectOfProperty(UObject* ExpectedOwner, FPropertyVisitorPath::Iterator PropertyIterator)
-{
-	if(!PropertyIterator)
-	{
-		return ExpectedOwner;
-	}
-
-	const UClass* Class = PropertyIterator->Property->GetOwnerClass();
-	UObject* PropertyOwner = ExpectedOwner;
-#if WITH_EDITOR
-	if (UE::StructIsInstanceDataObjectStruct(Class))
-	{
-		PropertyOwner = UE::FPropertyBagRepository::Get().FindInstanceDataObject(PropertyOwner);
-	}
-	checkf(PropertyOwner && PropertyOwner->GetClass()->IsChildOf(Class),
-		TEXT("Expected provided property to either be in the Owner class, or in it's IDO's class"));
-#endif
-	return PropertyOwner;
-}
-
 EOverriddenPropertyOperation FOverriddenPropertySet::GetOverriddenPropertyOperation(FPropertyVisitorPath::Iterator PropertyIterator, bool* bOutInheritedOperation /*= nullptr*/) const
 {
 	if (const FOverriddenPropertyNode* RootNode = OverriddenPropertyNodes.Find(RootNodeID))
 	{
-		return GetOverriddenPropertyOperation(*RootNode, PropertyIterator, bOutInheritedOperation, 
-			GetOwningObjectOfProperty(Owner, PropertyIterator));
+		return GetOverriddenPropertyOperation(*RootNode, PropertyIterator, bOutInheritedOperation, Owner);
 	}
 	return EOverriddenPropertyOperation::None;
 }
@@ -1348,8 +1326,7 @@ bool FOverriddenPropertySet::ClearOverriddenProperty(FPropertyVisitorPath::Itera
 {
 	if (FOverriddenPropertyNode* RootNode = OverriddenPropertyNodes.Find(RootNodeID))
 	{
-		return ClearOverriddenProperty(*RootNode, PropertyIterator, 
-			GetOwningObjectOfProperty(Owner, PropertyIterator));
+		return ClearOverriddenProperty(*RootNode, PropertyIterator, Owner);
 	}
 	return true;
 }
diff --git a/Engine/Source/Runtime/CoreUObject/Private/UObject/PropertyBagRepository.cpp b/Engine/Source/Runtime/CoreUObject/Private/UObject/PropertyBagRepository.cpp
--- a/Engine/Source/Runtime/CoreUObject/Private/UObject/PropertyBagRepository.cpp
+++ b/Engine/Source/Runtime/CoreUObject/Private/UObject/PropertyBagRepository.cpp
@@ -504,6 +504,12 @@ public:
 				PropertyChain.SetActiveMemberPropertyNode(ResolvedChainNode->GetValue());
 			}
 		}
+
+		// For add changes, SourceField probably has one more array index that needs to be added. Add that here
+		if (int32 AddChangeIndex = SourceEvent.GetArrayIndex(SourceField.GetName()); AddChangeIndex != INDEX_NONE)
+		{
+			ResolvedArrayIndices[SourceEvent.ObjectIteratorIndex].Add(ResolvedField.GetName(), AddChangeIndex);
+		}
 		
 		// make sure the Property and MemberProperty members were remapped as well
 		check(SourceEvent.Property != Property);
diff --git a/Engine/Source/Runtime/CoreUObject/Public/UObject/OverridableManager.h b/Engine/Source/Runtime/CoreUObject/Public/UObject/OverridableManager.h
--- a/Engine/Source/Runtime/CoreUObject/Public/UObject/OverridableManager.h
+++ b/Engine/Source/Runtime/CoreUObject/Public/UObject/OverridableManager.h
@@ -45,44 +45,28 @@ public:
 
 	FOverriddenPropertySet* Find(const UObject& Object)
 	{
-		const UObjectBase* KeyObject = GetKeyObject(Object);
-		FOverriddenPropertyAnnotation Annotation = GetAnnotation(KeyObject);
+		FOverriddenPropertyAnnotation Annotation = GetAnnotation(&Object);
 		return Annotation.OverriddenProperties.Get();
 	}
 
 	FOverriddenPropertySet& FindChecked(const UObject& Object)
 	{
-		const UObjectBase* KeyObject = GetKeyObject(Object);
-		FOverriddenPropertyAnnotation Annotation = GetAnnotation(KeyObject);
+		FOverriddenPropertyAnnotation Annotation = GetAnnotation(&Object);
 		checkf(!Annotation.IsDefault(), TEXT("Caller is expecting the object to have overridable serialization enabled"));
 		return *Annotation.OverriddenProperties.Get();
 	}
 
 	FOverriddenPropertySet& FindOrAdd(UObject& Object)
 	{
-		UObject* KeyObject = GetKeyObject(Object);
-		FOverriddenPropertyAnnotation Annotation = GetAnnotation(KeyObject);
+		FOverriddenPropertyAnnotation Annotation = GetAnnotation(&Object);
 		if (Annotation.IsDefault())
 		{
-			Annotation.OverriddenProperties = MakeShared<FOverriddenPropertySet>(*KeyObject);
+			Annotation.OverriddenProperties = MakeShared<FOverriddenPropertySet>(Object);
 			AddAnnotation(&Object, Annotation);
 		}
 
 		return *Annotation.OverriddenProperties;
 	}
-private:
-	static UObject* GetKeyObject(const UObject& Object)
-	{
-		const UObject* Found = nullptr;
-#if WITH_EDITOR
-		const UClass* Class = Object.GetClass();
-		if (Class && !Class->HasAnyClassFlags(EClassFlags::CLASS_Native))
-		{
-			Found = UE::FPropertyBagRepository::Get().FindInstanceForDataObject(&Object);
-		}
-#endif
-		return const_cast<UObject*>(Found ? Found : &Object);
-	}
 };
 
 #endif // WITH_EDITORONLY_DATA
