Fixed IDO structs and enums to be created in the transient package

#rnx

diff --git a/Engine/Source/Runtime/CoreUObject/Private/UObject/InstanceDataObjectUtils.cpp b/Engine/Source/Runtime/CoreUObject/Private/UObject/InstanceDataObjectUtils.cpp
--- a/Engine/Source/Runtime/CoreUObject/Private/UObject/InstanceDataObjectUtils.cpp
+++ b/Engine/Source/Runtime/CoreUObject/Private/UObject/InstanceDataObjectUtils.cpp
@@ -802,7 +802,7 @@ namespace UE
 					}
 				}
 
-				ConvertToInstanceDataObjectProperty(SuperProperty, Type, Super, SubTree, EnumNames);
+				ConvertToInstanceDataObjectProperty(SuperProperty, Type, Outer, SubTree, EnumNames);
 				TrySetContainsLoosePropertyMetadata(SuperProperty);
 			}
 
@@ -855,7 +855,7 @@ namespace UE
 							// skip loose types that have been explicitly excluded from IDOs
 							continue;
 						}
-						ConvertToInstanceDataObjectProperty(Property, Type, Result, It.GetNode().GetSubTree(), EnumNames);
+						ConvertToInstanceDataObjectProperty(Property, Type, Outer, It.GetNode().GetSubTree(), EnumNames);
 						MarkPropertyAsLoose(Property);	// note: make sure not to mark until AFTER conversion, as this can mutate property flags on nested struct fields
 						TrySetContainsLoosePropertyMetadata(Property);
 						LooseInstanceDataObjectProperties.Add(Property);
