[AsyncLoading2]

Replaces bIsZenPackage/bIsZenImport logic to use an Enum EPackageLoader to clarify which loader the package was loaded from since we can deal with packages that were created from neither (in-memory only packages).

This corrects an issue where packages created outside of the loaded are requested for load and incorrectly designated as loading from the IODispatcher (because we assumed all packages are for the IODispatcher unless proven found on the filesystem). However it's possible that a package is created in memory and the actual package data is meant to later be loadable from some loader later. In this case, if we see packages loaded from outside the loader, we should always deduce their load location rather than using the cached location in the GlobalImportStore.

#rnx

diff --git a/Engine/Source/Runtime/CoreUObject/Private/Serialization/AsyncLoading2.cpp b/Engine/Source/Runtime/CoreUObject/Private/Serialization/AsyncLoading2.cpp
--- a/Engine/Source/Runtime/CoreUObject/Private/Serialization/AsyncLoading2.cpp
+++ b/Engine/Source/Runtime/CoreUObject/Private/Serialization/AsyncLoading2.cpp
@@ -632,6 +632,25 @@ struct FPackageRequest
 	}
 };
 
+enum class EPackageLoader : uint8
+{
+	Unknown = 0, // Indicates the package was loaded from somewhere outside of the loader
+	LinkerLoad,
+	IoDispatcher
+};
+const TCHAR* LexToString(EPackageLoader PackageLoader)
+{
+	switch (PackageLoader)
+	{
+	case EPackageLoader::LinkerLoad: return TEXT("LinkerLoad");
+	case EPackageLoader::IoDispatcher: return TEXT("IoDispatcher");
+	case EPackageLoader::Unknown: return TEXT("Unknown");
+	default:
+		checkNoEntry();
+		return TEXT("Not Implemented");
+	}
+};
+
 struct FAsyncPackageDesc2
 {
 	// A unique request id for each external call to LoadPackage
@@ -660,17 +679,19 @@ struct FAsyncPackageDesc2
 	FPackagePath PackagePathToLoad;
 	// Package referencer
 	FPackageReferencer PackageReferencer;
-	// Packages with a a custom name can't be imported
+	// Packages with a custom name can't be imported
 	bool bCanBeImported;
-	// Package is a zen package or should be loaded using linkerload
-	bool bIsZenPackage;
+	// Indicates which loader to use when loading the package
+	EPackageLoader Loader;
 
 	static FAsyncPackageDesc2 FromPackageRequest(
 		FPackageRequest& Request,
 		FName UPackageName,
 		FPackageId PackageIdToLoad,
-		bool bIsZenPackage)
+		EPackageLoader Loader)
 	{
+		// Note, we intentionally create FAsyncPackageDesc2 for missing packages for deferred 
+		// processing so it's possible for the Loader to be EPackageLoader::Unknown
 		return FAsyncPackageDesc2
 		{
 			Request.RequestId,
@@ -691,7 +712,7 @@ struct FAsyncPackageDesc2
 #else
 			Request.CustomName.IsNone(),
 #endif
-			bIsZenPackage
+			Loader
 		};
 	}
 
@@ -701,8 +722,10 @@ struct FAsyncPackageDesc2
 		FPackageId ImportedPackageId,
 		FPackageId PackageIdToLoad,
 		FPackagePath&& PackagePathToLoad,
-		bool bIsZenPackage)
+		EPackageLoader Loader)
 	{
+		// Note, we intentionally create FAsyncPackageDesc2 for missing packages for deferred 
+		// processing so it's possible for the Loader to be EPackageLoader::Unknown
 		return FAsyncPackageDesc2
 		{
 			INDEX_NONE,
@@ -719,7 +742,7 @@ struct FAsyncPackageDesc2
 			MoveTemp(PackagePathToLoad),
 			FPackageReferencer::FromImport(ImportingPackageDesc.UPackageName),
 			true,
-			bIsZenPackage
+			Loader
 		};
 	}
 };
@@ -1131,7 +1154,7 @@ private:
 	int32 RefCount = 0;
 	int32 ExportCount = -1;
 	EPackageExtension PackageHeaderExtension = EPackageExtension::Unspecified;
-	bool bIsZenPackage = true;
+	EPackageLoader PackageLoader = EPackageLoader::Unknown;
 	bool bAreAllPublicExportsLoaded = false;
 	bool bIsMissing = false;
 	bool bHasFailed = false;
@@ -1168,9 +1191,9 @@ public:
 		return bHasFailed || bIsMissing;
 	}
 
-	inline bool IsZenPackage() const
+	inline EPackageLoader GetPackageLoader() const
 	{
-		return bIsZenPackage;
+		return PackageLoader;
 	}
 
 	inline EPackageExtension GetPackageHeaderExtension() const
@@ -1399,7 +1422,7 @@ public:
 		return *PackageRef;
 	}
 
-	inline FLoadedPackageRef& AddPackageRef(FPackageId PackageId, FName PackageNameIfKnown, bool bIsZenPackageIfKnown, EPackageExtension PackageHeaderExtensionIfKnown)
+	inline FLoadedPackageRef& AddPackageRef(FPackageId PackageId, FName PackageNameIfKnown, EPackageLoader PackageLoaderIfKnown, EPackageExtension PackageHeaderExtensionIfKnown)
 	{
 		LLM_SCOPE_BYNAME(TEXT("AsyncLoadPackageStore"));
 
@@ -1518,7 +1541,23 @@ public:
 
 			// Must be done last since we may reset PackageRef during RemoveUnreferencedObsoletePackage above
 			PackageRef.PackageHeaderExtension = PackageHeaderExtensionIfKnown;
-			PackageRef.bIsZenPackage = bIsZenPackageIfKnown;
+			PackageRef.PackageLoader = PackageLoaderIfKnown;
+		}
+		else
+		{
+			// Packages that exist only in memory might later be represented on disk, so if that changes
+			// allow specifying their extensions and loader if they are specified to keep our state in sync.
+			// Since our package has a RefCount > 0, we should only further specify/promote the extension and load
+			// location, and not allow lowering the extension or location to an unknown value (since it was
+			// either unknown upon creation or it was something of actual meaning that needs to remain the same)
+			if (PackageRef.PackageHeaderExtension == EPackageExtension::Unspecified)
+			{
+				PackageRef.PackageHeaderExtension = PackageHeaderExtensionIfKnown;
+			}
+			if (PackageRef.PackageLoader == EPackageLoader::Unknown)
+			{
+				PackageRef.PackageLoader = PackageLoaderIfKnown;
+			}
 		}
 		++PackageRef.RefCount;
 		return PackageRef;
@@ -1985,16 +2024,16 @@ public:
 		return !!GlobalImportStore.FindPackageRef(ImportedPackageId);
 	}
 
-	FLoadedPackageRef& AddImportedPackageReference(FPackageId ImportedPackageId, FName PackageNameIfKnown, bool bIsZenPackageIfKnown, EPackageExtension PackageHeaderExtensionIfKnown)
+	FLoadedPackageRef& AddImportedPackageReference(FPackageId ImportedPackageId, FName PackageNameIfKnown, EPackageLoader PackageLoaderIfKnown, EPackageExtension PackageHeaderExtensionIfKnown)
 	{
-		return GlobalImportStore.AddPackageRef(ImportedPackageId, PackageNameIfKnown, bIsZenPackageIfKnown, PackageHeaderExtensionIfKnown);
+		return GlobalImportStore.AddPackageRef(ImportedPackageId, PackageNameIfKnown, PackageLoaderIfKnown, PackageHeaderExtensionIfKnown);
 	}
 
 	void AddPackageReference(const FAsyncPackageDesc2& Desc)
 	{
 		if (Desc.bCanBeImported)
 		{
-			FLoadedPackageRef& PackageRef = GlobalImportStore.AddPackageRef(Desc.UPackageId, Desc.UPackageName, Desc.bIsZenPackage, Desc.PackagePathToLoad.GetHeaderExtension());
+			FLoadedPackageRef& PackageRef = GlobalImportStore.AddPackageRef(Desc.UPackageId, Desc.UPackageName, Desc.Loader, Desc.PackagePathToLoad.GetHeaderExtension());
 			PackageRef.ClearErrorFlags();
 		}
 	}
@@ -4037,16 +4076,16 @@ public:
 	}
 
 private:
-	bool TryGetExistingLoadedPackagePath(FPackageId InPackageIdToLoad, FPackagePath& OutPackagePath, bool& OutbIsZenPackage)
+	bool TryGetExistingLoadedPackagePath(FPackageId InPackageIdToLoad, FPackagePath& OutPackagePath, EPackageLoader& OutPackageLoader)
 	{
 		// Lookup if we have the package in the importstore already to avoid expensive calls in TryGetPackagePathFromFileSystem
-		// Note we must only look at valid, non-zero refcount, PackageRefs since we might not have yet cleared unloaded packages 
-		// or have yet to save and only reside in-memory currently.
-		if (FLoadedPackageRef* PackageRef = GlobalImportStore.FindPackageRef(InPackageIdToLoad); PackageRef && PackageRef->GetRefCount())
+		if (FLoadedPackageRef* PackageRef = GlobalImportStore.FindPackageRef(InPackageIdToLoad))
 		{
-			if (!PackageRef->HasErrors())
+			// If a package has failed or wasn't loaded from the loader, fail the lookup: we cannot trust these references
+			// and it's possible their actual state is now different than what was cached in the GlobalImportStore
+			if (!PackageRef->HasErrors() && PackageRef->GetPackageLoader() != EPackageLoader::Unknown)
 			{
-				OutbIsZenPackage = PackageRef->IsZenPackage();
+				OutPackageLoader = PackageRef->GetPackageLoader();
 				OutPackagePath.SetHeaderExtension(PackageRef->GetPackageHeaderExtension());
 				return true;
 			}
@@ -4866,6 +4905,7 @@ bool FAsyncLoadingThread2::CreateAsyncPackagesFromQueue(FAsyncLoadingThreadState
 		
 			FPackageRequest& Request = OptionalRequest.GetValue();
 			EPackageStoreEntryStatus PackageStatus = EPackageStoreEntryStatus::Missing;
+			EPackageLoader PackageLoader = EPackageLoader::IoDispatcher;
 			FPackageStoreEntry PackageEntry;
 			FName PackageNameToLoad = Request.PackagePath.GetPackageFName();
 			TCHAR NameBuffer[FName::StringBufferSize];
@@ -4901,6 +4941,7 @@ bool FAsyncLoadingThread2::CreateAsyncPackagesFromQueue(FAsyncLoadingThreadState
 			PackageStatus = PackageStore.GetPackageStoreEntry(PackageIdToLoad, UPackageName, PackageEntry);
 			if (PackageStatus == EPackageStoreEntryStatus::Missing)
 			{
+				PackageLoader = EPackageLoader::Unknown;
 				// While there is an active load request for (InName=/Temp/PackageABC_abc, InPackageToLoadFrom=/Game/PackageABC), then allow these requests too:
 				// (InName=/Temp/PackageA_abc, InPackageToLoadFrom=/Temp/PackageABC_abc) and (InName=/Temp/PackageABC_xyz, InPackageToLoadFrom=/Temp/PackageABC_abc)
 				FAsyncPackage2* Package = GetAsyncPackage(PackageIdToLoad);
@@ -4913,18 +4954,17 @@ bool FAsyncLoadingThread2::CreateAsyncPackagesFromQueue(FAsyncLoadingThreadState
 				}
 			}
 
-			bool bIsZenPackage = true;
 			if (!PackageNameToLoad.IsNone() && FAsyncLoadingThreadSettings::Get().bLooseFileLoadingEnabled)
 			{
 				FPackagePath::TryFromPackageName(PackageNameToLoad, Request.PackagePath);
 
-				if (TryGetExistingLoadedPackagePath(PackageIdToLoad, Request.PackagePath, bIsZenPackage))
+				if (TryGetExistingLoadedPackagePath(PackageIdToLoad, Request.PackagePath, PackageLoader))
 				{
 					PackageStatus = EPackageStoreEntryStatus::Ok;
 				}
 				else if (TryGetPackagePathFromFileSystem(PackageNameToLoad, UPackageName, Request.PackagePath))
 				{
-					bIsZenPackage = false;
+					PackageLoader = EPackageLoader::LinkerLoad;
 					PackageStatus = EPackageStoreEntryStatus::Ok;
 				}
 			}
@@ -4945,7 +4985,7 @@ bool FAsyncLoadingThread2::CreateAsyncPackagesFromQueue(FAsyncLoadingThreadState
 				UPackageName = Request.CustomName;
 			}
 
-			FAsyncPackageDesc2 PackageDesc = FAsyncPackageDesc2::FromPackageRequest(Request, UPackageName, PackageIdToLoad, bIsZenPackage);
+			FAsyncPackageDesc2 PackageDesc = FAsyncPackageDesc2::FromPackageRequest(Request, UPackageName, PackageIdToLoad, PackageLoader);
 			if (PackageStatus == EPackageStoreEntryStatus::Missing)
 			{
 				QueueMissingPackage(ThreadState, PackageDesc, MoveTemp(Request.PackageLoadedDelegate), MoveTemp(Request.PackageProgressDelegate));
@@ -4977,14 +5017,16 @@ bool FAsyncLoadingThread2::CreateAsyncPackagesFromQueue(FAsyncLoadingThreadState
 							ADD_CUSTOM_LOADTIMER_META(CreateAsyncPackage, PackageName, NameBuffer);
 
 						check(PackageStatus == EPackageStoreEntryStatus::Ok);
+						check(PackageLoader != EPackageLoader::Unknown);
 #if ALT2_ENABLE_LINKERLOAD_SUPPORT
-						if (!bIsZenPackage)
+						if (PackageLoader == EPackageLoader::LinkerLoad)
 						{
 							Package->InitializeLinkerLoadState(&PackageDesc.InstancingContext);
 						}
 						else
 #endif
 						{
+							check(PackageLoader == EPackageLoader::IoDispatcher);
 							InitializeAsyncPackageFromPackageStore(ThreadState, &IoBatch, Package, PackageEntry);
 						}
 						Package->StartLoading(ThreadState, IoBatch);
@@ -5502,6 +5544,7 @@ void FAsyncPackage2::ImportPackagesRecursiveInner(FAsyncLoadingThreadState2& Thr
 	{
 		FPackageId ImportedPackageId = ImportedPackageIds[LocalImportedPackageIndex];
 		EPackageStoreEntryStatus ImportedPackageStatus = EPackageStoreEntryStatus::Missing;
+		EPackageLoader ImportedPackageLoader = EPackageLoader::IoDispatcher;
 		FPackageStoreEntry ImportedPackageEntry;
 		FName ImportedPackageUPackageName = bHasImportedPackageNames ? ImportedPackageNames[LocalImportedPackageIndex] : NAME_None;
 		FName ImportedPackageNameToLoad = ImportedPackageUPackageName;
@@ -5562,32 +5605,35 @@ void FAsyncPackage2::ImportPackagesRecursiveInner(FAsyncLoadingThreadState2& Thr
 			}
 		}
 		ImportedPackageStatus = PackageStore.GetPackageStoreEntry(ImportedPackageIdToLoad, ImportedPackageUPackageName, ImportedPackageEntry);
+		if (ImportedPackageStatus == EPackageStoreEntryStatus::Missing)
+		{
+			ImportedPackageLoader = EPackageLoader::Unknown;
+		}
 
 		FPackagePath ImportedPackagePath;
-		bool bIsZenPackage = true;
-		bool bIsZenPackageImport = true;
+		EPackageLoader PackageLoader = EPackageLoader::IoDispatcher;
 		if (!ImportedPackageNameToLoad.IsNone() && FAsyncLoadingThreadSettings::Get().bLooseFileLoadingEnabled)
 		{
 			FPackagePath::TryFromPackageName(ImportedPackageNameToLoad, ImportedPackagePath);
-			bIsZenPackage = !LinkerLoadState.IsSet();
+			PackageLoader = LinkerLoadState.IsSet() ? EPackageLoader::LinkerLoad : PackageLoader;
 
-			if (AsyncLoadingThread.TryGetExistingLoadedPackagePath(ImportedPackageId, ImportedPackagePath, bIsZenPackageImport))
+			if (AsyncLoadingThread.TryGetExistingLoadedPackagePath(ImportedPackageId, ImportedPackagePath, ImportedPackageLoader))
 			{
 				ImportedPackageStatus = EPackageStoreEntryStatus::Ok;
 			}
 			else if (AsyncLoadingThread.TryGetPackagePathFromFileSystem(ImportedPackageNameToLoad, ImportedPackageUPackageName, ImportedPackagePath))
 			{
-				bIsZenPackageImport = false;
+				ImportedPackageLoader = EPackageLoader::LinkerLoad;
 				ImportedPackageStatus = EPackageStoreEntryStatus::Ok;
 			}
 		}
 
-		FLoadedPackageRef& ImportedPackageRef = ImportStore.AddImportedPackageReference(ImportedPackageId, ImportedPackageUPackageName, bIsZenPackageImport, ImportedPackagePath.GetHeaderExtension());
-		checkf(ImportedPackageStatus == EPackageStoreEntryStatus::Missing || ImportedPackageRef.IsZenPackage() == bIsZenPackageImport, TEXT("ImportPackageStatus(%d) != EPackageStoreEntryStatus::Missing(%d) || IsZenPackage(%d) != bIsZenPackageImport(%d)"), ImportedPackageStatus, EPackageStoreEntryStatus::Missing, ImportedPackageRef.IsZenPackage(), bIsZenPackageImport);
+		FLoadedPackageRef& ImportedPackageRef = ImportStore.AddImportedPackageReference(ImportedPackageId, ImportedPackageUPackageName, ImportedPackageLoader, ImportedPackagePath.GetHeaderExtension());
+		checkf(ImportedPackageStatus == EPackageStoreEntryStatus::Missing || ImportedPackageRef.GetPackageLoader() == ImportedPackageLoader, TEXT("ImportPackageStatus(%d) != EPackageStoreEntryStatus::Missing(%d) || ImportedPackageRef.GetPackageLoader()(%s) != ImportedPackageLoader(%s)"), ImportedPackageStatus, EPackageStoreEntryStatus::Missing, LexToString(ImportedPackageRef.GetPackageLoader()), LexToString(ImportedPackageLoader));
 		checkf(ImportedPackageStatus == EPackageStoreEntryStatus::Missing || ImportedPackageRef.GetPackageHeaderExtension() == ImportedPackagePath.GetHeaderExtension(), TEXT("ImportPackageStatus(%d) != EPackageStoreEntryStatus::Missing(%d) || ImportedPackageRef.GetPackageHeaderExtension()(%d) != ImportedPackagePath.GetHeaderExtension()(%d)"), ImportedPackageStatus, EPackageStoreEntryStatus::Missing, ImportedPackageRef.GetPackageHeaderExtension(), ImportedPackagePath.GetHeaderExtension());
 
 #if WITH_EDITOR
-		if (AsyncLoadingThread.UncookedPackageLoader && ImportedPackageStatus == EPackageStoreEntryStatus::Ok && !bIsZenPackageImport)
+		if (AsyncLoadingThread.UncookedPackageLoader && ImportedPackageStatus == EPackageStoreEntryStatus::Ok && ImportedPackageLoader == EPackageLoader::LinkerLoad)
 		{
 			UPackage* UncookedPackage = ImportedPackageRef.GetPackage();
 			if (!ImportedPackageRef.AreAllPublicExportsLoaded())
@@ -5656,7 +5702,7 @@ void FAsyncPackage2::ImportPackagesRecursiveInner(FAsyncLoadingThreadState2& Thr
 		bool bInserted = false;
 		bool bIsFullyLoaded = ImportedPackageRef.AreAllPublicExportsLoaded();
 #if ALT2_ENABLE_LINKERLOAD_SUPPORT
-		if (!bIsZenPackageImport && (!ImportedPackageRef.HasPackage() || (!ImportedPackageRef.GetPackage()->GetLinker() && !ImportedPackageRef.GetPackage()->HasAnyPackageFlags(PKG_InMemoryOnly))))
+		if (ImportedPackageLoader == EPackageLoader::LinkerLoad && (!ImportedPackageRef.HasPackage() || (!ImportedPackageRef.GetPackage()->GetLinker() && !ImportedPackageRef.GetPackage()->HasAnyPackageFlags(PKG_InMemoryOnly))))
 		{
 			// If we're importing a linker load package and it doesn't have its linker we need to reload it, otherwise we can't reliably link to its imports
 			// Note: Legacy loader path appears to do this only for uncooked packages in the editor?
@@ -5710,7 +5756,7 @@ void FAsyncPackage2::ImportPackagesRecursiveInner(FAsyncLoadingThreadState2& Thr
 		}
 		else
 		{
-			FAsyncPackageDesc2 PackageDesc = FAsyncPackageDesc2::FromPackageImport(Desc, ImportedPackageUPackageName, ImportedPackageId, ImportedPackageIdToLoad, MoveTemp(ImportedPackagePath), bIsZenPackageImport);
+			FAsyncPackageDesc2 PackageDesc = FAsyncPackageDesc2::FromPackageImport(Desc, ImportedPackageUPackageName, ImportedPackageId, ImportedPackageIdToLoad, MoveTemp(ImportedPackagePath), ImportedPackageLoader);
 			ImportedPackage = AsyncLoadingThread.FindOrInsertPackage(ThreadState, PackageDesc, bInserted, this);
 		}
 
@@ -5743,10 +5789,10 @@ void FAsyncPackage2::ImportPackagesRecursiveInner(FAsyncLoadingThreadState2& Thr
 
 		Header.ImportedAsyncPackagesView[LocalImportedPackageIndex] = ImportedPackage;
 
-		if (bIsZenPackage != bIsZenPackageImport)
+		if (PackageLoader != ImportedPackageLoader)
 		{
 			UE_ASYNC_PACKAGE_LOG(VeryVerbose, Desc, TEXT("ImportPackages: AddDependency"),
-				TEXT("Adding package dependency to %s import '%s'."), bIsZenPackageImport ? TEXT("cooked") : TEXT("non-cooked"), *ImportedPackage->Desc.UPackageName.ToString());
+				TEXT("Adding package dependency to %s import '%s'."), ImportedPackageLoader == EPackageLoader::IoDispatcher ? TEXT("cooked") : TEXT("non-cooked"), * ImportedPackage->Desc.UPackageName.ToString());
 
 			// When importing a linker load package from a zen package or vice versa we need to wait for all the exports in the imported package to be created
 			// and serialized before we can start processing our own exports
@@ -5763,10 +5809,10 @@ void FAsyncPackage2::ImportPackagesRecursiveInner(FAsyncLoadingThreadState2& Thr
 			{
 				check(ImportedPackageStatus == EPackageStoreEntryStatus::Ok);
 #if ALT2_ENABLE_LINKERLOAD_SUPPORT
-				if (!bIsZenPackageImport)
+				if (ImportedPackageLoader == EPackageLoader::LinkerLoad)
 				{
 					// Only propagate the instancing context if the imported package is also instanced and it isn't a zen package
-					ImportedPackage->InitializeLinkerLoadState(bIsZenPackage || !bIsInstanced ? nullptr : &LinkerLoadState->Linker->GetInstancingContext());
+					ImportedPackage->InitializeLinkerLoadState(PackageLoader == EPackageLoader::IoDispatcher || !bIsInstanced ? nullptr : &LinkerLoadState->Linker->GetInstancingContext());
 				}
 				else
 #endif
@@ -10376,12 +10422,7 @@ void FAsyncPackage2::DumpStateImpl(TSet<FAsyncPackage2*>& Set, int32 Indent, TMu
 
 	auto GetPackageType = [](FAsyncPackage2* Package)
 	{
-#if ALT2_ENABLE_LINKERLOAD_SUPPORT
-		const bool bIsZenPackage = !Package->LinkerLoadState.IsSet();
-#else
-		const bool bIsZenPackage = true;
-#endif
-		return bIsZenPackage ? TEXT("ZEN") : TEXT("LINKERLOAD");
+		return LexToString(Package->Desc.Loader);
 	};
 
 	auto FormatPackage = [&GetPackageType](FAsyncPackage2* Package)
