[Feature] Added console command that acts like a CVar for enabling/disabling all ido related features. This way QA doesn't have to enable so many cvars
[Minor] IDO CVar support: Rename of existing CVars

diff --git a/Engine/Source/Runtime/CoreUObject/Private/UObject/InstanceDataObjectUtils.cpp b/Engine/Source/Runtime/CoreUObject/Private/UObject/InstanceDataObjectUtils.cpp
--- a/Engine/Source/Runtime/CoreUObject/Private/UObject/InstanceDataObjectUtils.cpp
+++ b/Engine/Source/Runtime/CoreUObject/Private/UObject/InstanceDataObjectUtils.cpp
@@ -274,15 +274,15 @@ namespace UE
 		FConsoleVariableDelegate::CreateStatic(OnInstanceDataObjectSupportChanged)
 	);
 
-	FString ExcludedLoosePropertyTypesVar = TEXT("VerseFunctionProperty");
-	FAutoConsoleVariableRef ExcludedLoosePropertyTypesCVar(
-		TEXT("IDO.ExcludedLoosePropertyTypes"),
-		ExcludedLoosePropertyTypesVar,
+	FString ExcludedUnknownPropertyTypesVar = TEXT("VerseFunctionProperty");
+	FAutoConsoleVariableRef ExcludedUnknownPropertyTypesCVar(
+		TEXT("IDO.Unknowns.ExcludedTypes"),
+		ExcludedUnknownPropertyTypesVar,
 		TEXT("Comma separated list of property types that will be excluded from loose properties in IDOs.")
 	);
 
 	bool bEnableUninitializedUI = true;
-	FAutoConsoleVariableRef EnableUninitializedUI(
+	FAutoConsoleVariableRef EnableUninitializedUICVar(
 		TEXT("IDO.EnableUninitializedAlertUI"),
 		bEnableUninitializedUI,
 		TEXT("Enables alert information for uninitalized properties. Requires IDO.Enable=true")
@@ -290,30 +290,65 @@ namespace UE
 
 	// TODO: re-enable this ASAP. This disables most IDO features but disabling was necessary to unblock those experiencing IDO bugs
 	bool bEnableIDOImpersonationOnSave = false;
-	FAutoConsoleVariableRef EnableIDOImpersonationOnSave(
+	FAutoConsoleVariableRef EnableIDOImpersonationOnSaveCVar(
 		TEXT("IDO.Impersonation.EnableOnSave"),
 		bEnableIDOImpersonationOnSave,
 		TEXT("When enabled, IDOs will be saved instead of instances. Disabling this will stop data retention on save.")
 	);
 
 	bool bEnableIDOsForPrefabArchetypes = true;
-	FAutoConsoleVariableRef EnableIDOsForPrefabArchetypes(
+	FAutoConsoleVariableRef EnableIDOsForPrefabArchetypesCVar(
 		TEXT("IDO.EnablePrefabArchetypes"),
 		bEnableIDOsForPrefabArchetypes,
 		TEXT("When enabled, prefab archetypes can have IDOs generated for them")
 	);
 
 	bool bEnableIDOsForPrefabInstances = true;
-	FAutoConsoleVariableRef EnableIDOsForPrefabInstances(
+	FAutoConsoleVariableRef EnableIDOsForPrefabInstancesCVar(
 		TEXT("IDO.EnablePrefabInstances"),
 		bEnableIDOsForPrefabInstances,
 		TEXT("When enabled, prefab instances can have IDOs generated for them")
+		);
+
+	static void SetEnableAllIDOFeatures(bool bEnabled)
+	{
+		EnableIDOSupportCVar->Set(bEnabled);
+		EnableUninitializedUICVar->Set(bEnabled);
+		EnableIDOImpersonationOnSaveCVar->Set(bEnabled);
+		EnableIDOsForPrefabArchetypesCVar->Set(bEnabled);
+		EnableIDOsForPrefabInstancesCVar->Set(bEnabled);
+	}
+
+	static bool AreAllIDOFeaturesEnabled()
+	{
+		return
+			bEnableIDOSupport && 
+			bEnableUninitializedUI && 
+			bEnableIDOImpersonationOnSave && 
+			bEnableIDOsForPrefabArchetypes && 
+			bEnableIDOsForPrefabInstances;
+	}
+
+	FAutoConsoleCommand EnableAllIDOFeaturesCommand(
+		TEXT("IDO.EnableAllFeatures"),
+		TEXT("Call this method to toggle all IDO related features on"),
+		FConsoleCommandWithArgsAndOutputDeviceDelegate::CreateLambda([](const TArray< FString >& Args, FOutputDevice& OutputDevice)
+		{
+			if (Args.Num() == 1)
+			{
+				bool Val = Args[0] == TEXT("True") || Args[0] == TEXT("1");
+				SetEnableAllIDOFeatures(Val);
+			}
+			
+			OutputDevice.Logf(TEXT("IDO.EnableAllFeatures = \"%s\""), AreAllIDOFeaturesEnabled() ? TEXT("True") : TEXT("False"));
+			
+		})
 	);
 
-	static TSet<FString> GetExcludedLoosePropertyTypes()
+	static TSet<FString> GetExcludedUnknownPropertyTypes()
 	{
 		TArray<FString> Result;
-		ExcludedLoosePropertyTypesVar.ParseIntoArray(Result, TEXT(","));
+		ExcludedUnknownPropertyTypesVar.ParseIntoArray(Result, TEXT(","));
 		return TSet<FString>(Result);
 	}
 
@@ -826,7 +861,7 @@ namespace UE
 			Result->SetMetaData(NAME_ContainsLoosePropertiesMetadata, TEXT("True"));
 		}
 
-		TSet<FString> ExcludedLoosePropertyTypes = GetExcludedLoosePropertyTypes();
+		TSet<FString> ExcludedLoosePropertyTypes = GetExcludedUnknownPropertyTypes();
 
 		// Gather "loose" properties for child Struct
 		TArray<FProperty*> LooseInstanceDataObjectProperties;
diff --git a/Engine/Source/Runtime/CoreUObject/Private/UObject/PropertyBagRepository.cpp b/Engine/Source/Runtime/CoreUObject/Private/UObject/PropertyBagRepository.cpp
--- a/Engine/Source/Runtime/CoreUObject/Private/UObject/PropertyBagRepository.cpp
+++ b/Engine/Source/Runtime/CoreUObject/Private/UObject/PropertyBagRepository.cpp
@@ -868,7 +868,7 @@ namespace Private
 #if WITH_EDITOR
 	static bool bEnablePropertyBagPlaceholderObjectSupport = 0;
 	static FAutoConsoleVariableRef CVarEnablePropertyBagPlaceholderObjectSupport(
-		TEXT("SceneGraph.EnablePropertyBagPlaceholderObjectSupport"),
+		TEXT("IDO.Placeholder.Enable"),
 		bEnablePropertyBagPlaceholderObjectSupport,
 		TEXT("If true, allows placeholder types to be created in place of missing types in order to redirect serialization into a property bag."),
 		ECVF_Default
