remoteobject sync work:

we now expose three global sync functions:
GlobalServerSync
- this used to be called WaitForServersToSync and is the core function that simply synchronizes servers
GlobalServerSyncAndReturnBorrowedObjects
- this returns all borrowed objects back to their current owner server. when it returns, all owned objects are locally resident
GlobalServerSyncAndRecomputeObjectOwners
- this returns all borrowed objects to their owner and then re-evaluates which server should be the current owner, and pushes objects to their new owning server. when it returns, all owned objects are locally resident

- added FName+uint parameter to the syncs to give a descriptive name for debugging, and a monotonically incrementing sequence number. this is to ensure that when we sync servers, they all participate in the same sync and the same number of times, in order to catch bugs like one server conditionally doing a sync which would misalign everything

- OnEndFrame we now call GlobalServerSyncAndRecomputeObjectOwners

- removed the 'PendingMigrations' list, now we just recompute what to migrate on the fly and try to do it immediately, with some added network flush calls to work around the network buffer limits

- removed bIsWaitingForServersToSync, we now check bAllServersConnected inside GlobalServerSync and we don't send the sync ready messages until all servers are fully connected

- removed the PushAllActorsToPreferredServers timer tick because it was being called hundreds of times per frame, and now we have functions that we can use to explicitly insert sync points in the frame where we want it

diff --git a/Engine/Source/Runtime/CoreUObject/Private/UObject/RemoteExecutor.cpp b/Engine/Source/Runtime/CoreUObject/Private/UObject/RemoteExecutor.cpp
--- a/Engine/Source/Runtime/CoreUObject/Private/UObject/RemoteExecutor.cpp
+++ b/Engine/Source/Runtime/CoreUObject/Private/UObject/RemoteExecutor.cpp
@@ -10,7 +10,6 @@ namespace UE::RemoteExecutor
 {
 TDelegate<void()> TickNetworkDelegate;
 TDelegate<TTuple<FRemoteWorkPriority, TFunction<void(void)>>(void)> FetchNextDeferredRPCDelegate;
-TDelegate<void()> WaitForServersToSyncDelegate;
 
 TMulticastDelegate<void(FRemoteTransactionId)> OnTransactionStartingDelegate;
 TMulticastDelegate<void(FRemoteTransactionId)> OnTransactionCompletedDelegate;
@@ -277,12 +276,6 @@ void RegisterRemoteSubsystem(FRemoteSubsystemBase* Subsystem)
 	GRemoteExecutor.RegisterSubsystem(Subsystem);
 }
 
-void WaitForServersToSync()
-{
-	UE_LOG(LogRemoteExec, VeryVerbose, TEXT("-- WaitForServersToSync --"));
-	WaitForServersToSyncDelegate.ExecuteIfBound();
-}
-
 void AbortTransactionRequiresDependencies(FStringView Description)
 {
 	UE_AUTORTFM_OPEN
diff --git a/Engine/Source/Runtime/CoreUObject/Private/UObject/RemoteObject.cpp b/Engine/Source/Runtime/CoreUObject/Private/UObject/RemoteObject.cpp
--- a/Engine/Source/Runtime/CoreUObject/Private/UObject/RemoteObject.cpp
+++ b/Engine/Source/Runtime/CoreUObject/Private/UObject/RemoteObject.cpp
@@ -131,8 +131,6 @@ void InitRemoteObjects()
 	{
 		UE::RemoteExecutor::FetchNextDeferredRPCDelegate.BindStatic([]() { return TTuple<FRemoteWorkPriority, TFunction<void(void)>>(); });
 	}
-
-	FCoreDelegates::OnEndFrame.AddStatic(&UE::RemoteExecutor::WaitForServersToSync);
 }
 
 void RegisterRemoteObjectId(FRemoteObjectId ObjectId, FRemoteServerId ResidentServerId)
diff --git a/Engine/Source/Runtime/CoreUObject/Public/UObject/RemoteExecutor.h b/Engine/Source/Runtime/CoreUObject/Public/UObject/RemoteExecutor.h
--- a/Engine/Source/Runtime/CoreUObject/Public/UObject/RemoteExecutor.h
+++ b/Engine/Source/Runtime/CoreUObject/Public/UObject/RemoteExecutor.h
@@ -190,9 +190,6 @@ namespace UE::RemoteExecutor
 	/** Fetch the next deferred RPC request **/
 	extern COREUOBJECT_API TDelegate<TTuple<FRemoteWorkPriority, TFunction<void(void)>>(void)> FetchNextDeferredRPCDelegate;
 
-	/** Delegate that syncs all servers */
-	extern COREUOBJECT_API TDelegate<void()> WaitForServersToSyncDelegate;
-
 	/** Delegate executed when a ExecuteTransactional starts an AutoRTFM transaction */
 	extern COREUOBJECT_API TMulticastDelegate<void(FRemoteTransactionId)> OnTransactionStartingDelegate;
 
@@ -211,7 +208,9 @@ namespace UE::RemoteExecutor
 	COREUOBJECT_API void AbortTransactionAndAbandonWork(FStringView Description);
 	
 	/** Establishes a synchronization barrier across all servers */
-	COREUOBJECT_API void WaitForServersToSync();
-	
+	COREUOBJECT_API void GlobalServerSync(FName SyncName);
+	COREUOBJECT_API void GlobalServerSyncAndReturnBorrowedObjects(FName SyncName);
+	COREUOBJECT_API void GlobalServerSyncAndRecomputeObjectOwners(FName SyncName);
+
 	COREUOBJECT_API FRemoteWorkPriority CreateRootWorkPriority();
 }
