Corrected comment on Blueprint property meta specifier "BlueprintBaseOnly": The specifier shows classes that can be used as a base for Blueprints.

diff --git a/Engine/Source/Runtime/CoreUObject/Public/UObject/ObjectMacros.h b/Engine/Source/Runtime/CoreUObject/Public/UObject/ObjectMacros.h
--- a/Engine/Source/Runtime/CoreUObject/Public/UObject/ObjectMacros.h
+++ b/Engine/Source/Runtime/CoreUObject/Public/UObject/ObjectMacros.h
@@ -1287,7 +1287,7 @@ namespace UM
 		/// [PropertyMetadata] Used for SoftObjectPtr/SoftObjectPath properties. Comma separated list of Bundle names used inside PrimaryDataAssets to specify which bundles this reference is part of
 		AssetBundles,
 
-		/// [PropertyMetadata] Used for Subclass and SoftClass properties.  Indicates whether only blueprint classes should be shown in the class picker.
+		/// [PropertyMetadata] Used for Subclass and SoftClass properties.  Indicates whether the class picker should only show classes that can be used as a base for blueprints.
 		BlueprintBaseOnly,
 
 		/// [PropertyMetadata] Property defaults are generated by the Blueprint compiler and will not be copied when CopyPropertiesForUnrelatedObjects is called post-compile.
