Avoid reconstructing borrowed remote objects when they're being returned to their owning server and instead just reset their state to the archetype values

diff --git a/Engine/Source/Runtime/CoreUObject/Private/UObject/Obj.cpp b/Engine/Source/Runtime/CoreUObject/Private/UObject/Obj.cpp
--- a/Engine/Source/Runtime/CoreUObject/Private/UObject/Obj.cpp
+++ b/Engine/Source/Runtime/CoreUObject/Private/UObject/Obj.cpp
@@ -1924,7 +1924,15 @@ void UObject::SerializeScriptProperties( FStructuredArchive::FSlot Slot ) const
 		{
 			DiffObject = GetArchetype();
 		}
-		UStruct* DiffClass = HasAnyFlags(RF_ClassDefaultObject) ? ObjClass->GetSuperClass() : ObjClass;
+
+		// When migrating remote objects the only instance where we serialize a CDO is when resetting an object to its archetype state in which case
+		// we want to serialize against the actual object class to get the right delta
+#if UE_WITH_REMOTE_OBJECT_HANDLE
+		const bool bMigratingRemoteObjects = UnderlyingArchive.IsMigratingRemoteObjects();
+#else
+		constexpr bool bMigratingRemoteObjects = false;
+#endif
+		UStruct* DiffClass = (HasAnyFlags(RF_ClassDefaultObject) && !bMigratingRemoteObjects) ? ObjClass->GetSuperClass() : ObjClass;
 
 		// Query if this object data is being impersonated 
 		const UObject* ThisObject = this;
diff --git a/Engine/Source/Runtime/CoreUObject/Private/UObject/RemoteObjectSerialization.cpp b/Engine/Source/Runtime/CoreUObject/Private/UObject/RemoteObjectSerialization.cpp
--- a/Engine/Source/Runtime/CoreUObject/Private/UObject/RemoteObjectSerialization.cpp
+++ b/Engine/Source/Runtime/CoreUObject/Private/UObject/RemoteObjectSerialization.cpp
@@ -25,6 +25,12 @@ static FAutoConsoleVariableRef CVarRemoteObjectsMigrateFullHierarchy(
 	GRemoteObjectsMigrateFullHierarchy,
 	TEXT("Whether remote objects that are default subobjects should be always migrated with their parent objects"));
 
+int32 GResetBorrowedObjects = 1;
+static FAutoConsoleVariableRef CVarResetBorrowedObjects(
+	TEXT("ro.ResetBorrowedObjects"),
+	GRemoteObjectsMigrateFullHierarchy,
+	TEXT("Whether remote objects that were borrowed should be reset upon returning to their owner server instead of being reconstructed"));
+
 using FNameIndexType = FRemoteObjectPathName::FNameIndexType;
 
 FArchive& operator<<(FArchive& Ar, FRemoteObjectPathName& PathName)
@@ -132,6 +138,22 @@ FArchive& operator<<(FArchive& Ar, ERemoteReferenceType& RefType)
 	return Ar;
 }
 
+/**
+* Structure that holds information about a reference to an object
+* Helpful in avoiding calculating the same reference properties multiple times
+*/
+struct FRemoteObjectReferenceInfo
+{
+	UObject* Object = nullptr;
+	FRemoteObjectId Id;
+	ERemoteReferenceType Type = ERemoteReferenceType::None;
+	bool bIsSubobject = false;
+};
+
+/**
+* Base archive for serializing object data for migration
+* Implements serialization debugging functionality (see UE_WITH_REMOTEOBJECTARCHIVE_DEBUGGING)
+*/
 template <typename T>
 class TArchiveRemoteObjectBase : public T
 {
@@ -324,82 +346,128 @@ public:
 
 class FArchiveRemoteObjectWriter : public TArchiveRemoteObjectBase<FMemoryWriter>
 {
-	TArray<UObject*>& ObjectsToSerialize;
-	const TSet<UObject*>& SerializedObjects;
+protected:
+
+	TArray<UObject*> ObjectsToSerialize;
+	TSet<UObject*> SerializedObjects;
 	TMap<FName, FNameIndexType> NameMap;
 	TMap<UObject*, FNameIndexType> PathNameMap;
+	TMap<UObject*, UObject*>* ObjectToArchetypeMap = nullptr;
 
-	void WriteObjectReference(const FObjectPtr& ObjPtr)
-	{
-		using namespace UE::CoreUObject::Private;
-
-		ERemoteReferenceType Type = ERemoteReferenceType::None;
-#if UE_WITH_REMOTE_OBJECT_HANDLE
-		FRemoteObjectId Id = ObjPtr.GetRemoteId();
-#else
-		FRemoteObjectId Id;
-#endif
-		UObject* Obj = nullptr;
+public:
 
+	static FRemoteObjectReferenceInfo GetReferenceInfo(const FObjectPtr& ObjPtr, UObject* InRootObject)
+	{
+		FRemoteObjectReferenceInfo Info;
+		Info.Id = ObjPtr.GetRemoteId();
 		if (ObjPtr)
 		{
-			if (ObjPtr.IsResolved())
+			if (!ObjPtr.IsRemote())
 			{
-				Obj = ObjPtr.Get();
-				checkf(Obj, TEXT("Attempting to serialize resolved Object pointer but the pointer resolves to null (%s)"), *Id.ToString());
-				Obj->SetInternalFlags(EInternalObjectFlags::RemoteReference);
-				check(!UE::RemoteObject::Handle::IsRemote(Obj));
+				Info.Object = ObjPtr.Get();
+				Info.bIsSubobject = Info.Object->IsIn(InRootObject);
 
-				Id = FObjectHandleUtils::GetRemoteId(Obj);
-
-				if (Obj == RootObject || Obj->IsIn(RootObject) || (Id.GetServerId().IsValid() && Id.GetServerId() != UE::RemoteObject::GetGlobalServerId()))
+				if (Info.Object == InRootObject || Info.bIsSubobject || (Info.Id.GetServerId().IsValid() && Info.Id.GetServerId() != UE::RemoteObject::GetGlobalServerId()))
 				{
-					Type = ERemoteReferenceType::IdOnly;
+					Info.Type = ERemoteReferenceType::IdOnly;
 				}
 				else
 				{
-					
-					Type = ERemoteReferenceType::PathName;
+
+					Info.Type = ERemoteReferenceType::PathName;
 				}
 			}
 			else
 			{
-				Type = ERemoteReferenceType::IdOnly;
+				Info.Type = ERemoteReferenceType::IdOnly;
 			}
 		}
+		return Info;
+	}
 
-		*this << Type;
+protected:
+
+	virtual FRemoteObjectReferenceInfo GetReferenceInfo(const FObjectPtr& ObjPtr) const
+	{
+		return GetReferenceInfo(ObjPtr, RootObject);
+	}
+
+	virtual void PopulateObjectHeader(UObject* Object, FRemoteObjectHeader& OutHeader) const
+	{
+		using namespace UE::RemoteObject;
+		using namespace UE::RemoteObject::Private;
+		using namespace UE::CoreUObject::Private;
+
+		OutHeader.Name = Object->GetFName();
+		OutHeader.RemoteId = FObjectHandleUtils::GetRemoteId(Object);
+		OutHeader.Class = Object->GetClass();
+		UObject* Outer = FObjectHandleUtils::GetNonAccessTrackedOuterNoResolve(Object);
+		OutHeader.Outer = Outer;
+		OutHeader.Archetype = Object->GetArchetype();
+		OutHeader.InternalFlags = (int32)(Object->GetInternalFlags() & EInternalObjectFlags::Garbage);
+	}
+
+private:
+
+	void WriteObjectReference(const FObjectPtr& ObjPtr)
+	{
+		using namespace UE::CoreUObject::Private;
+
+		FRemoteObjectReferenceInfo Info = GetReferenceInfo(ObjPtr);
+
+		if (Info.Object)
+		{
+			if (UObject** ReplacementObject = (ObjectToArchetypeMap ? ObjectToArchetypeMap->Find(Info.Object) : nullptr))
+			{
+				Info.Object = *ReplacementObject;
+				Info.Id = FRemoteObjectId(Info.Object);
+			}
+			Info.Object->SetInternalFlags(EInternalObjectFlags::RemoteReference);
+		}
+
+		*this << Info.Type;
 		// Always serialize unique id as objects may not exist on the other server and then we may need to pull them from this server
-		if (Type != ERemoteReferenceType::None)
+		if (Info.Type != ERemoteReferenceType::None)
 		{
-			*this << Id;
+			*this << Info.Id;
 		}
 
-		if (Type == ERemoteReferenceType::IdOnly)
+		if (Info.Type == ERemoteReferenceType::IdOnly)
 		{
 			// Add subobjects of the root object to the list of objects to serialize
-			if (Obj && Obj->IsIn(RootObject))
+			if (Info.bIsSubobject)
 			{
-				ObjectsToSerialize.Add(Obj);
+				ObjectsToSerialize.Add(Info.Object);
 			}
 		}
-		else if (Type == ERemoteReferenceType::PathName)
+		else if (Info.Type == ERemoteReferenceType::PathName)
 		{
-			FNameIndexType PathNameIndex = AddPathNameToNameMap(Obj);
+			FNameIndexType PathNameIndex = AddPathNameToNameMap(Info.Object);
 			*this << PathNameIndex;
 		}
 	}
 
 public:
 
-	FArchiveRemoteObjectWriter(UObject* InRootObject, FRemoteObjectData& OutObjectData, TArray<UObject*>& OutObjectsToSerialize, const TSet<UObject*>& InSerializedObjects)
-		: TArchiveRemoteObjectBase<FMemoryWriter>(OutObjectData, TEXT("RemoteObjectWriter"))
-		, ObjectsToSerialize(OutObjectsToSerialize)
-		, SerializedObjects(InSerializedObjects)
+	FArchiveRemoteObjectWriter(UObject* InRootObject, FRemoteObjectData& OutObjectData, const TCHAR* ArchiveName = nullptr)
+		: TArchiveRemoteObjectBase<FMemoryWriter>(OutObjectData, ArchiveName ? ArchiveName : TEXT("RemoteObjectWriter"))
 	{
 		static_assert(sizeof(FNameIndexType) == sizeof(FRemoteObjectPathName::FNameIndexType), "Name Index type must match the type used for storing remote object path name indices");
 
 		SetRootObject(InRootObject);
+		ObjectsToSerialize.Add(InRootObject);
+	}
+
+	virtual UObject* GetArchetypeFromLoader(const UObject* Obj) override
+	{
+		if (ObjectToArchetypeMap)
+		{
+			if (UObject** FoundArchetype = ObjectToArchetypeMap->Find(Obj))
+			{
+				return *FoundArchetype;
+			}
+		}
+		return nullptr;
 	}
 
 	FNameIndexType AddNameToNameMap(const FName& Name)
@@ -480,6 +548,18 @@ public:
 
 		return *this;
 	}
+
+	TArray<UObject*>& GetObjectsToSerialize()
+	{
+		return ObjectsToSerialize;
+	}
+
+	void SerializeRemoteObject(UObject* Object, FRemoteObjectHeader& OutHeader)
+	{
+		PopulateObjectHeader(Object, OutHeader);
+		AddNameToNameMap(OutHeader.Name);
+		Object->Serialize(*this);
+	}
 };
 
 template <typename T>
@@ -514,6 +594,67 @@ struct TRemoteObjectArchiveScope
 using FRemoteObjectWriterScope = TRemoteObjectArchiveScope<FMemoryWriter>;
 using FRemoteObjectReaderScope = TRemoteObjectArchiveScope<FMemoryReader>;
 
+/**
+* Helper archive that serializes the difference between archetypes and their instances
+* This is achieved using delta serialization but the data we serialize against is coming from instances of the archetypes
+* Effectively this is the opposite how delta serialization normally works which serializes instances of archetypes agains the archetypes
+*/
+class FArchetypeDeltaWriter : public FArchiveRemoteObjectWriter
+{
+	TMap<UObject*, UObject*>& SubstituteArchetypeMap;
+
+public:
+
+	FArchetypeDeltaWriter(UObject* InRootObject, FRemoteObjectData& OutObjectData, TMap<UObject*, UObject*>& InSubstituteArchetypeMap)
+		: FArchiveRemoteObjectWriter(InRootObject, OutObjectData, TEXT("RemoteArchetypeDeltaWriter"))
+		, SubstituteArchetypeMap(InSubstituteArchetypeMap)
+	{
+	}
+
+	virtual UObject* GetArchetypeFromLoader(const UObject* Obj) override
+	{
+		// Since this archive serializes the archetype we want the archetype of that archetype to be its instance (effectively reversing the object -> archetype relationship)
+		// this way we will delta serialize the difference between the archetype and its instance
+		if (UObject** FoundArchetype = SubstituteArchetypeMap.Find(Obj))
+		{
+			return *FoundArchetype;
+		}
+		return nullptr;
+	}
+
+protected:
+
+	virtual FRemoteObjectReferenceInfo GetReferenceInfo(const FObjectPtr& ObjPtr) const override
+	{
+		FRemoteObjectReferenceInfo Info = FArchiveRemoteObjectWriter::GetReferenceInfo(ObjPtr, RootObject);
+		if (Info.Object)
+		{
+			// In addition to substituting the archetype of the serialized objects (which are archetypes themselves) we want to substitue references to these objects as well (to pretend we're serializing their instances)
+			// this way when we deserialize the serialized data on top of the instances all references to instances of subobjects point to the right objects (instances, not archetypes)
+			if (UObject** SubstituteObject = SubstituteArchetypeMap.Find(Info.Object))
+			{
+				Info.Object = *SubstituteObject;
+				Info.Id = FRemoteObjectId(Info.Object);
+			}
+		}
+		return Info;
+	}
+
+	virtual void PopulateObjectHeader(UObject* Object, FRemoteObjectHeader& OutHeader) const override
+	{
+		// As per the comment in GetReferenceInfo - when serializing archetype data pretend we're actually serializing instances of the archetype so substitute
+		// the archetype object with its instance when serializing the object header
+		UObject* HeaderObject = Object;
+		if (UObject** ReplacementObject = SubstituteArchetypeMap.Find(Object))
+		{
+			HeaderObject = *ReplacementObject;
+		}
+
+		FArchiveRemoteObjectWriter::PopulateObjectHeader(HeaderObject, OutHeader);
+	}
+};
+
+
 FArchive& operator<<(FArchive& Ar, FRemoteObjectHeader& Header)
 {
 	Ar << Header.Name;
@@ -526,27 +667,11 @@ FArchive& operator<<(FArchive& Ar, FRemoteObjectHeader& Header)
 	return Ar;
 }
 
-void SerializeRemoteObject(UObject* Object, FArchive& Ar, FRemoteObjectHeader& OutHeader)
-{
-	using namespace UE::RemoteObject;
-	using namespace UE::RemoteObject::Private;
-	using namespace UE::CoreUObject::Private;
-
-	OutHeader.Name = Object->GetFName();
-	OutHeader.RemoteId = FObjectHandleUtils::GetRemoteId(Object);
-	OutHeader.Class = Object->GetClass();
-	UObject* Outer = FObjectHandleUtils::GetNonAccessTrackedOuterNoResolve(Object);
-	OutHeader.Outer = Outer;
-	OutHeader.Archetype = Object->GetArchetype();
-	OutHeader.InternalFlags = (int32)(Object->GetInternalFlags() & EInternalObjectFlags::Garbage);
-
-	Object->Serialize(Ar);
-}
-
 class FArchiveRemoteObjectReader : public TArchiveRemoteObjectBase<FMemoryReader>
 {
 	const TArray<FName>& Names;
 	const TArray<UObject*>& ResolvedPathNameObjects;
+	const ERemoteObjectSerializationFlags DeserializeFlags = ERemoteObjectSerializationFlags::None;
 
 	ERemoteReferenceType ReadObjectReference(FObjectPtr& Value)
 	{
@@ -564,14 +689,24 @@ class FArchiveRemoteObjectReader : public TArchiveRemoteObjectBase<FMemoryReader
 				FNameIndexType PathNameIndex = TNumericLimits<FNameIndexType>::Max();
 				*this << PathNameIndex;
 
-				// Try to resolve path name immediately as we expect the object to exist in memory.
-				UObject* Obj = ResolvedPathNameObjects[PathNameIndex];
-				if (Obj)
+				// In some situations (like resetting an object to its archetype state) we might want to preserve
+				// references to remote objects because we might end up migrating them mid-deserialization 
+				// Overwriting them could also potentially discard any changes made to them on another server
+				if (EnumHasAnyFlags(DeserializeFlags, ERemoteObjectSerializationFlags::PreserveRemoteReferences) && Value.IsRemote())
 				{
-					Obj->SetInternalFlags(EInternalObjectFlags::RemoteReference);
-					Value = FObjectPtr(Obj);
 					bNeedsResolvingWithId = false;
 				}
+				else
+				{
+					// Try to resolve path name immediately as we expect the object to exist in memory.
+					UObject* Obj = ResolvedPathNameObjects[PathNameIndex];
+					if (Obj)
+					{
+						Obj->SetInternalFlags(EInternalObjectFlags::RemoteReference);
+						Value = FObjectPtr(Obj);
+						bNeedsResolvingWithId = false;
+					}
+				}
 			}
 			if (bNeedsResolvingWithId)
 			{
@@ -592,10 +727,11 @@ class FArchiveRemoteObjectReader : public TArchiveRemoteObjectBase<FMemoryReader
 
 public:
 
-	FArchiveRemoteObjectReader(FRemoteObjectData& InObjectData, const TArray<UObject*>& InResolvedPathNames)
+	FArchiveRemoteObjectReader(FRemoteObjectData& InObjectData, const TArray<UObject*>& InResolvedPathNames, ERemoteObjectSerializationFlags InDeserializeFlags)
 		: TArchiveRemoteObjectBase<FMemoryReader>(InObjectData, TEXT("RemoteObjectReader"))
 		, Names(InObjectData.Names)
 		, ResolvedPathNameObjects(InResolvedPathNames)
+		, DeserializeFlags(InDeserializeFlags)
 	{
 	}
 
@@ -693,6 +829,39 @@ public:
 	}
 };
 
+
+class FArchiveSubObjectGatherer : public FArchiveUObject
+{
+	UObject* RootObject = nullptr;
+	TMap<UObject*, UObject*>& ArchetypeToObjectMap;
+	TArray<UObject*>& ObjectsToSerialize;
+
+public:
+
+	FArchiveSubObjectGatherer(UObject* InRootObject, TMap<UObject*, UObject*>& InArchetypeToObjectMap, TArray<UObject*>& InObjectsToSerialize)
+		: RootObject(InRootObject)
+		, ArchetypeToObjectMap(InArchetypeToObjectMap)
+		, ObjectsToSerialize(InObjectsToSerialize)
+	{
+		ArchetypeToObjectMap.Add(RootObject->GetArchetype(), RootObject);
+		ArIsObjectReferenceCollector = true;
+		SetIsPersistent(false);
+		SetMigratingRemoteObjects(true);
+		SetPortFlags(PPF_MigratingRemoteObjects);
+	}
+
+	virtual FArchive& operator<<(FObjectPtr& Value) override
+	{
+		FRemoteObjectReferenceInfo Info = FArchiveRemoteObjectWriter::GetReferenceInfo(Value, RootObject);
+		if (Info.bIsSubobject)
+		{
+			ArchetypeToObjectMap.Add(Info.Object->GetArchetype(), Info.Object);
+			ObjectsToSerialize.Add(Info.Object);
+		}
+		return *this;
+	}
+};
+
 FRemoteObjectConstructionOverrides::FRemoteObjectConstructionOverrides(const TArray<FRemoteObjectHeader>& InObjectHeaders)
 {
 	using namespace UE::RemoteObject::Handle;
@@ -711,7 +880,7 @@ FRemoteObjectConstructionOverrides::FRemoteObjectConstructionOverrides(const TAr
 	}
 }
 
-UObject* ConstructRemoteObject(const FRemoteObjectHeader& Header, FRemoteObjectConstructionOverrides* ConstructionOverrides)
+UObject* ConstructRemoteObject(const FRemoteObjectHeader& Header, FRemoteObjectConstructionOverrides* ConstructionOverrides, ERemoteObjectSerializationFlags DeserializeFlags)
 {
 	using namespace UE::RemoteObject::Handle;
 	using namespace UE::RemoteObject::Private;
@@ -753,7 +922,10 @@ UObject* ConstructRemoteObject(const FRemoteObjectHeader& Header, FRemoteObjectC
 	}
 
 	// If not or the object is marked as remote (which means we brought it back before it was GC'd) (re)construct it
-	if (!Object || IsRemote(Object))
+	// Unless we explcitly want to re-use existing objects and skip re-construction to avoid side-effects
+	// Note that even then an object may not exist on this server (it could've been constructed on a different server when its owner was migrated or was simply GC'd)
+	const bool bSkipConstruction = Object && (!IsRemote(Object) || EnumHasAnyFlags(DeserializeFlags, ERemoteObjectSerializationFlags::UseExistingObjects));
+	if (!bSkipConstruction)
 	{
 		FStaticConstructObjectParameters Params(Class);
 		Params.Outer = Outer;
@@ -774,6 +946,11 @@ UObject* ConstructRemoteObject(const FRemoteObjectHeader& Header, FRemoteObjectC
 		}
 	}
 
+	if (!Object)
+	{
+		return nullptr;
+	}
+
 	checkf(FRemoteObjectId(Object) == Header.RemoteId, TEXT("Created an object with a different ID:%s than requested:%s"), *FRemoteObjectId(Object).ToString(), *Header.RemoteId.ToString());
 
 	MarkAsLocal(Object);
@@ -820,80 +997,139 @@ void ResolvePathNames(FRemoteObjectData& InObjectData, TArray<UObject*>& OutReso
 	}
 }
 
-
-FRemoteObjectData* SerializeObjectData(UObject* InObject, TSet<UObject*>& OutObjects)
+void SerializeObjectDataInternal(FArchiveRemoteObjectWriter& Ar, UObject* RequestedObject, FRemoteObjectId RequestedObjectId, TSet<UObject*>& OutObjects)
 {
 	using namespace UE::RemoteObject::Handle;
 
-	FRemoteObjectData* ObjectData = new FRemoteObjectData();
-	TArray<UObject*> ObjectsToSerialize;
-	TArray<FRemoteObjectPathName> PathNames;
 	TArray<FRemoteObjectHeader> ObjectHeaders;
 	int32 Version = 0;
 	int64 HeaderOffset = 0;
-	FRemoteObjectId RequestedObjectId(InObject);
-	UObject* Object = InObject;
 
-	if (GRemoteObjectsMigrateFullHierarchy)
+	int64 OffsetOfHeaderOffset = 0;
 	{
-		Object = FindCanonicalRootObjectForSerialization(Object);
+		FRemoteObjectWriterScope Scope(Ar, TEXT("Header"));
+		Ar << Version;
+		Ar << RequestedObjectId;
+		OffsetOfHeaderOffset = Ar.Tell();
+		Ar << HeaderOffset;
 	}
-	
+
+	int32 SerializedObjectIndex = 0;
+	bool bSerializedRequestedObject = false;
+
+	do
 	{
-		FArchiveRemoteObjectWriter Ar(Object, *ObjectData, ObjectsToSerialize, OutObjects);
-		ObjectsToSerialize.Add(Object);
+		for (; SerializedObjectIndex < Ar.GetObjectsToSerialize().Num(); ++SerializedObjectIndex)
+		{
+			UObject* ObjectToSerialize = Ar.GetObjectsToSerialize()[SerializedObjectIndex];
+			if (!OutObjects.Contains(ObjectToSerialize))
+			{
+				FRemoteObjectWriterScope Scope(Ar, ObjectToSerialize);
+				OutObjects.Add(ObjectToSerialize);
+				FRemoteObjectHeader& Header = ObjectHeaders.Emplace_GetRef();
+				Ar.SerializeRemoteObject(ObjectToSerialize, Header);
+				Header.NextOffset = Ar.Tell();
+			}
+		}
 
-		int64 OffsetOfHeaderOffset = 0;
+		bSerializedRequestedObject = OutObjects.Contains(RequestedObject);
+		if (!bSerializedRequestedObject)
 		{
-			FRemoteObjectWriterScope Scope(Ar, TEXT("Header"));
-			Ar << Version;
-			Ar << RequestedObjectId;
-			OffsetOfHeaderOffset = Ar.Tell();
-			Ar << HeaderOffset;
+			// InObject was a default subobject (see GRemoteObjectsMigrateFullHierarchy) but when we serialized its parent 
+			// it turned out that the parent had no direct reference to InObject in which case we need to manually add InObject to ObjectsToSerialize list
+			Ar.GetObjectsToSerialize().Add(RequestedObject);
 		}
+	} while (!bSerializedRequestedObject);
 
-		int32 SerializedObjectIndex = 0;
-		bool bSerializedRequestedObject = false;
+	HeaderOffset = Ar.Tell();
+	Ar.Seek(OffsetOfHeaderOffset);
+	Ar << HeaderOffset;
+	Ar.Seek(HeaderOffset);
 
-		do
-		{
-			for (; SerializedObjectIndex < ObjectsToSerialize.Num(); ++SerializedObjectIndex)
-			{
-				UObject* ObjectToSerialize = ObjectsToSerialize[SerializedObjectIndex];
-				if (!OutObjects.Contains(ObjectToSerialize))
-				{
-					FRemoteObjectWriterScope Scope(Ar, ObjectToSerialize);
-					OutObjects.Add(ObjectToSerialize);
-					FRemoteObjectHeader& Header = ObjectHeaders.Emplace_GetRef();
-					SerializeRemoteObject(ObjectToSerialize, Ar, Header);
-					Ar.AddNameToNameMap(Header.Name);
-				}
-			}
+	{
+		FRemoteObjectWriterScope Scope(Ar, TEXT("ObjectHeaders"));
+		Ar << ObjectHeaders;
+	}
+}
 
-			bSerializedRequestedObject = OutObjects.Contains(InObject);
-			if (!bSerializedRequestedObject)
-			{
-				// InObject was a default subobject (see GRemoteObjectsMigrateFullHierarchy) but when we serialized its parent 
-				// it turned out that the parent had no direct reference to InObject in which case we need to manually add InObject to ObjectsToSerialize list
-				ObjectsToSerialize.Add(InObject);
-			}
-		} while (!bSerializedRequestedObject);
+FRemoteObjectData* SerializeObjectData(UObject* InObject, TSet<UObject*>& OutObjects)
+{
+	using namespace UE::RemoteObject::Handle;
 
-		HeaderOffset = Ar.Tell();
-		Ar.Seek(OffsetOfHeaderOffset);
-		Ar << HeaderOffset;
-		Ar.Seek(HeaderOffset);
+	UObject* Object = InObject;
+	if (GRemoteObjectsMigrateFullHierarchy)
+	{
+		Object = FindCanonicalRootObjectForSerialization(Object);
+	}
+
+	FRemoteObjectData* ObjectData = new FRemoteObjectData();	
+	FRemoteObjectId RequestedObjectId(InObject);
+
+	{
+		FArchiveRemoteObjectWriter Ar(Object, *ObjectData);
+		SerializeObjectDataInternal(Ar, InObject, RequestedObjectId, OutObjects);
+	}
 
+	return ObjectData;
+}
+
+void ResetRemoteObject(UObject* InObject)
+{
+	UObject* Object = InObject;
+	if (GRemoteObjectsMigrateFullHierarchy)
+	{
+		Object = FindCanonicalRootObjectForSerialization(Object);
+	}
+
+	TMap<UObject*, UObject*> ReverseArchetypeToObjectMap;
+
+	// Serialize Object to gather its subobjects and map the subobject archetypes to their respective instances
+	// This map will also be used to replace archetype object pathnames in the serialized archetype data to their instances' pathnames
+	{
+		TArray<UObject*> ObjectsToSerialize;
+		TSet<UObject*> SerializedObjects;
+		FArchiveSubObjectGatherer SubobjectGather(Object, ReverseArchetypeToObjectMap, ObjectsToSerialize);
+		ObjectsToSerialize.Add(Object);
+
+		for (int32 Index = 0; Index < ObjectsToSerialize.Num(); ++Index)
 		{
-			FRemoteObjectWriterScope Scope(Ar, TEXT("ObjectHeaders"));
-			Ar << ObjectHeaders;
+			UObject* Obj = ObjectsToSerialize[Index];
+			if (!SerializedObjects.Contains(Obj))
+			{
+				Obj->Serialize(SubobjectGather);
+				SerializedObjects.Add(Obj);
+			}
 		}
 	}
 
-	return ObjectData;
+	FRemoteObjectData* ArchetypeDelta = new FRemoteObjectData();
+	{
+		// Serialize the Object Archetype against the Object (and its subobjects)
+		// This will produce a delta between the archetypes and their instances which will then be used to restore the instances' state to the archetypes'
+		FRemoteObjectId RequestedObjectId(Object);
+		UObject* Archetype = Object->GetArchetype();
+		FArchetypeDeltaWriter Ar(Archetype, *ArchetypeDelta, ReverseArchetypeToObjectMap);
+		TSet<UObject*> ArchetypeSerializedObjects;
+		SerializeObjectDataInternal(Ar, Archetype, RequestedObjectId, ArchetypeSerializedObjects);
+	}
+
+	{
+		const FRemoteServerId LocalServerId = UE::RemoteObject::GetGlobalServerId();
+
+		// Deserialize archetype data on top of the object (and its subobjects) to restore their state to the archetype values
+		TArray<UObject*> DeserializedObjects;
+		TArray<FRemoteObjectId> DeserializedIds;
+		// When deserializing archetype delta we want to:
+		// Preserve any references to remote objects that have not been migrated yet (this is because we can't generate archetype delta for them because they don't exist on this server and we don't have their data)
+		// Use any existing objects and deserialize archetype delta on top of them avoiding re-construction which may lead to undiserable side-effects
+		// Additionally we don't want to recursively re-enter this function so let the deserialization process know we're already resetting migrated object(s)
+		ERemoteObjectSerializationFlags DeserializationFlags = ERemoteObjectSerializationFlags::PreserveRemoteReferences | ERemoteObjectSerializationFlags::UseExistingObjects | ERemoteObjectSerializationFlags::Resetting;
+		DeserializeObjectData(LocalServerId, *ArchetypeDelta, DeserializedIds, DeserializedObjects, DeserializationFlags);
+	}
+	delete ArchetypeDelta;
 }
 
-int32 DeserializeObjectData(FRemoteServerId OwnerServerId, FRemoteObjectData& ObjectData, TArray<FRemoteObjectId>& OutObjectRemoteIds, TArray<UObject*>& OutObjects)
+int32 DeserializeObjectData(FRemoteServerId OwnerServerId, FRemoteObjectData& ObjectData, TArray<FRemoteObjectId>& OutObjectRemoteIds, TArray<UObject*>& OutObjects, ERemoteObjectSerializationFlags DeserializeFlags)
 {
 	using namespace UE::RemoteObject::Handle;
 	using namespace UE::RemoteObject::Private;
@@ -907,7 +1143,7 @@ int32 DeserializeObjectData(FRemoteServerId OwnerServerId, FRemoteObjectData& Ob
 
 	ResolvePathNames(ObjectData, ResolvedPathNameObjects);
 
-	FArchiveRemoteObjectReader Ar(ObjectData, ResolvedPathNameObjects);
+	FArchiveRemoteObjectReader Ar(ObjectData, ResolvedPathNameObjects, DeserializeFlags);
 
 	{
 		FRemoteObjectReaderScope Scope(Ar, TEXT("Header"));
@@ -943,6 +1179,14 @@ int32 DeserializeObjectData(FRemoteServerId OwnerServerId, FRemoteObjectData& Ob
 			}
 		}
 
+		const bool bReturningBorrowedObject = GResetBorrowedObjects && IsOwned(ObjectHeaders[0].RemoteId);
+		if (bReturningBorrowedObject)
+		{
+			// Root object is owned by this server so we're receiving an object that was borrowed by another server.
+			// In this case we don't need to reconstruct anything and we can re-use the objects that are already in memory.
+			DeserializeFlags |= ERemoteObjectSerializationFlags::UseExistingObjects;
+		}
+
 		// Construct all objects first
 		FRemoteObjectConstructionOverrides ConstructionOverrides(ObjectHeaders);
 		OutObjectRemoteIds.Reserve(ObjectHeaders.Num());
@@ -950,35 +1194,60 @@ int32 DeserializeObjectData(FRemoteServerId OwnerServerId, FRemoteObjectData& Ob
 		for (int32 ObjectIndex = 0; ObjectIndex < ObjectHeaders.Num(); ++ObjectIndex)
 		{
 			OutObjectRemoteIds.Add(ObjectHeaders[ObjectIndex].RemoteId);
-			OutObjects.Add(ConstructRemoteObject(ObjectHeaders[ObjectIndex], &ConstructionOverrides));
+			OutObjects.Add(ConstructRemoteObject(ObjectHeaders[ObjectIndex], &ConstructionOverrides, DeserializeFlags));
 		}
 
-		// ensure remote object stubs are created and ownership is assigned
-		for (int32 ObjectIndex = 0; ObjectIndex < OutObjects.Num(); ++ObjectIndex)
+		// If we're already resetting a borrowed object we don't want to change the ownership until the object is deserialized using remote server data
+		// and we also don't want to recursively reset the object
+		const bool bResetting = EnumHasAnyFlags(DeserializeFlags, ERemoteObjectSerializationFlags::Resetting);
+		if (!bResetting)
 		{
-			UObject* Object = OutObjects[ObjectIndex];
-			UE::RemoteObject::Private::RegisterRemoteObjectId(FRemoteObjectId(Object), LocalServerId);
-			ChangeOwnerServerId(Object, OwnerServerId);
+			// ensure remote object stubs are created and ownership is assigned
+			for (int32 ObjectIndex = 0; ObjectIndex < OutObjects.Num(); ++ObjectIndex)
+			{
+				if (UObject* Object = OutObjects[ObjectIndex])
+				{
+					UE::RemoteObject::Private::RegisterRemoteObjectId(FRemoteObjectId(Object), LocalServerId);
+					ChangeOwnerServerId(Object, OwnerServerId);
+				}
+			}
+
+			if (bReturningBorrowedObject)
+			{
+				ResetRemoteObject(OutObjects[0]);
+			}
 		}
 
 		// Deserialize all objects
 		Ar.SetRootObject(OutObjects.Num() ? OutObjects[0] : nullptr);
 		for (int32 ObjectIndex = 0; ObjectIndex < OutObjects.Num(); ++ObjectIndex)
 		{
-			UObject* Object = OutObjects[ObjectIndex];
-			FRemoteObjectReaderScope Scope(Ar, Object);
-			Object->Serialize(Ar);
-			if (RequestedObjectId == FRemoteObjectId(Object))
+			if (UObject* Object = OutObjects[ObjectIndex])
 			{
-				RequestedObjectIndex = ObjectIndex;
+				FRemoteObjectReaderScope Scope(Ar, Object);
+				Object->Serialize(Ar);
+				if (RequestedObjectIndex == -1 && RequestedObjectId == FRemoteObjectId(Object))
+				{
+					RequestedObjectIndex = ObjectIndex;
+				}
+			}
+			else
+			{
+				Ar.Seek(ObjectHeaders[ObjectIndex].NextOffset);
 			}
 		}
 
-		// Patch any unresolved remote references
-		FArchiveRemoteReferencePatcher PatchAr(OutObjects);
-		for (UObject* Object : OutObjects)
+		// Patch any unresolved remote references (skip if we're resetting an object to its CDO state)
+		if (!bResetting)
 		{
-			Object->Serialize(PatchAr);
+			FArchiveRemoteReferencePatcher PatchAr(OutObjects);
+			for (UObject* Object : OutObjects)
+			{
+				if (Object)
+				{
+					Object->Serialize(PatchAr);
+				}
+			}
 		}
 	}
 	checkf(RequestedObjectIndex >= 0, TEXT("Received remote object data but the requested object (%s) was not deserialized"), *RequestedObjectId.ToString());
diff --git a/Engine/Source/Runtime/CoreUObject/Private/UObject/RemoteObjectSerialization.h b/Engine/Source/Runtime/CoreUObject/Private/UObject/RemoteObjectSerialization.h
--- a/Engine/Source/Runtime/CoreUObject/Private/UObject/RemoteObjectSerialization.h
+++ b/Engine/Source/Runtime/CoreUObject/Private/UObject/RemoteObjectSerialization.h
@@ -12,6 +12,15 @@ namespace UE::RemoteObject::Handle { struct FRemoteObjectStub; }
 
 namespace UE::RemoteObject::Serialization
 {
+	enum class ERemoteObjectSerializationFlags : uint8
+	{
+		None = 0,
+		UseExistingObjects = (1 << 0),			// If possible re-use existing UObjects and don't reconstruct them when deserializing object data
+		PreserveRemoteReferences = (1 << 1),		// Don't overwrite references to objects that are remote
+		Resetting = (1 << 2)						// Indicates that the serialization process is resetting an object to its archetype state
+	};
+	ENUM_CLASS_FLAGS(ERemoteObjectSerializationFlags);
+
 	/**
 	* Basic information needed to construct deserialized remote object (this information does not get serialized in UObject::Serialize())
 	*/
@@ -23,6 +32,7 @@ namespace UE::RemoteObject::Serialization
 		TObjectPtr<UObject> Outer;
 		TObjectPtr<UObject> Archetype;
 		int32 InternalFlags = 0;
+		int64 NextOffset = 0;
 
 		/** Transient (SerialNumber is local-only but it's stored here for convenience when deserializing object data) */
 		int32 SerialNumber = 0;
@@ -60,20 +70,23 @@ namespace UE::RemoteObject::Serialization
 
 	/**
 	* Serializes an object and its subobject (or if the object is a default subobject, its parent and the parent's subobjects)
-	* @param Object Object to be serialized
+	* @param InObject Object to be serialized
 	* @param OutObjects All objects that have been serialized (Object and its subobjects and/or parent)
 	* @return Remote object data representing the serialized objects
 	*/
-	FRemoteObjectData* SerializeObjectData(UObject* Object, TSet<UObject*>& OutObjects);
+	FRemoteObjectData* SerializeObjectData(UObject* InObject, TSet<UObject*>& OutObjects);
 
 	/**
 	* Deserializes remote object data
+	* @param OwnerServerId New Owner Server Id of the objects that are about to be created and deserialized (it may be identicaly to their previous owner if the objects were borrowed)
 	* @param ObjectData to be deserialized
+	* @param OutObjectRemoteIds Remote IDs of the deserialized objects
 	* @param OutReceivedObjects All deserialized objects
+	* @param DeserializeFlags Flags modifying the behavior of the deserialization process
 	* @return Index of an object in OutReceivedObjects that was the main object the migration request was triggered for 
 	* (usually 0 but if a migration requests a default subobject then its parent is also migrated and the return value will be > 0)
 	*/
-	int32 DeserializeObjectData(FRemoteServerId OwnerServerId, FRemoteObjectData& ObjectData, TArray<FRemoteObjectId>& OutObjectRemoteIds, TArray<UObject*>& OutReceivedObjects);
+	int32 DeserializeObjectData(FRemoteServerId OwnerServerId, FRemoteObjectData& ObjectData, TArray<FRemoteObjectId>& OutObjectRemoteIds, TArray<UObject*>& OutReceivedObjects, ERemoteObjectSerializationFlags DeserializeFlags = ERemoteObjectSerializationFlags::None);
 
 	/**
 	* Finds the canonical 'root' object that is used for remote object serialization
