[CoreRedirects]

- Cleans up FCoreRedirectObjectName::HasValidCharacter to internally handle Verse paths which have special exemptions to our object path rules. Previously exclusions for Verse Packages were handled inside the adding  of redirects but would were missing for removal. That issue was fixed in 39985639 but these changes change cleans up the code and adds some tests.

#rnx

diff --git a/Engine/Source/Runtime/CoreUObject/Private/UObject/CoreRedirects.cpp b/Engine/Source/Runtime/CoreUObject/Private/UObject/CoreRedirects.cpp
--- a/Engine/Source/Runtime/CoreUObject/Private/UObject/CoreRedirects.cpp
+++ b/Engine/Source/Runtime/CoreUObject/Private/UObject/CoreRedirects.cpp
@@ -859,13 +859,23 @@ FName FCoreRedirectObjectName::GetSearchKey(ECoreRedirectFlags Type) const
 
 bool FCoreRedirectObjectName::HasValidCharacters(ECoreRedirectFlags Type) const
 {
-	static FString InvalidRedirectCharacters = TEXT("\"',|&!~\n\r\t@#(){}[]=;^%$`");
-	static FString InvalidObjectRedirectCharacters = TEXT(".\n\r\t"); // Object and field names in Blueprints are very permissive...
+	if (FPackageName::IsVersePackage(WriteToString<NAME_SIZE>(PackageName)))
+	{
+		// Verse allows the visible character ASCII set minus these CR, LF, TAB chars.
+		static FString InvalidVerseCharacters = TEXT("\n\r\t"); 
+		return  ObjectName.IsValidXName(InvalidVerseCharacters) && 
+				OuterName.IsValidXName(InvalidVerseCharacters) &&
+				PackageName.IsValidXName(InvalidVerseCharacters);
+	}
+
+	static FString InvalidRedirectNameCharacters = TEXT("\"',|&!~\n\r\t@#(){}[]=;^%$`");
+	static FString InvalidObjectNameCharacters = TEXT(".\n\r\t"); // Object and field names in Blueprints are very permissive
+	const bool bUsePermissiveObjectNames = EnumHasAnyFlags(Type, 
+		ECoreRedirectFlags::Type_Object | ECoreRedirectFlags::Type_Property | ECoreRedirectFlags::Type_Function);
 
-	return ObjectName.IsValidXName(
-		EnumHasAnyFlags(Type, ECoreRedirectFlags::Type_Object | ECoreRedirectFlags::Type_Property| ECoreRedirectFlags::Type_Function)
-		? InvalidObjectRedirectCharacters : InvalidRedirectCharacters)
-		&& OuterName.IsValidXName(InvalidRedirectCharacters) && PackageName.IsValidXName(InvalidRedirectCharacters);
+	return	ObjectName.IsValidXName(bUsePermissiveObjectNames ? InvalidObjectNameCharacters : InvalidRedirectNameCharacters) &&
+			OuterName.IsValidXName(InvalidRedirectNameCharacters) &&
+			PackageName.IsValidXName(InvalidRedirectNameCharacters);
 }
 
 bool FCoreRedirectObjectName::ExpandNames(const FStringView InStringView, FName& OutName, FName& OutOuter, FName &OutPackage)
@@ -2162,6 +2172,72 @@ bool FCoreRedirects::RunTests()
 	}
 	OldNames.Empty();
 
+	// ObjectNames have almost no restrictions but PackageNames are more restrictive
+	if (FCoreRedirectObjectName(TEXT("/Foo/Foo.Foo:\nNew")).HasValidCharacters(ECoreRedirectFlags::Type_Object))
+	{
+		bSuccess = false;
+		UE_LOG(LogCoreRedirects, Error, TEXT("FCoreRedirect Test Failed: HasValidCharacters should report Object redirects with \\n in the ObjectName as invalid"));
+	}
+	if (FCoreRedirectObjectName(TEXT("/Foo/Foo.\nFoo:New")).HasValidCharacters(ECoreRedirectFlags::Type_Object))
+	{
+		bSuccess = false;
+		UE_LOG(LogCoreRedirects, Error, TEXT("FCoreRedirect Test Failed: HasValidCharacters should report Object redirects with \\n in the OuterName as invalid"));
+	}
+	if (FCoreRedirectObjectName(TEXT("/Foo/\nFoo.Foo:New")).HasValidCharacters(ECoreRedirectFlags::Type_Object))
+	{
+		bSuccess = false;
+		UE_LOG(LogCoreRedirects, Error, TEXT("FCoreRedirect Test Failed: HasValidCharacters should report Object redirects with \\n in the PackageName as invalid"));
+	}
+	if (FCoreRedirectObjectName(TEXT("/Foo/_Verse/Foo.Foo:\nNew")).HasValidCharacters(ECoreRedirectFlags::Type_Object))
+	{
+		bSuccess = false;
+		UE_LOG(LogCoreRedirects, Error, TEXT("FCoreRedirect Test Failed: HasValidCharacters should report verse Object redirects with \\n in the ObjectName as invalid"));
+	}
+	if (FCoreRedirectObjectName(TEXT("/Foo/_Verse/Foo.\nFoo:New")).HasValidCharacters(ECoreRedirectFlags::Type_Object))
+	{
+		bSuccess = false;
+		UE_LOG(LogCoreRedirects, Error, TEXT("FCoreRedirect Test Failed: HasValidCharacters should report verse Object redirects with \\n in the OuterName as invalid"));
+	}
+	if (FCoreRedirectObjectName(TEXT("/Foo/_Verse/\nFoo.Foo:New")).HasValidCharacters(ECoreRedirectFlags::Type_Object))
+	{
+		bSuccess = false;
+		UE_LOG(LogCoreRedirects, Error, TEXT("FCoreRedirect Test Failed: HasValidCharacters should report verse Object redirects with \\n in the PackageName as invalid"));
+	}
+
+	// Blueprints allow almost anything to be an ObjectName, but should complain about outers and package names with special characters
+	if (!FCoreRedirectObjectName(TEXT("/Foo/Foo.Foo:$New")).HasValidCharacters(ECoreRedirectFlags::Type_Object))
+	{
+		bSuccess = false;
+		UE_LOG(LogCoreRedirects, Error, TEXT("FCoreRedirect Test Failed: HasValidCharacters should report Object redirects with $ in the ObjectName as valid"));
+	}
+	if (FCoreRedirectObjectName(TEXT("/Foo/Foo.$Foo:New")).HasValidCharacters(ECoreRedirectFlags::Type_Object))
+	{
+		bSuccess = false;
+		UE_LOG(LogCoreRedirects, Error, TEXT("FCoreRedirect Test Failed: HasValidCharacters should report Object redirects with $ in the OuterName as invalid"));
+	}
+	if (FCoreRedirectObjectName(TEXT("/Foo/$Foo.Foo:New")).HasValidCharacters(ECoreRedirectFlags::Type_Object))
+	{
+		bSuccess = false;
+		UE_LOG(LogCoreRedirects, Error, TEXT("FCoreRedirect Test Failed: HasValidCharacters should report Object redirects with $ in the PackageName as invalid"));
+	}
+	// Verse is much more permissive and allows $ anywhere
+	if (!FCoreRedirectObjectName(TEXT("/Foo/_Verse/Foo.Foo:$New")).HasValidCharacters(ECoreRedirectFlags::Type_Object))
+	{
+		bSuccess = false;
+		UE_LOG(LogCoreRedirects, Error, TEXT("FCoreRedirect Test Failed: HasValidCharacters should report verse Object redirects with $ in the ObjectName as valid"));
+	}
+	if (!FCoreRedirectObjectName(TEXT("/Foo/_Verse/Foo.$Foo:New")).HasValidCharacters(ECoreRedirectFlags::Type_Object))
+	{
+		bSuccess = false;
+		UE_LOG(LogCoreRedirects, Error, TEXT("FCoreRedirect Test Failed: HasValidCharacters should report verse Object redirects with $ in the OuterName as valid"));
+	}
+	if (!FCoreRedirectObjectName(TEXT("/Foo/_Verse/$Foo.Foo:New")).HasValidCharacters(ECoreRedirectFlags::Type_Object))
+	{
+		bSuccess = false;
+		UE_LOG(LogCoreRedirects, Error, TEXT("FCoreRedirect Test Failed: HasValidCharacters should report verse Object redirects with $ in the PackageName as valid"));
+	}
+
+
 	// Check removed
 	if (!IsKnownMissing(ECoreRedirectFlags::Type_Package, FCoreRedirectObjectName(TEXT("/Game/RemovedPackage"))))
 	{
@@ -2628,8 +2704,8 @@ bool FCoreRedirects::AddRedirectListUnderWriteLock(TArrayView<const FCoreRedirec
 		if (!!(LockedContext.Get().GetFlags() & FCoreRedirectsContext::EFlags::ValidateAddedRedirects)
 			&& !EnumHasAnyFlags(NewRedirect.RedirectFlags, ECoreRedirectFlags::Type_Asset))
 		{
-			if ((!NewRedirect.OldName.HasValidCharacters(NewRedirect.RedirectFlags) && !FPackageName::IsVersePackage(NewRedirect.OldName.PackageName.ToString()))
-				|| (!NewRedirect.NewName.HasValidCharacters(NewRedirect.RedirectFlags) && !FPackageName::IsVersePackage(NewRedirect.NewName.PackageName.ToString())))
+			if (!NewRedirect.OldName.HasValidCharacters(NewRedirect.RedirectFlags)
+				|| !NewRedirect.NewName.HasValidCharacters(NewRedirect.RedirectFlags))
 			{
 				UE_LOG(LogCoreRedirects, Error, TEXT("AddRedirect(%s) failed to add redirect from %s to %s with invalid characters!"), *SourceString, *NewRedirect.OldName.ToString(), *NewRedirect.NewName.ToString());
 				continue;
@@ -2781,8 +2857,8 @@ bool FCoreRedirects::RemoveRedirectList(TArrayView<const FCoreRedirect> Redirect
 			continue;
 		}
 
-		if ((!RedirectToRemove.OldName.HasValidCharacters(RedirectToRemove.RedirectFlags) && !FPackageName::IsVersePackage(RedirectToRemove.OldName.PackageName.ToString()))
-			|| (!RedirectToRemove.NewName.HasValidCharacters(RedirectToRemove.RedirectFlags) && !FPackageName::IsVersePackage(RedirectToRemove.NewName.PackageName.ToString())))
+		if (!RedirectToRemove.OldName.HasValidCharacters(RedirectToRemove.RedirectFlags) ||
+			!RedirectToRemove.NewName.HasValidCharacters(RedirectToRemove.RedirectFlags))
 		{
 			UE_LOG(LogCoreRedirects, Error, TEXT("RemoveRedirect(%s) failed to remove redirect from %s to %s with invalid characters!"), *SourceString, *RedirectToRemove.OldName.ToString(), *RedirectToRemove.NewName.ToString());
 			continue;
