[CoreRedirects]

- Treats all wildcard match types the same and leverage substring fuzzy matching for prefix and suffix matches to potentially early out of expensive string comparisons. This change accelerates Suffix and Prefix matches at the cost of potentially increasing the false positive rate of substring matches  due to increasing the pool of strings to compare against as possible hits when fuzzy matching for all cases (prefix, suffix and substring).  However no noticeable change has been seen for common use common large map tests and string comparisons run less than 0.9% of the time (in cases for header patching  where prefix matches are used a lot, string comparisons are now done < 0.3% of the time.
- When defining MatchWilcar=True for INI define reirects, prefix an suffix redirects are longer treate as substring comparisons.
- Removes the temporary support for "MatchWildcardStrict" when specifying redirects by INI added in 39985639, and replaces usage of "MatchWildcardStrict" with "MatchWildcard"

#rnx

diff --git a/Engine/Source/Runtime/CoreUObject/Private/UObject/CoreRedirects.cpp b/Engine/Source/Runtime/CoreUObject/Private/UObject/CoreRedirects.cpp
--- a/Engine/Source/Runtime/CoreUObject/Private/UObject/CoreRedirects.cpp
+++ b/Engine/Source/Runtime/CoreUObject/Private/UObject/CoreRedirects.cpp
@@ -35,23 +35,23 @@ DEFINE_LOG_CATEGORY(LogCoreRedirects);
 
 #define COREREDIRECT_STATS 0
 #if COREREDIRECT_STATS
-static int sSubstringLookups = 0;
-static int sSubstringPredictHit = 0;
-static int sSubstringPredictMiss = 0;
-FAutoConsoleCommand StaticLoadTrackerDump(TEXT("CoreRedirects.SubstringPrediction.StatsDump"), TEXT(""), FConsoleCommandDelegate::CreateLambda([]()
+static int sWildcardLookups = 0;
+static int sWildcardPredictHit = 0;
+static int sWildcardPredictMiss = 0;
+FAutoConsoleCommand StaticLoadTrackerDump(TEXT("CoreRedirects.WildcardPrediction.StatsDump"), TEXT(""), FConsoleCommandDelegate::CreateLambda([]()
 	{
-		UE_LOG(LogCoreRedirects, Log, TEXT(
+		UE_LOG(LogCoreRedirects, Log, TEXT("\n\n"
 			"==========================================\n"
-			"Dumping Substring Stats\n\n"
+			"Dumping Wildcard Stats\n\n"
 			"Lookups:\t\t\t\t%d\n"
-			"Prediction Hit:\t\t%d\n"
+			"Prediction Hit:\t\t\t%d\n"
 			"Prediction Miss:\t\t%d\n"
 			"=========================================="
-		), sSubstringLookups, sSubstringPredictHit, sSubstringPredictMiss);
+		),  sWildcardLookups, sWildcardPredictHit, sWildcardPredictMiss)
 	}));
 
-#define COREREDIRECT_STATS_UPDATE_PREDICTION_LOOKUP() do{ sSubstringLookups++; }while(false)
-#define COREREDIRECT_STATS_UPDATE_PREDICTION_RESULT(bFound) do{ if(bFound){ sSubstringPredictHit++; } else { sSubstringPredictMiss++; }}while(false)
+#define COREREDIRECT_STATS_UPDATE_PREDICTION_LOOKUP() do{ sWildcardLookups++; }while(false)
+#define COREREDIRECT_STATS_UPDATE_PREDICTION_RESULT(bFound) do{ if(bFound){ sWildcardPredictHit++; } else { sWildcardPredictMiss++; }}while(false)
 #else
 #define COREREDIRECT_STATS_UPDATE_PREDICTION_LOOKUP()
 #define COREREDIRECT_STATS_UPDATE_PREDICTION_RESULT(bFound)
@@ -62,8 +62,8 @@ FAutoConsoleCommand StaticLoadTrackerDump(TEXT("CoreRedirects.SubstringPredictio
 * without relying on internal details of FName. Due to the sheer volume of CoreRedirects, it's 
 * significant to avoid the copies where possible. We maintain UTF8 encoding for FNames when dealing with Wildcard
 * redirects since UTF8 overlays ASCII allowing pure ASCII strings to be compact, greatly improving PredictMatch efficiency.
-* We also need a common encoding for how we store substring redirects and how we query them to ensure we do not generate
-* false negatives in our fuzzy searchs using PredictMatch; UTF8 will allow us to still predict non-ASCII substrings.
+* We also need a common encoding for how we store wildcard redirects and how we query them to ensure we do not generate
+* false negatives in our fuzzy searchs using PredictMatch; UTF8 will allow us to still predict non-ASCII wildcards.
 */
 namespace UE::CoreRedirects::Private
 {
@@ -509,7 +509,6 @@ void FCoreRedirects::FWildcardData::Add(const FCoreRedirect& Redirect)
 	if (Redirect.IsSubstringMatch())
 	{
 		Substrings.Add(Redirect);
-		AddPredictionWords(Redirect);
 	}
 	else if (Redirect.IsPrefixMatch())
 	{
@@ -520,6 +519,7 @@ void FCoreRedirects::FWildcardData::Add(const FCoreRedirect& Redirect)
 		check(Redirect.IsSuffixMatch());
 		Suffixes.Add(Redirect);
 	}
+	AddPredictionWords(Redirect);
 }
 
 void FCoreRedirects::FWildcardData::AddPredictionWords(const FCoreRedirect& Redirect)
@@ -538,7 +538,7 @@ void FCoreRedirects::FWildcardData::AddPredictionWords(const FCoreRedirect& Redi
 
 		// Since we only predict based on a small window of characters (8), 
 		// try to use characters that will be more unique by removing common prefixes.
-		// Note, this does not affect the actual substring matching, only
+		// Note, this does not affect the actual wildcard matching, only
 		// our prediction and more specifically our false positive rate.
 		FUtf8StringView CommonPrefixes[] = 
 		{ "/", "Script/", "Temp/", "Extra/", "Memory/", "Config/", "Game/", "Engine/", "Transient/", "Niagara/"};
@@ -585,6 +585,16 @@ void FCoreRedirects::FWildcardData::Rebuild()
 {
 	PredictMatch.Reset();
 
+	for (const FCoreRedirect& Redirect : Prefixes)
+	{
+		AddPredictionWords(Redirect);
+	}
+
+	for (const FCoreRedirect& Redirect : Suffixes)
+	{
+		AddPredictionWords(Redirect);
+	}
+
 	for (const FCoreRedirect& Redirect : Substrings)
 	{
 		AddPredictionWords(Redirect);
@@ -613,39 +623,40 @@ bool FCoreRedirects::FWildcardData::Matches(ECoreRedirectFlags InFlags, const FC
 	}
 
 	const FCoreRedirectObjectUtf8Name ObjectUtf8Name(InName);
-	if (bPrefixMatch)
+	
+	COREREDIRECT_STATS_UPDATE_PREDICTION_LOOKUP();
+	// Perform a fuzzy match against all known wildcard (prefix, suffix and substrings since they all overlap). 
+	// We may have false positives so we must still perform a stronger check on a fuzzy match. 
+	// However we can eliminate the common failure case quickly with a fast fuzzy match.
+	if (MatchSubstringApproximate(ObjectUtf8Name))
 	{
-		for (const FCoreRedirect& CheckRedirect : Prefixes)
+		// We found a possible wildcard match, so now do a slow match to find if we really match
+		if (bPrefixMatch)
 		{
-			if (MatchPrefix(CheckRedirect.OldName, ObjectUtf8Name, bPartialRHS))
+			for (const FCoreRedirect& CheckRedirect : Prefixes)
 			{
-				bFound = true;
-				OutFoundRedirects.Add(&CheckRedirect);
+				if (MatchPrefix(CheckRedirect.OldName, ObjectUtf8Name, bPartialRHS))
+				{
+					bFound = true;
+					OutFoundRedirects.Add(&CheckRedirect);
+				}
 			}
 		}
-	}
 
-	if (bSuffixMatch)
-	{
-		for (const FCoreRedirect& CheckRedirect : Suffixes)
+		if (bSuffixMatch)
 		{
-			if (MatchSuffix(CheckRedirect.OldName, ObjectUtf8Name, bPartialRHS))
+			for (const FCoreRedirect& CheckRedirect : Suffixes)
 			{
-				bFound = true;
-				OutFoundRedirects.Add(&CheckRedirect);
+				if (MatchSuffix(CheckRedirect.OldName, ObjectUtf8Name, bPartialRHS))
+				{
+					bFound = true;
+					OutFoundRedirects.Add(&CheckRedirect);
+				}
 			}
 		}
-	}
-	
-	if (bSubstringMatch && !Substrings.IsEmpty())
-	{
-		COREREDIRECT_STATS_UPDATE_PREDICTION_LOOKUP();
-		// Perform a fuzzy match against all known substrings. We may have false positives
-		// so we must still perform a stronger check on a fuzzy match. However we can eliminate
-		// the common failure case quickly with a fast fuzzy match.
-		if (MatchSubstringApproximate(ObjectUtf8Name))
+
+		if (bSubstringMatch )
 		{
-			// We found a possible substring match, so now do a slow match to find if we really match
 			for (const FCoreRedirect& CheckRedirect : Substrings)
 			{
 				if (MatchSubstring(CheckRedirect.OldName, ObjectUtf8Name, bPartialRHS))
@@ -654,8 +665,9 @@ bool FCoreRedirects::FWildcardData::Matches(ECoreRedirectFlags InFlags, const FC
 					OutFoundRedirects.Add(&CheckRedirect);
 				}
 			}
-			COREREDIRECT_STATS_UPDATE_PREDICTION_RESULT(bFound);
 		}
+			
+		COREREDIRECT_STATS_UPDATE_PREDICTION_RESULT(bFound);
 	}
 
 	return bFound;
@@ -2540,11 +2552,6 @@ bool FCoreRedirects::ReadRedirectsFromIni(const FString& IniName)
 				FParse::Bool(*ValueString, TEXT("MatchSubstring="), bMatchSubstring);
 				FParse::Bool(*ValueString, TEXT("MatchWildcard="), bMatchWildcard);
 
-				// This works around UE-247674.
-				bool bMatchWildcardStrict = false;
-				FParse::Bool(*ValueString, TEXT("MatchWildcardStrict="), bMatchWildcardStrict);
-				bMatchWildcard |= bMatchWildcardStrict;
-
 				FParse::Value(*ValueString, TEXT("OldName="), OldName);
 				FParse::Value(*ValueString, TEXT("NewName="), NewName);
 
@@ -2609,12 +2616,6 @@ bool FCoreRedirects::ReadRedirectsFromIni(const FString& IniName)
 							OldName.RightChopInline(WildcardMarker.Len(), EAllowShrinking::No);
 						}
 
-						// UE-209583 - Convert prefix and suffix matches to substring matches since prefix/suffix is currently unoptimized
-						if (!bMatchWildcardStrict && (bMatchPrefix || bMatchSuffix))
-						{
-							bMatchSubstring = true;
-						}
-
 						NewFlags |= bMatchSubstring ? ECoreRedirectFlags::Option_MatchSubstring : ECoreRedirectFlags::None;
 					}
 
@@ -2951,7 +2952,7 @@ bool FCoreRedirects::RemoveSingleRedirectUnderWriteLock(const FCoreRedirect& Red
 	{
 		if (bIsWildcardMatch)
 		{
-			// We removed a substring redirect so we need to regenerate our prediction 
+			// We removed a wildcard redirect so we need to regenerate our prediction 
 			// tables to avoid unnecessary false positives
 			ExistingNameMap->Wildcards->Rebuild();
 		}
