traceviewer:
- provide a text name for each transactional work
- first pass support for nested regions

diff --git a/Engine/Source/Runtime/CoreUObject/Private/UObject/RemoteExecutor.cpp b/Engine/Source/Runtime/CoreUObject/Private/UObject/RemoteExecutor.cpp
--- a/Engine/Source/Runtime/CoreUObject/Private/UObject/RemoteExecutor.cpp
+++ b/Engine/Source/Runtime/CoreUObject/Private/UObject/RemoteExecutor.cpp
@@ -9,9 +9,9 @@ DEFINE_LOG_CATEGORY(LogRemoteExec);
 namespace UE::RemoteExecutor
 {
 TDelegate<void()> TickNetworkDelegate;
-TDelegate<TTuple<FRemoteWorkPriority, TFunction<void(void)>>(void)> FetchNextDeferredRPCDelegate;
+TDelegate<TTuple<FName, FRemoteWorkPriority, TFunction<void(void)>>(void)> FetchNextDeferredRPCDelegate;
 
-TMulticastDelegate<void(FRemoteTransactionId)> OnTransactionStartingDelegate;
+TMulticastDelegate<void(FRemoteTransactionId, FName)> OnTransactionStartingDelegate;
 TMulticastDelegate<void(FRemoteTransactionId)> OnTransactionCompletedDelegate;
 TMulticastDelegate<void(FRemoteTransactionId)> OnTransactionAbortedDelegate;
 }
@@ -31,6 +31,7 @@ FArchive& operator<<(FArchive& Ar, FRemoteTransactionId& Id)
 
 struct FRemoteExecutorWork
 {
+	FName Name;
 	FRemoteTransactionId RequestId;
 	FRemoteWorkPriority Priority;
 	TFunction<void(void)> Work;
@@ -61,9 +62,10 @@ public:
 		return FRemoteTransactionId((RawRequestId++ % 0x800000u) + 1000u);
 	}
 
-	void EnqueueWork(FRemoteWorkPriority InWorkPriority, const TFunction<void(void)>& InWork)
+	void EnqueueWork(FName WorkName, FRemoteWorkPriority InWorkPriority, const TFunction<void(void)>& InWork)
 	{
 		FRemoteExecutorWork& NewWork = PendingWorks.Emplace_GetRef();
+		NewWork.Name = WorkName;
 		NewWork.Priority = InWorkPriority;
 		NewWork.Work = InWork;
 	}
@@ -85,10 +87,10 @@ public:
 			// fetch pending RPC work
 			for (;;)
 			{
-				TTuple<FRemoteWorkPriority, TFunction<void(void)>> NewWork = UE::RemoteExecutor::FetchNextDeferredRPCDelegate.Execute();
-				if (NewWork.Get<0>().IsValid())
+				TTuple<FName, FRemoteWorkPriority, TFunction<void(void)>> NewWork = UE::RemoteExecutor::FetchNextDeferredRPCDelegate.Execute();
+				if (NewWork.Get<1>().IsValid())
 				{
-					EnqueueWork(NewWork.Get<0>(), NewWork.Get<1>());
+					EnqueueWork(NewWork.Get<0>(), NewWork.Get<1>(), NewWork.Get<2>());
 				}
 				else
 				{
@@ -100,8 +102,9 @@ public:
 			{
 				FRemoteExecutorWork& PendingWork = PendingWorks[PendingWorkIndex];
 
-				UE_LOG(LogRemoteExec, VeryVerbose, TEXT("PendingWork[%d] RequestId %s %s"),
+				UE_LOG(LogRemoteExec, VeryVerbose, TEXT("PendingWork[%d] '%s' RequestId %s %s"),
 					PendingWorkIndex,
+					*PendingWork.Name.ToString(),
 					*PendingWork.RequestId.ToString(),
 					*PendingWork.Priority.ToString());
 			}
@@ -191,7 +194,7 @@ public:
 				ExecutingWork.ExecutionAttempts++;
 
 				FRemoteTransactionId RequestId = ExecutingWork.RequestId;
-				UE::RemoteExecutor::OnTransactionStartingDelegate.Broadcast(RequestId);
+				UE::RemoteExecutor::OnTransactionStartingDelegate.Broadcast(RequestId, ExecutingWork.Name);
 				AutoRTFM::ETransactionResult WorkTransactResult = AutoRTFM::Transact([&ExecutingWork]() { ExecutingWork.Work(); });
 				bool bWorkComplete = false;
 				bool bWorkTransactionAborted = false;
@@ -257,12 +260,12 @@ public:
 
 FRemoteExecutor GRemoteExecutor;
 
-void ExecuteTransactionalInternal(FRemoteWorkPriority InWorkPriority, const TFunctionRef<void(void)>& InWork)
+void ExecuteTransactionalInternal(FName WorkName, FRemoteWorkPriority InWorkPriority, const TFunctionRef<void(void)>& InWork)
 {
 #if UE_WITH_REMOTE_OBJECT_HANDLE
 	check(!AutoRTFM::IsClosed());
 
-	GRemoteExecutor.EnqueueWork(InWorkPriority, [&InWork]() { InWork(); });
+	GRemoteExecutor.EnqueueWork(WorkName, InWorkPriority, [&InWork]() { InWork(); });
 	GRemoteExecutor.ExecutePendingWork();
 #else
 	InWork();
@@ -304,19 +307,19 @@ FRemoteWorkPriority CreateRootWorkPriority()
 	return RootWorkPriority;
 }
 
-void ExecuteTransactionalWithExplicitPriority(FRemoteWorkPriority WorkPriority, const TFunctionRef<void(void)>& Work)
+void ExecuteTransactionalWithExplicitPriority(FName WorkName, FRemoteWorkPriority WorkPriority, const TFunctionRef<void(void)>& Work)
 {
 #if UE_WITH_REMOTE_OBJECT_HANDLE
-	ExecuteTransactionalInternal(WorkPriority, Work);
+	ExecuteTransactionalInternal(WorkName, WorkPriority, Work);
 #else
 	Work();
 #endif
 }
 
-void ExecuteTransactional(const TFunctionRef<void(void)>& Work)
+void ExecuteTransactional(FName WorkName, const TFunctionRef<void(void)>& Work)
 {
 	FRemoteWorkPriority RootWorkPriority = FRemoteWorkPriority::CreateRootWorkPriority(UE::RemoteObject::GetGlobalServerId(), GRemoteExecutor.GenerateNextTransactionId());
-	ExecuteTransactionalWithExplicitPriority(RootWorkPriority, Work);
+	ExecuteTransactionalWithExplicitPriority(WorkName, RootWorkPriority, Work);
 }
 
 } // namespace UE::RemoteExecutor
diff --git a/Engine/Source/Runtime/CoreUObject/Private/UObject/RemoteObject.cpp b/Engine/Source/Runtime/CoreUObject/Private/UObject/RemoteObject.cpp
--- a/Engine/Source/Runtime/CoreUObject/Private/UObject/RemoteObject.cpp
+++ b/Engine/Source/Runtime/CoreUObject/Private/UObject/RemoteObject.cpp
@@ -129,7 +129,7 @@ void InitRemoteObjects()
 
 	if (!UE::RemoteExecutor::FetchNextDeferredRPCDelegate.IsBound())
 	{
-		UE::RemoteExecutor::FetchNextDeferredRPCDelegate.BindStatic([]() { return TTuple<FRemoteWorkPriority, TFunction<void(void)>>(); });
+		UE::RemoteExecutor::FetchNextDeferredRPCDelegate.BindStatic([]() { return TTuple<FName, FRemoteWorkPriority, TFunction<void(void)>>(); });
 	}
 }
 
diff --git a/Engine/Source/Runtime/CoreUObject/Public/UObject/RemoteExecutor.h b/Engine/Source/Runtime/CoreUObject/Public/UObject/RemoteExecutor.h
--- a/Engine/Source/Runtime/CoreUObject/Public/UObject/RemoteExecutor.h
+++ b/Engine/Source/Runtime/CoreUObject/Public/UObject/RemoteExecutor.h
@@ -188,10 +188,10 @@ namespace UE::RemoteExecutor
 	extern COREUOBJECT_API TDelegate<void()> TickNetworkDelegate;
 
 	/** Fetch the next deferred RPC request **/
-	extern COREUOBJECT_API TDelegate<TTuple<FRemoteWorkPriority, TFunction<void(void)>>(void)> FetchNextDeferredRPCDelegate;
+	extern COREUOBJECT_API TDelegate<TTuple<FName, FRemoteWorkPriority, TFunction<void(void)>>(void)> FetchNextDeferredRPCDelegate;
 
 	/** Delegate executed when a ExecuteTransactional starts an AutoRTFM transaction */
-	extern COREUOBJECT_API TMulticastDelegate<void(FRemoteTransactionId)> OnTransactionStartingDelegate;
+	extern COREUOBJECT_API TMulticastDelegate<void(FRemoteTransactionId, FName)> OnTransactionStartingDelegate;
 
 	/** Delegate executed when a ExecuteTransactional observes a completed AutoRTFM transaction */
 	extern COREUOBJECT_API TMulticastDelegate<void(FRemoteTransactionId)> OnTransactionCompletedDelegate;
@@ -201,8 +201,8 @@ namespace UE::RemoteExecutor
 
 	COREUOBJECT_API void RegisterRemoteSubsystem(FRemoteSubsystemBase* Subsystem);
 
-	COREUOBJECT_API void ExecuteTransactional(const TFunctionRef<void(void)>& Work);
-	COREUOBJECT_API void ExecuteTransactionalWithExplicitPriority(FRemoteWorkPriority WorkPriority, const TFunctionRef<void(void)>& Work);
+	COREUOBJECT_API void ExecuteTransactional(FName WorkName, const TFunctionRef<void(void)>& Work);
+	COREUOBJECT_API void ExecuteTransactionalWithExplicitPriority(FName WorkName, FRemoteWorkPriority WorkPriority, const TFunctionRef<void(void)>& Work);
 
 	COREUOBJECT_API void AbortTransactionRequiresDependencies(FStringView Description);
 	COREUOBJECT_API void AbortTransactionAndAbandonWork(FStringView Description);
diff --git a/Engine/Source/Runtime/Engine/Private/Actor.cpp b/Engine/Source/Runtime/Engine/Private/Actor.cpp
--- a/Engine/Source/Runtime/Engine/Private/Actor.cpp
+++ b/Engine/Source/Runtime/Engine/Private/Actor.cpp
@@ -6381,7 +6381,8 @@ void AActor::LifeSpanExpired()
 {
 #if UE_WITH_REMOTE_OBJECT_HANDLE
 	// World Timer Manager is not yet transactional, so let's make this call transactional to make up for it
-	UE::RemoteExecutor::ExecuteTransactional([this] { Destroy(); });
+	static FName TransactionalWorkName(TEXT("AActor::LifeSpanExpired"));
+	UE::RemoteExecutor::ExecuteTransactional(TransactionalWorkName, [this] { Destroy(); });
 #else
 	Destroy();
 #endif
diff --git a/Engine/Source/Runtime/Engine/Private/Collision/WorldCollisionAsync.cpp b/Engine/Source/Runtime/Engine/Private/Collision/WorldCollisionAsync.cpp
--- a/Engine/Source/Runtime/Engine/Private/Collision/WorldCollisionAsync.cpp
+++ b/Engine/Source/Runtime/Engine/Private/Collision/WorldCollisionAsync.cpp
@@ -808,7 +808,8 @@ void UWorld::FinishAsyncTrace()
 	};
 	if (PhysicsQueryHandler != nullptr)
 	{
-		UE::RemoteExecutor::ExecuteTransactional(DoWorkCallback);
+		static FName TransactionalWorkName(TEXT("UWorld::FinishAsyncTrace"));
+		UE::RemoteExecutor::ExecuteTransactional(TransactionalWorkName, DoWorkCallback);
 	}
 #endif
 
diff --git a/Engine/Source/Runtime/Engine/Private/DataReplication.cpp b/Engine/Source/Runtime/Engine/Private/DataReplication.cpp
--- a/Engine/Source/Runtime/Engine/Private/DataReplication.cpp
+++ b/Engine/Source/Runtime/Engine/Private/DataReplication.cpp
@@ -1390,7 +1390,8 @@ bool FObjectReplicator::ReceivedRPC(FNetBitReader& Reader, const FReplicationFla
 #if UE_WITH_REMOTE_OBJECT_HANDLE
 					if (Connection->Driver->bReplicateTransactionally)
 					{
-						UE::RemoteExecutor::ExecuteTransactional([Object, Function, Parms]()
+						static FName TransactionalWorkName(TEXT("FObjectReplicator ProcessEvent"));
+						UE::RemoteExecutor::ExecuteTransactional(TransactionalWorkName, [Object, Function, Parms]()
 							{
 								CallProcessEventForReceivedRPC(Object, Function, Parms);
 							});
diff --git a/Engine/Source/Runtime/Engine/Private/NetDriver.cpp b/Engine/Source/Runtime/Engine/Private/NetDriver.cpp
--- a/Engine/Source/Runtime/Engine/Private/NetDriver.cpp
+++ b/Engine/Source/Runtime/Engine/Private/NetDriver.cpp
@@ -1036,7 +1036,8 @@ void UNetDriver::TickFlush(float DeltaSeconds)
 #if UE_WITH_REMOTE_OBJECT_HANDLE
 			if (bReplicateTransactionally)
 			{
-				UE::RemoteExecutor::ExecuteTransactional([this, DeltaSeconds]()
+				static FName TransactionalWorkName(TEXT("ServerReplicateActors"));
+				UE::RemoteExecutor::ExecuteTransactional(TransactionalWorkName, [this, DeltaSeconds]()
 				{
 					ServerReplicateActors(DeltaSeconds);
 				});
diff --git a/Engine/Source/Runtime/Engine/Private/TickTaskManager.cpp b/Engine/Source/Runtime/Engine/Private/TickTaskManager.cpp
--- a/Engine/Source/Runtime/Engine/Private/TickTaskManager.cpp
+++ b/Engine/Source/Runtime/Engine/Private/TickTaskManager.cpp
@@ -320,7 +320,8 @@ public:
 
 			if (Target->bRunTransactionally)
 			{
-				UE::RemoteExecutor::ExecuteTransactional(ExecuteTickWork);
+				static FName TransactionalWorkName(TEXT("FTickFunction"));
+				UE::RemoteExecutor::ExecuteTransactional(TransactionalWorkName, ExecuteTickWork);
 			}
 			else
 			{
