PR #10996: Use unsigned constants to avoid signed overflow UB

diff --git a/Engine/Source/Runtime/Core/Public/GenericPlatform/GenericPlatformMath.h b/Engine/Source/Runtime/Core/Public/GenericPlatform/GenericPlatformMath.h
--- a/Engine/Source/Runtime/Core/Public/GenericPlatform/GenericPlatformMath.h
+++ b/Engine/Source/Runtime/Core/Public/GenericPlatform/GenericPlatformMath.h
@@ -810,7 +810,7 @@ struct FGenericPlatformMath
 	/** @return Rounds the given number up to the next highest power of two. */
 	[[nodiscard]] static constexpr FORCEINLINE uint32 RoundUpToPowerOfTwo(uint32 Arg)
 	{
-		return 1 << CeilLogTwo(Arg);
+		return 1u << CeilLogTwo(Arg);
 	}
 
 	[[nodiscard]] static constexpr FORCEINLINE uint64 RoundUpToPowerOfTwo64(uint64 V)
diff --git a/Engine/Source/Runtime/Core/Public/Microsoft/MicrosoftPlatformMath.h b/Engine/Source/Runtime/Core/Public/Microsoft/MicrosoftPlatformMath.h
--- a/Engine/Source/Runtime/Core/Public/Microsoft/MicrosoftPlatformMath.h
+++ b/Engine/Source/Runtime/Core/Public/Microsoft/MicrosoftPlatformMath.h
@@ -57,7 +57,7 @@ struct FMicrosoftPlatformMathBase : public TUnrealPlatformMathSSE4Base<FGenericP
 
 	static FORCEINLINE uint32 RoundUpToPowerOfTwo(uint32 Arg)
 	{
-		return 1 << CeilLogTwo(Arg);
+		return 1u << CeilLogTwo(Arg);
 	}
 
 	static FORCEINLINE uint64 RoundUpToPowerOfTwo64(uint64 Arg)
diff --git a/Engine/Source/Runtime/CoreUObject/Public/UObject/DynamicallyTypedValue.h b/Engine/Source/Runtime/CoreUObject/Public/UObject/DynamicallyTypedValue.h
--- a/Engine/Source/Runtime/CoreUObject/Public/UObject/DynamicallyTypedValue.h
+++ b/Engine/Source/Runtime/CoreUObject/Public/UObject/DynamicallyTypedValue.h
@@ -49,7 +49,7 @@ namespace UE
 
 		SIZE_T GetNumBytes() const { return NumBytes; }
 		uint8 GetMinAlignmentLogTwo() const { return MinAlignmentLogTwo; }
-		uint32 GetMinAlignment() const { return 1 << MinAlignmentLogTwo; }
+		uint32 GetMinAlignment() const { return 1u << MinAlignmentLogTwo; }
 		EContainsReferences GetContainsReferences() const { return ContainsReferences; }
 
 	private:
