Added BLAKE3 hashing to FPropertyTypeName, FPropertyPathNameTree, FUnknownEnumNames

diff --git a/Engine/Source/Runtime/Core/Public/Templates/TypeHash.h b/Engine/Source/Runtime/Core/Public/Templates/TypeHash.h
--- a/Engine/Source/Runtime/Core/Public/Templates/TypeHash.h
+++ b/Engine/Source/Runtime/Core/Public/Templates/TypeHash.h
@@ -214,6 +214,13 @@ template <typename T>
 template <typename T>
 [[nodiscard]] FORCEINLINE uint32 GetTypeHashHelper(const T& V) { return static_cast<uint32>(GetTypeHash(V)); }
 
+// Use this when in a scope that has AppendHash to dispatch AppendHash using unqualified lookup.
+template <typename BuilderType, typename ArgType>
+inline void DispatchAppendHash(BuilderType& Builder, const ArgType& Arg)
+{
+	AppendHash(Builder, Arg);
+}
+
 #if UE_ENABLE_INCLUDE_ORDER_DEPRECATED_IN_5_6
 #include "Templates/IsEnum.h"
 #endif
diff --git a/Engine/Source/Runtime/CoreUObject/Internal/UObject/PropertyPathNameTree.h b/Engine/Source/Runtime/CoreUObject/Internal/UObject/PropertyPathNameTree.h
--- a/Engine/Source/Runtime/CoreUObject/Internal/UObject/PropertyPathNameTree.h
+++ b/Engine/Source/Runtime/CoreUObject/Internal/UObject/PropertyPathNameTree.h
@@ -11,6 +11,8 @@
 
 #define UE_API COREUOBJECT_API
 
+class FBlake3;
+
 namespace UE
 {
 
@@ -45,6 +47,12 @@ class FPropertyPathNameTree
 		{
 			return HashCombineFast(GetTypeHash(Key.Name), GetTypeHash(Key.Type));
 		}
+
+		friend inline void AppendHash(FBlake3& Builder, const FKey& Key)
+		{
+			DispatchAppendHash(Builder, Key.Name);
+			DispatchAppendHash(Builder, Key.Type);
+		}
 	};
 
 	struct FValue
@@ -94,8 +102,8 @@ public:
 	 * @param StartIndex   Index of the segment within Path at which to start searching.
 	 * @return An accessor for the node at the path if it exists.
 	 */
-	UE_API FNode Find(const FPropertyPathName& Path, int32 StartIndex = 0);
-	inline FConstNode Find(const FPropertyPathName& Path, int32 StartIndex = 0) const
+	[[nodiscard]] UE_API FNode Find(const FPropertyPathName& Path, int32 StartIndex = 0);
+	[[nodiscard]] inline FConstNode Find(const FPropertyPathName& Path, int32 StartIndex = 0) const
 	{
 		return {const_cast<FPropertyPathNameTree*>(this)->Find(Path, StartIndex).Value};
 	}
@@ -123,9 +131,11 @@ public:
 		TMap<FKey, FValue>::TConstIterator NodeIt;
 	};
 
-	inline FConstIterator CreateConstIterator() const { return FConstIterator(Nodes.CreateConstIterator()); }
+	[[nodiscard]] inline FConstIterator CreateConstIterator() const { return FConstIterator(Nodes.CreateConstIterator()); }
 
 private:
+	UE_API friend void AppendHash(FBlake3& Builder, const FPropertyPathNameTree& Tree);
+
 	TMap<FKey, FValue> Nodes;
 };
 
diff --git a/Engine/Source/Runtime/CoreUObject/Internal/UObject/PropertyStateTracking.h b/Engine/Source/Runtime/CoreUObject/Internal/UObject/PropertyStateTracking.h
--- a/Engine/Source/Runtime/CoreUObject/Internal/UObject/PropertyStateTracking.h
+++ b/Engine/Source/Runtime/CoreUObject/Internal/UObject/PropertyStateTracking.h
@@ -11,6 +11,7 @@
 #include "Templates/SharedPointerFwd.h"
 #include "UObject/Class.h"
 
+class FBlake3;
 class FName;
 class FProperty;
 class FStructuredArchiveRecord;
@@ -183,6 +184,8 @@ public:
 	void Destroy();
 
 private:
+	friend void AppendHash(FBlake3& Builder, const FUnknownEnumNames& EnumNames);
+
 	const UObject* Owner = nullptr;
 };
 
diff --git a/Engine/Source/Runtime/CoreUObject/Private/UObject/PropertyPathNameTree.cpp b/Engine/Source/Runtime/CoreUObject/Private/UObject/PropertyPathNameTree.cpp
--- a/Engine/Source/Runtime/CoreUObject/Private/UObject/PropertyPathNameTree.cpp
+++ b/Engine/Source/Runtime/CoreUObject/Private/UObject/PropertyPathNameTree.cpp
@@ -2,6 +2,7 @@
 
 #include "UObject/PropertyPathNameTree.h"
 
+#include "Hash/Blake3.h"
 #include "Misc/StringBuilder.h"
 #include "Templates/TypeHash.h"
 
@@ -74,4 +75,28 @@ void FPropertyPathNameTree::FNode::SetTag(const FPropertyTag& Tag)
 	}
 }
 
+void AppendHash(FBlake3& Builder, const FPropertyPathNameTree& Tree)
+{
+	TArray<FPropertyPathNameTree::FKey, TInlineAllocator<16>> Keys;
+	Tree.Nodes.GetKeys(Keys);
+	Keys.Sort();
+
+	for (const FPropertyPathNameTree::FKey& Key : Keys)
+	{
+		AppendHash(Builder, Key.Name);
+		AppendHash(Builder, Key.Type);
+
+		const FPropertyPathNameTree::FValue& Value = Tree.Nodes.FindChecked(Key);
+
+		// Most of the tag is represented in the key and does not need to be hashed twice.
+		Builder.Update(&Value.Tag.PropertyGuid, sizeof(Value.Tag.PropertyGuid));
+		Builder.Update(&Value.Tag.SerializeType, sizeof(Value.Tag.SerializeType));
+
+		if (Value.SubTree)
+		{
+			AppendHash(Builder, *Value.SubTree);
+		}
+	}
+}
+
 } // UE
diff --git a/Engine/Source/Runtime/CoreUObject/Private/UObject/PropertyStateTracking.cpp b/Engine/Source/Runtime/CoreUObject/Private/UObject/PropertyStateTracking.cpp
--- a/Engine/Source/Runtime/CoreUObject/Private/UObject/PropertyStateTracking.cpp
+++ b/Engine/Source/Runtime/CoreUObject/Private/UObject/PropertyStateTracking.cpp
@@ -4,9 +4,12 @@
 
 #if WITH_EDITORONLY_DATA
 
+#include "Algo/Copy.h"
+#include "Algo/Sort.h"
 #include "Containers/Array.h"
 #include "Containers/Map.h"
 #include "Containers/Set.h"
+#include "Hash/Blake3.h"
 #include "String/ParseTokens.h"
 #include "Templates/SharedPointer.h"
 #include "UObject/Class.h"
@@ -173,6 +176,31 @@ void FUnknownEnumNames::Destroy()
 	GetUnknownEnumNamesAnnotations().RemoveAnnotation(Owner);
 }
 
+void AppendHash(FBlake3& Builder, const FUnknownEnumNames& EnumNames)
+{
+	if (TSharedPtr<TMap<FPropertyTypeName, FUnknownEnumNamesAnnotation::FInfo>> Enums = GetUnknownEnumNamesAnnotations().GetAnnotation(EnumNames.Owner).Enums)
+	{
+		TArray<FPropertyTypeName, TInlineAllocator<4>> Keys;
+		Enums->GetKeys(Keys);
+		Algo::Sort(Keys);
+
+		for (const FPropertyTypeName& Key : Keys)
+		{
+			AppendHash(Builder, Key);
+			const FUnknownEnumNamesAnnotation::FInfo& Info = Enums->FindChecked(Key);
+			TArray<FName, TInlineAllocator<4>> Names;
+			Names.Reserve(Info.Names.Num());
+			Algo::Copy(Info.Names, Names);
+			Algo::Sort(Names, FNameLexicalLess());
+			for (const FName& Name : Names)
+			{
+				AppendHash(Builder, Name);
+			}
+			Builder.Update(&Info.bHasFlags, sizeof(Info.bHasFlags));
+		}
+	}
+}
+
 } // UE
 
 #endif // WITH_EDITORONLY_DATA
diff --git a/Engine/Source/Runtime/CoreUObject/Private/UObject/PropertyTypeName.cpp b/Engine/Source/Runtime/CoreUObject/Private/UObject/PropertyTypeName.cpp
--- a/Engine/Source/Runtime/CoreUObject/Private/UObject/PropertyTypeName.cpp
+++ b/Engine/Source/Runtime/CoreUObject/Private/UObject/PropertyTypeName.cpp
@@ -8,6 +8,7 @@
 #include "Containers/Map.h"
 #include "HAL/CriticalSection.h"
 #include "HAL/LowLevelMemTracker.h"
+#include "Hash/Blake3.h"
 #include "Misc/ScopeRWLock.h"
 #include "Misc/StringBuilder.h"
 #include "String/Find.h"
@@ -95,6 +96,17 @@ struct FPropertyTypeNameNodeProxy
 		return Hash;
 	}
 
+	friend inline void AppendHash(FBlake3& Builder, const FPropertyTypeNameNodeProxy& Proxy)
+	{
+		const FPropertyTypeNameNode* Node = Proxy.First;
+		for (int32 Remaining = 1; Remaining > 0; --Remaining, ++Node)
+		{
+			AppendHash(Builder, Node->Name);
+			Builder.Update(&Node->InnerCount, sizeof(Node->InnerCount));
+			Remaining += Node->InnerCount;
+		}
+	}
+
 	friend inline bool operator==(const FPropertyTypeNameNodeProxy& Lhs, const FPropertyTypeNameNodeProxy& Rhs)
 	{
 		const FPropertyTypeNameNode* LhsNode = Lhs.First;
@@ -302,12 +314,18 @@ bool FPropertyTypeName::IsEnum(FName EnumName) const
 	return First->InnerCount > 0 && (First->Name == NAME_EnumProperty || First->Name == NAME_ByteProperty) && First[1].Name == EnumName;
 }
 
-[[nodiscard]] uint32 GetTypeHash(const FPropertyTypeName& TypeName)
+uint32 GetTypeHash(const FPropertyTypeName& TypeName)
 {
 	const FPropertyTypeNameNode* First = GPropertyTypeNameTable.ResolveByIndex(TypeName.Index);
 	return GetTypeHash(FPropertyTypeNameNodeProxy{First});
 }
 
+void AppendHash(FBlake3& Builder, const FPropertyTypeName& TypeName)
+{
+	const FPropertyTypeNameNode* First = GPropertyTypeNameTable.ResolveByIndex(TypeName.Index);
+	AppendHash(Builder, FPropertyTypeNameNodeProxy{First});
+}
+
 bool operator==(const FPropertyTypeName& Lhs, const FPropertyTypeName& Rhs)
 {
 	if (Lhs.Index == Rhs.Index)
diff --git a/Engine/Source/Runtime/CoreUObject/Public/UObject/PropertyTypeName.h b/Engine/Source/Runtime/CoreUObject/Public/UObject/PropertyTypeName.h
--- a/Engine/Source/Runtime/CoreUObject/Public/UObject/PropertyTypeName.h
+++ b/Engine/Source/Runtime/CoreUObject/Public/UObject/PropertyTypeName.h
@@ -47,9 +47,9 @@ class FPropertyTypeName
 public:
 	FPropertyTypeName() = default;
 
-	UE_API explicit FPropertyTypeName(const FProperty* Property);
+	[[nodiscard]] UE_API explicit FPropertyTypeName(const FProperty* Property);
 
-	inline bool IsEmpty() const
+	[[nodiscard]] inline bool IsEmpty() const
 	{
 		return Index == 0;
 	}
@@ -60,7 +60,7 @@ public:
 	 * Example: MapProperty(StructProperty(KeyStruct),EnumProperty(ByteEnum,ByteProperty))
 	 * - GetName() -> MapProperty
 	 */
-	UE_API FName GetName() const;
+	[[nodiscard]] UE_API FName GetName() const;
 
 	/**
 	 * Returns the number of type parameters under the root of this property type name.
@@ -68,7 +68,7 @@ public:
 	 * Example: MapProperty(StructProperty(KeyStruct),EnumProperty(ByteEnum,ByteProperty))
 	 * - GetParameterCount() -> 2
 	 */
-	UE_API int32 GetParameterCount() const;
+	[[nodiscard]] UE_API int32 GetParameterCount() const;
 
 	/**
 	 * Returns the indexed parameter under the root of this property type name.
@@ -79,7 +79,7 @@ public:
 	 * - GetParameter(0) -> StructProperty(KeyStruct)
 	 * - GetParameter(1) -> EnumProperty(ByteEnum,ByteProperty)
 	 */
-	UE_API FPropertyTypeName GetParameter(int32 ParamIndex) const;
+	[[nodiscard]] UE_API FPropertyTypeName GetParameter(int32 ParamIndex) const;
 
 	/**
 	 * Returns the indexed parameter type name under the root of this property type name.
@@ -90,7 +90,7 @@ public:
 	 * - GetParameterName(0) -> StructProperty
 	 * - GetParameterName(1) -> EnumProperty
 	 */
-	inline FName GetParameterName(int32 ParamIndex) const
+	[[nodiscard]] inline FName GetParameterName(int32 ParamIndex) const
 	{
 		return GetParameter(ParamIndex).GetName();
 	}
@@ -104,12 +104,13 @@ public:
 	}
 
 	/** Returns true if this is StructProperty with a first parameter of StructName. */
-	UE_API bool IsStruct(FName StructName) const;
+	[[nodiscard]] UE_API bool IsStruct(FName StructName) const;
 	/** Returns true if this is EnumProperty or ByteProperty with a first parameter of EnumName. */
-	UE_API bool IsEnum(FName EnumName) const;
+	[[nodiscard]] UE_API bool IsEnum(FName EnumName) const;
 
 private:
 	UE_API friend uint32 GetTypeHash(const FPropertyTypeName& TypeName);
+	UE_API friend void AppendHash(FBlake3& Builder, const FPropertyTypeName& TypeName);
 
 	UE_API friend bool operator==(const FPropertyTypeName& Lhs, const FPropertyTypeName& Rhs);
 	UE_API friend bool operator<(const FPropertyTypeName& Lhs, const FPropertyTypeName& Rhs);
@@ -168,7 +169,7 @@ public:
 	UE_API void EndParameters();
 
 	/** Build from the types that have been added. */
-	UE_API FPropertyTypeName Build() const;
+	[[nodiscard]] UE_API FPropertyTypeName Build() const;
 
 	/** Resets the builder to allow it to be reused. */
 	UE_API void Reset();
@@ -181,7 +182,7 @@ public:
 	 * @param Name A name in the format returned by operator<<(FStringBuilderBase&, const FPropertyTypeName&).
 	 * @return true if a name was parsed into the builder, false on failure.
 	 */
-	UE_API bool TryParse(FStringView Name);
+	[[nodiscard]] UE_API bool TryParse(FStringView Name);
 
 private:
 	TArray<FPropertyTypeNameNode, TInlineAllocator<10>> Nodes;
diff --git a/Engine/Source/Runtime/CoreUObject/Tests/UObject/PropertyTypeNameTest.cpp b/Engine/Source/Runtime/CoreUObject/Tests/UObject/PropertyTypeNameTest.cpp
--- a/Engine/Source/Runtime/CoreUObject/Tests/UObject/PropertyTypeNameTest.cpp
+++ b/Engine/Source/Runtime/CoreUObject/Tests/UObject/PropertyTypeNameTest.cpp
@@ -486,7 +486,7 @@ TEST_CASE_NAMED(FPropertyTypeNameConcurrencyTest, "CoreUObject::PropertyTypeName
 			Builder.EndParameters();
 		}
 
-		Builder.Build();
+		FPropertyTypeName Name = Builder.Build();
 	});
 }
 
