Fix ExperimentalAlwaysOverriden when put on a struct, the subproperties did not inherit it.
Unblocking someone, but we should add unittests for this.

#rnx

diff --git a/Engine/Source/Runtime/CoreUObject/Private/UObject/OverriddenPropertySet.cpp b/Engine/Source/Runtime/CoreUObject/Private/UObject/OverriddenPropertySet.cpp
--- a/Engine/Source/Runtime/CoreUObject/Private/UObject/OverriddenPropertySet.cpp
+++ b/Engine/Source/Runtime/CoreUObject/Private/UObject/OverriddenPropertySet.cpp
@@ -92,11 +92,6 @@ EOverriddenPropertyOperation FOverridableSerializationLogic::GetOverriddenProper
 	// It does not mean that if we have no record of an overriden operation that a subobject might have one, need to traverse all possible subobjects. 
 	if (CurrentProperty)
 	{
-		if (CurrentProperty->HasAnyPropertyFlags(CPF_ExperimentalAlwaysOverriden))
-		{
-			return EOverriddenPropertyOperation::Replace;
-		}
-
 		// In the case of a CDO owning default value, we might need to serialize it to keep its value.
 		if (OverriddenProperties && OverriddenProperties->IsCDOOwningProperty(*CurrentProperty))
 		{
@@ -1404,7 +1399,8 @@ EOverriddenPropertyOperation FOverriddenPropertySet::GetOverriddenPropertyOperat
 
 	TArray<class FProperty*, TInlineAllocator<8>>::TConstIterator PropertyIterator = PropertyChain.GetRootIterator();
 	const FOverriddenPropertyNode* OverriddenPropertyNode = &ParentPropertyNode;
-	while (PropertyIterator && OverriddenPropertyNode && OverriddenPropertyNode->Operation != EOverriddenPropertyOperation::Replace)
+	while (PropertyIterator && OverriddenPropertyNode && 
+		((OverriddenPropertyNode->Operation != EOverriddenPropertyOperation::Replace) || ((*PropertyIterator)->HasAnyPropertyFlags(CPF_ExperimentalAlwaysOverriden))))
 	{
 		const FProperty* CurrentProperty = (*PropertyIterator);
 		if (const FOverriddenPropertyNodeID* CurrentPropKey = OverriddenPropertyNode->SubPropertyNodeKeys.Find(CurrentProperty))
@@ -1420,6 +1416,11 @@ EOverriddenPropertyOperation FOverriddenPropertySet::GetOverriddenPropertyOperat
 		++PropertyIterator;
 	}
 
+	if (PropertyIterator && (*PropertyIterator)->HasAnyPropertyFlags(CPF_ExperimentalAlwaysOverriden))
+	{
+		return EOverriddenPropertyOperation::Replace;
+	}
+
 	return OverriddenPropertyNode ? OverriddenPropertyNode->Operation : EOverriddenPropertyOperation::None;
 }
 
