[Feature] Reenabled IDOs for transient objects. They had been previously been disabled as a temp workaround but now IDOs are stable enough that we can reenable them
[Refactor] This change also effects some cvars that previously only effected prefabs. Now they will impact all blueprint generated objects so they've been renamed accordingly

diff --git a/Engine/Source/Runtime/CoreUObject/Private/UObject/InstanceDataObjectUtils.cpp b/Engine/Source/Runtime/CoreUObject/Private/UObject/InstanceDataObjectUtils.cpp
--- a/Engine/Source/Runtime/CoreUObject/Private/UObject/InstanceDataObjectUtils.cpp
+++ b/Engine/Source/Runtime/CoreUObject/Private/UObject/InstanceDataObjectUtils.cpp
@@ -296,18 +296,18 @@ namespace UE
 		TEXT("When enabled, IDOs will be saved instead of instances. Disabling this will stop data retention on save.")
 	);
 
-	bool bEnableIDOsForPrefabArchetypes = true;
-	FAutoConsoleVariableRef EnableIDOsForPrefabArchetypesCVar(
-		TEXT("IDO.EnablePrefabArchetypes"),
-		bEnableIDOsForPrefabArchetypes,
-		TEXT("When enabled, prefab archetypes can have IDOs generated for them")
+	bool bEnableIDOsForBlueprintArchetypes = true;
+	FAutoConsoleVariableRef EnableIDOsForBlueprintArchetypesCVar(
+		TEXT("IDO.EnableBlueprintArchetypes"),
+		bEnableIDOsForBlueprintArchetypes,
+		TEXT("When enabled, blueprint archetypes (and prefab archetypes) can have IDOs generated for them")
 	);
 
-	bool bEnableIDOsForPrefabInstances = true;
-	FAutoConsoleVariableRef EnableIDOsForPrefabInstancesCVar(
-		TEXT("IDO.EnablePrefabInstances"),
-		bEnableIDOsForPrefabInstances,
-		TEXT("When enabled, prefab instances can have IDOs generated for them")
+	bool bEnableIDOsForBlueprintInstances = true;
+	FAutoConsoleVariableRef EnableIDOsForBlueprintInstancesCVar(
+		TEXT("IDO.EnableBlueprintInstances"),
+		bEnableIDOsForBlueprintInstances,
+		TEXT("When enabled, blueprint instances (and prefab instances) can have IDOs generated for them")
 		);
 
 	static void SetEnableAllIDOFeatures(bool bEnabled)
@@ -315,8 +315,8 @@ namespace UE
 		EnableIDOSupportCVar->Set(bEnabled);
 		EnableUninitializedUICVar->Set(bEnabled);
 		EnableIDOImpersonationOnSaveCVar->Set(bEnabled);
-		EnableIDOsForPrefabArchetypesCVar->Set(bEnabled);
-		EnableIDOsForPrefabInstancesCVar->Set(bEnabled);
+		EnableIDOsForBlueprintArchetypesCVar->Set(bEnabled);
+		EnableIDOsForBlueprintInstancesCVar->Set(bEnabled);
 	}
 
 	static bool AreAllIDOFeaturesEnabled()
@@ -325,8 +325,8 @@ namespace UE
 			bEnableIDOSupport && 
 			bEnableUninitializedUI && 
 			bEnableIDOImpersonationOnSave && 
-			bEnableIDOsForPrefabArchetypes && 
-			bEnableIDOsForPrefabInstances;
+			bEnableIDOsForBlueprintArchetypes && 
+			bEnableIDOsForBlueprintInstances;
 	}
 
 	FAutoConsoleCommand EnableAllIDOFeaturesCommand(
@@ -367,7 +367,7 @@ namespace UE
 		return bEnableIDOImpersonationOnSave;
 	}
 
-	static const UObject* GetPrefabObject(const UObject* InObject)
+	static const UObject* GetBlueprintGeneratedObject(const UObject* InObject)
 	{
 		static const FName NAME_BlueprintGeneratedClass(TEXT("BlueprintGeneratedClass"));
 		const UObject* Current = InObject;
@@ -414,51 +414,23 @@ namespace UE
 			return false;
 		}
 
-		// respect flags for disabling the generation of prefab archetypes and/or their instances
-		if (!bEnableIDOsForPrefabArchetypes || !bEnableIDOsForPrefabInstances)
+		// respect flags for disabling the generation of blueprint or prefab archetypes and/or their instances
+		if (!bEnableIDOsForBlueprintArchetypes || !bEnableIDOsForBlueprintInstances)
 		{
-			if (const UObject* PrefabObject = GetPrefabObject(InObject))
+			if (const UObject* BlueprintGeneratedObject = GetBlueprintGeneratedObject(InObject))
 			{
-				const bool bCdoMatchesPrefab = PrefabObject->GetClass()->ClassDefaultObject == PrefabObject;
-				if (!bEnableIDOsForPrefabArchetypes && bCdoMatchesPrefab)
+				const bool bIsArchetype = BlueprintGeneratedObject->GetClass()->ClassDefaultObject == BlueprintGeneratedObject;
+				if (!bEnableIDOsForBlueprintArchetypes && bIsArchetype)
 				{
 					return false;
 				}
-				if (!bEnableIDOsForPrefabInstances && !bCdoMatchesPrefab)
+				if (!bEnableIDOsForBlueprintInstances && !bIsArchetype)
 				{
 					return false;
 				}
 			}
 		}
-
-		// TODO: Temp! Remove with the conditions below.
-		while (ObjClass && ObjClass->GetClass()->GetFName() != NAME_VerseClass)
-		{
-			ObjClass = ObjClass->GetSuperClass();
-		}
-
-		if (ObjClass)
-		{
-			// TODO: Temp! Don't generate IDOs for anything within a creative device
-			for (UObject* Outer = InObject->GetOuter(); Outer; Outer = Outer->GetOuter())
-			{
-				if (Outer->GetClass()->GetFName() == NAME_VerseDevice)
-				{
-					return false;
-				}
-			}
-
-			// TODO: Temp! Don't generate IDOs for anything transient that isn't a CDO.
-			if (!InObject->HasAnyFlags(RF_ClassDefaultObject))
-			{
-				const UPackage* Package = InObject->GetPackage();
-				if (!Package || Package->HasAnyFlags(RF_Transient) || Package == GetTransientPackage())
-				{
-					return false;
-				}
-			}
-		}
-
+		
 		return true;
 	}
 
