UE: Remove FAutomationTestAttemptToFindShortTypeNamesInMetaData from the smoke filter for now since it is much more expensive after 40093672

diff --git a/Engine/Source/Runtime/CoreUObject/Private/UObject/Class.cpp b/Engine/Source/Runtime/CoreUObject/Private/UObject/Class.cpp
--- a/Engine/Source/Runtime/CoreUObject/Private/UObject/Class.cpp
+++ b/Engine/Source/Runtime/CoreUObject/Private/UObject/Class.cpp
@@ -4580,7 +4580,7 @@ FAutoConsoleCommandWithWorldAndArgs GCmdListShortTypeNamesInMetaData(
 			FStructUtils::AttemptToFindShortTypeNamesInMetaData();
 		}));
 
-IMPLEMENT_CUSTOM_SIMPLE_AUTOMATION_TEST(FAutomationTestAttemptToFindShortTypeNamesInMetaData, FAutomationTestUObjectClassBase, "UObject.Class AttemptToFindShortTypeNamesInMetaData", EAutomationTestFlags::EditorContext | EAutomationTestFlags_ApplicationContextMask | EAutomationTestFlags::ServerContext | EAutomationTestFlags::SmokeFilter)
+IMPLEMENT_CUSTOM_SIMPLE_AUTOMATION_TEST(FAutomationTestAttemptToFindShortTypeNamesInMetaData, FAutomationTestUObjectClassBase, "UObject.Class AttemptToFindShortTypeNamesInMetaData", EAutomationTestFlags::EditorContext | EAutomationTestFlags_ApplicationContextMask | EAutomationTestFlags::ServerContext)
 bool FAutomationTestAttemptToFindShortTypeNamesInMetaData::RunTest(const FString& Parameters)
 {
 	// This test is not necessary when running under UHT so just skip it in that config.
