UE: Add a filter to metadata test because it needs at least one filter and we took it out of smoke

diff --git a/Engine/Source/Runtime/CoreUObject/Private/UObject/Class.cpp b/Engine/Source/Runtime/CoreUObject/Private/UObject/Class.cpp
--- a/Engine/Source/Runtime/CoreUObject/Private/UObject/Class.cpp
+++ b/Engine/Source/Runtime/CoreUObject/Private/UObject/Class.cpp
@@ -4580,7 +4580,7 @@ FAutoConsoleCommandWithWorldAndArgs GCmdListShortTypeNamesInMetaData(
 			FStructUtils::AttemptToFindShortTypeNamesInMetaData();
 		}));
 
-IMPLEMENT_CUSTOM_SIMPLE_AUTOMATION_TEST(FAutomationTestAttemptToFindShortTypeNamesInMetaData, FAutomationTestUObjectClassBase, "UObject.Class AttemptToFindShortTypeNamesInMetaData", EAutomationTestFlags::EditorContext | EAutomationTestFlags_ApplicationContextMask | EAutomationTestFlags::ServerContext)
+IMPLEMENT_CUSTOM_SIMPLE_AUTOMATION_TEST(FAutomationTestAttemptToFindShortTypeNamesInMetaData, FAutomationTestUObjectClassBase, "UObject.Class AttemptToFindShortTypeNamesInMetaData", EAutomationTestFlags::EditorContext | EAutomationTestFlags_ApplicationContextMask | EAutomationTestFlags::ServerContext | EAutomationTestFlags::EngineFilter)
 bool FAutomationTestAttemptToFindShortTypeNamesInMetaData::RunTest(const FString& Parameters)
 {
 	// This test is not necessary when running under UHT so just skip it in that config.
