traceview: add hooks to report every accessed object, also added a Frame region

diff --git a/Engine/Source/Runtime/CoreUObject/Private/UObject/RemoteObject.cpp b/Engine/Source/Runtime/CoreUObject/Private/UObject/RemoteObject.cpp
--- a/Engine/Source/Runtime/CoreUObject/Private/UObject/RemoteObject.cpp
+++ b/Engine/Source/Runtime/CoreUObject/Private/UObject/RemoteObject.cpp
@@ -320,6 +320,7 @@ UObject* ResolveObject(const FRemoteObjectStub* Stub, ERemoteReferenceType RefTy
 		// don't want to accidentally migrate them back mid-purge and if the Object memory is still valid (but has EInternalObjectFlags::Remote flag)
 		// we can just return it and let the owner finish its cleanup.
 		// In case of weak object ptrs it's relatively safe to just return null if the object doesn't exist on this server (see CanResolveObject)
+		TouchResidentObject(Object);
 		return Object;
 	}
 
@@ -345,6 +346,7 @@ UObject* ResolveObject(const FRemoteObjectStub* Stub, ERemoteReferenceType RefTy
 
 UObject* ResolveObject(UObject* Object, ERemoteReferenceType RefType /*= ERemoteReferenceType::Strong*/)
 {
+	using namespace UE::RemoteObject::Transfer;
 	using namespace UE::RemoteObject::Private;
 
 	// Begin/FinishDestroy overrides may attempt to access subobjects of objects that have been migrated in which case we
@@ -356,9 +358,16 @@ UObject* ResolveObject(UObject* Object, ERemoteReferenceType RefType /*= ERemote
 		checkf(Stub, TEXT("Failed to find remote object stub for %s"), *GetPathNameSafe(Object));
 		return ResolveObject(Stub, RefType);
 	}
+
+	TouchResidentObject(Object);
 	return Object;
 }
 
+void TouchResidentObject(UObject* Object)
+{
+	UE::RemoteObject::Transfer::TouchResidentObject(Object);
+}
+
 bool CanResolveObject(FRemoteObjectId ObjectId)
 {
 	using namespace UE::RemoteObject::Private;
diff --git a/Engine/Source/Runtime/CoreUObject/Private/UObject/RemoteObjectTransfer.cpp b/Engine/Source/Runtime/CoreUObject/Private/UObject/RemoteObjectTransfer.cpp
--- a/Engine/Source/Runtime/CoreUObject/Private/UObject/RemoteObjectTransfer.cpp
+++ b/Engine/Source/Runtime/CoreUObject/Private/UObject/RemoteObjectTransfer.cpp
@@ -29,6 +29,7 @@ TMulticastDelegate<void(const FRemoteObjectData& /*ObjectData*/, FRemoteServerId
 TMulticastDelegate<void(const TSet<UObject*>& /*MigratedObjects*/, FRemoteServerId /*RemoteServerId*/, EObjectMigrationFlags /*MigrationFlags*/)> OnObjectsSentDelegate;
 TMulticastDelegate<void(const FRemoteObjectData& /*ObjectData*/, FRemoteServerId /*RemoteServerId*/, EObjectMigrationFlags /*MigrationFlags*/)> OnObjectDataReceivedDelegate;
 TMulticastDelegate<void(const TArray<UObject*>& /*MigratedObjects*/, FRemoteServerId /*RemoteServerId*/, EObjectMigrationFlags /*MigrationFlags*/)> OnObjectsReceivedDelegate;
+TMulticastDelegate<void(FRemoteTransactionId /*RequestId*/, FRemoteObjectId /*ObjectId*/)> OnObjectTouchedDelegate;
 
 TDelegate<void(FRemoteTransactionId /*RequestId*/, FRemoteObjectId /*ObjectId*/, FRemoteServerId /*DatabaseId*/, FRemoteServerId /*OwnerServerId*/, FRemoteObjectData& /*ObjectData*/)> StoreRemoteObjectDataDelegate;
 TDelegate<void(FRemoteTransactionId /*RequestId*/, FRemoteWorkPriority /*RequestPriority*/, FRemoteObjectId /*ObjectId*/, FRemoteServerId /*DatabaseId*/, FRemoteServerId /*DestinationServerId*/)> RestoreRemoteObjectDataDelegate;
@@ -134,6 +135,13 @@ public:
 					UE_LOG(LogRemoteObject, Verbose, TEXT("\tsent request for %s %s (%d pending requests)"), *RequiredObjectId.ToString(), *HighestPriorityRequest->Priority.ToString(), PendingObjectRequests.Num());
 				}
 			}
+			else if (!RequiredCanonicalRootObjectId.IsValid())
+			{
+				UObject* Object = StaticFindObjectFastInternal(RequiredObjectId);
+				UObject* const RootObject = FindCanonicalRootObjectForSerialization(Object);
+				const FRemoteObjectId CanonicalRootObjectId = UE::CoreUObject::Private::FObjectHandleUtils::GetRemoteId(RootObject);
+				ActiveRequest->RequiredObjectsCanonicalRoot[RequiredObjectIndex] = CanonicalRootObjectId;
+			}
 		}
 	}
 
@@ -344,6 +352,8 @@ void MigrateObjectFromRemoteServer(FRemoteObjectId ObjectId, FRemoteServerId Cur
 
 			GTransferQueue.ActiveRequest->RequiredObjects.Add(ObjectId);
 			GTransferQueue.ActiveRequest->RequiredObjectsCanonicalRoot.Add(CanonicalRootObjectId);
+
+			OnObjectTouchedDelegate.Broadcast(GTransferQueue.ActiveRequest->RequestId, ObjectId);
 		}
 		
 		if (UE_LOG_ACTIVE(LogRemoteObject, VeryVerbose))
@@ -361,6 +371,38 @@ void MigrateObjectFromRemoteServer(FRemoteObjectId ObjectId, FRemoteServerId Cur
 	UE::RemoteExecutor::AbortTransactionRequiresDependencies(TEXT("object migration"));
 }
 
+void TouchResidentObject(UObject* Object)
+{
+	using namespace UE::RemoteObject::Serialization;
+
+	if (!Object)
+	{
+		return;
+	}
+
+	if (AutoRTFM::IsClosed() && GTransferQueue.ActiveRequest)
+	{
+		UE_AUTORTFM_OPEN
+		{
+			const FRemoteObjectId ObjectId = UE::CoreUObject::Private::FObjectHandleUtils::GetRemoteId(Object);
+
+			// add it to the active request's list of required objects
+			if (GTransferQueue.ActiveRequest->RequiredObjects.Find(ObjectId) == INDEX_NONE)
+			{
+				GTransferQueue.ActiveRequest->RequiredObjects.Add(ObjectId);
+				
+				// NOTE - even though this object is resident, we can't fetch the canonical root
+				// object because it could trigger endless recursion into this function which is
+				// called when object pointers get resolved
+				FRemoteObjectId CanonicalRootObjectId;
+				GTransferQueue.ActiveRequest->RequiredObjectsCanonicalRoot.Add(CanonicalRootObjectId);
+
+				OnObjectTouchedDelegate.Broadcast(GTransferQueue.ActiveRequest->RequestId, ObjectId);
+			}
+		};
+	}
+}
+
 void MigrateObjectToRemoteServer(FRemoteTransactionId RequestId, UObject* Object, FRemoteServerId DestinationServerId)
 {
 	FRemoteWorkPriority RootWorkPriority = UE::RemoteExecutor::CreateRootWorkPriority();
@@ -618,6 +660,8 @@ void PrefetchRemoteObject(FRemoteObjectId ObjectId, UObject* DestinationOuter)
 			GTransferQueue.ActiveRequest->RequiredObjects.Add(ObjectId);
 			GTransferQueue.ActiveRequest->RequiredObjectsCanonicalRoot.Add(ObjectId);
 			UE_LOG(LogRemoteObject, VeryVerbose, TEXT("PrefetchRemoteObject(%s) : obj %s"), *GTransferQueue.ActiveRequest->GetRequestId().ToString(), *ObjectId.ToString());
+
+			OnObjectTouchedDelegate.Broadcast(GTransferQueue.ActiveRequest->RequestId, ObjectId);
 		}
 	};
 }
diff --git a/Engine/Source/Runtime/CoreUObject/Public/UObject/ObjectHandle.h b/Engine/Source/Runtime/CoreUObject/Public/UObject/ObjectHandle.h
--- a/Engine/Source/Runtime/CoreUObject/Public/UObject/ObjectHandle.h
+++ b/Engine/Source/Runtime/CoreUObject/Public/UObject/ObjectHandle.h
@@ -453,7 +453,9 @@ namespace UE::CoreUObject::Private
 		}
 		else
 		{
-			return ReadObjectHandlePointerNoCheck(LocalHandle);
+			ResolvedObject = ReadObjectHandlePointerNoCheck(LocalHandle);
+			UE::RemoteObject::Handle::TouchResidentObject(ResolvedObject);
+			return ResolvedObject;
 		}
 		Handle = MakeObjectHandle(ResolvedObject);
 		return ResolvedObject;
diff --git a/Engine/Source/Runtime/CoreUObject/Public/UObject/RemoteObject.h b/Engine/Source/Runtime/CoreUObject/Public/UObject/RemoteObject.h
--- a/Engine/Source/Runtime/CoreUObject/Public/UObject/RemoteObject.h
+++ b/Engine/Source/Runtime/CoreUObject/Public/UObject/RemoteObject.h
@@ -70,6 +70,9 @@ namespace UE::RemoteObject::Handle
 	* @return Resolved object
 	*/
 	COREUOBJECT_API UObject* ResolveObject(UObject* Object, ERemoteReferenceType RefType = ERemoteReferenceType::Strong);
+	
+	COREUOBJECT_API void TouchResidentObject(UObject* Object);
+
 
 	/**
 	* Resolves a remote object
diff --git a/Engine/Source/Runtime/CoreUObject/Public/UObject/RemoteObjectTransfer.h b/Engine/Source/Runtime/CoreUObject/Public/UObject/RemoteObjectTransfer.h
--- a/Engine/Source/Runtime/CoreUObject/Public/UObject/RemoteObjectTransfer.h
+++ b/Engine/Source/Runtime/CoreUObject/Public/UObject/RemoteObjectTransfer.h
@@ -92,6 +92,11 @@ namespace UE::RemoteObject::Transfer
 	*/
 	COREUOBJECT_API void MigrateObjectFromRemoteServer(FRemoteObjectId ObjectId, FRemoteServerId CurrentOwnerServerId, UObject* DestinationOuter = nullptr);
 
+	/**
+	* Reports code that touches a resident object
+	*/
+	COREUOBJECT_API void TouchResidentObject(UObject* Object);
+
 	/**
 	* Utility that converts FRemoteObjectId to TObjectPtr of a given type. Will not resolve the object if it's not local.
 	*/
@@ -150,6 +155,9 @@ namespace UE::RemoteObject::Transfer
 	 */
 	extern COREUOBJECT_API TMulticastDelegate<void(const TSet<UObject*>& /*MigratedObjects*/, FRemoteServerId /*RemoteServerId*/, EObjectMigrationFlags /*MigrationFlags*/)> OnObjectsSentDelegate;
 
+	/** Delegate executed when an object has been accessed by a transaction */
+	extern COREUOBJECT_API TMulticastDelegate<void(FRemoteTransactionId /*RequestId*/, FRemoteObjectId /*ObjectId*/)> OnObjectTouchedDelegate;
+
 	/** Delegate that stores locally unrachable object data into a database. */
 	extern COREUOBJECT_API TDelegate<void(FRemoteTransactionId /*RequestId*/, FRemoteObjectId /*ObjectId*/, FRemoteServerId /*DatabaseId*/, FRemoteServerId /*OwnerServerId*/, FRemoteObjectData& /*ObjectData*/)> StoreRemoteObjectDataDelegate;
 
