Allow access to protected and private class members in intrinsic class init code
Currently it's impossble to make any of the GC-exposed members private.

diff --git a/Engine/Source/Runtime/CoreUObject/Public/UObject/ObjectMacros.h b/Engine/Source/Runtime/CoreUObject/Public/UObject/ObjectMacros.h
--- a/Engine/Source/Runtime/CoreUObject/Public/UObject/ObjectMacros.h
+++ b/Engine/Source/Runtime/CoreUObject/Public/UObject/ObjectMacros.h
@@ -1879,6 +1879,7 @@ public: \
 	static void StaticRegisterNatives##TClass() {} \
 	DECLARE_SERIALIZER(TClass) \
 	DEFINE_DEFAULT_OBJECT_INITIALIZER_CONSTRUCTOR_CALL(TClass) \
+	static void IntrinsicClassInit(UClass* Class); \
 	static UObject* __VTableCtorCaller(FVTableHelper& Helper) \
 	{ \
 		return new (EC_InternalUseOnlyConstructor, (UObject*)GetTransientPackage(), NAME_None, RF_NeedLoad | RF_ClassDefaultObject | RF_TagGarbageTemp) TClass(Helper); \
@@ -1892,6 +1893,7 @@ public: \
 	static void StaticRegisterNatives##TClass() {} \
 	DECLARE_SERIALIZER(TClass) \
 	DEFINE_DEFAULT_OBJECT_INITIALIZER_CONSTRUCTOR_CALL(TClass) \
+	static void IntrinsicClassInit(UClass* Class); \
 	static UObject* __VTableCtorCaller(FVTableHelper& Helper) \
 	{ \
 		return new (EC_InternalUseOnlyConstructor, (UObject*)GetTransientPackage(), NAME_None, RF_NeedLoad | RF_ClassDefaultObject | RF_TagGarbageTemp) TClass(Helper); \
@@ -1902,6 +1904,7 @@ public: \
 	static void StaticRegisterNatives##TClass() {} \
 	DECLARE_SERIALIZER(TClass) \
 	DEFINE_DEFAULT_OBJECT_INITIALIZER_CONSTRUCTOR_CALL(TClass) \
+	static void IntrinsicClassInit(UClass* Class); \
 	static UObject* __VTableCtorCaller(FVTableHelper& Helper) \
 	{ \
 		return new (EC_InternalUseOnlyConstructor, (UObject*)GetTransientPackage(), NAME_None, RF_NeedLoad | RF_ClassDefaultObject | RF_TagGarbageTemp) TClass(Helper); \
@@ -1915,6 +1918,7 @@ public: \
 	static void StaticRegisterNatives##TClass() {} \
 	DECLARE_SERIALIZER(TClass) \
 	DEFINE_DEFAULT_OBJECT_INITIALIZER_CONSTRUCTOR_CALL(TClass) \
+	static void IntrinsicClassInit(UClass* Class); \
 	static UObject* __VTableCtorCaller(FVTableHelper& Helper) \
 	{ \
 		return new (EC_InternalUseOnlyConstructor, (UObject*)GetTransientPackage(), NAME_None, RF_NeedLoad | RF_ClassDefaultObject | RF_TagGarbageTemp) TClass(Helper); \
@@ -1925,6 +1929,7 @@ public: \
 	static void StaticRegisterNatives##TClass() {} \
 	DECLARE_SERIALIZER(TClass) \
 	DEFINE_DEFAULT_OBJECT_INITIALIZER_CONSTRUCTOR_CALL(TClass) \
+	static void IntrinsicClassInit(UClass* Class); \
 	static UObject* __VTableCtorCaller(FVTableHelper& Helper) \
 	{ \
 		return new (EC_InternalUseOnlyConstructor, (UObject*)GetTransientPackage(), NAME_None, RF_NeedLoad | RF_ClassDefaultObject | RF_TagGarbageTemp) TClass(Helper); \
@@ -2133,11 +2138,15 @@ private:
 			UClass* Class = TClass::StaticClass(); \
 			UObjectForceRegistration(Class); \
 			check(Class->GetSuperClass() == SuperClass); \
-			InitCode \
+			TClass::IntrinsicClassInit(Class); \
 			Class->StaticLink(); \
 			return Class; \
 		} \
 	}; \
+	void TClass::IntrinsicClassInit(UClass* Class) \
+	{ \
+		InitCode \
+	} \
 	UClass* Z_Construct_UClass_##TClass() \
 	{ \
 		if (!Z_Registration_Info_UClass_##TClass.OuterSingleton) \
