Optimize FSoftObjectPath construction from a UObject* by avoiding roundtrips to string and parsing by manually walking the outer chain instead.

diff --git a/Engine/Source/Runtime/CoreUObject/Private/UObject/SoftObjectPath.cpp b/Engine/Source/Runtime/CoreUObject/Private/UObject/SoftObjectPath.cpp
--- a/Engine/Source/Runtime/CoreUObject/Private/UObject/SoftObjectPath.cpp
+++ b/Engine/Source/Runtime/CoreUObject/Private/UObject/SoftObjectPath.cpp
@@ -4,6 +4,7 @@
 #include "UObject/PropertyPortFlags.h"
 #include "UObject/UnrealType.h"
 #include "UObject/ObjectRedirector.h"
+#include "UObject/Package.h"
 #include "Misc/AsciiSet.h"
 #include "Misc/PackageName.h"
 #include "Misc/StringBuilder.h"
@@ -257,6 +258,61 @@ void FSoftObjectPath::SetPath(FUtf8StringView Path)
 	SetPath(Wide);
 }
 
+void FSoftObjectPath::SetPath(const UObject* InObject)
+{
+	if (!InObject)
+	{
+		Reset();
+		return;
+	}
+
+	const UObject* ObjectOuter = InObject->GetOuter();
+
+	// Fast path: InObject is a package
+	if (!ObjectOuter)
+	{
+		AssetPath.TrySetPath(CastChecked<UPackage>(InObject)->GetFName(), NAME_None);
+		SubPathString.Empty();
+		return;
+	}
+	
+	// Fast path: InObject is a top-level asset
+	if (!ObjectOuter->GetOuter())
+	{
+		AssetPath.TrySetPath(CastChecked<UPackage>(ObjectOuter)->GetFName(), InObject->GetFName());
+		SubPathString.Empty();
+		return;
+	}
+
+	// Slow path: walk the outer chain and construct the soft object path parts
+	TArray<FName, TInlineAllocator<64>> SubObjectPath;
+
+	const UObject* CurrentObject = InObject;
+	const UObject* CurrentObjectOuter = ObjectOuter;
+
+	while (CurrentObjectOuter->GetOuter())
+	{
+		SubObjectPath.Push(CurrentObject->GetFName());
+		CurrentObject = CurrentObjectOuter;
+		CurrentObjectOuter = CurrentObject->GetOuter();
+	}
+
+	AssetPath.TrySetPath(CastChecked<UPackage>(CurrentObjectOuter)->GetFName(), CurrentObject->GetFName());
+
+	TStringBuilder<1024> SubPathStringBuilder;
+
+	while (!SubObjectPath.IsEmpty())
+	{
+		const FName ObjectName = SubObjectPath.Pop(EAllowShrinking::No);
+		ObjectName.AppendString(SubPathStringBuilder);
+		SubPathStringBuilder.AppendChar(TEXT('.'));
+	}
+
+	SubPathStringBuilder.RemoveSuffix(1);
+
+	SubPathString = SubPathStringBuilder.ToString();
+}
+
 #if WITH_EDITOR
 	extern bool* GReportSoftObjectPathRedirects;
 #endif
diff --git a/Engine/Source/Runtime/CoreUObject/Public/UObject/SoftObjectPath.h b/Engine/Source/Runtime/CoreUObject/Public/UObject/SoftObjectPath.h
--- a/Engine/Source/Runtime/CoreUObject/Public/UObject/SoftObjectPath.h
+++ b/Engine/Source/Runtime/CoreUObject/Public/UObject/SoftObjectPath.h
@@ -94,21 +94,36 @@ struct FSoftObjectPath
 	template <typename T>
 	[[nodiscard]] FSoftObjectPath(const TObjectPtr<T>& InObject)
 	{
-		SetPath(InObject.GetPathName());
+		// Avoid conversion to string if the object is resolved
+		if (InObject.IsResolved())
+		{
+			// reinterpret_cast here is used because T may be an incomplete type, but we only support using TObjectPtr for UObject-derived types 
+			SetPath(reinterpret_cast<const UObject*>(InObject.Get()));
+		}
+		else
+		{
+			SetPath(InObject.GetPathName());
+		}
 	}
 
 	[[nodiscard]] FSoftObjectPath(const FObjectPtr& InObject)
 	{
-		SetPath(InObject.GetPathName());
+		// Avoid conversion to string if the object is resolved
+		if (InObject.IsResolved())
+		{
+			SetPath(InObject.Get());
+		}
+		else
+		{
+			SetPath(InObject.GetPathName());
+		}
 	}
 
 	/** Construct from an existing object in memory */
 	[[nodiscard]] FSoftObjectPath(const UObject* InObject)
 	{
-		if (InObject)
-		{
-			SetPath(InObject->GetPathName());
-		}
+		// Avoid conversion to string 
+		SetPath(InObject);
 	}
 
 	/** Static methods for more meaningful construction sites. */
@@ -245,6 +260,7 @@ struct FSoftObjectPath
 	COREUOBJECT_API void SetPath(FWideStringView Path);
 	COREUOBJECT_API void SetPath(FAnsiStringView Path);
 	COREUOBJECT_API void SetPath(FUtf8StringView Path);
+	COREUOBJECT_API void SetPath(const UObject* InObject);
 	void SetPath(const WIDECHAR* Path)
 	{
 		SetPath(FWideStringView(Path));
