Fix a bug whereby the transactionally safe async loader could hold an invalid `UPackage` reference.

diff --git a/Engine/Source/Runtime/CoreUObject/Private/Serialization/TransactionallySafeAsyncLoading.cpp b/Engine/Source/Runtime/CoreUObject/Private/Serialization/TransactionallySafeAsyncLoading.cpp
--- a/Engine/Source/Runtime/CoreUObject/Private/Serialization/TransactionallySafeAsyncLoading.cpp
+++ b/Engine/Source/Runtime/CoreUObject/Private/Serialization/TransactionallySafeAsyncLoading.cpp
@@ -5,6 +5,7 @@
 #include "AutoRTFM/AutoRTFM.h"
 #include "Containers/Map.h"
 #include "UObject/LinkerInstancingContext.h"
+#include "UObject/Package.h"
 #include "Async/Mutex.h"
 #include "Misc/ScopeLock.h"
 
@@ -94,11 +95,21 @@ public:
 				{
 					UE::TScopeLock _(Mutex);
 
+					// Check if we've cached the package previously.
 					if (FMapPayload* const Result = PackageCache.Find(PackagePathAsString))
 					{
-						// Need to copy it because as soon as we give up our mutex the package cache
-						// could have been resized thus invalidating our underlying memory allocation.
-						MapPayload = *Result;
+						// Check if the result is still valid. It could become invalid if GC had collected
+						// the underlying `UPackage` for instance.
+						if (Result->IsStillValid())
+						{
+							// Need to copy it because as soon as we give up our mutex the package cache
+							// could have been resized thus invalidating our underlying memory allocation.
+							MapPayload = *Result;
+						}
+						else
+						{
+							UE_LOG(LogStreaming, Display, TEXT("A previously loaded cached package `%s` was garbage collected, and we are having to reload it."), *PackagePathAsString);
+						}
 					}
 				}
 
@@ -154,7 +165,8 @@ public:
 					// close.
 					(void)AutoRTFM::Close([&]
 					{
-						InCompletionDelegate.ExecuteIfBound(MapPayload->PackageName, MapPayload->LoadedPackage, MapPayload->Result);
+						ensure(MapPayload->IsStillValid());
+						InCompletionDelegate.ExecuteIfBound(MapPayload->PackageName, MapPayload->LoadedPackage.Get(), MapPayload->Result);
 					});
 					return PackageCacheHit;
 				}
@@ -428,12 +440,28 @@ private:
 
 	struct FMapPayload final
 	{
-		FMapPayload(FName PackageName, UPackage* LoadedPackage, EAsyncLoadingResult::Type Result) :
-			PackageName(PackageName), LoadedPackage(LoadedPackage), Result(Result) {}
+		FMapPayload() = default;
+		FMapPayload(FName PackageName, UPackage* const LoadedPackage, EAsyncLoadingResult::Type Result) :
+			PackageName(PackageName), LoadedPackage(MakeWeakObjectPtr(LoadedPackage)), Result(Result) {}
 
 		FName PackageName;
-		UPackage* LoadedPackage;
+		TWeakObjectPtr<UPackage> LoadedPackage;
 		EAsyncLoadingResult::Type Result;
+
+		// Tells us whether a given map payload entry is still valid (EG. a cached loaded package didn't get GC'ed).
+		bool IsStillValid() const
+		{
+			// If we succeeded in the load and at one point had a valid `LoadedPackage`, then check whether it
+			// became null (meaning the GC did its thing).
+			if (EAsyncLoadingResult::Succeeded == Result)
+			{
+				return LoadedPackage.IsValid();
+			}
+			else
+			{
+				return true;
+			}
+		}
 	};
 
 	TMap<FString, FMapPayload> PackageCache;
@@ -477,7 +505,7 @@ private:
 				FMapPayload Payload(Name, Package, Type);
 
 				UE::TScopeLock _(Mutex);
-				PackageCache.Add(PackagePathStr, Payload);
+				PackageCache.FindOrAdd(PackagePathStr) = Payload;
 			}
 
 			CompletionDelegate.ExecuteIfBound(Name, Package, Type);
