Remote Objects:  Some fixes that were missing in the borrowing concept.  We should wait until after PostMigrate to change ownership so we can give PostMigrate some context about where the object originated from before processing. If changed before PostMigrate, it would give the impression we own objects we've never seen before.

#rnx

diff --git a/Engine/Source/Runtime/CoreUObject/Private/UObject/RemoteObjectSerialization.cpp b/Engine/Source/Runtime/CoreUObject/Private/UObject/RemoteObjectSerialization.cpp
--- a/Engine/Source/Runtime/CoreUObject/Private/UObject/RemoteObjectSerialization.cpp
+++ b/Engine/Source/Runtime/CoreUObject/Private/UObject/RemoteObjectSerialization.cpp
@@ -52,7 +52,7 @@ void SerializeNameTables(FArchive& Ar, FRemoteObjectData& InObjectData)
 {
 	FNameIndexType NumNames = IntCastChecked<FNameIndexType>(InObjectData.Names.Num());
 	Ar << NumNames;
-	for (FName& Name : InObjectData.Names)
+	for (const FName& Name : InObjectData.Names)
 	{
 		FString NameString = Name.ToString();
 		Ar << NameString;
@@ -978,7 +978,7 @@ UObject* ConstructRemoteObject(const FRemoteObjectHeader& Header, FRemoteObjectC
 	return Object;
 }
 
-void ResolvePathNames(FRemoteObjectData& InObjectData, TArray<UObject*>& OutResolvedObjects)
+void ResolvePathNames(const FRemoteObjectData& InObjectData, TArray<UObject*>& OutResolvedObjects)
 {
 	for (const FRemoteObjectPathName& PathName : InObjectData.PathNames)
 	{
@@ -1114,8 +1114,6 @@ void ResetRemoteObject(UObject* InObject)
 	}
 
 	{
-		const FRemoteServerId LocalServerId = UE::RemoteObject::GetGlobalServerId();
-
 		// Deserialize archetype data on top of the object (and its subobjects) to restore their state to the archetype values
 		TArray<UObject*> DeserializedObjects;
 		TArray<FRemoteObjectId> DeserializedIds;
@@ -1124,12 +1122,18 @@ void ResetRemoteObject(UObject* InObject)
 		// Use any existing objects and deserialize archetype delta on top of them avoiding re-construction which may lead to undiserable side-effects
 		// Additionally we don't want to recursively re-enter this function so let the deserialization process know we're already resetting migrated object(s)
 		ERemoteObjectSerializationFlags DeserializationFlags = ERemoteObjectSerializationFlags::PreserveRemoteReferences | ERemoteObjectSerializationFlags::UseExistingObjects | ERemoteObjectSerializationFlags::Resetting;
-		DeserializeObjectData(LocalServerId, *ArchetypeDelta, DeserializedIds, DeserializedObjects, DeserializationFlags);
+		DeserializeObjectData(*ArchetypeDelta, DeserializedIds, DeserializedObjects, DeserializationFlags);
+
+		const FRemoteServerId LocalServerId = UE::RemoteObject::GetGlobalServerId();
+		for (UObject* DeserializedObject : DeserializedObjects)
+		{
+			UE::RemoteObject::Handle::ChangeOwnerServerId(DeserializedObject, LocalServerId);
+		}
 	}
 	delete ArchetypeDelta;
 }
 
-int32 DeserializeObjectData(FRemoteServerId OwnerServerId, FRemoteObjectData& ObjectData, TArray<FRemoteObjectId>& OutObjectRemoteIds, TArray<UObject*>& OutObjects, ERemoteObjectSerializationFlags DeserializeFlags)
+int32 DeserializeObjectData(FRemoteObjectData& ObjectData, TArray<FRemoteObjectId>& OutObjectRemoteIds, TArray<UObject*>& OutObjects, ERemoteObjectSerializationFlags DeserializeFlags)
 {
 	using namespace UE::RemoteObject::Handle;
 	using namespace UE::RemoteObject::Private;
@@ -1202,13 +1206,12 @@ int32 DeserializeObjectData(FRemoteServerId OwnerServerId, FRemoteObjectData& Ob
 		const bool bResetting = EnumHasAnyFlags(DeserializeFlags, ERemoteObjectSerializationFlags::Resetting);
 		if (!bResetting)
 		{
-			// ensure remote object stubs are created and ownership is assigned
+			// ensure remote object stubs are created (ownership will be assigned after PostMigrate)
 			for (int32 ObjectIndex = 0; ObjectIndex < OutObjects.Num(); ++ObjectIndex)
 			{
 				if (UObject* Object = OutObjects[ObjectIndex])
 				{
 					UE::RemoteObject::Private::RegisterRemoteObjectId(FRemoteObjectId(Object), LocalServerId);
-					ChangeOwnerServerId(Object, OwnerServerId);
 				}
 			}
 
diff --git a/Engine/Source/Runtime/CoreUObject/Private/UObject/RemoteObjectSerialization.h b/Engine/Source/Runtime/CoreUObject/Private/UObject/RemoteObjectSerialization.h
--- a/Engine/Source/Runtime/CoreUObject/Private/UObject/RemoteObjectSerialization.h
+++ b/Engine/Source/Runtime/CoreUObject/Private/UObject/RemoteObjectSerialization.h
@@ -78,7 +78,6 @@ namespace UE::RemoteObject::Serialization
 
 	/**
 	* Deserializes remote object data
-	* @param OwnerServerId New Owner Server Id of the objects that are about to be created and deserialized (it may be identicaly to their previous owner if the objects were borrowed)
 	* @param ObjectData to be deserialized
 	* @param OutObjectRemoteIds Remote IDs of the deserialized objects
 	* @param OutReceivedObjects All deserialized objects
@@ -86,7 +85,7 @@ namespace UE::RemoteObject::Serialization
 	* @return Index of an object in OutReceivedObjects that was the main object the migration request was triggered for 
 	* (usually 0 but if a migration requests a default subobject then its parent is also migrated and the return value will be > 0)
 	*/
-	int32 DeserializeObjectData(FRemoteServerId OwnerServerId, FRemoteObjectData& ObjectData, TArray<FRemoteObjectId>& OutObjectRemoteIds, TArray<UObject*>& OutReceivedObjects, ERemoteObjectSerializationFlags DeserializeFlags = ERemoteObjectSerializationFlags::None);
+	int32 DeserializeObjectData(FRemoteObjectData& ObjectData, TArray<FRemoteObjectId>& OutObjectRemoteIds, TArray<UObject*>& OutReceivedObjects, ERemoteObjectSerializationFlags DeserializeFlags = ERemoteObjectSerializationFlags::None);
 
 	/**
 	* Finds the canonical 'root' object that is used for remote object serialization
diff --git a/Engine/Source/Runtime/CoreUObject/Private/UObject/RemoteObjectTransfer.cpp b/Engine/Source/Runtime/CoreUObject/Private/UObject/RemoteObjectTransfer.cpp
--- a/Engine/Source/Runtime/CoreUObject/Private/UObject/RemoteObjectTransfer.cpp
+++ b/Engine/Source/Runtime/CoreUObject/Private/UObject/RemoteObjectTransfer.cpp
@@ -137,10 +137,12 @@ public:
 			}
 			else if (!RequiredCanonicalRootObjectId.IsValid())
 			{
-				UObject* Object = StaticFindObjectFastInternal(RequiredObjectId);
-				UObject* const RootObject = FindCanonicalRootObjectForSerialization(Object);
-				const FRemoteObjectId CanonicalRootObjectId = UE::CoreUObject::Private::FObjectHandleUtils::GetRemoteId(RootObject);
-				ActiveRequest->RequiredObjectsCanonicalRoot[RequiredObjectIndex] = CanonicalRootObjectId;
+				if (UObject* Object = StaticFindObjectFastInternal(RequiredObjectId))
+				{
+					UObject* const RootObject = FindCanonicalRootObjectForSerialization(Object);
+					const FRemoteObjectId CanonicalRootObjectId = UE::CoreUObject::Private::FObjectHandleUtils::GetRemoteId(RootObject);
+					ActiveRequest->RequiredObjectsCanonicalRoot[RequiredObjectIndex] = CanonicalRootObjectId;
+				}
 			}
 		}
 	}
@@ -258,7 +260,7 @@ UObject* MigrateObjectFromObjectData(FRemoteServerId OwnerServerId, FRemoteObjec
 
 	TArray<FRemoteObjectId> ReceivedObjectRemoteIds;
 	TArray<UObject*> ReceivedObjects;
-	int32 RequestedObjectIndex = DeserializeObjectData(OwnerServerId, ObjectData, ReceivedObjectRemoteIds, ReceivedObjects);
+	int32 RequestedObjectIndex = DeserializeObjectData(ObjectData, ReceivedObjectRemoteIds, ReceivedObjects);
 
 	checkf(ReceivedObjects.Num() > 0, TEXT("PeerId:%s unable to deserialize object data (%d bytes)"), *GetGlobalServerId().ToString(), ObjectData.Bytes.Num());
 
@@ -282,16 +284,32 @@ UObject* MigrateObjectFromObjectData(FRemoteServerId OwnerServerId, FRemoteObjec
 	}
 
 #if UE_WITH_REMOTE_OBJECT_HANDLE
+	const bool bAssignedOwnership = EnumHasAnyFlags(MigrationFlags, EObjectMigrationFlags::ChangeOwnership);
+
 	FPostMigrateParams PostMigrateParams;
 	PostMigrateParams.RequestingServerId = RemoteServerId;
 	PostMigrateParams.DestinationServerId = GetGlobalServerId();
-	PostMigrateParams.MigrationFlags = MigrationFlags;
 	for (UObject* ReceivedObject : ReceivedObjects)
 	{
+		PostMigrateParams.MigrationFlags = MigrationFlags;
+
+		// Recompute the flags as having only changed ownership if we're being assigned ownership we don't already have
+		const bool bAlreadyOwned = UE::RemoteObject::Handle::IsOwned(ReceivedObject);
+		if (bAssignedOwnership && bAlreadyOwned)
+		{
+			EnumRemoveFlags(PostMigrateParams.MigrationFlags, EObjectMigrationFlags::ChangeOwnership);
+		}
+
 		ReceivedObject->PostMigrate(PostMigrateParams);
 	}
 #endif
 
+	// Set the ownership to the correct server id
+	for (UObject* ReceivedObject : ReceivedObjects)
+	{
+		UE::RemoteObject::Handle::ChangeOwnerServerId(ReceivedObject, OwnerServerId);
+	}
+
 	// Notify that objects have been migrated from a remote server
 	OnObjectsReceivedDelegate.Broadcast(ReceivedObjects, RemoteServerId, MigrationFlags);
 
