Replaced UObject references with TNotNull for OverridableManager system.

diff --git a/Engine/Plugins/Experimental/InstanceDataObjectFixupTool/Source/Private/InstanceDataObjectFixupPanel.cpp b/Engine/Plugins/Experimental/InstanceDataObjectFixupTool/Source/Private/InstanceDataObjectFixupPanel.cpp
--- a/Engine/Plugins/Experimental/InstanceDataObjectFixupTool/Source/Private/InstanceDataObjectFixupPanel.cpp
+++ b/Engine/Plugins/Experimental/InstanceDataObjectFixupTool/Source/Private/InstanceDataObjectFixupPanel.cpp
@@ -839,7 +839,7 @@ void FInstanceDataObjectFixupPanel::RedirectProperty(const FPropertyPath& From,
 		if (SourceProperty->SameType(DestinationProperty))
 		{
 			SourceProperty->CopyCompleteValue(Destination, Source);
-			FOverridableManager::Get().GetOverriddenProperties(*Instance)->SetOverriddenPropertyOperation(EOverriddenPropertyOperation::Modified, nullptr, DestinationProperty);
+			FOverridableManager::Get().GetOverriddenProperties(Instance)->SetOverriddenPropertyOperation(EOverriddenPropertyOperation::Modified, nullptr, DestinationProperty);
 		}
 		else
 		{
diff --git a/Engine/Source/Editor/DetailCustomizations/Private/MathStructProxyCustomizations.cpp b/Engine/Source/Editor/DetailCustomizations/Private/MathStructProxyCustomizations.cpp
--- a/Engine/Source/Editor/DetailCustomizations/Private/MathStructProxyCustomizations.cpp
+++ b/Engine/Source/Editor/DetailCustomizations/Private/MathStructProxyCustomizations.cpp
@@ -546,7 +546,7 @@ bool FMatrixStructCustomization<T>::FlushValues( TWeakPtr<IPropertyHandle> Prope
 				OuterObjects[ValueIndex]->GetArchetypeInstances(ArchetypeInstances);
 				for (UObject* ArchetypeInstance : ArchetypeInstances)
 				{
-					if (!FOverridableManager::Get().IsEnabled(*ArchetypeInstance))
+					if (!FOverridableManager::Get().IsEnabled(ArchetypeInstance))
 					{
 						UE::Math::TMatrix<T>* CurrentValue = reinterpret_cast<UE::Math::TMatrix<T>*>(PropertyHandle->GetValueBaseAddress(reinterpret_cast<uint8*>(ArchetypeInstance)));
 						if (CurrentValue && CurrentValue->Equals(PreviousValue))
@@ -707,7 +707,7 @@ bool FTransformStructCustomization<T>::FlushValues( TWeakPtr<IPropertyHandle> Pr
 				OuterObjects[ValueIndex]->GetArchetypeInstances(ArchetypeInstances);
 				for (UObject* ArchetypeInstance : ArchetypeInstances)
 				{
-					if (!FOverridableManager::Get().IsEnabled(*ArchetypeInstance))
+					if (!FOverridableManager::Get().IsEnabled(ArchetypeInstance))
 					{
 						UE::Math::TTransform<T>* CurrentValue = reinterpret_cast<UE::Math::TTransform<T>*>(PropertyHandle->GetValueBaseAddress(reinterpret_cast<uint8*>(ArchetypeInstance)));
 						if (CurrentValue && CurrentValue->Equals(PreviousValue))
@@ -853,7 +853,7 @@ bool FQuatStructCustomization<T>::FlushValues(TWeakPtr<IPropertyHandle> Property
 				OuterObjects[ValueIndex]->GetArchetypeInstances(ArchetypeInstances);
 				for (UObject* ArchetypeInstance : ArchetypeInstances)
 				{
-					if (!FOverridableManager::Get().IsEnabled(*ArchetypeInstance))
+					if (!FOverridableManager::Get().IsEnabled(ArchetypeInstance))
 					{
 						UE::Math::TQuat<T>* CurrentValue = reinterpret_cast<UE::Math::TQuat<T>*>(PropertyHandle->GetValueBaseAddress(reinterpret_cast<uint8*>(ArchetypeInstance)));
 						if (CurrentValue && CurrentValue->Equals(PreviousValue))
diff --git a/Engine/Source/Editor/Kismet/Private/BlueprintCompilationManager.cpp b/Engine/Source/Editor/Kismet/Private/BlueprintCompilationManager.cpp
--- a/Engine/Source/Editor/Kismet/Private/BlueprintCompilationManager.cpp
+++ b/Engine/Source/Editor/Kismet/Private/BlueprintCompilationManager.cpp
@@ -2150,7 +2150,7 @@ void FBlueprintCompilationManagerImpl::SaveConstructionHashes(const TArray<FComp
 		{
 			if( const UObject* CDO = OriginalClass->ClassDefaultObject )
 			{
-				if( FOverridableManager::Get().IsEnabled(*CDO) )
+				if( FOverridableManager::Get().IsEnabled(CDO) )
 				{
 					return;
 				}
@@ -2669,9 +2669,9 @@ void FBlueprintCompilationManagerImpl::ReinstanceBatch(TArray<FReinstancingJob>&
 			OldCDO = OldClass->ClassDefaultObject;
 			if (OldCDO && ReinstancingJob.Reinstancer.IsValid())
 			{
-				// Object using overridable serialization need to use delta serialization for it do work appropriatly
+				// Object using overridable serialization need to use delta serialization for it do work appropriately
 				// Eventually we should do this to all BP classes CDO no matter what
-				const bool bUsingOverrideSerialization = FOverridableManager::Get().IsEnabled(*OldCDO);
+				const bool bUsingOverrideSerialization = FOverridableManager::Get().IsEnabled(OldCDO);
 				const bool bUseDeltaSerialization = bUsingOverrideSerialization ? true : (ReinstancingJob.Reinstancer.IsValid() ? ReinstancingJob.Reinstancer->bUseDeltaSerializationToCopyProperties : false);
 				UClass* NewClass = ReinstancingJob.OldToNew.Value;
 
@@ -2837,7 +2837,7 @@ void FBlueprintCompilationManagerImpl::ReinstanceBatch(TArray<FReinstancingJob>&
 				{ 
 					bool bIsArchetype = 
 						Obj->HasAnyFlags(RF_ArchetypeObject|RF_InheritableComponentTemplate)
-						|| (FOverridableManager::Get().IsEnabled(*Obj) && Obj->GetOutermostObject()->HasAnyFlags(RF_ClassDefaultObject))
+						|| (FOverridableManager::Get().IsEnabled(Obj) && Obj->GetOutermostObject()->HasAnyFlags(RF_ClassDefaultObject))
 						|| Obj->GetTypedOuter<UBlueprintGeneratedClass>()
 						|| Obj->GetTypedOuter<UBlueprint>();
 					// remove if this is not an archetype or its already in the transient package, note
@@ -2887,7 +2887,7 @@ void FBlueprintCompilationManagerImpl::ReinstanceBatch(TArray<FReinstancingJob>&
 							}
 
 							// this handles nested subobjects:
-							if (FOverridableManager::Get().IsEnabled(*Archetype))
+							if (FOverridableManager::Get().IsEnabled(Archetype))
 							{
 								ForEachObjectWithOuter(Iter, [Archetype, &ArchetypeReferencers](UObject* SubObject)
 								{
diff --git a/Engine/Source/Editor/PropertyEditor/Private/DetailsViewStyle.cpp b/Engine/Source/Editor/PropertyEditor/Private/DetailsViewStyle.cpp
--- a/Engine/Source/Editor/PropertyEditor/Private/DetailsViewStyle.cpp
+++ b/Engine/Source/Editor/PropertyEditor/Private/DetailsViewStyle.cpp
@@ -130,7 +130,7 @@ TAttribute<EVisibility> FOverridesWidgetStyleKey::GetVisibilityAttribute(const T
 		
 	    if ( OverriddenObjectWeakPtr.IsValid() )
 	    {
-		    UObject& OverriddenObject = *OverriddenObjectWeakPtr.Get();
+		    UObject* OverriddenObject = OverriddenObjectWeakPtr.Get();
 
 			bool bIsVisible = false;
 			if (bIsProperty)
diff --git a/Engine/Source/Editor/PropertyEditor/Private/PropertyHandleImpl.cpp b/Engine/Source/Editor/PropertyEditor/Private/PropertyHandleImpl.cpp
--- a/Engine/Source/Editor/PropertyEditor/Private/PropertyHandleImpl.cpp
+++ b/Engine/Source/Editor/PropertyEditor/Private/PropertyHandleImpl.cpp
@@ -551,7 +551,7 @@ FPropertyAccess::Result FPropertyValueImpl::ImportText( const TArray<FObjectBase
 				{
 					// propagate the changes to instances unless the object is using overridable serialization, 
 					// the propagation is done via reinstantiation
-					if (!FOverridableManager::Get().IsEnabled(*CurObject))
+					if (!FOverridableManager::Get().IsEnabled(CurObject))
 					{
 						InPropertyNode->PropagatePropertyChange(CurObject, *NewValue, PreviousContainerValue.IsEmpty() ? PreviousValue : PreviousContainerValue);
 					}
diff --git a/Engine/Source/Editor/UnrealEd/Private/EditorObject.cpp b/Engine/Source/Editor/UnrealEd/Private/EditorObject.cpp
--- a/Engine/Source/Editor/UnrealEd/Private/EditorObject.cpp
+++ b/Engine/Source/Editor/UnrealEd/Private/EditorObject.cpp
@@ -34,7 +34,6 @@
 #include "InstancedFoliage.h"
 #include "Components/BrushComponent.h"
 #include "Algo/Transform.h"
-#include "UObject/OverridableManager.h"
 #include "UObject/OverriddenPropertySet.h"
 
 DEFINE_LOG_CATEGORY_STATIC(LogEditorObject, Log, All);
@@ -268,7 +267,7 @@ static const TCHAR* ImportProperties(
 	FOverriddenPropertySet* OverriddenProperties = nullptr;
 	if (SubobjectOuter && FOverridableSerializationLogic::HasCapabilities(FOverridableSerializationLogic::ECapabilities::T3DSerialization))
 	{
-		OverriddenProperties = FOverridableManager::Get().GetOverriddenProperties(*SubobjectOuter);
+		OverriddenProperties = FOverridableManager::Get().GetOverriddenProperties(SubobjectOuter);
 	}
 	FEnableOverridableSerializationScope Scope(OverriddenProperties != nullptr, OverriddenProperties);
 
@@ -661,7 +660,7 @@ static const TCHAR* ImportProperties(
 					//
 					// Ideally this rule should be applied more widely but for now we limit it to the Prefab system
 					// which we could identify using the fact it uses overridable serialization.
-					const bool bShouldUseDedicatedInstancingGraph = FOverridableManager::Get().IsEnabled(*SubobjectOuter);
+					const bool bShouldUseDedicatedInstancingGraph = FOverridableManager::Get().IsEnabled(SubobjectOuter);
 					FObjectInstancingGraph InstanceGraphForNewObject;
 
 					ComponentTemplate = NewObject<UObject>(
@@ -729,7 +728,7 @@ static const TCHAR* ImportProperties(
 					}
 					else
 					{
-						FOverridableManager::Get().Disable(*ComponentTemplate);
+						FOverridableManager::Get().Disable(ComponentTemplate);
 					}
 				}
 
@@ -1346,7 +1345,7 @@ static const TCHAR* ImportCreateSubObjectsStep(
 					}
 					else
 					{
-						FOverridableManager::Get().Disable(*ComponentTemplate);
+						FOverridableManager::Get().Disable(ComponentTemplate);
 					}
 				}
 
@@ -1405,7 +1404,7 @@ static const TCHAR* ImportPropertiesStep(
 	FOverriddenPropertySet* OverriddenProperties = nullptr;
 	if (SubobjectOuter && FOverridableSerializationLogic::HasCapabilities(FOverridableSerializationLogic::ECapabilities::T3DSerialization))
 	{
-		OverriddenProperties = FOverridableManager::Get().GetOverriddenProperties(*SubobjectOuter);
+		OverriddenProperties = FOverridableManager::Get().GetOverriddenProperties(SubobjectOuter);
 	}
 	FEnableOverridableSerializationScope Scope(OverriddenProperties != nullptr, OverriddenProperties);
 
diff --git a/Engine/Source/Editor/UnrealEd/Private/Factories/EditorFactories.cpp b/Engine/Source/Editor/UnrealEd/Private/Factories/EditorFactories.cpp
--- a/Engine/Source/Editor/UnrealEd/Private/Factories/EditorFactories.cpp
+++ b/Engine/Source/Editor/UnrealEd/Private/Factories/EditorFactories.cpp
@@ -5391,7 +5391,7 @@ void FCustomizableTextObjectFactory::ProcessBuffer(UObject* InParent, EObjectFla
 					}
 					else
 					{
-						FOverridableManager::Get().Disable(*CreatedObject);
+						FOverridableManager::Get().Disable(CreatedObject);
 					}
 				}
 
diff --git a/Engine/Source/Editor/UnrealEd/Private/Kismet2/KismetReinstanceUtilities.cpp b/Engine/Source/Editor/UnrealEd/Private/Kismet2/KismetReinstanceUtilities.cpp
--- a/Engine/Source/Editor/UnrealEd/Private/Kismet2/KismetReinstanceUtilities.cpp
+++ b/Engine/Source/Editor/UnrealEd/Private/Kismet2/KismetReinstanceUtilities.cpp
@@ -2292,7 +2292,7 @@ static void ReplaceObjectHelper(UObject*& OldObject, UClass* OldClass, UObject*&
 	Options.bSkipCompilerGeneratedDefaults = true;
 	Options.bOnlyHandleDirectSubObjects = true;
 	Options.OptionalReplacementMappings = &OldToNewInstanceMap;
-	if (FOverridableManager::Get().IsEnabled(*OldObject))
+	if (FOverridableManager::Get().IsEnabled(OldObject))
 	{
 		Options.bReplaceInternalReferenceUponRead = true;
 		Options.OptionalOldToNewClassMappings = &OldToNewClassMap;
@@ -3218,7 +3218,7 @@ void FBlueprintCompileReinstancer::CopyPropertiesForUnrelatedObjects(UObject* Ol
 	Params.OptionalReplacementMappings = OldToNewInstanceMap;
 	Params.bOnlyHandleDirectSubObjects = bOnlyHandleDirectSubObjects;
 	// Overridable serialization needs this to be able to merge back containers of subobjects.
-	if (FOverridableManager::Get().IsEnabled(*OldObject))
+	if (FOverridableManager::Get().IsEnabled(OldObject))
 	{
 		Params.bReplaceInternalReferenceUponRead = true;
 		Params.OptionalOldToNewClassMappings = OldToNewClassMap;
@@ -3239,7 +3239,7 @@ void FBlueprintCompileReinstancer::PreCreateSubObjectsForReinstantiation(
 	// Scoped only for OS object for now as it is the only case so far that has problem with this filtering.
 	TSet<UObject*> OldInstancedSubObjects;
 	TSet<UObject*>* OldInstancedSubObjectsPtr = nullptr;
-	if (!FOverridableManager::Get().IsEnabled(*OldObject) || !OldObject->HasAnyFlags(RF_NeedPostLoad))
+	if (!FOverridableManager::Get().IsEnabled(OldObject) || !OldObject->HasAnyFlags(RF_NeedPostLoad))
 	{
 		OldInstancedSubObjectsPtr = &OldInstancedSubObjects;
 		UE::Private::FReplaceReferenceHelper::GetOwnedSubobjectsRecursive(OldObject, OldInstancedSubObjects);
diff --git a/Engine/Source/Runtime/CoreUObject/Private/UObject/Class.cpp b/Engine/Source/Runtime/CoreUObject/Private/UObject/Class.cpp
--- a/Engine/Source/Runtime/CoreUObject/Private/UObject/Class.cpp
+++ b/Engine/Source/Runtime/CoreUObject/Private/UObject/Class.cpp
@@ -1557,7 +1557,7 @@ struct FSerializationControlExtensionContext
 		EClassSerializationControlExtension SerializationExtension = EClassSerializationControlExtension::NoExtension;
 	
 		// Overridable serialization information initialization
-		if (FOverriddenPropertySet* ObjectOverriddenProperties = FOverridableManager::Get().GetOverriddenProperties(*(UObject*)Data))
+		if (FOverriddenPropertySet* ObjectOverriddenProperties = FOverridableManager::Get().GetOverriddenProperties((UObject*)Data))
 		{
 			SerializationExtension |= EClassSerializationControlExtension::OverridableSerializationInformation;
 			bEnableOverridableSerialization = true;
@@ -1611,7 +1611,7 @@ void UStruct::SerializeVersionedTaggedProperties(FStructuredArchive::FSlot Slot,
 					ControlContext.OverriddenProperties = nullptr;
 					if (UnderlyingArchive.ArMergeOverrides)
 					{
-						ControlContext.OverriddenProperties = FOverridableManager::Get().GetOverriddenProperties(*(UObject*)Data);
+						ControlContext.OverriddenProperties = FOverridableManager::Get().GetOverriddenProperties((UObject*)Data);
 					}
 					if (!ControlContext.OverriddenProperties)
 					{
diff --git a/Engine/Source/Runtime/CoreUObject/Private/UObject/CoreNative.cpp b/Engine/Source/Runtime/CoreUObject/Private/UObject/CoreNative.cpp
--- a/Engine/Source/Runtime/CoreUObject/Private/UObject/CoreNative.cpp
+++ b/Engine/Source/Runtime/CoreUObject/Private/UObject/CoreNative.cpp
@@ -205,7 +205,7 @@ UObject* FObjectInstancingGraph::GetInstancedSubobject( UObject* SourceSubobject
 		if ( bShouldInstance )
 		{
 			// If the CurrentValue is within the SourceRoot, lets use it to instantiate as it must have come from the merge result of the serialization
-			const bool bIsInstantiatingSubObjectForOverridableSerialization = SourceSubobject && FOverridableManager::Get().NeedSubObjectTemplateInstantiation(*SourceSubobject);
+			const bool bIsInstantiatingSubObjectForOverridableSerialization = SourceSubobject && FOverridableManager::Get().NeedSubObjectTemplateInstantiation(SourceSubobject);
 			if (bIsInstantiatingSubObjectForOverridableSerialization && SourceSubobject != CurrentValue && CurrentValue->IsIn(SourceRoot))
 			{
 				SourceSubobject = CurrentValue;
diff --git a/Engine/Source/Runtime/CoreUObject/Private/UObject/Obj.cpp b/Engine/Source/Runtime/CoreUObject/Private/UObject/Obj.cpp
--- a/Engine/Source/Runtime/CoreUObject/Private/UObject/Obj.cpp
+++ b/Engine/Source/Runtime/CoreUObject/Private/UObject/Obj.cpp
@@ -528,7 +528,7 @@ void UObject::PreEditChange( FEditPropertyChain& PropertyAboutToChange )
 		PropagatePreEditChange(Objects, PropertyAboutToChange);
 	}
 
-	FOverridableManager::Get().PreOverrideProperty(*this, PropertyAboutToChange);
+	FOverridableManager::Get().PreOverrideProperty(this, PropertyAboutToChange);
 }
 
 
@@ -601,7 +601,7 @@ void UObject::PostEditChangeChainProperty(FPropertyChangedChainEvent& PropertyCh
 
 	PostEditChangeProperty(PropertyEvent);
 
-	FOverridableManager::Get().PostOverrideProperty(*this, PropertyChangedEvent, PropertyChangedEvent.PropertyChain);
+	FOverridableManager::Get().PostOverrideProperty(this, PropertyChangedEvent, PropertyChangedEvent.PropertyChain);
 	UE::FPropertyBagRepository::PostEditChangeChainProperty(this, PropertyChangedEvent);
 }
 
@@ -620,7 +620,7 @@ void UObject::PropagatePreEditChange( TArray<UObject*>& AffectedObjects, FEditPr
 {
 	// This feature is not essential for overridable serialization
 	// Disable it until we fix the fact that do it will create overrides on all the instances
-	if (FOverridableManager::Get().IsEnabled(*this))
+	if (FOverridableManager::Get().IsEnabled(this))
 	{
 		return;
 	}
@@ -712,7 +712,7 @@ void UObject::PropagatePostEditChange( TArray<UObject*>& AffectedObjects, FPrope
 {
 	// This feature is not essential for overridable serialization
 	// Disable it until we fix the fact that do it will create overrides on all the instances
-	if (FOverridableManager::Get().IsEnabled(*this))
+	if (FOverridableManager::Get().IsEnabled(this))
 	{
 		return;
 	}
@@ -1404,7 +1404,7 @@ void UObject::PostLoadSubobjects( FObjectInstancingGraph* OuterInstanceGraph/*=N
 		// 
 		// @todo_LoadPerf - Consider modifying to bypass this for cooked packages if the class doesn't include any instanced editor-only fields.
 		// Currently we do not have a flag that indicates this scenario, and it might be expensive to iterate over the linked property chain here.
-		if (!FPlatformProperties::RequiresCookedData() || FOverridableManager::Get().NeedSubObjectTemplateInstantiation(*this))
+		if (!FPlatformProperties::RequiresCookedData() || FOverridableManager::Get().NeedSubObjectTemplateInstantiation(this))
 		{
 			FObjectInstancingGraph CurrentInstanceGraph;
 
@@ -1711,13 +1711,13 @@ void UObject::Serialize(FStructuredArchive::FRecord Record)
 		// Handle derived UClass objects (exact UClass objects are native only and shouldn't be touched)
 		if (ObjClass != UClass::StaticClass())
 		{
-			// Serializing the overriden properties if using UPS (like object duplication)
+			// Serializing the overridden properties if using UPS (like object duplication)
 			// Also needed for transactions using TPS
 			if (!UnderlyingArchive.IsCooking() && 
 				!UnderlyingArchive.IsLoadingFromCookedPackage() &&
 				(UnderlyingArchive.UseUnversionedPropertySerialization() || UnderlyingArchive.IsTransacting()))
 			{
-				FOverridableManager::Get().SerializeOverriddenProperties(*this, Record);
+				FOverridableManager::Get().SerializeOverriddenProperties(this, Record);
 #if WITH_EDITORONLY_DATA
 				UE::FInitializedPropertyValueState(this).Serialize(Record);
 #endif
diff --git a/Engine/Source/Runtime/CoreUObject/Private/UObject/OverridableManager.cpp b/Engine/Source/Runtime/CoreUObject/Private/UObject/OverridableManager.cpp
--- a/Engine/Source/Runtime/CoreUObject/Private/UObject/OverridableManager.cpp
+++ b/Engine/Source/Runtime/CoreUObject/Private/UObject/OverridableManager.cpp
@@ -19,7 +19,7 @@ void FOverridableManager::Create()
 	}
 }
 
-bool FOverridableManager::IsEnabled(const UObject& Object)
+bool FOverridableManager::IsEnabled(TNotNull<const UObject*> Object)
 {
 #if WITH_EDITORONLY_DATA
 	return OverriddenObjectAnnotations.IsEnabled(Object);
@@ -28,36 +28,36 @@ bool FOverridableManager::IsEnabled(const UObject& Object)
 #endif // WITH_EDITORONLY_DATA
 }
 
-void FOverridableManager::Enable(UObject& Object)
+void FOverridableManager::Enable(TNotNull<UObject*> Object)
 {
 #if WITH_EDITORONLY_DATA
 	OverriddenObjectAnnotations.FindOrAdd(Object);
 #endif // WITH_EDITORONLY_DATA
 }
 
-void FOverridableManager::Disable(UObject& Object, bool bPropagateToSubObject/*= false*/)
+void FOverridableManager::Disable(TNotNull<UObject*> Object, bool bPropagateToSubObject/*= false*/)
 {
 #if WITH_EDITORONLY_DATA
-	OverriddenObjectAnnotations.RemoveAnnotation(&Object);
+	OverriddenObjectAnnotations.RemoveAnnotation(Object);
 
 	if (bPropagateToSubObject)
 	{
 		TArray<UObject*> InstancedSubObjects;
 		constexpr bool bIncludeNestedObjects = false;
-		GetObjectsWithOuter(&Object, InstancedSubObjects, bIncludeNestedObjects);
+		GetObjectsWithOuter(Object, InstancedSubObjects, bIncludeNestedObjects);
 		for (UObject* InstancedSubObject : InstancedSubObjects)
 		{
 			checkf(InstancedSubObject, TEXT("Expecting non null SubObjects"));
-			DisableInstancedSubObject(Object, *InstancedSubObject);
+			DisableInstancedSubObject(Object, InstancedSubObject);
 		}
 	}
 #endif // WITH_EDITORONLY_DATA
 }
 
-void FOverridableManager::DisableInstancedSubObject(const UObject& Object, UObject& InstancedSubObject)
+void FOverridableManager::DisableInstancedSubObject(TNotNull<const UObject*> Object, TNotNull<UObject*> InstancedSubObject)
 {
 #if WITH_EDITORONLY_DATA
-	if (InstancedSubObject.IsIn(&Object))
+	if (InstancedSubObject->IsIn(Object))
 	{
 		Disable(InstancedSubObject, /*bPropagateToSubObject*/ true);
 	}
@@ -65,13 +65,13 @@ void FOverridableManager::DisableInstancedSubObject(const UObject& Object, UObje
 }
 
 
-void FOverridableManager::InheritEnabledFrom(UObject& Object, const UObject* DefaultData)
+void FOverridableManager::InheritEnabledFrom(TNotNull<UObject*> Object, const UObject* DefaultData)
 {
 #if WITH_EDITORONLY_DATA
 	if (!OverriddenObjectAnnotations.IsEnabled(Object))
 	{
-		const UObject* Outer = Object.GetOuter();
-		if ((Outer && IsEnabled(*Outer)) || (DefaultData && IsEnabled(*DefaultData)))
+		const UObject* Outer = Object->GetOuter();
+		if ((Outer && IsEnabled(Outer)) || (DefaultData && IsEnabled(DefaultData)))
 		{
 			Enable(Object);
 		}
@@ -79,7 +79,7 @@ void FOverridableManager::InheritEnabledFrom(UObject& Object, const UObject* Def
 #endif // WITH_EDITORONLY_DATA
 }
 
-bool FOverridableManager::NeedSubObjectTemplateInstantiation(const UObject& Object)
+bool FOverridableManager::NeedSubObjectTemplateInstantiation(TNotNull<const UObject*> Object)
 {
 #if WITH_EDITORONLY_DATA
 	if (const FOverriddenPropertySet* OverriddenProperties = OverriddenObjectAnnotations.Find(Object))
@@ -88,11 +88,11 @@ bool FOverridableManager::NeedSubObjectTemplateInstantiation(const UObject& Obje
 	}
 	return false;
 #else
-	return NeedsSubobjectTemplateInstantiation.Get(&Object);
+	return NeedsSubobjectTemplateInstantiation.Get(Object);
 #endif // WITH_EDITORONLY_DATA
 }
 
-FOverriddenPropertySet* FOverridableManager::GetOverriddenProperties(UObject& Object)
+FOverriddenPropertySet* FOverridableManager::GetOverriddenPropertiesInternal(TNotNull<UObject*> Object)
 {
 #if WITH_EDITORONLY_DATA
 	return OverriddenObjectAnnotations.Find(Object);
@@ -101,7 +101,7 @@ FOverriddenPropertySet* FOverridableManager::GetOverriddenProperties(UObject& Ob
 #endif // WITH_EDITORONLY_DATA
 }
 
-const FOverriddenPropertySet* FOverridableManager::GetOverriddenProperties(const UObject& Object)
+const FOverriddenPropertySet* FOverridableManager::GetOverriddenPropertiesInternal(TNotNull<const UObject*> Object)
 {
 #if WITH_EDITORONLY_DATA
 	return OverriddenObjectAnnotations.Find(Object);
@@ -114,14 +114,14 @@ const FOverriddenPropertySet* FOverridableManager::GetOverriddenProperties(const
 FOverriddenPropertySet* FOverridableManager::RestoreOverrideOperation(TNotNull<UObject*> Object, EOverriddenPropertyOperation Operation, const bool bNeedsSubobjectTemplateInstantiation)
 {
 	// When sub property overrides are going to be restored, it will turn this object back to modified, so let's not restore that operation.
-	// This allows cleanup of the modified state if we endup not overriding any sub properties
+	// This allows cleanup of the modified state if we end up not overriding any sub properties
 	if (Operation == EOverriddenPropertyOperation::Modified)
 	{
 		Operation = EOverriddenPropertyOperation::None;
 	}
 
 #if WITH_EDITORONLY_DATA
-	FOverriddenPropertySet& ObjectOverriddenProperties = OverriddenObjectAnnotations.FindOrAdd(*Object);
+	FOverriddenPropertySet& ObjectOverriddenProperties = OverriddenObjectAnnotations.FindOrAdd(Object);
 	ObjectOverriddenProperties.Reset();
 	ObjectOverriddenProperties.SetOverriddenPropertyOperation(Operation, /*CurrentPropertyChain*/nullptr, /*Property*/nullptr);
 	ObjectOverriddenProperties.bNeedsSubobjectTemplateInstantiation = bNeedsSubobjectTemplateInstantiation;
@@ -142,22 +142,22 @@ FOverriddenPropertySet* FOverridableManager::RestoreOverrideOperation(TNotNull<U
 
 void FOverridableManager::RestoreOverrideState(TNotNull<const UObject*> OldObject, TNotNull<UObject*> NewObject)
 {
-	if (const FOverriddenPropertySet* OldOverriddenProperties = GetOverriddenProperties(*OldObject))
+	if (const FOverriddenPropertySet* OldOverriddenProperties = GetOverriddenProperties(OldObject))
 	{
-		if (FOverriddenPropertySet* NewOverriddenProperties = GetOverriddenProperties(*NewObject))
+		if (FOverriddenPropertySet* NewOverriddenProperties = GetOverriddenProperties(NewObject))
 		{
 			NewOverriddenProperties->RestoreOverriddenState(*OldOverriddenProperties);
 		}
 	}
 }
 
-EOverriddenState FOverridableManager::GetOverriddenState(UObject& Object)
+EOverriddenState FOverridableManager::GetOverriddenState(TNotNull<UObject*> Object)
 {
 #if WITH_EDITORONLY_DATA
 	if(const FOverriddenPropertySet* OverriddenProperties = GetOverriddenProperties(Object))
 	{
 		// Transient should not affect the overridden state
-		if (Object.HasAnyFlags(RF_Transient))
+		if (Object->HasAnyFlags(RF_Transient))
 		{
 			return EOverriddenState::NoOverrides;
 		}
@@ -179,19 +179,19 @@ EOverriddenState FOverridableManager::GetOverriddenState(UObject& Object)
 			case EOverriddenPropertyOperation::None:
 				break;
 			default:
-				checkf(false, TEXT("Unsupported operation(%s) on object(%p:%s)"), *GetOverriddenOperationString(Operation), &Object, *Object.GetName());
+				checkf(false, TEXT("Unsupported operation(%s) on object(%p:%s)"), *GetOverriddenOperationString(Operation), static_cast<UObject*>(Object), *Object->GetName());
 		}
 
-		auto GetSubObjectState = [this](UObject& Object, auto& Recurse) -> EOverriddenState
+		auto GetSubObjectState = [this](TNotNull<UObject*> Object, auto& Recurse) -> EOverriddenState
 		{
 			// Need to check subobjects
-			for (TPropertyValueIterator<FObjectProperty> ObjPropIt(Object.GetClass(), &Object); ObjPropIt; ++ObjPropIt)
+			for (TPropertyValueIterator<FObjectProperty> ObjPropIt(Object->GetClass(), Object); ObjPropIt; ++ObjPropIt)
 			{
 				if (UObject* InstancedSubObject = ObjPropIt.Key()->GetObjectPropertyValue(ObjPropIt.Value()))
 				{
-					if (InstancedSubObject->IsIn(&Object))
+					if (InstancedSubObject->IsIn(Object))
 					{
-						if (const FOverriddenPropertySet* OverriddenProperties = GetOverriddenProperties(*InstancedSubObject))
+						if (const FOverriddenPropertySet* OverriddenProperties = GetOverriddenProperties(InstancedSubObject))
 						{
 							const EOverriddenPropertyOperation Operation = OverriddenProperties->GetOverriddenPropertyOperation((FArchiveSerializedPropertyChain*)nullptr, (FProperty*)nullptr);
 							if (Operation != EOverriddenPropertyOperation::None)
@@ -199,7 +199,7 @@ EOverriddenState FOverridableManager::GetOverriddenState(UObject& Object)
 								return EOverriddenState::HasOverrides;
 							}
 
-							if (Recurse(*InstancedSubObject, Recurse) != EOverriddenState::NoOverrides)
+							if (Recurse(InstancedSubObject, Recurse) != EOverriddenState::NoOverrides)
 							{
 								return EOverriddenState::HasOverrides;
 							}
@@ -221,28 +221,28 @@ EOverriddenState FOverridableManager::GetOverriddenState(UObject& Object)
 	return EOverriddenState::NoOverrides;
 }
 
-void FOverridableManager::OverrideObject(UObject& Object)
+void FOverridableManager::OverrideObject(TNotNull<UObject*> Object)
 {
 #if WITH_EDITORONLY_DATA
 	if (FOverriddenPropertySet* ThisObjectOverriddenProperties = OverriddenObjectAnnotations.Find(Object))
 	{
 		// Passing no property node means we are overriding the object itself
-		ThisObjectOverriddenProperties->OverrideProperty(FPropertyVisitorPath::InvalidIterator(), &Object);
+		ThisObjectOverriddenProperties->OverrideProperty(FPropertyVisitorPath::InvalidIterator(), Object);
 	}
 #endif // WITH_EDITORONLY_DATA
 }
 
-void FOverridableManager::OverrideInstancedSubObject(const UObject& Object, UObject& InstancedSubObject)
+void FOverridableManager::OverrideInstancedSubObject(TNotNull<const UObject*> Object, TNotNull<UObject*> InstancedSubObject)
 {
 #if WITH_EDITORONLY_DATA
-	if (InstancedSubObject.IsIn(&Object))
+	if (InstancedSubObject->IsIn(Object))
 	{
 		OverrideObject(InstancedSubObject);
 	}
 #endif // WITH_EDITORONLY_DATA
 }
 
-void FOverridableManager::PropagateOverrideToInstancedSubObjects(UObject& Object)
+void FOverridableManager::PropagateOverrideToInstancedSubObjects(TNotNull<UObject*> Object)
 {
 #if WITH_EDITORONLY_DATA
 	// Do not use TPropertyValueIterator<FObjectProperty> to be consistent with PropagateClearOverridesToInstancedSubObjects
@@ -250,26 +250,26 @@ void FOverridableManager::PropagateOverrideToInstancedSubObjects(UObject& Object
 	// the object class setup correctly when this is called from PostInit
 	TArray<UObject*> InstancedSubObjects;
 	constexpr bool bIncludeNestedObjects = false;
-	GetObjectsWithOuter(&Object, InstancedSubObjects, bIncludeNestedObjects);
+	GetObjectsWithOuter(Object, InstancedSubObjects, bIncludeNestedObjects);
 	for (UObject* InstancedSubObject : InstancedSubObjects)
 	{
 		checkf(InstancedSubObject, TEXT("Expecting non null SubObjects"));
-		OverrideInstancedSubObject(Object, *InstancedSubObject);
+		OverrideInstancedSubObject(Object, InstancedSubObject);
 	}
 #endif // WITH_EDITORONLY_DATA
 }
 
-void FOverridableManager::OverrideProperty(UObject& Object, const FPropertyVisitorPath& PropertyPath)
+void FOverridableManager::OverrideProperty(TNotNull<UObject*> Object, const FPropertyVisitorPath& PropertyPath)
 {
 #if WITH_EDITORONLY_DATA
 	if (FOverriddenPropertySet* ThisObjectOverriddenProperties = OverriddenObjectAnnotations.Find(Object))
 	{
-		ThisObjectOverriddenProperties->OverrideProperty(PropertyPath.GetRootIterator(), &Object);
+		ThisObjectOverriddenProperties->OverrideProperty(PropertyPath.GetRootIterator(), Object);
 	}
 #endif // WITH_EDITORONLY_DATA
 }
 
-bool FOverridableManager::ClearOverriddenProperty(UObject& Object, FPropertyVisitorPath::Iterator PropertyIterator)
+bool FOverridableManager::ClearOverriddenProperty(TNotNull<UObject*> Object, FPropertyVisitorPath::Iterator PropertyIterator)
 {
 #if WITH_EDITORONLY_DATA
 	if (FOverriddenPropertySet* ThisObjectOverriddenProperties = OverriddenObjectAnnotations.Find(Object))
@@ -280,31 +280,31 @@ bool FOverridableManager::ClearOverriddenProperty(UObject& Object, FPropertyVisi
 	return false;
 }
 
-void FOverridableManager::PreOverrideProperty(UObject& Object, const FPropertyVisitorPath& PropertyPath)
+void FOverridableManager::PreOverrideProperty(TNotNull<UObject*> Object, const FPropertyVisitorPath& PropertyPath)
 {
 #if WITH_EDITORONLY_DATA
 	NotifyPropertyChange(EPropertyNotificationType::PreEdit, Object, PropertyPath.GetRootIterator(), EPropertyChangeType::Unspecified);
 #endif // WITH_EDITORONLY_DATA
 }
 
-void FOverridableManager::PostOverrideProperty(UObject& Object, const FPropertyVisitorPath& PropertyPath, const EPropertyChangeType::Type ChangeType)
+void FOverridableManager::PostOverrideProperty(TNotNull<UObject*> Object, const FPropertyVisitorPath& PropertyPath, const EPropertyChangeType::Type ChangeType)
 {
 #if WITH_EDITORONLY_DATA
 	NotifyPropertyChange(EPropertyNotificationType::PostEdit, Object, PropertyPath.GetRootIterator(), ChangeType);
 #endif // WITH_EDITORONLY_DATA
 }
 
-void FOverridableManager::NotifyPropertyChange(const EPropertyNotificationType Notification, UObject& Object, FPropertyVisitorPath::Iterator PropertyIterator, const EPropertyChangeType::Type ChangeType)
+void FOverridableManager::NotifyPropertyChange(const EPropertyNotificationType Notification, TNotNull<UObject*> Object, FPropertyVisitorPath::Iterator PropertyIterator, const EPropertyChangeType::Type ChangeType)
 {
 #if WITH_EDITORONLY_DATA
 	if (FOverriddenPropertySet* ThisObjectOverriddenProperties = OverriddenObjectAnnotations.Find(Object))
 	{
-		ThisObjectOverriddenProperties->NotifyPropertyChange(Notification, PropertyIterator, ChangeType, &Object);
+		ThisObjectOverriddenProperties->NotifyPropertyChange(Notification, PropertyIterator, ChangeType, Object);
 	}
 #endif // WITH_EDITORONLY_DATA
 }
 
-EOverriddenPropertyOperation FOverridableManager::GetOverriddenPropertyOperation(UObject& Object, FPropertyVisitorPath::Iterator PropertyIterator, bool* bOutInheritedOperation)
+EOverriddenPropertyOperation FOverridableManager::GetOverriddenPropertyOperation(TNotNull<UObject*> Object, FPropertyVisitorPath::Iterator PropertyIterator, bool* bOutInheritedOperation)
 {
 #if WITH_EDITORONLY_DATA
 	if (const FOverriddenPropertySet* ThisObjectOverriddenProperties = OverriddenObjectAnnotations.Find(Object))
@@ -315,7 +315,7 @@ EOverriddenPropertyOperation FOverridableManager::GetOverriddenPropertyOperation
 	return EOverriddenPropertyOperation::None;
 }
 
-void FOverridableManager::ClearOverrides(UObject& Object)
+void FOverridableManager::ClearOverrides(TNotNull<UObject*> Object)
 {
 #if WITH_EDITORONLY_DATA
 	if(FOverriddenPropertySet* ThisObjectOverriddenProperties = OverriddenObjectAnnotations.Find(Object))
@@ -326,17 +326,17 @@ void FOverridableManager::ClearOverrides(UObject& Object)
 #endif // WITH_EDITORONLY_DATA
 }
 
-void FOverridableManager::ClearInstancedSubObjectOverrides(const UObject& Object, UObject& InstancedSubObject)
+void FOverridableManager::ClearInstancedSubObjectOverrides(TNotNull<const UObject*> Object, TNotNull<UObject*> InstancedSubObject)
 {
 #if WITH_EDITORONLY_DATA
-	if (InstancedSubObject.IsIn(&Object))
+	if (InstancedSubObject->IsIn(Object))
 	{
 		ClearOverrides(InstancedSubObject);
 	}
 #endif // WITH_EDITORONLY_DATA
 }
 
-void FOverridableManager::PropagateClearOverridesToInstancedSubObjects(UObject& Object)
+void FOverridableManager::PropagateClearOverridesToInstancedSubObjects(TNotNull<UObject*> Object)
 {
 #if WITH_EDITORONLY_DATA
 	// Cannot use TPropertyValueIterator<FObjectProperty> here because the object might not have 
@@ -344,18 +344,18 @@ void FOverridableManager::PropagateClearOverridesToInstancedSubObjects(UObject&
 	// @todo: need to figure out why this is a problem. Maybe we should not even this method during PostInit.
 	TArray<UObject*> InstancedSubObjects;
 	constexpr bool bIncludeNestedObjects = false;
-	GetObjectsWithOuter(&Object, InstancedSubObjects, bIncludeNestedObjects);
+	GetObjectsWithOuter(Object, InstancedSubObjects, bIncludeNestedObjects);
 	for (UObject* InstancedSubObject : InstancedSubObjects)
 	{
 		checkf(InstancedSubObject, TEXT("Expecting non null SubObjects"));
 
 		// There are some cases where the property has information about that should be an instanced subobject, but it is not owned by us.
-		ClearInstancedSubObjectOverrides(Object, *InstancedSubObject);
+		ClearInstancedSubObjectOverrides(Object, InstancedSubObject);
 	}
 #endif // WITH_EDITORONLY_DATA
 }
 
-void FOverridableManager::SerializeOverriddenProperties(UObject& Object, FStructuredArchive::FRecord ObjectRecord)
+void FOverridableManager::SerializeOverriddenProperties(TNotNull<UObject*> Object, FStructuredArchive::FRecord ObjectRecord)
 {
 #if WITH_EDITORONLY_DATA
 	const FArchiveState& ArchiveState = ObjectRecord.GetArchiveState();
@@ -368,7 +368,7 @@ void FOverridableManager::SerializeOverriddenProperties(UObject& Object, FStruct
 
 		if (ArchiveState.IsLoading())
 		{
-			OverriddenProperties = RestoreOverrideOperation(&Object, Operation, /*bNeedsSubobjectTemplateInstantiation*/false);
+			OverriddenProperties = RestoreOverrideOperation(Object, Operation, /*bNeedsSubobjectTemplateInstantiation*/false);
 			checkf(OverriddenProperties, TEXT("Expecting a overridden property set returned"));
 		}
 
diff --git a/Engine/Source/Runtime/CoreUObject/Private/UObject/OverriddenPropertySet.cpp b/Engine/Source/Runtime/CoreUObject/Private/UObject/OverriddenPropertySet.cpp
--- a/Engine/Source/Runtime/CoreUObject/Private/UObject/OverriddenPropertySet.cpp
+++ b/Engine/Source/Runtime/CoreUObject/Private/UObject/OverriddenPropertySet.cpp
@@ -267,7 +267,7 @@ FOverriddenPropertyNodeID FOverriddenPropertyNodeID::FromMapKey(const FProperty*
 	{
 		if (const UObject* Object = KeyObjectProperty->GetObjectPropertyValue(KeyData))
 		{
-			return FOverriddenPropertyNodeID(*Object);
+			return FOverriddenPropertyNodeID(Object);
 		}
 	}
 	else
@@ -292,7 +292,7 @@ int32 FOverriddenPropertyNodeID::ToMapInternalIndex(FScriptMapHelper& MapHelper)
 		{
 			if (UObject* CurrentObject = KeyObjectProperty->GetObjectPropertyValue(MapHelper.GetKeyPtr(It)))
 			{
-				if ((*this) == FOverriddenPropertyNodeID(*CurrentObject))
+				if ((*this) == FOverriddenPropertyNodeID(CurrentObject))
 				{
 					return It.GetInternalIndex();
 				}
@@ -430,7 +430,7 @@ EOverriddenPropertyOperation FOverriddenPropertySet::GetOverriddenPropertyOperat
 				if (UObject* SubObject = ObjectProperty->GetObjectPropertyValue(SubValuePtr))
 				{
 					// This should not be needed in the property grid, as it should already been called on the subobject.
-					return OverridableManager.GetOverriddenPropertyOperation(*SubObject, NextPropertyIterator, bOutInheritedOperation);
+					return OverridableManager.GetOverriddenPropertyOperation(SubObject, NextPropertyIterator, bOutInheritedOperation);
 				}
 			}
 		}
@@ -452,12 +452,12 @@ EOverriddenPropertyOperation FOverriddenPropertySet::GetOverriddenPropertyOperat
 						if (NextPropertyIterator)
 						{
 							// Forward any sub queries to the subobject
-							return OverridableManager.GetOverriddenPropertyOperation(*SubObject, NextPropertyIterator, bOutInheritedOperation);
+							return OverridableManager.GetOverriddenPropertyOperation(SubObject, NextPropertyIterator, bOutInheritedOperation);
 						}
 						else if(CurrentOverriddenPropertyNode)
 						{
 							// Caller wants to know about any override state on the reference of the subobject itself
-							const FOverriddenPropertyNodeID  SubObjectID(*SubObject);
+							const FOverriddenPropertyNodeID  SubObjectID(SubObject);
 							if (const FOverriddenPropertyNodeID* CurrentPropKey = CurrentOverriddenPropertyNode->SubPropertyNodeKeys.Find(SubObjectID))
 							{
 								const FOverriddenPropertyNode* SubObjectOverriddenPropertyNode = OverriddenPropertyNodes.Find(*CurrentPropKey);
@@ -492,7 +492,7 @@ EOverriddenPropertyOperation FOverriddenPropertySet::GetOverriddenPropertyOperat
 					{
 						if (UObject* ValueSubObject = ValueInstancedObjectProperty->GetObjectPropertyValue(MapHelper.GetValuePtr(InternalMapIndex)))
 						{
-							return OverridableManager.GetOverriddenPropertyOperation(*ValueSubObject, NextPropertyIterator, bOutInheritedOperation);
+							return OverridableManager.GetOverriddenPropertyOperation(ValueSubObject, NextPropertyIterator, bOutInheritedOperation);
 						}
 					}
 				}
@@ -563,11 +563,11 @@ bool FOverriddenPropertySet::ClearOverriddenProperty(FOverriddenPropertyNode& Pa
 			{
 				if (NextPropertyIterator)
 				{
-					return OverridableManager.ClearOverriddenProperty(*SubObject, NextPropertyIterator);
+					return OverridableManager.ClearOverriddenProperty(SubObject, NextPropertyIterator);
 				}
 				else
 				{
-					OverridableManager.ClearOverrides(*SubObject);
+					OverridableManager.ClearOverrides(SubObject);
 					bClearedOverrides = true;
 				}
 			}
@@ -593,7 +593,7 @@ bool FOverriddenPropertySet::ClearOverriddenProperty(FOverriddenPropertyNode& Pa
 						{
 							if (UObject* SubObject = InnerObjectProperty->GetObjectPropertyValue(ArrayHelper.GetElementPtr(i)))
 							{
-								OverridableManager.ClearInstancedSubObjectOverrides(*Owner, *SubObject);
+								OverridableManager.ClearInstancedSubObjectOverrides(Owner, SubObject);
 							}
 						}
 						bClearedOverrides = true;
@@ -604,16 +604,16 @@ bool FOverriddenPropertySet::ClearOverriddenProperty(FOverriddenPropertyNode& Pa
 						{
 							if (NextPropertyIterator)
 							{
-								return OverridableManager.ClearOverriddenProperty(*SubObject, NextPropertyIterator);
+								return OverridableManager.ClearOverriddenProperty(SubObject, NextPropertyIterator);
 							}
 							else if(CurrentOverriddenPropertyNode)
 							{
-								const FOverriddenPropertyNodeID  SubObjectID(*SubObject);
+								const FOverriddenPropertyNodeID  SubObjectID(SubObject);
 								FOverriddenPropertyNodeID CurrentPropKey;
 								if (CurrentOverriddenPropertyNode->SubPropertyNodeKeys.RemoveAndCopyValue(SubObjectID, CurrentPropKey))
 								{
 									verifyf(OverriddenPropertyNodes.Remove(CurrentPropKey), TEXT("Expecting a node to be removed"));
-									OverridableManager.ClearInstancedSubObjectOverrides(*Owner, *SubObject);
+									OverridableManager.ClearInstancedSubObjectOverrides(Owner, SubObject);
 									return true;
 								}
 							}
@@ -642,7 +642,7 @@ bool FOverriddenPropertySet::ClearOverriddenProperty(FOverriddenPropertyNode& Pa
 					checkf(MapProperty->ValueProp, TEXT("Expecting a value type for Maps"));
 					if (UObject* ValueSubObject = ValueInstancedObjectProperty ? ValueInstancedObjectProperty->GetObjectPropertyValue(MapHelper.GetValuePtr(InternalMapIndex)) : nullptr)
 					{
-						return OverridableManager.ClearOverriddenProperty(*ValueSubObject, NextPropertyIterator);
+						return OverridableManager.ClearOverriddenProperty(ValueSubObject, NextPropertyIterator);
 					}
 				}
 			}
@@ -657,7 +657,7 @@ bool FOverriddenPropertySet::ClearOverriddenProperty(FOverriddenPropertyNode& Pa
 					{
 						if (UObject* ValueSubObject = ValueInstancedObjectProperty->GetObjectPropertyValue(MapHelper.GetValuePtr(It.GetInternalIndex())))
 						{
-							OverridableManager.ClearInstancedSubObjectOverrides(*Owner, *ValueSubObject);
+							OverridableManager.ClearInstancedSubObjectOverrides(Owner, ValueSubObject);
 						}
 					}
 				}
@@ -676,7 +676,7 @@ bool FOverriddenPropertySet::ClearOverriddenProperty(FOverriddenPropertyNode& Pa
 					if (UObject* ValueSubObject = ValueInstancedObjectProperty ? ValueInstancedObjectProperty->GetObjectPropertyValue(MapHelper.GetValuePtr(InternalMapIndex)) : nullptr)
 					{
 						// In the case of a instanced subobject, clear all the overrides on the subobject as well
-						OverridableManager.ClearInstancedSubObjectOverrides(*Owner, *ValueSubObject);
+						OverridableManager.ClearInstancedSubObjectOverrides(Owner, ValueSubObject);
 					}
 
 					return true;
@@ -755,7 +755,7 @@ void FOverriddenPropertySet::NotifyPropertyChange(FOverriddenPropertyNode* Paren
 			checkf(RootNode, TEXT("Expecting to always have a "));
 			if (RootNode == ParentPropertyNode)
 			{
-				OverridableManager.PropagateOverrideToInstancedSubObjects(*Owner);
+				OverridableManager.PropagateOverrideToInstancedSubObjects(Owner);
 			}
 		}
 		return;
@@ -841,12 +841,12 @@ void FOverriddenPropertySet::NotifyPropertyChange(FOverriddenPropertyNode* Paren
 							{
 								if(SubPropertyNode)
 								{
-									const FOverriddenPropertyNodeID SubObjectID(*SubObject);
+									const FOverriddenPropertyNodeID SubObjectID(SubObject);
 									FOverriddenPropertyNode& SubObjectNode = FindOrAddNode(*SubPropertyNode, SubObjectID);
 									SubObjectNode.Operation = EOverriddenPropertyOperation::Replace;
 								}
 
-								OverridableManager.OverrideInstancedSubObject(*Owner, *SubObject);
+								OverridableManager.OverrideInstancedSubObject(Owner, SubObject);
 							}
 						}
 					};
@@ -858,12 +858,12 @@ void FOverriddenPropertySet::NotifyPropertyChange(FOverriddenPropertyNode* Paren
 						{
 							if(SubPropertyNode)
 							{
-								const FOverriddenPropertyNodeID  AddedSubObjectID(*AddedSubObject);
+								const FOverriddenPropertyNodeID  AddedSubObjectID(AddedSubObject);
 								FOverriddenPropertyNode& AddedSubObjectNode = FindOrAddNode(*SubPropertyNode, AddedSubObjectID);
 								AddedSubObjectNode.Operation = EOverriddenPropertyOperation::Add;
 
 								// Notify the subobject that it was added
-								if (FOverriddenPropertySet* AddedSubObjectOverriddenProperties = OverridableManager.GetOverriddenProperties(*AddedSubObject))
+								if (FOverriddenPropertySet* AddedSubObjectOverriddenProperties = OverridableManager.GetOverriddenProperties(AddedSubObject))
 								{
 									AddedSubObjectOverriddenProperties->bWasAdded = true;
 								}
@@ -882,7 +882,7 @@ void FOverriddenPropertySet::NotifyPropertyChange(FOverriddenPropertyNode* Paren
 								UObject* RemovedSubObjectArchetype = RemovedSubObject->GetArchetype();
 								if (RemovedSubObjectArchetype && !RemovedSubObjectArchetype->HasAnyFlags(RF_ClassDefaultObject))
 								{
-									const FOverriddenPropertyNodeID RemovedSubObjectID (*RemovedSubObjectArchetype);
+									const FOverriddenPropertyNodeID RemovedSubObjectID (RemovedSubObjectArchetype);
 									FOverriddenPropertyNode& RemovedSubObjectNode = FindOrAddNode(*SubPropertyNode, RemovedSubObjectID);
 									if (RemovedSubObjectNode.Operation == EOverriddenPropertyOperation::Add)
 									{
@@ -902,7 +902,7 @@ void FOverriddenPropertySet::NotifyPropertyChange(FOverriddenPropertyNode* Paren
 								else
 								{
 									// Figure out if it is a remove of a previously added element
-									const FOverriddenPropertyNodeID RemovedSubObjectID (*RemovedSubObject);
+									const FOverriddenPropertyNodeID RemovedSubObjectID (RemovedSubObject);
 									if (const FOverriddenPropertyNodeID* AddedSubObjectID = SubPropertyNode->SubPropertyNodeKeys.Find(RemovedSubObjectID))
 									{
 										FOverriddenPropertyNode* AddedSubObjectNode = OverriddenPropertyNodes.Find(*AddedSubObjectID);
@@ -1003,7 +1003,7 @@ void FOverriddenPropertySet::NotifyPropertyChange(FOverriddenPropertyNode* Paren
 					if (UObject* SubObject = InnerObjectProperty->GetObjectPropertyValue(ArrayHelper.GetElementPtr(ArrayIndex)))
 					{
 						// This should not be needed in the property grid, as it should already been called on the subobject itself.
-						OverridableManager.NotifyPropertyChange(Notification, *SubObject, NextPropertyIterator, ChangeType);
+						OverridableManager.NotifyPropertyChange(Notification, SubObject, NextPropertyIterator, ChangeType);
 						return;
 					}
 				}
@@ -1098,11 +1098,11 @@ void FOverriddenPropertySet::NotifyPropertyChange(FOverriddenPropertyNode* Paren
 					// @todo support instanced object as a key in maps
 					//if (UObject* KeySubObject = KeyInstancedObjectProperty ? KeyInstancedObjectProperty->GetObjectPropertyValue(MapHelper.GetKeyPtr(*It)) : nullptr)
 					//{
-					//	OverridableManager.OverrideInstancedSubObject(*Owner, *KeySubObject);
+					//	OverridableManager.OverrideInstancedSubObject(Owner, KeySubObject);
 					//}
 					if (UObject* ValueSubObject = ValueInstancedObjectProperty ? ValueInstancedObjectProperty->GetObjectPropertyValue(MapHelper.GetValuePtr(It.GetInternalIndex())) : nullptr)
 					{
-						OverridableManager.OverrideInstancedSubObject(*Owner, *ValueSubObject);
+						OverridableManager.OverrideInstancedSubObject(Owner, ValueSubObject);
 					}
 				}
 			};
@@ -1232,7 +1232,7 @@ void FOverriddenPropertySet::NotifyPropertyChange(FOverriddenPropertyNode* Paren
 			if (UObject* SubObject = ValueInstancedObjectProperty ? ValueInstancedObjectProperty->GetObjectPropertyValue(MapHelper.GetValuePtr(InternalMapIndex)) : nullptr)
 			{
 				// This should not be needed in the property grid, as it should already been called on the subobject.
-				OverridableManager.NotifyPropertyChange(Notification, *SubObject, NextPropertyIterator, ChangeType);
+				OverridableManager.NotifyPropertyChange(Notification, SubObject, NextPropertyIterator, ChangeType);
 				return;
 			}
 		}
@@ -1264,7 +1264,7 @@ void FOverriddenPropertySet::NotifyPropertyChange(FOverriddenPropertyNode* Paren
 		else if (UObject* SubObject = ObjectProperty->GetObjectPropertyValue(SubValuePtr))
 		{
 			// This should not be needed in the property grid, as it should already been called on the subobject.
-			OverridableManager.NotifyPropertyChange(Notification, *SubObject, NextPropertyIterator, ChangeType);
+			OverridableManager.NotifyPropertyChange(Notification, SubObject, NextPropertyIterator, ChangeType);
 		}
 		return;
 	}
@@ -1475,13 +1475,13 @@ FOverriddenPropertyNode* FOverriddenPropertySet::SetSubPropertyOperation(EOverri
 
 FOverriddenPropertyNode* FOverriddenPropertySet::SetSubObjectOperation(EOverriddenPropertyOperation Operation, FOverriddenPropertyNode& Node, TNotNull<UObject*> SubObject)
 {
-	const FOverriddenPropertyNodeID SubObjectID(*SubObject);
+	const FOverriddenPropertyNodeID SubObjectID(SubObject);
 	FOverriddenPropertyNode* SubObjectNode = SetSubPropertyOperation(Operation, Node, SubObjectID);
 
 	if (Operation == EOverriddenPropertyOperation::Add)
 	{
 		// Notify the subobject that it was added
-		if (FOverriddenPropertySet* AddedSubObjectOverriddenProperties = FOverridableManager::Get().GetOverriddenProperties(*SubObject))
+		if (FOverriddenPropertySet* AddedSubObjectOverriddenProperties = FOverridableManager::Get().GetOverriddenProperties(SubObject))
 		{
 			AddedSubObjectOverriddenProperties->bWasAdded = true;
 		}
diff --git a/Engine/Source/Runtime/CoreUObject/Private/UObject/PropertyArray.cpp b/Engine/Source/Runtime/CoreUObject/Private/UObject/PropertyArray.cpp
--- a/Engine/Source/Runtime/CoreUObject/Private/UObject/PropertyArray.cpp
+++ b/Engine/Source/Runtime/CoreUObject/Private/UObject/PropertyArray.cpp
@@ -458,7 +458,7 @@ void FArrayProperty::SerializeItem(FStructuredArchive::FSlot Slot, void* Value,
 							{
 								if (UObject* CurrentObject = InnerObjectProperty->GetObjectPropertyValue(ArrayHelper.GetElementPtr(i)); CurrentObject && (!bFilterExternalObjects || !CurrentObject->IsPackageExternal()))
 								{
-									if (ObjectToFind == FOverriddenPropertyNodeID(*CurrentObject))
+									if (ObjectToFind == FOverriddenPropertyNodeID(CurrentObject))
 									{
 										return i;
 									}
diff --git a/Engine/Source/Runtime/CoreUObject/Private/UObject/SavePackage/PackageHarvester.cpp b/Engine/Source/Runtime/CoreUObject/Private/UObject/SavePackage/PackageHarvester.cpp
--- a/Engine/Source/Runtime/CoreUObject/Private/UObject/SavePackage/PackageHarvester.cpp
+++ b/Engine/Source/Runtime/CoreUObject/Private/UObject/SavePackage/PackageHarvester.cpp
@@ -696,7 +696,7 @@ void FPackageHarvester::TryHarvestExportInternal(UObject* InObject, UE::SavePack
 
 	// Objects that has overridable serialization enabled, needs to separately export their subobject independently of the property serializing it.
 	// This is because it needs to make a difference between the pointer that is overridden vs some properties inside the instanced sub object that are overridden.
-	if (FOverridableManager::Get().IsEnabled(*InObject))
+	if (FOverridableManager::Get().IsEnabled(InObject))
 	{
 		for (TPropertyValueIterator<FObjectProperty> ObjPropIt(InObject->GetClass(), InObject); ObjPropIt; ++ObjPropIt)
 		{
diff --git a/Engine/Source/Runtime/CoreUObject/Private/UObject/UObjectGlobals.cpp b/Engine/Source/Runtime/CoreUObject/Private/UObject/UObjectGlobals.cpp
--- a/Engine/Source/Runtime/CoreUObject/Private/UObject/UObjectGlobals.cpp
+++ b/Engine/Source/Runtime/CoreUObject/Private/UObject/UObjectGlobals.cpp
@@ -3927,7 +3927,7 @@ void UObject::PostInitProperties()
 	GetClass()->CreatePersistentUberGraphFrame(this, true);
 	if (FOverridableManager* Manager = FOverridableManager::TryGet())
 	{
-		Manager->ClearOverrides(*this);
+		Manager->ClearOverrides(this);
 	}
 }
 
@@ -4500,7 +4500,7 @@ void FObjectInitializer::InitProperties(UObject* Obj, UClass* DefaultsClass, UOb
 
 	if (FOverridableManager* Manager = FOverridableManager::TryGet())
 	{
-		Manager->InheritEnabledFrom(*Obj, DefaultData);
+		Manager->InheritEnabledFrom(Obj, DefaultData);
 	}
 
 	UClass* Class = Obj->GetClass();
diff --git a/Engine/Source/Runtime/CoreUObject/Public/UObject/OverridableManager.h b/Engine/Source/Runtime/CoreUObject/Public/UObject/OverridableManager.h
--- a/Engine/Source/Runtime/CoreUObject/Public/UObject/OverridableManager.h
+++ b/Engine/Source/Runtime/CoreUObject/Public/UObject/OverridableManager.h
@@ -38,31 +38,31 @@ struct FOverriddenPropertyAnnotation
 class FOverriddenPropertyAnnotations : public FUObjectAnnotationSparse<FOverriddenPropertyAnnotation, true/*bAutoRemove*/>
 {
 public:
-	bool IsEnabled(const UObject& Object)
+	bool IsEnabled(TNotNull<const UObject*> Object)
 	{
-		return !GetAnnotation(&Object).IsDefault();
+		return !GetAnnotation(Object).IsDefault();
 	}
 
-	FOverriddenPropertySet* Find(const UObject& Object)
+	FOverriddenPropertySet* Find(TNotNull<const UObject*> Object)
 	{
-		FOverriddenPropertyAnnotation Annotation = GetAnnotation(&Object);
+		FOverriddenPropertyAnnotation Annotation = GetAnnotation(Object);
 		return Annotation.OverriddenProperties.Get();
 	}
 
-	FOverriddenPropertySet& FindChecked(const UObject& Object)
+	FOverriddenPropertySet& FindChecked(TNotNull<const UObject*> Object)
 	{
-		FOverriddenPropertyAnnotation Annotation = GetAnnotation(&Object);
+		FOverriddenPropertyAnnotation Annotation = GetAnnotation(Object);
 		checkf(!Annotation.IsDefault(), TEXT("Caller is expecting the object to have overridable serialization enabled"));
 		return *Annotation.OverriddenProperties.Get();
 	}
 
-	FOverriddenPropertySet& FindOrAdd(UObject& Object)
+	FOverriddenPropertySet& FindOrAdd(TNotNull<UObject*> Object)
 	{
-		FOverriddenPropertyAnnotation Annotation = GetAnnotation(&Object);
+		FOverriddenPropertyAnnotation Annotation = GetAnnotation(Object);
 		if (Annotation.IsDefault())
 		{
 			Annotation.OverriddenProperties = MakeShared<FOverriddenPropertySet>(Object);
-			AddAnnotation(&Object, Annotation);
+			AddAnnotation(Object, Annotation);
 		}
 
 		return *Annotation.OverriddenProperties;
@@ -115,48 +115,57 @@ public:
 	 * Lookup if the for the specified object has overridable serialization enabled
 	 * @param Object to check against
 	 * @return true if it uses the overridable serialization */
-	 COREUOBJECT_API bool IsEnabled(const UObject& Object);
+	 COREUOBJECT_API bool IsEnabled(TNotNull<const UObject*> Object);
 
 	/**
 	 * Sets on the specified object to use overridable serialization
 	 * @param Object to enable */
-	COREUOBJECT_API void Enable(UObject& Object);
+	COREUOBJECT_API void Enable(TNotNull<UObject*> Object);
 
 	/**
 	 * Sets on the specified object to not use overridable serialization
 	 * @param Object to disable
 	 * @param bPropagateToSubObjects true if wants to propagate the disable to all subobjects */
-	COREUOBJECT_API void Disable(UObject& Object, bool bPropagateToSubObjects = false);
+	COREUOBJECT_API void Disable(TNotNull<UObject*> Object, bool bPropagateToSubObjects = false);
 
 	/**
 	 * Sets on the specified object to not use overridable serialization
 	 * Sets the specified instanced subobject to not use overridable serialization if it is really owned by the referencer
 	 * @param Object referencing this subobject
 	 * @param InstancedSubObject the sub object to override all its properties */
-	COREUOBJECT_API void DisableInstancedSubObject(const UObject& Object, UObject& InstancedSubObject);
+	COREUOBJECT_API void DisableInstancedSubObject(TNotNull<const UObject*> Object, TNotNull<UObject*> InstancedSubObject);
 
 	/**
 	 * Inherit if the specified object should enable overridable serialization. It inherits it from either its default object or its outer.
 	 * @param Object to be inherited on
 	 * @param DefaultData its default object */
-	COREUOBJECT_API void InheritEnabledFrom(UObject& Object, const UObject* DefaultData);
+	COREUOBJECT_API void InheritEnabledFrom(TNotNull<UObject*> Object, const UObject* DefaultData);
 
 	/**
 	 * Return true if this object needs subobject template instantiation
 	 * @param Object to be querying about */
-	COREUOBJECT_API bool NeedSubObjectTemplateInstantiation(const UObject& Object);
+	COREUOBJECT_API bool NeedSubObjectTemplateInstantiation(TNotNull<const UObject*> Object);
 
 	/**
 	 * Retrieve the overridden properties for the specified object
 	 * @param Object to fetch the overridden properties
 	 * @return the overridden properties if the object have overridable serialization enabled */
-	COREUOBJECT_API FOverriddenPropertySet* GetOverriddenProperties(UObject& Object);
-
-	/**
-	 * Retrieve the overridden properties for the specified object
-	 * @param Object to fetch the overridden properties
-	 * @return the overridden properties if the object have overridable serialization enabled */
-	COREUOBJECT_API const FOverriddenPropertySet* GetOverriddenProperties(const UObject& Object);
+	template<
+		typename InputType
+		UE_REQUIRES(std::is_convertible_v<InputType, TNotNull<const UObject*>> || std::is_convertible_v<InputType, TNotNull<UObject*>>)
+	>
+	auto GetOverriddenProperties(InputType&& Object)
+	{
+		// Use a template because the compiler cannot discern which method to use if UObject* is passed (both TNotNulls accept UObject*).
+		if constexpr (std::is_convertible_v<InputType, TNotNull<UObject*>>)
+		{
+			return GetOverriddenPropertiesInternal(TNotNull<UObject*>(Object));
+		}
+		else
+		{
+			return GetOverriddenPropertiesInternal(TNotNull<const UObject*>(Object));
+		}
+	}
 
 	/**
 	 * Restore the override operation from a saved state on this object, and it will enable it if it wasn't already enabled
@@ -178,53 +187,53 @@ public:
 	 * Retrieve the overridden state for the specified object
 	 * @param Object to fetch the overridden properties
 	 * @return the overridden state if the object has overridable serialization enabled */
-	COREUOBJECT_API EOverriddenState GetOverriddenState(UObject& Object);
+	COREUOBJECT_API EOverriddenState GetOverriddenState(TNotNull<UObject*> Object);
 
 	/**
 	 * Override the entire object properties and all its instanced subobjects
 	 * @param Object to override all it properties */
-	COREUOBJECT_API void OverrideObject(UObject& Object);
+	COREUOBJECT_API void OverrideObject(TNotNull<UObject*> Object);
 
 	/**
 	 * Override all the properties of the specified instanced subobject if it is really owned by the referencer
 	 * @param Object referencing this subobject
 	 * @param InstancedSubObject the sub object to override all its properties */
-	COREUOBJECT_API void OverrideInstancedSubObject(const UObject& Object, UObject& InstancedSubObject);
+	COREUOBJECT_API void OverrideInstancedSubObject(TNotNull<const UObject*> Object, TNotNull<UObject*> InstancedSubObject);
 
 	/**
 	 * Propagate override to all instanced sub object of the specified object
 	 * @param Object to propagate from */
-	COREUOBJECT_API void PropagateOverrideToInstancedSubObjects(UObject& Object);
+	COREUOBJECT_API void PropagateOverrideToInstancedSubObjects(TNotNull<UObject*> Object);
 
 	/**
 	 * Clears all the overrides on the specified object
 	 * @param Object to clear them on */
-	COREUOBJECT_API void ClearOverrides(UObject& Object);
+	COREUOBJECT_API void ClearOverrides(TNotNull<UObject*> Object);
 
 	/**
 	 * Clear all overrides of the specified instanced subobject if it is really owned by the referencer
 	 * @param Object referencing the subobject
 	 * @param InstancedSubObject to clear overrides on */
-	COREUOBJECT_API void ClearInstancedSubObjectOverrides(const UObject& Object, UObject& InstancedSubObject);
+	COREUOBJECT_API void ClearInstancedSubObjectOverrides(TNotNull<const UObject*> Object, TNotNull<UObject*> InstancedSubObject);
 
 	/**
 	 * Propagate the clear overrides to all instanced suboject of the specified object
 	 * @param Object*/
-	COREUOBJECT_API void PropagateClearOverridesToInstancedSubObjects(UObject& Object);
+	COREUOBJECT_API void PropagateClearOverridesToInstancedSubObjects(TNotNull<UObject*> Object);
 
 	/**
 	 * Override a specific property of an object (Helper methods to call Pre/PostOverride)
 	 * Note: Supports object that does not have overridable serialization enabled
 	 * @param Object owning the property
 	 * @param PropertyPath leading to the property that is about to be overridden */
-	COREUOBJECT_API void OverrideProperty(UObject& Object, const FPropertyVisitorPath& PropertyPath);
+	COREUOBJECT_API void OverrideProperty(TNotNull<UObject*> Object, const FPropertyVisitorPath& PropertyPath);
 
 	/**
 	 * Clears an overridden properties specified by the property chain
 	 * @param Object owning the property to clear
 	 * @param PropertyPath to the property to clear from the root of the specified object
 	 * @return true if the property was successfully cleared. */
-	bool ClearOverriddenProperty(UObject& Object, const FPropertyVisitorPath& PropertyPath)
+	bool ClearOverriddenProperty(TNotNull<UObject*> Object, const FPropertyVisitorPath& PropertyPath)
 	{
 		return ClearOverriddenProperty(Object, PropertyPath.GetRootIterator());
 	}
@@ -234,7 +243,7 @@ public:
 	 * Note: Supports object that does not have overridable serialization enabled
 	 * @param Object owning the property
 	 * @param PropertyPath leading to the property about to be overridden */
-	COREUOBJECT_API void PreOverrideProperty(UObject& Object, const FPropertyVisitorPath& PropertyPath);
+	COREUOBJECT_API void PreOverrideProperty(TNotNull<UObject*> Object, const FPropertyVisitorPath& PropertyPath);
 
 	/**
 	 * To be called after the property of the specified object was overridden
@@ -242,7 +251,7 @@ public:
 	 * @param Object owning the property
 	 * @param PropertyPath leading to the property that was overridden
  	 * @param ChangeType of the current operation */
-	COREUOBJECT_API void PostOverrideProperty(UObject& Object, const FPropertyVisitorPath& PropertyPath, const EPropertyChangeType::Type ChangeType);
+	COREUOBJECT_API void PostOverrideProperty(TNotNull<UObject*> Object, const FPropertyVisitorPath& PropertyPath, const EPropertyChangeType::Type ChangeType);
 
 	/**
 	 * Retrieve the overridable operation from the specified property path
@@ -250,7 +259,7 @@ public:
 	 * @param PropertyPath leading to the property the caller is interested in
 	 * @param bOutInheritedOperation optional parameter to know if the operation returned was inherited from a parent property
 	 * @return the current type of override operation on the property */
-	EOverriddenPropertyOperation GetOverriddenPropertyOperation(UObject& Object, const FPropertyVisitorPath& PropertyPath, bool* bOutInheritedOperation = nullptr)
+	EOverriddenPropertyOperation GetOverriddenPropertyOperation(TNotNull<UObject*> Object, const FPropertyVisitorPath& PropertyPath, bool* bOutInheritedOperation = nullptr)
 	{
 		return GetOverriddenPropertyOperation(Object, PropertyPath.GetRootIterator(), bOutInheritedOperation);
 	}
@@ -259,7 +268,7 @@ public:
 	 * Serializes the overridden properties of the specified object into the record
 	 * @param Object to serialize the overridden property
 	 * @param ObjectRecord the record to use for serialization */
-	COREUOBJECT_API void SerializeOverriddenProperties(UObject& Object, FStructuredArchive::FRecord ObjectRecord);
+	COREUOBJECT_API void SerializeOverriddenProperties(TNotNull<UObject*> Object, FStructuredArchive::FRecord ObjectRecord);
 
 	/**
 	 * Override a specific property of an object (Helper methods to call Pre/PostOverride)
@@ -267,7 +276,7 @@ public:
 	 * @param Object owning the property
 	 * @param PropertyEvent information about the type of change including any container item index
 	 * @param PropertyChain leading to the property that is about to be overridden */
-	void OverrideProperty(UObject& Object, const FPropertyChangedEvent& PropertyEvent, const FEditPropertyChain& PropertyChain)
+	void OverrideProperty(TNotNull<UObject*> Object, const FPropertyChangedEvent& PropertyEvent, const FEditPropertyChain& PropertyChain)
 	{
 		OverrideProperty(Object, FPropertyVisitorPath(PropertyEvent, PropertyChain));
 	}
@@ -278,7 +287,7 @@ public:
 	 * @param PropertyEvent only needed to know about the container item index in any
 	 * @param PropertyChain to the property to clear from the root of the specified object
 	 * @return true if the property was successfully cleared. */
-	bool ClearOverriddenProperty(UObject& Object, const FPropertyChangedEvent& PropertyEvent, const FEditPropertyChain& PropertyChain)
+	bool ClearOverriddenProperty(TNotNull<UObject*> Object, const FPropertyChangedEvent& PropertyEvent, const FEditPropertyChain& PropertyChain)
 	{
 		return ClearOverriddenProperty(Object, FPropertyVisitorPath(PropertyEvent, PropertyChain));
 	}
@@ -288,7 +297,7 @@ public:
 	 * Note: Supports object that does not have overridable serialization enabled
 	 * @param Object owning the property
 	 * @param PropertyChain leading to the property about to be overridden */
-	void PreOverrideProperty(UObject& Object, const FEditPropertyChain& PropertyChain)
+	void PreOverrideProperty(TNotNull<UObject*> Object, const FEditPropertyChain& PropertyChain)
 	{
 		PreOverrideProperty(Object, FPropertyVisitorPath(FPropertyChangedEvent(nullptr), PropertyChain));
 	}
@@ -299,7 +308,7 @@ public:
 	 * @param Object owning the property
 	 * @param PropertyEvent information about the type of change including any container item index
 	 * @param PropertyChain leading to the property that was overridden */
-	void PostOverrideProperty(UObject& Object, const FPropertyChangedEvent& PropertyEvent, const FEditPropertyChain& PropertyChain)
+	void PostOverrideProperty(TNotNull<UObject*> Object, const FPropertyChangedEvent& PropertyEvent, const FEditPropertyChain& PropertyChain)
 	{
 		PostOverrideProperty(Object, FPropertyVisitorPath(PropertyEvent, PropertyChain), PropertyEvent.ChangeType);
 	}
@@ -311,7 +320,7 @@ public:
 	 * @param PropertyChain leading to the property the caller is interested in
 	 * @param bOutInheritedOperation optional parameter to know if the operation returned was inherited from a parent property
 	 * @return the current type of override operation on the property */
-	FORCEINLINE EOverriddenPropertyOperation GetOverriddenPropertyOperation(UObject& Object, const FPropertyChangedEvent& PropertyEvent, const FEditPropertyChain& PropertyChain, bool* bOutInheritedOperation = nullptr)
+	FORCEINLINE EOverriddenPropertyOperation GetOverriddenPropertyOperation(TNotNull<UObject*> Object, const FPropertyChangedEvent& PropertyEvent, const FEditPropertyChain& PropertyChain, bool* bOutInheritedOperation = nullptr)
 	{
 		return GetOverriddenPropertyOperation(Object, FPropertyVisitorPath(PropertyEvent, PropertyChain), bOutInheritedOperation);
 	}
@@ -341,8 +350,10 @@ protected:
 #endif // WITH_EDITORONLY_DATA
 
 	friend FOverriddenPropertySet;
-	COREUOBJECT_API void NotifyPropertyChange(const EPropertyNotificationType Notification, UObject& Object, FPropertyVisitorPath::Iterator PropertyIterator, const EPropertyChangeType::Type ChangeType);
-	COREUOBJECT_API EOverriddenPropertyOperation GetOverriddenPropertyOperation(UObject& Object, FPropertyVisitorPath::Iterator PropertyIterator, bool* bOutInheritedOperation = nullptr);
-	COREUOBJECT_API bool ClearOverriddenProperty(UObject& Object, FPropertyVisitorPath::Iterator PropertyIterator);
+	COREUOBJECT_API void NotifyPropertyChange(const EPropertyNotificationType Notification, TNotNull<UObject*> Object, FPropertyVisitorPath::Iterator PropertyIterator, const EPropertyChangeType::Type ChangeType);
+	COREUOBJECT_API EOverriddenPropertyOperation GetOverriddenPropertyOperation(TNotNull<UObject*> Object, FPropertyVisitorPath::Iterator PropertyIterator, bool* bOutInheritedOperation = nullptr);
+	COREUOBJECT_API bool ClearOverriddenProperty(TNotNull<UObject*> Object, FPropertyVisitorPath::Iterator PropertyIterator);
+	COREUOBJECT_API FOverriddenPropertySet* GetOverriddenPropertiesInternal(TNotNull<UObject*> Object);
+	COREUOBJECT_API const FOverriddenPropertySet* GetOverriddenPropertiesInternal(TNotNull<const UObject*> Object);
 };
 
diff --git a/Engine/Source/Runtime/CoreUObject/Public/UObject/OverriddenPropertySet.h b/Engine/Source/Runtime/CoreUObject/Public/UObject/OverriddenPropertySet.h
--- a/Engine/Source/Runtime/CoreUObject/Public/UObject/OverriddenPropertySet.h
+++ b/Engine/Source/Runtime/CoreUObject/Public/UObject/OverriddenPropertySet.h
@@ -202,12 +202,12 @@ struct FOverriddenPropertyNodeID
 
 	FOverriddenPropertyNodeID(const FProperty* Property = nullptr);
 
-	FOverriddenPropertyNodeID(const UObject& InObject)
-		: Object(&InObject)
+	FOverriddenPropertyNodeID(TNotNull<const UObject*> InObject)
+		: Object(InObject)
 	{
 		// Note: Using ObjectIndex by itself is not sufficient for an enduring unique identifier
 		// as re-instantiation can cause a reuse of the index for another object. Appending the serial solves this issue
-		const int32 ObjectIndex = GUObjectArray.ObjectToIndex(&InObject);
+		const int32 ObjectIndex = GUObjectArray.ObjectToIndex(InObject);
 		Path = *(FString::Printf(TEXT("%d%d"), ObjectIndex, GUObjectArray.AllocateSerialNumber(ObjectIndex)));
 	}
 	FOverriddenPropertyNodeID(const FOverriddenPropertyNodeID& ParentNodeID, const FOverriddenPropertyNodeID& SubNodeID);
@@ -319,8 +319,8 @@ struct COREUOBJECT_API FOverriddenPropertySet
 public:
 
 	FOverriddenPropertySet() = default;
-	FOverriddenPropertySet(UObject& InOwner)
-		: Owner(&InOwner)
+	FOverriddenPropertySet(TNotNull<UObject*> InOwner)
+		: Owner(InOwner)
 	{}
 
 	/**
diff --git a/Engine/Source/Runtime/Engine/Private/UnrealEngine.cpp b/Engine/Source/Runtime/Engine/Private/UnrealEngine.cpp
--- a/Engine/Source/Runtime/Engine/Private/UnrealEngine.cpp
+++ b/Engine/Source/Runtime/Engine/Private/UnrealEngine.cpp
@@ -17880,7 +17880,7 @@ void UEngine::CopyPropertiesForUnrelatedObjects(UObject* OldObject, UObject* New
 	// may be missing, but testing will determine that:
 	const auto CollectAllSubobjects = [&Params](UObject* Object, TArray<UObject*>& OutSubobjectArray)
 	{
-		auto ShouldExcludeSubobjectFromCopy = FOverridableManager::Get().IsEnabled(*Object) ?
+		auto ShouldExcludeSubobjectFromCopy = FOverridableManager::Get().IsEnabled(Object) ?
 			[](UObject* SubObject)
 			{
 				// Here we are trying to prevent an extra copy of COD direct subobjects. The problem is that IsDefaultSubobject() 
diff --git a/Engine/Source/Runtime/Engine/Private/UnrealExporter.cpp b/Engine/Source/Runtime/Engine/Private/UnrealExporter.cpp
--- a/Engine/Source/Runtime/Engine/Private/UnrealExporter.cpp
+++ b/Engine/Source/Runtime/Engine/Private/UnrealExporter.cpp
@@ -34,7 +34,7 @@ DEFINE_LOG_CATEGORY_STATIC(LogExporter, Log, All);
 
 namespace UnrealExporterPrivate
 {
-	FOverriddenPropertySet* GetOverriddenProperties(UObject& Object)
+	FOverriddenPropertySet* GetOverriddenProperties(TNotNull<UObject*> Object)
 	{
 		if (FOverridableSerializationLogic::HasCapabilities(FOverridableSerializationLogic::ECapabilities::T3DSerialization))
 		{
@@ -488,7 +488,7 @@ void UExporter::EmitBeginObject( FOutputDevice& Ar, UObject* Obj, uint32 PortFla
 			Ar.Logf(TEXT(" Archetype=%s"), *FObjectPropertyBase::GetExportPath(Archetype, nullptr, /*ExportRootScope =*/nullptr, (PortFlags | PPF_Delimited) & ~PPF_ExportsNotFullyQualified));
 		}
 
-		if (FOverriddenPropertySet* ObjectOverriddenProperties = UnrealExporterPrivate::GetOverriddenProperties(*Obj))
+		if (FOverriddenPropertySet* ObjectOverriddenProperties = UnrealExporterPrivate::GetOverriddenProperties(Obj))
 		{
 			const EOverriddenPropertyOperation Operation = ObjectOverriddenProperties->GetOverriddenPropertyOperation((FArchiveSerializedPropertyChain*)nullptr, (FProperty*)nullptr);
 			Ar.Logf(TEXT(" OverriddenOperation=%s"), *GetOverriddenOperationString(Operation));
@@ -680,8 +680,8 @@ void ExportProperties_Overrides
 {
 	check(ObjectClass != NULL);
 
-	FOverriddenPropertySet* OverriddenProperties = Object ? UnrealExporterPrivate::GetOverriddenProperties(*(UObject*)Object) : nullptr;
-	FEnableOverridableSerializationScope Scope(OverriddenProperties!= nullptr, OverriddenProperties);
+	FOverriddenPropertySet* OverriddenProperties = Object ? UnrealExporterPrivate::GetOverriddenProperties((UObject*)Object) : nullptr;
+	FEnableOverridableSerializationScope Scope(OverriddenProperties != nullptr, OverriddenProperties);
 
 	for (FProperty* Property = ObjectClass->PropertyLink; Property; Property = Property->PropertyLinkNext)
 	{
@@ -849,7 +849,7 @@ void ExportProperties_Overrides
 							{
 								if (UObject* CurrentObject = InnerObjectProperty->GetObjectPropertyValue(ArrayHelper.GetElementPtr(i)))
 								{
-									if (ObjectToFind == FOverriddenPropertyNodeID(*CurrentObject))
+									if (ObjectToFind == FOverriddenPropertyNodeID(CurrentObject))
 									{
 										return i;
 									}
