Remap meta data object keys when renaming a package.

diff --git a/Engine/Source/Runtime/CoreUObject/Private/UObject/MetaData.cpp b/Engine/Source/Runtime/CoreUObject/Private/UObject/MetaData.cpp
--- a/Engine/Source/Runtime/CoreUObject/Private/UObject/MetaData.cpp
+++ b/Engine/Source/Runtime/CoreUObject/Private/UObject/MetaData.cpp
@@ -278,6 +278,14 @@ void UDEPRECATED_MetaData::InitializeRedirectMap()
 
 TMap<FName, FName> FMetaData::KeyRedirectMap;
 
+void FMetaData::RemapObjectKeys(FName OldPackageName, FName NewPackageName)
+{
+	for (TMap<FSoftObjectPath, TMap<FName, FString>>::TIterator It(ObjectMetaDataMap); It; ++It)
+	{
+		It.Key().RemapPackage(OldPackageName, NewPackageName);
+	}
+}
+
 /**
  * Return the value for the given key in the given property
  * @param Object the object to lookup the metadata for
diff --git a/Engine/Source/Runtime/CoreUObject/Private/UObject/Package.cpp b/Engine/Source/Runtime/CoreUObject/Private/UObject/Package.cpp
--- a/Engine/Source/Runtime/CoreUObject/Private/UObject/Package.cpp
+++ b/Engine/Source/Runtime/CoreUObject/Private/UObject/Package.cpp
@@ -378,6 +378,33 @@ bool UPackage::IsPostLoadThreadSafe() const
 	return true;
 }
 
+bool UPackage::Rename(const TCHAR* InName, UObject* NewOuter, ERenameFlags Flags)
+{
+#if WITH_METADATA
+	const FName OldPackageName = GetFName();
+#endif
+
+	if (!Super::Rename(InName, NewOuter, Flags))
+	{
+		return false;
+	}
+
+	if (Flags & REN_Test)
+	{
+		return true;
+	}
+
+#if WITH_METADATA
+	const FName NewPackageName = GetFName();
+	if (OldPackageName != NewPackageName)
+	{
+		MetaData.RemapObjectKeys(OldPackageName, NewPackageName);
+	}
+#endif
+
+	return true;
+}
+
 #if WITH_EDITORONLY_DATA
 
 FIoHash UPackage::GetSavedHash()
diff --git a/Engine/Source/Runtime/CoreUObject/Private/UObject/SoftObjectPath.cpp b/Engine/Source/Runtime/CoreUObject/Private/UObject/SoftObjectPath.cpp
--- a/Engine/Source/Runtime/CoreUObject/Private/UObject/SoftObjectPath.cpp
+++ b/Engine/Source/Runtime/CoreUObject/Private/UObject/SoftObjectPath.cpp
@@ -631,6 +631,11 @@ bool FSoftObjectPath::SerializeFromMismatchedTag(struct FPropertyTag const& Tag,
 	return bReturn;
 }
 
+bool FSoftObjectPath::RemapPackage(FName OldPackageName, FName NewPackageName)
+{
+	return AssetPath.RemapPackage(OldPackageName, NewPackageName);
+}
+
 UObject* FSoftObjectPath::TryLoad(FUObjectSerializeContext* InLoadContext) const
 {
 	UObject* LoadedObject = nullptr;
diff --git a/Engine/Source/Runtime/CoreUObject/Private/UObject/TopLevelAssetPath.cpp b/Engine/Source/Runtime/CoreUObject/Private/UObject/TopLevelAssetPath.cpp
--- a/Engine/Source/Runtime/CoreUObject/Private/UObject/TopLevelAssetPath.cpp
+++ b/Engine/Source/Runtime/CoreUObject/Private/UObject/TopLevelAssetPath.cpp
@@ -77,6 +77,16 @@ void FTopLevelAssetPath::ToString(FString& OutString) const
 	AppendString(OutString);
 }
 
+bool FTopLevelAssetPath::RemapPackage(FName OldPackageName, FName NewPackageName)
+{
+	if (PackageName == OldPackageName)
+	{
+		PackageName = NewPackageName;
+		return true;
+	}
+	return false;
+}
+
 bool FTopLevelAssetPath::TrySetPath(FName InPackageName, FName InAssetName)
 {
 	PackageName = InPackageName;
diff --git a/Engine/Source/Runtime/CoreUObject/Public/UObject/MetaData.h b/Engine/Source/Runtime/CoreUObject/Public/UObject/MetaData.h
--- a/Engine/Source/Runtime/CoreUObject/Public/UObject/MetaData.h
+++ b/Engine/Source/Runtime/CoreUObject/Public/UObject/MetaData.h
@@ -73,6 +73,9 @@ public:
 	 */
 	TMap< FName, FString > RootMetaDataMap;
 
+	/** Remap all object keys when renaming the owner package. */
+	void RemapObjectKeys(FName OldPackageName, FName NewPackageName);
+
 public:
 	// MetaData utility functions
 
diff --git a/Engine/Source/Runtime/CoreUObject/Public/UObject/Package.h b/Engine/Source/Runtime/CoreUObject/Public/UObject/Package.h
--- a/Engine/Source/Runtime/CoreUObject/Public/UObject/Package.h
+++ b/Engine/Source/Runtime/CoreUObject/Public/UObject/Package.h
@@ -343,6 +343,7 @@ public:
 	virtual bool NeedsLoadForClient() const override { return true; }
 	virtual bool NeedsLoadForServer() const override { return true; }
 	virtual bool IsPostLoadThreadSafe() const override;
+	virtual bool Rename(const TCHAR* NewName = nullptr, UObject* NewOuter = nullptr, ERenameFlags Flags = REN_None) override;
 
 #if WITH_EDITORONLY_DATA
 	UE_DEPRECATED(5.5, "No longer used; skiponlyeditoronly is used instead and tracks editoronly references via savepackage results.")
diff --git a/Engine/Source/Runtime/CoreUObject/Public/UObject/SoftObjectPath.h b/Engine/Source/Runtime/CoreUObject/Public/UObject/SoftObjectPath.h
--- a/Engine/Source/Runtime/CoreUObject/Public/UObject/SoftObjectPath.h
+++ b/Engine/Source/Runtime/CoreUObject/Public/UObject/SoftObjectPath.h
@@ -274,6 +274,9 @@ struct FSoftObjectPath
 		SetPath(FStringView(Path));
 	}
 
+	/** Remap the package to handle renames. */
+	COREUOBJECT_API bool RemapPackage(FName OldPackageName, FName NewPackageName);
+
 	/**
 	 * Attempts to load the asset, this will call LoadObject which can be very slow
 	 * @param InLoadContext Optional load context when called from nested load callstack
diff --git a/Engine/Source/Runtime/CoreUObject/Public/UObject/TopLevelAssetPath.h b/Engine/Source/Runtime/CoreUObject/Public/UObject/TopLevelAssetPath.h
--- a/Engine/Source/Runtime/CoreUObject/Public/UObject/TopLevelAssetPath.h
+++ b/Engine/Source/Runtime/CoreUObject/Public/UObject/TopLevelAssetPath.h
@@ -80,6 +80,9 @@ struct FTopLevelAssetPath
 	bool TrySetPath(const CharType* Path) { return TrySetPath(TStringView<CharType>(Path)); }
 	bool TrySetPath(const FString& Path) { return TrySetPath(FStringView(Path)); }
 
+	/** Remap the package to handle renames. */
+	COREUOBJECT_API bool RemapPackage(FName OldPackageName, FName NewPackageName);
+
 	/** Return the package name part e.g. /Path/To/Package as an FName. */
 	FName GetPackageName() const { return PackageName; }
 
