[CoreRedirects]

Treats Verse paths with '.' in the ObjectName as invalid until UE-248198 adds proper support for Verse paths which can contain quoted identifiers which may contain '.'.

#rnx

diff --git a/Engine/Source/Runtime/CoreUObject/Private/UObject/CoreRedirects.cpp b/Engine/Source/Runtime/CoreUObject/Private/UObject/CoreRedirects.cpp
--- a/Engine/Source/Runtime/CoreUObject/Private/UObject/CoreRedirects.cpp
+++ b/Engine/Source/Runtime/CoreUObject/Private/UObject/CoreRedirects.cpp
@@ -871,17 +871,21 @@ FName FCoreRedirectObjectName::GetSearchKey(ECoreRedirectFlags Type) const
 
 bool FCoreRedirectObjectName::HasValidCharacters(ECoreRedirectFlags Type) const
 {
+	static FString InvalidObjectNameCharacters = TEXT(".\n\r\t"); // Object and field names in Blueprints are very permissive
 	if (FPackageName::IsVersePackage(WriteToString<NAME_SIZE>(PackageName)))
 	{
 		// Verse allows the visible character ASCII set minus these CR, LF, TAB chars.
-		static FString InvalidVerseCharacters = TEXT("\n\r\t"); 
-		return  ObjectName.IsValidXName(InvalidVerseCharacters) && 
+		// However CoreRedirects currently does not support names with '.' which verse does allow. 
+		// Until support is added (UE-248198) we reject such names since they can't work correctly.
+		// For redirecting verse names, a verse specific transformation must occur before using FCoreRedirects
+		// which ensures '.' is no longer used as a non-delimiting character.
+		static FString InvalidVerseCharacters = TEXT("\n\r\t");
+		return  ObjectName.IsValidXName(InvalidObjectNameCharacters) &&
 				OuterName.IsValidXName(InvalidVerseCharacters) &&
 				PackageName.IsValidXName(InvalidVerseCharacters);
 	}
 
 	static FString InvalidRedirectNameCharacters = TEXT("\"',|&!~\n\r\t@#(){}[]=;^%$`");
-	static FString InvalidObjectNameCharacters = TEXT(".\n\r\t"); // Object and field names in Blueprints are very permissive
 	const bool bUsePermissiveObjectNames = EnumHasAnyFlags(Type, 
 		ECoreRedirectFlags::Type_Object | ECoreRedirectFlags::Type_Property | ECoreRedirectFlags::Type_Function);
 
