FSoftObjectPath::GetTypeHash optimizations:
- Convert double usage of HashCombine to a single call to HashCombineFast.
- Don't go into FString::GetTypeHash if the sub path striung is empty to avoid a function call.

diff --git a/Engine/Source/Runtime/CoreUObject/Public/UObject/SoftObjectPath.h b/Engine/Source/Runtime/CoreUObject/Public/UObject/SoftObjectPath.h
--- a/Engine/Source/Runtime/CoreUObject/Public/UObject/SoftObjectPath.h
+++ b/Engine/Source/Runtime/CoreUObject/Public/UObject/SoftObjectPath.h
@@ -397,10 +397,13 @@ struct FSoftObjectPath
 
 	[[nodiscard]] FORCEINLINE friend uint32 GetTypeHash(FSoftObjectPath const& This)
 	{
-		uint32 Hash = 0;
+		uint32 Hash = GetTypeHash(This.AssetPath);
+			
+		if (!This.SubPathString.IsEmpty())
+		{
+			Hash = HashCombineFast(Hash, GetTypeHash(This.SubPathString));
+		}
 
-		Hash = HashCombine(Hash, GetTypeHash(This.AssetPath));
-		Hash = HashCombine(Hash, GetTypeHash(This.SubPathString));
 		return Hash;
 	}
 
