Fix two incorrect `PopOnAbortHandler`'s where we didn't actually push.

diff --git a/Engine/Source/Runtime/Core/Public/Stats/LightweightStats.h b/Engine/Source/Runtime/Core/Public/Stats/LightweightStats.h
--- a/Engine/Source/Runtime/Core/Public/Stats/LightweightStats.h
+++ b/Engine/Source/Runtime/Core/Public/Stats/LightweightStats.h
@@ -80,15 +80,14 @@ namespace UE::Stats::Private
 
 		FORCEINLINE ~TScopeCycleCounterBase()
 		{
-			AutoRTFM::PopOnAbortHandler(this);
-
-			UE_AUTORTFM_OPEN
+			if (bPop)
 			{
-				if (bPop)
+				AutoRTFM::PopOnAbortHandler(this);
+				UE_AUTORTFM_OPEN
 				{
 					FPlatformMisc::EndNamedEvent();
-				}
-			};
+				};
+			}
 		}
 
 		FORCEINLINE static bool CanEmitStat(const TStatId& Stat)
diff --git a/Engine/Source/Runtime/CoreUObject/Public/UObject/UObjectBaseUtility.h b/Engine/Source/Runtime/CoreUObject/Public/UObject/UObjectBaseUtility.h
--- a/Engine/Source/Runtime/CoreUObject/Public/UObject/UObjectBaseUtility.h
+++ b/Engine/Source/Runtime/CoreUObject/Public/UObject/UObjectBaseUtility.h
@@ -1080,15 +1080,15 @@ public:
 	FORCEINLINE_STATS ~FScopeCycleCounterUObject()
 	{
 #if ENABLE_STATNAMEDEVENTS_UOBJECT && CPUPROFILERTRACE_ENABLED
-		AutoRTFM::PopOnAbortHandler(this);
-
-		UE_AUTORTFM_OPEN
+		if (bPop)
 		{
-			if (bPop)
+			AutoRTFM::PopOnAbortHandler(this);
+
+			UE_AUTORTFM_OPEN
 			{
 				FCpuProfilerTrace::OutputEndEvent();
-			}
-		};
+			};
+		}
 #endif
 	}
 };
