Add LoadAsync wrappers to TSoftObjectPtr and TSoftClassPtr to better expose it as an alternative to LoadSynchronous when converting loading code to be asynchronous

diff --git a/Engine/Source/Runtime/CoreUObject/Public/UObject/SoftObjectPtr.h b/Engine/Source/Runtime/CoreUObject/Public/UObject/SoftObjectPtr.h
--- a/Engine/Source/Runtime/CoreUObject/Public/UObject/SoftObjectPtr.h
+++ b/Engine/Source/Runtime/CoreUObject/Public/UObject/SoftObjectPtr.h
@@ -473,6 +473,19 @@ public:
 		return Cast<T>(Asset);
 	}
 
+	/**
+	 * Attempts to asynchronously load the object referenced by this soft pointer.
+	 * This is a wrapper around the LoadAsync function in SoftObjectPath, and the delegate is responsible for validating it loaded the correct type
+	 *
+	 * @param	InCompletionDelegate	Delegate to be invoked when the async load finishes, this will execute on the game thread as soon as the load succeeds or fails
+	 * @param	InOptionalParams		Optional parameters for async loading the asset
+	 * @return Unique ID associated with this load request (the same object or package can be associated with multiple IDs).
+	 */
+	int32 LoadAsync(FLoadSoftObjectPathAsyncDelegate InCompletionDelegate, FLoadAssetAsyncOptionalParams InOptionalParams = FLoadAssetAsyncOptionalParams()) const
+	{
+		return SoftObjectPtr.ToSoftObjectPath().LoadAsync(MoveTemp(InCompletionDelegate), MoveTemp(InOptionalParams));
+	}
+
 	/**  
 	 * Test if this points to a live UObject
 	 *
@@ -926,6 +939,19 @@ public:
 		return Class;
 	}
 
+	/**
+	 * Attempts to asynchronously load the object referenced by this soft pointer.
+	 * This is a wrapper around the LoadAsync function in SoftObjectPath, and the delegate is responsible for validating it loaded the correct type
+	 *
+	 * @param	InCompletionDelegate	Delegate to be invoked when the async load finishes, this will execute on the game thread as soon as the load succeeds or fails
+	 * @param	InOptionalParams		Optional parameters for async loading the asset
+	 * @return Unique ID associated with this load request (the same object or package can be associated with multiple IDs).
+	 */
+	int32 LoadAsync(FLoadSoftObjectPathAsyncDelegate InCompletionDelegate, FLoadAssetAsyncOptionalParams InOptionalParams = FLoadAssetAsyncOptionalParams()) const
+	{
+		return SoftObjectPtr.ToSoftObjectPath().LoadAsync(MoveTemp(InCompletionDelegate), MoveTemp(InOptionalParams));
+	}
+
 	inline void Serialize(FArchive& Ar)
 	{
 		Ar << static_cast<FSoftObjectPtr&>(SoftObjectPtr);
