Eliminate needless access tracking when ClassDefaultObject pointer has its alignment checked within IsValidLowLevelFast.

diff --git a/Engine/Source/Runtime/CoreUObject/Private/UObject/UObjectBase.cpp b/Engine/Source/Runtime/CoreUObject/Private/UObject/UObjectBase.cpp
--- a/Engine/Source/Runtime/CoreUObject/Private/UObject/UObjectBase.cpp
+++ b/Engine/Source/Runtime/CoreUObject/Private/UObject/UObjectBase.cpp
@@ -422,7 +422,8 @@ bool UObjectBase::IsValidLowLevelFast(bool bRecursive /*= true*/) const
     return false;
 	}
 	// These should all be non-NULL (except CDO-alignment check which should be 0)
-	if (Class == nullptr || Class->ClassDefaultObject == nullptr || ((UPTRINT)Class->ClassDefaultObject & AlignmentCheck) != 0)
+	const UObject* ClassDefaultObject = ObjectPtr_Private::Friend::NoAccessTrackingGet(Class->ClassDefaultObject);
+	if (Class == nullptr || ClassDefaultObject == nullptr || ((UPTRINT)ClassDefaultObject & AlignmentCheck) != 0)
 	{
 #if WITH_LIVE_CODING
 		// When live coding is re-instancing blueprint generated classes, we have to clear out the default object so it can get 
