CoreUObject: include property flags when computing the hash of a property bag

diff --git a/Engine/Source/Runtime/CoreUObject/Private/StructUtils/PropertyBag.cpp b/Engine/Source/Runtime/CoreUObject/Private/StructUtils/PropertyBag.cpp
--- a/Engine/Source/Runtime/CoreUObject/Private/StructUtils/PropertyBag.cpp
+++ b/Engine/Source/Runtime/CoreUObject/Private/StructUtils/PropertyBag.cpp
@@ -62,9 +62,9 @@ namespace UE::StructUtils::Private
 	uint64 CalcPropertyDescHash(const FPropertyBagPropertyDesc& Desc)
 	{
 #if WITH_EDITORONLY_DATA
-		const uint32 Hashes[] = { GetTypeHash(Desc.ID), GetTypeHash(Desc.Name), GetTypeHash(Desc.ValueType), GetTypeHash(Desc.ContainerTypes), GetTypeHash(Desc.MetaData) };
+		const uint32 Hashes[] = { GetTypeHash(Desc.ID), GetTypeHash(Desc.Name), GetTypeHash(Desc.ValueType), GetTypeHash(Desc.ContainerTypes), GetTypeHash(Desc.PropertyFlags), GetTypeHash(Desc.MetaData) };
 #else
-		const uint32 Hashes[] = { GetTypeHash(Desc.ID), GetTypeHash(Desc.Name), GetTypeHash(Desc.ValueType), GetTypeHash(Desc.ContainerTypes)};
+		const uint32 Hashes[] = { GetTypeHash(Desc.ID), GetTypeHash(Desc.Name), GetTypeHash(Desc.ValueType), GetTypeHash(Desc.ContainerTypes), GetTypeHash(Desc.PropertyFlags) };
 #endif
 		return CityHash64WithSeed(reinterpret_cast<const char*>(Hashes), sizeof(Hashes), GetObjectHash(Desc.ValueTypeObject));
 	}
