Converted FSoftObjectPath::SubPathString to FUtf8String and added new accessors.
Added placeholders for deprecating the relevant FString-taking functions.
Added a new TStringOverload type which greatly helps in migrating APIs, so wide and narrow arguments resolve to the expected overload.

diff --git a/Engine/Source/Runtime/Core/Public/Containers/StringOverload.h b/Engine/Source/Runtime/Core/Public/Containers/StringOverload.h
new file mode 100644
--- /dev/null
+++ b/Engine/Source/Runtime/Core/Public/Containers/StringOverload.h
@@ -0,0 +1,119 @@
+// Copyright Epic Games, Inc. All Rights Reserved.
+
+#pragma once
+
+#include "Templates/Requires.h"
+#include "Traits/ElementType.h"
+#include "Traits/IsCharEncodingCompatibleWith.h"
+#include "Traits/IsContiguousContainer.h"
+#include "Traits/IsImplicitlyConstructible.h"
+#include "Traits/IsTString.h"
+#include <type_traits>
+
+// A type which wraps a TString parameter but fails (SFINAE) construction if passed a string argument of an incompatible character type.
+// This will allow FString (TString<WIDECHAR>) and FUtf8String (TString<UTF8CHAR>) overloads to sit next to each other and direct
+// pointer, array or other TString arguments to the right overload.
+//
+// This is important because of the legacy of UE's string constructors, which takes anything that looks like a string, of any character
+// width, so overloading - say - FString and FUtf8String will cause overload resolution problems.
+//
+// The argument should be copied/moved into the real TString instance in order to use it, as this type has no string manipulation
+// ability.
+//
+// This argument is equivalent to a pass-by-value toboth the lvalue/rvalue overload case.  Once , and can be converted back to regular
+// const TString&/TString&& overloads once overloading is no longer required.
+//
+// Example:
+//   void Func(TStringOverload<FWideString>); // called by Func(WideString), Func(MoveTemp(WideString)) or Func(TEXT("String"))
+//   void Func(TStringOverload<FUtf8String>); // called by Func(Utf8String), Func(MoveTemp(Utf8String)), Func(UTF8TEXT("String")) or Func("String") or Func({})
+template <typename StringType>
+struct TStringOverload
+{
+private:
+	static_assert(std::is_same_v<StringType, std::remove_cv_t<StringType>> && TIsTString_V<StringType>, "TStringOverload expects an unqualified string object type");
+
+	using CharType = typename StringType::ElementType;
+
+	template <typename ArgType>
+	static constexpr bool IsValidArgType()
+	{
+		using DecayedArgType = std::decay_t<ArgType>;
+
+		if constexpr (!TIsImplicitlyConstructible_V<StringType, DecayedArgType>)
+		{
+			return false;
+		}
+		else if constexpr (TIsContiguousContainer<DecayedArgType>::Value)
+		{
+			return TIsCharEncodingCompatibleWith_V<TElementType_T<DecayedArgType>, CharType>;
+		}
+		else if constexpr (std::is_pointer_v<DecayedArgType>)
+		{
+			return TIsCharEncodingCompatibleWith_V<std::remove_const_t<std::remove_pointer_t<DecayedArgType>>, CharType>;
+		}
+		else
+		{
+			static_assert(sizeof(ArgType) == 0, "Unsupported argument type passed to TStringOverload");
+		}
+	}
+
+public:
+	// This Construct allows us to default construct a TStringOverload using {}.
+	// It only exists for a narrow string, because {} needs to resolve one way or another, and we want it
+	// to be FUtf8String by default.
+	template <
+		typename T = CharType
+		UE_REQUIRES(sizeof(T) == 1)
+	>
+	TStringOverload()
+	{
+	}
+
+	template <
+		typename ArgType
+		UE_REQUIRES(IsValidArgType<ArgType>())
+	>
+	TStringOverload(ArgType&& Arg)
+		: String((ArgType&&)Arg)
+	{
+	}
+
+	StringType&& MoveTemp()
+	{
+		// Shouldn't need a const overload, because these types are only intended as by-non-const-value function parameters, and the only
+		// thing you can do with it is move it into a real string object.
+		return (StringType&&)String;
+	}
+
+	StringType String;
+};
+
+namespace UE::Core::Private
+{
+	template <typename CharType, typename ArgType>
+	inline ArgType&& CheckCharType(ArgType&& Arg)
+	{
+		using DecayedArgType = std::decay_t<ArgType>;
+
+		if constexpr (TIsContiguousContainer<DecayedArgType>::Value)
+		{
+			static_assert(std::is_same_v<const TElementType_T<DecayedArgType>, const CharType>, "Incorrect string type passed to UE_PRIVATE_TO_STRING macro");
+		}
+		else if constexpr (std::is_pointer_v<DecayedArgType>)
+		{
+			static_assert(std::is_same_v<const std::remove_pointer_t<DecayedArgType>, const CharType>, "Incorrect string type passed to UE_PRIVATE_TO_STRING macro");
+		}
+		else
+		{
+			static_assert(sizeof(ArgType) == 0, "Unsupported argument type passed to TStringOverload");
+		}
+		return (ArgType&&)Arg;
+	}
+}
+
+// NOT FOR GENERAL USE!
+//
+// These macros exist to make it easier to denote where wide<->narrow conversions are needed due to legacy APIs,
+// so they can be located and fixed later.
+#define UE_PRIVATE_TO_UTF8_STRING(Str) FUtf8String(UE::Core::Private::CheckCharType<WIDECHAR>(Str))
+#define UE_PRIVATE_TO_WIDE_STRING(Str) FWideString(UE::Core::Private::CheckCharType<UTF8CHAR>(Str))
diff --git a/Engine/Source/Runtime/CoreUObject/Private/UObject/SoftObjectPath.cpp b/Engine/Source/Runtime/CoreUObject/Private/UObject/SoftObjectPath.cpp
--- a/Engine/Source/Runtime/CoreUObject/Private/UObject/SoftObjectPath.cpp
+++ b/Engine/Source/Runtime/CoreUObject/Private/UObject/SoftObjectPath.cpp
@@ -36,12 +36,7 @@ FName FSoftObjectPath::NAME_Untracked(TEXT("Untracked"));
 /** Static methods for more meaningful construction sites. */
 FSoftObjectPath FSoftObjectPath::ConstructFromPackageAssetSubpath(FName InPackageName, FName InAssetName, const FString& InSubPathString)
 {
-	return FSoftObjectPath(FTopLevelAssetPath(InPackageName, InAssetName), InSubPathString);
-}
-
-FSoftObjectPath FSoftObjectPath::ConstructFromPackageAssetSubpath(FName InPackageName, FName InAssetName, FString&& InSubPathString)
-{
-	return FSoftObjectPath(FTopLevelAssetPath(InPackageName, InAssetName), MoveTemp(InSubPathString));
+	return FSoftObjectPath(FTopLevelAssetPath(InPackageName, InAssetName), FUtf8String(InSubPathString));
 }
 
 FSoftObjectPath FSoftObjectPath::ConstructFromPackageAsset(FName InPackageName, FName InAssetName)
@@ -49,14 +44,14 @@ FSoftObjectPath FSoftObjectPath::ConstructFromPackageAsset(FName InPackageName,
 	return FSoftObjectPath(FTopLevelAssetPath(InPackageName, InAssetName), {});
 }
 
-FSoftObjectPath FSoftObjectPath::ConstructFromAssetPathAndSubpath(FTopLevelAssetPath InAssetPath, const FString& InSubPathString)
+FSoftObjectPath FSoftObjectPath::ConstructFromAssetPathAndSubpath(FTopLevelAssetPath InAssetPath, TStringOverload<FWideString> InSubPathString)
 {
-	return FSoftObjectPath(InAssetPath, InSubPathString);
+	return FSoftObjectPath(InAssetPath, FUtf8String(InSubPathString.MoveTemp()));
 }
 
-FSoftObjectPath FSoftObjectPath::ConstructFromAssetPathAndSubpath(FTopLevelAssetPath InAssetPath, FString&& InSubPathString)
+FSoftObjectPath FSoftObjectPath::ConstructFromAssetPathAndSubpath(FTopLevelAssetPath InAssetPath, TStringOverload<FUtf8String> InSubPathString)
 {
-	return FSoftObjectPath(InAssetPath, MoveTemp(InSubPathString));
+	return FSoftObjectPath(InAssetPath, InSubPathString.MoveTemp());
 }
 
 FSoftObjectPath FSoftObjectPath::ConstructFromAssetPath(FTopLevelAssetPath InAssetPath)
@@ -174,10 +169,22 @@ static FString GetObjectBeingSerializedForSoftObjectPath()
 	return Result;
 }
 
-void FSoftObjectPath::SetPath(const FTopLevelAssetPath& InAssetPath, FString InSubPathString)
+void FSoftObjectPath::SetPath(const FTopLevelAssetPath& InAssetPath)
+{
+	AssetPath = InAssetPath;
+	SubPathString.Empty();
+}
+
+void FSoftObjectPath::SetPath(const FTopLevelAssetPath& InAssetPath, TStringOverload<FWideString> InSubPathString)
+{
+	AssetPath = InAssetPath;
+	SubPathString = FUtf8String(InSubPathString.MoveTemp());
+}
+
+void FSoftObjectPath::SetPath(const FTopLevelAssetPath& InAssetPath, TStringOverload<FUtf8String> InSubPathString)
 {
 	AssetPath = InAssetPath;
-	SubPathString = MoveTemp(InSubPathString);
+	SubPathString = InSubPathString.MoveTemp();
 }
 
 void FSoftObjectPath::SetPath(FWideStringView Path)
@@ -239,8 +246,8 @@ void FSoftObjectPath::SetPath(FWideStringView Path)
 
 		// Replace delimiters in subpath string with . to normalize
 		AssetPath = FTopLevelAssetPath(FName(PackageNameView), FName(AssetNameView));
-		SubPathString = Path;
-		SubPathString.ReplaceCharInline(SUBOBJECT_DELIMITER_CHAR, '.');
+		SubPathString = UE_PRIVATE_TO_UTF8_STRING(Path);
+		SubPathString.ReplaceCharInline(UTF8TEXT(SUBOBJECT_DELIMITER_CHAR_ANSI), UTF8TEXT('.'));
 	}
 }
 
@@ -655,7 +662,7 @@ UObject* FSoftObjectPath::TryLoad(FUObjectSerializeContext* InLoadContext) const
 			// it knows how to load that specific object.
 			if (!LoadedObject && TopLevelObject)
 			{
-				TopLevelObject->ResolveSubobject(*SubPathString, LoadedObject, /*bLoadIfExists*/true);
+				TopLevelObject->ResolveSubobject(*UE_PRIVATE_TO_WIDE_STRING(SubPathString), LoadedObject, /*bLoadIfExists*/true);
 			}
 		}
 		else
@@ -770,7 +777,7 @@ UObject* FSoftObjectPath::ResolveObjectInternal() const
 		// it knows how to load that specific object.
 		if (TopLevelObject)
 		{
-			TopLevelObject->ResolveSubobject(*SubPathString, FoundObject, /*bLoadIfExists*/false);
+			TopLevelObject->ResolveSubobject(*UE_PRIVATE_TO_WIDE_STRING(SubPathString), FoundObject, /*bLoadIfExists*/false);
 		}
 	}
 
@@ -1016,7 +1023,7 @@ void SerializeForLog(FCbWriter& Writer, const FSoftObjectPath& Value)
 	Writer.AddString(ANSITEXTVIEW("$text"), WriteToUtf8String<256>(Value));
 	Writer.AddString(ANSITEXTVIEW("PackageName"), WriteToUtf8String<256>(Value.GetLongPackageFName()));
 	Writer.AddString(ANSITEXTVIEW("AssetName"), WriteToUtf8String<256>(Value.GetAssetFName()));
-	Writer.AddString(ANSITEXTVIEW("SubPath"), Value.GetSubPathString());
+	Writer.AddString(ANSITEXTVIEW("SubPath"), Value.GetSubPathUtf8String());
 	Writer.EndObject();
 }
 
diff --git a/Engine/Source/Runtime/CoreUObject/Public/UObject/SoftObjectPath.h b/Engine/Source/Runtime/CoreUObject/Public/UObject/SoftObjectPath.h
--- a/Engine/Source/Runtime/CoreUObject/Public/UObject/SoftObjectPath.h
+++ b/Engine/Source/Runtime/CoreUObject/Public/UObject/SoftObjectPath.h
@@ -6,9 +6,10 @@
 #include "Containers/Array.h"
 #include "Containers/Set.h"
 #include "Containers/StringFwd.h"
+#include "Containers/StringOverload.h"
 #include "Containers/StringView.h"
+#include "Containers/Utf8String.h"
 #include "Containers/UnrealString.h"
-#include "Containers/StringView.h"
 #include "CoreTypes.h"
 #include "HAL/ThreadSafeCounter.h"
 #include "HAL/ThreadSingleton.h"
@@ -33,6 +34,9 @@ class FOutputDevice;
 struct FPropertyTag;
 struct FUObjectSerializeContext;
 
+// Wrapper for UE_DEPRECATED that can be selectively enabled.
+#define UE_SOFTOBJECTPATH_FSTRING_DEPRECATED(Version, Message) /*UE_DEPRECATED(Version, Message)*/
+
 /** Delegate called on completion of async loading a soft object. The UObject will be null if the load failed */
 DECLARE_DELEGATE_TwoParams(FLoadSoftObjectPathAsyncDelegate, const FSoftObjectPath&, UObject*);
 
@@ -56,20 +60,25 @@ struct FSoftObjectPath
 	{
 		SetPath(FStringView(Path));
 	}
-	[[nodiscard]] explicit FSoftObjectPath(FTopLevelAssetPath InAssetPath, FString InSubPathString)
+	UE_SOFTOBJECTPATH_FSTRING_DEPRECATED(5.6, "FSoftObjectPath with a wide string subpath has been deprecated - please use a UTF8 subpath instead.")
+	[[nodiscard]] explicit FSoftObjectPath(FTopLevelAssetPath InAssetPath, TStringOverload<FWideString> InSubPathString)
 	{
-		SetPath(InAssetPath, MoveTemp(InSubPathString));
+		SetPath(InAssetPath, FUtf8String(InSubPathString.MoveTemp()));
+	}
+	[[nodiscard]] explicit FSoftObjectPath(FTopLevelAssetPath InAssetPath, TStringOverload<FUtf8String> InSubPathString)
+	{
+		SetPath(InAssetPath, InSubPathString.MoveTemp());
 	}
 	UE_DEPRECATED(5.6, "FSoftObjectPath which takes a package name and asset name has been deprecated - please pass FTopLevelAssetPath instead.")
-	[[nodiscard]] explicit FSoftObjectPath(FName InPackageName, FName InAssetName, FString InSubPathString)
+	[[nodiscard]] explicit FSoftObjectPath(FName InPackageName, FName InAssetName, const FString& InSubPathString)
 	{
-		SetPath(FTopLevelAssetPath(InPackageName, InAssetName), MoveTemp(InSubPathString));
+		SetPath(FTopLevelAssetPath(InPackageName, InAssetName), FUtf8String(InSubPathString));
 	}
 
 	/** Explicitly extend a top-level object path with an empty subobject path. */
 	[[nodiscard]] explicit FSoftObjectPath(FTopLevelAssetPath InAssetPath)
 	{
-		SetPath(InAssetPath, FString());
+		SetPath(InAssetPath);
 	}
 	[[nodiscard]] explicit FSoftObjectPath(FWideStringView Path)
 	{
@@ -129,12 +138,11 @@ struct FSoftObjectPath
 	/** Static methods for more meaningful construction sites. */
 	UE_DEPRECATED(5.6, "FSoftObjectPath::ConstructFromPackageAssetSubpath has been deprecated - please use ConstructFromAssetPathAndSubpath instead.")
 	[[nodiscard]] COREUOBJECT_API static FSoftObjectPath ConstructFromPackageAssetSubpath(FName InPackageName, FName InAssetName, const FString& InSubPathString);
-	UE_DEPRECATED(5.6, "FSoftObjectPath::ConstructFromPackageAssetSubpath has been deprecated - please use ConstructFromAssetPathAndSubpath instead.")
-	[[nodiscard]] COREUOBJECT_API static FSoftObjectPath ConstructFromPackageAssetSubpath(FName InPackageName, FName InAssetName, FString&& InSubPathString);
 	UE_DEPRECATED(5.6, "FSoftObjectPath::ConstructFromPackageAsset has been deprecated - please use ConstructFromAssetPath instead.")
 	[[nodiscard]] COREUOBJECT_API static FSoftObjectPath ConstructFromPackageAsset(FName InPackageName, FName InAssetName);
-	[[nodiscard]] COREUOBJECT_API static FSoftObjectPath ConstructFromAssetPathAndSubpath(FTopLevelAssetPath InAssetPath, const FString& InSubPathString);
-	[[nodiscard]] COREUOBJECT_API static FSoftObjectPath ConstructFromAssetPathAndSubpath(FTopLevelAssetPath InAssetPath, FString&& InSubPathString);
+	UE_SOFTOBJECTPATH_FSTRING_DEPRECATED(5.6, "FSoftObjectPath::ConstructFromAssetPathAndSubpath with a wide string subpath has been deprecated - please use a UTF8 subpath instead.")
+	[[nodiscard]] COREUOBJECT_API static FSoftObjectPath ConstructFromAssetPathAndSubpath(FTopLevelAssetPath InAssetPath, TStringOverload<FWideString> InSubPathString);
+	[[nodiscard]] COREUOBJECT_API static FSoftObjectPath ConstructFromAssetPathAndSubpath(FTopLevelAssetPath InAssetPath, TStringOverload<FUtf8String> InSubPathString);
 	[[nodiscard]] COREUOBJECT_API static FSoftObjectPath ConstructFromAssetPath(FTopLevelAssetPath InAssetPath);
 	[[nodiscard]] COREUOBJECT_API static FSoftObjectPath ConstructFromStringPath(FString&& InPath);
 	[[nodiscard]] COREUOBJECT_API static FSoftObjectPath ConstructFromStringPath(FStringView InPath);
@@ -204,7 +212,7 @@ struct FSoftObjectPath
 	/** Returns this path without the SubPath component, restricting the result to a top level asset but keeping the type as FSoftObjectPath in contrast to GetAssetPath. */
 	[[nodiscard]] FSoftObjectPath GetWithoutSubPath() const
 	{
-		return FSoftObjectPath(AssetPath, {});
+		return FSoftObjectPath(AssetPath);
 	}
 
 	/** Returns string version of asset path, including both package and asset but not sub object */
@@ -219,14 +227,24 @@ struct FSoftObjectPath
 	}
 
 	/** Returns the sub path, which is often empty */
-	[[nodiscard]] FORCEINLINE const FString& GetSubPathString() const
+	UE_SOFTOBJECTPATH_FSTRING_DEPRECATED(5.6, "FSoftObjectPath::GetSubPathString has been deprecated - please use GetSubPathUtf8String() instead.")
+	[[nodiscard]] FORCEINLINE FString GetSubPathString() const
+	{
+		return FString(SubPathString);
+	}
+	[[nodiscard]] FORCEINLINE const FUtf8String& GetSubPathUtf8String() const
 	{
 		return SubPathString;
 	}
 
-	FORCEINLINE void SetSubPathString(const FString& InSubPathString)
+	UE_SOFTOBJECTPATH_FSTRING_DEPRECATED(5.6, "FSoftObjectPath::ConstructFromAssetPathAndSubpath with a wide string subpath has been deprecated - please use a UTF8 subpath instead.")
+	FORCEINLINE void SetSubPathString(TStringOverload<FWideString> InSubPathString)
+	{
+		SubPathString = FUtf8String(InSubPathString.MoveTemp());
+	}
+	FORCEINLINE void SetSubPathString(TStringOverload<FUtf8String> InSubPathString)
 	{
-		SubPathString = InSubPathString;
+		SubPathString = InSubPathString.MoveTemp();
 	}
 
 	/** Returns /package/path, leaving off the asset name and sub object */
@@ -256,7 +274,10 @@ struct FSoftObjectPath
 	}
 
 	/** Sets asset path of this reference based on a string path */
-	COREUOBJECT_API void SetPath(const FTopLevelAssetPath& InAssetPath, FString InSubPathString = FString());
+	COREUOBJECT_API void SetPath(const FTopLevelAssetPath& InAssetPath);
+	UE_SOFTOBJECTPATH_FSTRING_DEPRECATED(5.6, "FSoftObjectPath::SetPath with a wide string subpath has been deprecated - please use a UTF8 subpath instead.")
+	COREUOBJECT_API void SetPath(const FTopLevelAssetPath& InAssetPath, TStringOverload<FWideString> InSubPathString);
+	COREUOBJECT_API void SetPath(const FTopLevelAssetPath& InAssetPath, TStringOverload<FUtf8String> InSubPathString);
 	COREUOBJECT_API void SetPath(FWideStringView Path);
 	COREUOBJECT_API void SetPath(FAnsiStringView Path);
 	COREUOBJECT_API void SetPath(FUtf8StringView Path);
@@ -442,7 +463,7 @@ private:
 	FTopLevelAssetPath AssetPath;
 
 	/** Optional FString for subobject within an asset. This is the sub path after the : */
-	FString SubPathString;
+	FUtf8String SubPathString;
 
 	/** Package names currently being duplicated, needed by FixupForPIE */
 	static COREUOBJECT_API TSet<FName> PIEPackageNames;
@@ -672,7 +693,7 @@ struct FSoftObjectPathFixupArchive : public FArchiveUObject
 		{
 			if (!Value.IsNull() && Value.GetAssetPathString().Equals(OldAssetPathString, ESearchCase::IgnoreCase))
 			{
-				Value = FSoftObjectPath(NewAssetPath, Value.GetSubPathString());
+				Value = FSoftObjectPath(NewAssetPath, Value.GetSubPathUtf8String());
 			}
 		})
 	{
