Make FObjectPropertyBase::ConstructDefaultObjectValueIfNecessary fail gracefully instead of dereferencing null if PropertyClass is null.

diff --git a/Engine/Source/Runtime/CoreUObject/Private/UObject/PropertyBaseObject.cpp b/Engine/Source/Runtime/CoreUObject/Private/UObject/PropertyBaseObject.cpp
--- a/Engine/Source/Runtime/CoreUObject/Private/UObject/PropertyBaseObject.cpp
+++ b/Engine/Source/Runtime/CoreUObject/Private/UObject/PropertyBaseObject.cpp
@@ -816,7 +816,14 @@ UObject* FObjectPropertyBase::ConstructDefaultObjectValueIfNecessary(UObject* Ex
 	// We cannot know which object to create if the class is abstract or if there is a required interface
 	if (!NewDefaultObjectValue)
 	{
-		if (PropertyClass->HasAnyClassFlags(CLASS_Abstract))
+		if (!ensureMsgf(PropertyClass, TEXT("Malformed %s: PropertyClass=null."), *GetFullName()))
+		{
+			if (OutFailureReason)
+			{
+				*OutFailureReason = TEXT("malformed property: PropertyClass=null");
+			}
+		}
+		else if (PropertyClass->HasAnyClassFlags(CLASS_Abstract))
 		{
 			if (OutFailureReason)
 			{
