Add an option for the async loader to pre-populate the package's public export map with all export hashes to prevent quadratic insert cost when creating exports.

#tests Forced sorted array data structure enabled in editor & loaded multiple maps

diff --git a/Engine/Source/Runtime/CoreUObject/Private/Serialization/AsyncLoading2.cpp b/Engine/Source/Runtime/CoreUObject/Private/Serialization/AsyncLoading2.cpp
--- a/Engine/Source/Runtime/CoreUObject/Private/Serialization/AsyncLoading2.cpp
+++ b/Engine/Source/Runtime/CoreUObject/Private/Serialization/AsyncLoading2.cpp
@@ -499,6 +499,13 @@ static FAutoConsoleVariableRef CVarDetectRaceDuringLoading(
 );
 #endif // USING_INSTRUMENTATION
 
+static bool GAsyncLoading2_PreInsertPublicExports = false;
+static FAutoConsoleVariableRef CVar_PreInsertPublicExports(
+	TEXT("s.PreInsertPublicExports"),
+	GAsyncLoading2_PreInsertPublicExports,
+	TEXT("After reading the package summary for a package, bulk insert all export hashes for public exports into the loaded package's map"),
+	ECVF_Default);
+
 CSV_DECLARE_CATEGORY_MODULE_EXTERN(CORE_API, Basic);
 CSV_DECLARE_CATEGORY_MODULE_EXTERN(CORE_API, FileIO);
 
@@ -517,6 +524,34 @@ FString FormatPackageId(FPackageId PackageId)
 #endif
 }
 
+// Check if an export should be loaded or not on the current platform/build
+// Returns true if the export should be skipped, false if it should be kept
+static bool AsyncLoading2_ShouldSkipLoadingExport(const EExportFilterFlags FilterFlags)
+{
+#if WITH_EDITOR
+		return false;
+#elif UE_SERVER
+		return !!(static_cast<uint32>(FilterFlags) & static_cast<uint32>(EExportFilterFlags::NotForServer));
+#elif !WITH_SERVER_CODE
+		return !!(static_cast<uint32>(FilterFlags) & static_cast<uint32>(EExportFilterFlags::NotForClient));
+#else
+		static const bool bIsDedicatedServer = !GIsClient && GIsServer;
+		static const bool bIsClientOnly = GIsClient && !GIsServer;
+
+		if (bIsDedicatedServer && static_cast<uint32>(FilterFlags) & static_cast<uint32>(EExportFilterFlags::NotForServer))
+		{
+			return true;
+		}
+
+		if (bIsClientOnly && static_cast<uint32>(FilterFlags) & static_cast<uint32>(EExportFilterFlags::NotForClient))
+		{
+			return true;
+		}
+
+		return false;
+#endif
+}
+
 #if WITH_EDITOR
 
 FName ApplyPackageNameRedirections(FName PackageName)
@@ -763,12 +798,12 @@ private:
 
 	class FPublicExportMap
 	{
+		static const int32 InvalidValue = -1;
 // The AGGRESSIVE_MEMORY_SAVING implementation allocates only as many keys and values as necessary to load a package's exports
 // which is space efficient but results in more memory traffic as exports are Store()'d or Remove()'d.
 // The !AGGRESSIVE_MEMORY_SAVING implementation uses more memory but provides O(1) Store() and Remove()
 #if AGGRESSIVE_MEMORY_SAVING && !WITH_EDITOR
 	public:
-		static const int32 InvalidValue = -1;
 
 		FPublicExportMap() = default;
 
@@ -837,7 +872,106 @@ private:
 			Count = NewCount;
 		}
 
-		void Store(uint64 ExportHash, UObject* Object)
+		void PreInsertPublicExports(TConstArrayView<FExportMapEntry> Exports)
+		{
+			if (Count <= 1) 
+			{ 
+				int32 PublicExportCount = 0;
+				for (const FExportMapEntry& Entry : Exports)
+				{
+					if (Entry.PublicExportHash && !AsyncLoading2_ShouldSkipLoadingExport(Entry.FilterFlags))
+					{
+						PublicExportCount++;
+					}
+				}
+
+				// Either: 
+				// 	We're empty so we can just insert all the hashes with an associated value of -1 and sort.
+				//  We have a single pair in the SingleItemKey/SingleItemValue slots that we want to keep the mapping for.
+				// So we reserve space for the new keys under the assumption that it is most likely that the old key is also present in Exports.
+				uint64 OldKey = Count == 1 ? SingleItemKey : 0;
+				int32 OldValue = Count == 1 ? SingleItemValue : InvalidValue;
+				Reserve(PublicExportCount);	
+				TArrayView<uint64> Keys = GetKeys();
+				TArrayView<int32> Values = GetValues();
+				int32 Index = 0;
+				for (const FExportMapEntry& Entry : Exports)
+				{
+					if (Entry.PublicExportHash && !AsyncLoading2_ShouldSkipLoadingExport(Entry.FilterFlags))
+					{
+						Keys[Index++] = Entry.PublicExportHash;
+					}
+				}
+				Algo::Sort(Keys);
+				// Set all values to -1
+				FMemory::Memset(Values.GetData(), 0xFF, Values.NumBytes());
+				// If we had a single object previously, re-store its old value.
+				// If the old key was present in Exports, this will overwrite the -1 we just stored.
+				// If it wasn't in Exports, this will grow the storage, but we consider that unlikely at present.
+				if (OldKey)
+				{
+					Store(OldKey, OldValue);
+				}
+			}
+			else
+			{
+				const int32 OldCount = Count;
+				TArrayView<uint64> OldKeys = GetKeys();
+				TArrayView<int32> OldValues = GetValues();
+				TArray<uint64> KeysToAdd;
+				for (const FExportMapEntry& Entry : Exports)
+				{
+					if (Entry.PublicExportHash && !AsyncLoading2_ShouldSkipLoadingExport(Entry.FilterFlags))
+					{
+						int32 Index = Algo::LowerBound(OldKeys, Entry.PublicExportHash);
+						if (Index >= Count || OldKeys[Index] != Entry.PublicExportHash)
+						{
+							KeysToAdd.Add(Entry.PublicExportHash);
+						}
+					}
+				}
+				
+				// We expect the most common case here to be KeysToAdd.Num() == 0 as we're likely to be reloading the same package data.
+				// If we aren't, it implies some UGC/delivered content scenario where the old content was not fully GC'd before attemping 
+				// to reload the newly downloaded version.
+				if (KeysToAdd.Num() > 0)
+				{
+					// Sort keys to add so that we can interleave old and new keys to maintain sorting
+					Algo::Sort(KeysToAdd);
+					const int32 KeysToAddCount = KeysToAdd.Num();
+
+					// At this point, Count > 1 and KeysToAdd.Num() >= 1, so we don't need to consider the single object case and will always allocate.
+					const int32 NewCount = OldCount + KeysToAddCount;
+					const uint64 NewKeysSize = NewCount * sizeof(uint64);
+					const uint64 NewValuesSize = NewCount * sizeof(int32);
+					uint8* NewAllocation = reinterpret_cast<uint8*>(FMemory::Malloc(NewKeysSize + NewValuesSize));
+
+					TArrayView<uint64> NewKeys = MakeArrayView(reinterpret_cast<uint64*>(NewAllocation), NewCount);
+					TArrayView<int32> NewValues = MakeArrayView(reinterpret_cast<int32*>(NewAllocation + sizeof(uint64) * NewCount), NewCount);
+					for (int32 InsertIndex = 0, OldIndex = 0, ToAddIndex = 0; InsertIndex < NewCount; ++InsertIndex)
+					{
+						if (OldKeys.IsValidIndex(OldIndex)
+						&& (!KeysToAdd.IsValidIndex(ToAddIndex) || OldKeys[OldIndex] < KeysToAdd[ToAddIndex]))
+						{
+							NewKeys[InsertIndex] = OldKeys[OldIndex];
+							NewValues[InsertIndex] = OldValues[OldIndex];
+							++OldIndex;
+						}
+						else 
+						{
+							NewKeys[InsertIndex] = KeysToAdd[ToAddIndex];
+							NewValues[InsertIndex] = InvalidValue;
+							++ToAddIndex;
+						}
+					}
+					FMemory::Free(Allocation);
+					Allocation = NewAllocation;
+					Count = NewCount;
+				}
+			}
+		}
+
+		void Store(uint64 ExportHash, int32 ObjectIndex)
 		{
 			TArrayView<uint64> Keys = GetKeys();
 			TArrayView<int32> Values = GetValues();
@@ -845,7 +979,7 @@ private:
 			if (Index < Count && Keys[Index] == ExportHash)
 			{
 				// Slot already exists so reuse it
-				Values[Index] = GUObjectArray.ObjectToIndex(Object);
+				Values[Index] = ObjectIndex;
 				return;
 			}
 			if (Count == 0 || Keys[0] != 0)
@@ -866,7 +1000,7 @@ private:
 				FMemory::Memmove(Values.GetData(), Values.GetData() + 1, Index * sizeof(int32));
 			}
 			Keys[Index] = ExportHash;
-			Values[Index] = GUObjectArray.ObjectToIndex(Object);
+			Values[Index] = ObjectIndex;
 		}
 
 		bool Remove(uint64 ExportHash)
@@ -883,20 +1017,17 @@ private:
 			return false;
 		}
 
-		UObject* Find(uint64 ExportHash)
+		// Returns the value stored if any, returing InvalidValue for not present or removed values
+		int32 Find(uint64 ExportHash)
 		{
 			TArrayView<uint64> Keys = GetKeys();
 			int32 Index = Algo::LowerBound(Keys, ExportHash);
 			if (Index < Count && Keys[Index] == ExportHash)
 			{
 				TArrayView<int32> Values = GetValues();
-				int32 ObjectIndex = Values[Index];
-				if (ObjectIndex >= 0)
-				{
-					return static_cast<UObject*>(GUObjectArray.IndexToObject(ObjectIndex)->Object);
-				}
+				return Values[Index];
 			}
-			return nullptr;
+			return InvalidValue;
 		}
 
 		[[nodiscard]] bool PinForGC(TArray<int32>& OutUnreachableObjectIndices)
@@ -1031,6 +1162,10 @@ private:
 			uint8* Allocation = nullptr;
 			uint64 SingleItemKey;
 		};
+		
+		// Total number of keys stored, not necessarily number of keys that map to valid values
+		// Count == 1 means we use SingleItemKey and SingleItemValue 
+		// Count > 1 means we use Allocation 
 		int32 Count = 0;
 		int32 SingleItemValue = -1;
 #else
@@ -1042,9 +1177,23 @@ private:
 			Data.Reserve(NewCount);
 		}
 
-		void Store(uint64 ExportHash, UObject* Object)
+		void PreInsertPublicExports(TConstArrayView<FExportMapEntry> Exports)
 		{
-			Data.Add(ExportHash, GUObjectArray.ObjectToIndex(Object));
+			// Just behave like Reserve for the TMap version as we have cheap insertion
+			int32 PublicExportCount = 0;
+			for (const FExportMapEntry& Entry : Exports)
+			{
+				if (Entry.PublicExportHash && !AsyncLoading2_ShouldSkipLoadingExport(Entry.FilterFlags))
+				{
+					PublicExportCount++;
+				}
+			}
+			Data.Reserve(PublicExportCount);
+		}
+
+		void Store(uint64 ExportHash, int32 ObjectIndex)
+		{
+			Data.Add(ExportHash, ObjectIndex);
 		}
 
 		bool Remove(uint64 ExportHash)
@@ -1057,14 +1206,14 @@ private:
 			return Data.Num();
 		}
 
-		UObject* Find(uint64 ExportHash)
+		int32 Find(uint64 ExportHash)
 		{
 			int32* ObjectIndex = Data.Find(ExportHash);
 			if (ObjectIndex)
 			{
-				return static_cast<UObject*>(GUObjectArray.IndexToObject(*ObjectIndex)->Object);
+				return *ObjectIndex;
 			}
-			return nullptr;
+			return InvalidValue;
 		}
 
 		[[nodiscard]] bool PinForGC(TArray<int32>& OutUnreachableObjectIndices)
@@ -1296,9 +1445,15 @@ public:
 		PublicExportMap.Reserve(PublicExportCount);
 	}
 
+	/** Insert all the given exports into the map by their hash, storing a non-present/null value for any objects that are not already in the map */
+	void PreInsertPublicExports(TConstArrayView<FExportMapEntry> Exports)
+	{
+		PublicExportMap.PreInsertPublicExports(Exports);
+	}
+
 	void StorePublicExport(uint64 ExportHash, UObject* Object)
 	{
-		PublicExportMap.Store(ExportHash, Object);
+		PublicExportMap.Store(ExportHash, GUObjectArray.ObjectToIndex(Object));
 	}
 
 	void RemovePublicExport(uint64 ExportHash, FName ObjectName = NAME_None)
@@ -1314,7 +1469,12 @@ public:
 
 	UObject* GetPublicExport(uint64 ExportHash)
 	{
-		return PublicExportMap.Find(ExportHash);
+		int32 ObjectIndex = PublicExportMap.Find(ExportHash);
+		if (ObjectIndex >= 0)
+		{
+			return static_cast<UObject*>(GUObjectArray.IndexToObject(ObjectIndex)->Object);
+		}
+		return nullptr;
 	}
 
 	void PinPublicExportsForGC(TArray<int32>& OutUnreachableObjectIndices)
@@ -6198,7 +6358,14 @@ EEventLoadNodeExecutionResult FAsyncPackage2::ProcessLinkerLoadPackageSummary(FA
 		FLoadedPackageRef& PackageRef = AsyncLoadingThread.GlobalImportStore.FindPackageRefChecked(Desc.UPackageId, Desc.UPackageName);
 		if (!bLoadHasFailed)
 		{
-			PackageRef.ReserveSpaceForPublicExports(LinkerLoadState->LinkerLoadHeaderData.ExportMap.Num());
+			if (GAsyncLoading2_PreInsertPublicExports)
+			{
+				PackageRef.PreInsertPublicExports(LinkerLoadState->LinkerLoadHeaderData.ExportMap);
+			}
+			else
+			{
+				PackageRef.ReserveSpaceForPublicExports(LinkerLoadState->LinkerLoadHeaderData.ExportMap.Num());
+			}
 #if WITH_METADATA
 			// Create metadata object, this needs to happen before any other package wants to use our exports
 			LinkerLoadState->MetaDataIndex = LinkerLoadState->Linker->LoadMetaDataFromExportMap(false);
@@ -6707,31 +6874,38 @@ EEventLoadNodeExecutionResult FAsyncPackage2::Event_ProcessPackageSummary(FAsync
 #endif
 		if (Package->Desc.bCanBeImported)
 		{
-			int32 PublicExportsCount = 0;
-			for (const FExportMapEntry& Export : Package->HeaderData.ExportMap)
+			FLoadedPackageRef& PackageRef = Package->AsyncLoadingThread.GlobalImportStore.FindPackageRefChecked(Package->Desc.UPackageId, Package->Desc.UPackageName);
+			if (GAsyncLoading2_PreInsertPublicExports)
 			{
-				if (Export.PublicExportHash)
-				{
-					++PublicExportsCount;
-				}
+				PackageRef.PreInsertPublicExports(Package->HeaderData.ExportMap);
 			}
-#if WITH_EDITOR
-			if (Package->OptionalSegmentHeaderData.IsSet())
+			else
 			{
-				for (const FExportMapEntry& Export : Package->OptionalSegmentHeaderData->ExportMap)
+				int32 PublicExportsCount = 0;
+				for (const FExportMapEntry& Export : Package->HeaderData.ExportMap)
 				{
 					if (Export.PublicExportHash)
 					{
 						++PublicExportsCount;
 					}
 				}
-			}
+#if WITH_EDITOR
+				if (Package->OptionalSegmentHeaderData.IsSet())
+				{
+					for (const FExportMapEntry& Export : Package->OptionalSegmentHeaderData->ExportMap)
+					{
+						if (Export.PublicExportHash)
+						{
+							++PublicExportsCount;
+						}
+					}
+				}
 #endif
-			FLoadedPackageRef& PackageRef = Package->AsyncLoadingThread.GlobalImportStore.FindPackageRefChecked(Package->Desc.UPackageId, Package->Desc.UPackageName);
-			if (PublicExportsCount)
-			{
-				PackageRef.ReserveSpaceForPublicExports(PublicExportsCount);
-			}
+				if (PublicExportsCount)
+				{
+					PackageRef.ReserveSpaceForPublicExports(PublicExportsCount);
+				}
+		}
 		}
 
 		check(Package->Desc.PackageIdToLoad == FPackageId::FromName(Package->HeaderData.PackageName));
@@ -7188,33 +7362,7 @@ void FAsyncPackage2::EventDrivenCreateExport(const FAsyncPackageHeaderData& Head
 		ObjectName = Header.NameMap.GetName(Export.ObjectName);
 	}
 
-	auto FilterExport = [](const EExportFilterFlags FilterFlags) -> bool
-	{
-#if WITH_EDITOR
-		return false;
-#elif UE_SERVER
-		return !!(static_cast<uint32>(FilterFlags) & static_cast<uint32>(EExportFilterFlags::NotForServer));
-#elif !WITH_SERVER_CODE
-		return !!(static_cast<uint32>(FilterFlags) & static_cast<uint32>(EExportFilterFlags::NotForClient));
-#else
-		static const bool bIsDedicatedServer = !GIsClient && GIsServer;
-		static const bool bIsClientOnly = GIsClient && !GIsServer;
-
-		if (bIsDedicatedServer && static_cast<uint32>(FilterFlags) & static_cast<uint32>(EExportFilterFlags::NotForServer))
-		{
-			return true;
-		}
-
-		if (bIsClientOnly && static_cast<uint32>(FilterFlags) & static_cast<uint32>(EExportFilterFlags::NotForClient))
-		{
-			return true;
-		}
-
-		return false;
-#endif
-	};
-
-	ExportObject.bFiltered = FilterExport(Export.FilterFlags);
+	ExportObject.bFiltered = AsyncLoading2_ShouldSkipLoadingExport(Export.FilterFlags);
 	if (ExportObject.bFiltered || ExportObject.bExportLoadFailed)
 	{
 		if (ExportObject.bExportLoadFailed)
