Enable packages to contain VCells

This change bumps the format version but leaves the new fields and tables empty.

diff --git a/Engine/Source/Developer/AssetTools/Private/AssetHeaderPatcher.cpp b/Engine/Source/Developer/AssetTools/Private/AssetHeaderPatcher.cpp
--- a/Engine/Source/Developer/AssetTools/Private/AssetHeaderPatcher.cpp
+++ b/Engine/Source/Developer/AssetTools/Private/AssetHeaderPatcher.cpp
@@ -691,6 +691,8 @@ enum class ESummaryOffset
 	GatherableTextDataTable,
 	ImportTable,
 	ExportTable,
+	CellImportTable,
+	CellExportTable,
 	DependsTable,
 	SoftPackageReferenceList,
 	SearchableNamesMap,
@@ -800,6 +802,8 @@ void PatchSummaryOffsets(FPackageFileSummary& Dst, int64 OffsetFrom, int64 Offse
 				{ UE_POPULATE_OFFSET_INFO(MetaDataOffset) },
 				{ UE_POPULATE_OFFSET_INFO(ImportOffset) },
 				{ UE_POPULATE_OFFSET_INFO(ExportOffset) },
+				{ UE_POPULATE_OFFSET_INFO(CellImportOffset) },
+				{ UE_POPULATE_OFFSET_INFO(CellExportOffset) },
 				{ UE_POPULATE_OFFSET_INFO(DependsOffset) },
 				{ UE_POPULATE_OFFSET_INFO(SoftPackageReferencesOffset) },
 				{ UE_POPULATE_OFFSET_INFO(SearchableNamesOffset) },
@@ -1311,6 +1315,11 @@ FAssetHeaderPatcher::EResult FAssetHeaderPatcherInner::PatchHeader_Deserialize()
 		UE_LOG(LogAssetHeaderPatcher, Error, TEXT("Asset %s has an unexpected DataResourceOffset"), *SrcAsset);
 		return EResult::ErrorUnexpectedSectionOrder;
 	}
+	if (Summary.CellExportCount > 0 || Summary.CellImportCount > 0)
+	{
+		UE_LOG(LogAssetHeaderPatcher, Error, TEXT("Asset %s contains unexpected VCells"), *SrcAsset);
+		return EResult::ErrorUnexpectedSectionOrder;
+	}
 
 	if (Summary.NameCount > 0)
 	{
diff --git a/Engine/Source/Developer/IoStoreUtilities/Private/PackageStoreOptimizer.cpp b/Engine/Source/Developer/IoStoreUtilities/Private/PackageStoreOptimizer.cpp
--- a/Engine/Source/Developer/IoStoreUtilities/Private/PackageStoreOptimizer.cpp
+++ b/Engine/Source/Developer/IoStoreUtilities/Private/PackageStoreOptimizer.cpp
@@ -905,6 +905,8 @@ void FPackageStoreOptimizer::FinalizePackageHeader(FPackageStorePackage* Package
 	HeaderArchive.Serialize(ImportMapArchive.GetWriterData(), ImportMapArchive.Tell());
 	PackageSummary->ExportMapOffset = HeaderArchive.Tell();
 	HeaderArchive.Serialize(ExportMapArchive.GetWriterData(), ExportMapArchive.Tell());
+	PackageSummary->CellImportMapOffset = HeaderArchive.Tell();
+	PackageSummary->CellExportMapOffset = HeaderArchive.Tell();
 	// raw arrays of 4-byte aligned items
 	PackageSummary->ExportBundleEntriesOffset = HeaderArchive.Tell();
 	HeaderArchive.Serialize(ExportBundleEntriesArchive.GetWriterData(), ExportBundleEntriesArchive.Tell());
diff --git a/Engine/Source/Runtime/Core/Public/UObject/ObjectVersion.h b/Engine/Source/Runtime/Core/Public/UObject/ObjectVersion.h
--- a/Engine/Source/Runtime/Core/Public/UObject/ObjectVersion.h
+++ b/Engine/Source/Runtime/Core/Public/UObject/ObjectVersion.h
@@ -90,6 +90,9 @@ enum class EUnrealEngineObjectUE5Version : uint32
 	// Added meta data serialization offset to for saved, versioned packages
 	METADATA_SERIALIZATION_OFFSET,
 
+	// Added VCells to the object graph
+	VERSE_CELLS,
+
 	// -----<new versions can be added before this line>-------------------------------------------------
 	// - this needs to be the last line (see note below)
 	AUTOMATIC_VERSION_PLUS_ONE,
diff --git a/Engine/Source/Runtime/CoreUObject/Private/Serialization/ZenPackageHeader.cpp b/Engine/Source/Runtime/CoreUObject/Private/Serialization/ZenPackageHeader.cpp
--- a/Engine/Source/Runtime/CoreUObject/Private/Serialization/ZenPackageHeader.cpp
+++ b/Engine/Source/Runtime/CoreUObject/Private/Serialization/ZenPackageHeader.cpp
@@ -60,7 +60,7 @@ FZenPackageHeader FZenPackageHeader::MakeView(FMemoryView Memory, FString& OutEr
 		(PackageSummary->ExportMapOffset - PackageSummary->ImportMapOffset) / sizeof(FPackageObjectIndex));
 	PackageHeader.ExportMap = TArrayView<const FExportMapEntry>(
 		reinterpret_cast<const FExportMapEntry*>(PackageHeaderDataPtr + PackageSummary->ExportMapOffset),
-		(PackageSummary->ExportBundleEntriesOffset - PackageSummary->ExportMapOffset) / sizeof(FExportMapEntry));
+		(PackageSummary->CellImportMapOffset - PackageSummary->ExportMapOffset) / sizeof(FExportMapEntry));
 	PackageHeader.ExportCount = PackageHeader.ExportMap.Num();
 
 	const uint64 ExportBundleEntriesSize = PackageHeader.PackageSummary->DependencyBundleHeadersOffset - PackageHeader.PackageSummary->ExportBundleEntriesOffset;
diff --git a/Engine/Source/Runtime/CoreUObject/Private/UObject/PackageFileSummary.cpp b/Engine/Source/Runtime/CoreUObject/Private/UObject/PackageFileSummary.cpp
--- a/Engine/Source/Runtime/CoreUObject/Private/UObject/PackageFileSummary.cpp
+++ b/Engine/Source/Runtime/CoreUObject/Private/UObject/PackageFileSummary.cpp
@@ -240,6 +240,12 @@ void operator<<(FStructuredArchive::FSlot Slot, FPackageFileSummary& Sum)
 		Record << SA_VALUE(TEXT("ExportCount"), Sum.ExportCount) << SA_VALUE(TEXT("ExportOffset"), Sum.ExportOffset);
 		Record << SA_VALUE(TEXT("ImportCount"), Sum.ImportCount) << SA_VALUE(TEXT("ImportOffset"), Sum.ImportOffset);
 
+		if (Sum.FileVersionUE >= EUnrealEngineObjectUE5Version::VERSE_CELLS)
+		{
+			Record << SA_VALUE(TEXT("CellExportCount"), Sum.CellExportCount) << SA_VALUE(TEXT("CellExportOffset"), Sum.CellExportOffset);
+			Record << SA_VALUE(TEXT("CellImportCount"), Sum.CellImportCount) << SA_VALUE(TEXT("CellImportOffset"), Sum.CellImportOffset);
+		}
+
 		if (Sum.FileVersionUE >= EUnrealEngineObjectUE5Version::METADATA_SERIALIZATION_OFFSET)
 		{
 			Record << SA_VALUE(TEXT("MetaDataOffset"), Sum.MetaDataOffset);
diff --git a/Engine/Source/Runtime/CoreUObject/Private/UObject/SavePackage2.cpp b/Engine/Source/Runtime/CoreUObject/Private/UObject/SavePackage2.cpp
--- a/Engine/Source/Runtime/CoreUObject/Private/UObject/SavePackage2.cpp
+++ b/Engine/Source/Runtime/CoreUObject/Private/UObject/SavePackage2.cpp
@@ -1369,6 +1369,8 @@ ESavePackageResult BuildLinker(FSaveContext& SaveContext)
 			ImportSortHelper.SortImports(Linker);
 		}
 		Linker->Summary.ImportCount = Linker->ImportMap.Num();
+
+		Linker->Summary.CellImportCount = 0;
 	}
 
 	// Build ExportMap & Package Netplay data
@@ -1396,6 +1398,8 @@ ESavePackageResult BuildLinker(FSaveContext& SaveContext)
 			ExportSortHelper.SortExports(Linker);
 		}
 		Linker->Summary.ExportCount = Linker->ExportMap.Num();
+
+		Linker->Summary.CellExportCount = 0;
 	}
 
 	// Build Linker Reverse Mapping
@@ -1973,6 +1977,16 @@ ESavePackageResult WritePackageHeader(FStructuredArchive::FRecord& StructuredArc
 	}
 	SaveContext.OffsetAfterExportMap = (int32)Linker->Tell();
 
+	// Save Dummy Cell Import Map, overwritten later.
+	{
+		Linker->Summary.CellImportOffset = (int32)Linker->Tell();
+	}
+
+	// Save Dummy Cell Export Map, overwritten later.
+	{
+		Linker->Summary.CellExportOffset = (int32)Linker->Tell();
+	}
+
 	// Save Depend Map
 	{
 		SCOPED_SAVETIMER(UPackage_Save_WriteDependsMap);
@@ -2558,6 +2572,19 @@ ESavePackageResult UpdatePackageHeader(FStructuredArchive::FRecord& StructuredAr
 		check(Linker->Tell() == SaveContext.OffsetAfterExportMap);
 	}
 
+	// Write Real Cell Import Map
+	if (!SaveContext.IsTextFormat())
+	{
+		Linker->Seek(Linker->Summary.CellImportOffset);
+		FStructuredArchive::FStream CellExportTableStream = StructuredArchiveRoot.EnterStream(TEXT("CellImportTable"));
+	}
+	// Write Real Cell Export Map
+	if (!SaveContext.IsTextFormat())
+	{
+		Linker->Seek(Linker->Summary.CellExportOffset);
+		FStructuredArchive::FStream CellExportTableStream = StructuredArchiveRoot.EnterStream(TEXT("CellExportTable"));
+	}
+
 	// Figure out if at leat one export is marked as an asset
 	bool bContainsAsset = false;
 	for (FObjectExport& Export : Linker->ExportMap)
diff --git a/Engine/Source/Runtime/CoreUObject/Public/Serialization/AsyncLoading2.h b/Engine/Source/Runtime/CoreUObject/Public/Serialization/AsyncLoading2.h
--- a/Engine/Source/Runtime/CoreUObject/Public/Serialization/AsyncLoading2.h
+++ b/Engine/Source/Runtime/CoreUObject/Public/Serialization/AsyncLoading2.h
@@ -284,6 +284,8 @@ struct FZenPackageSummary
 	int32 ImportedPublicExportHashesOffset;
 	int32 ImportMapOffset;
 	int32 ExportMapOffset;
+	int32 CellImportMapOffset;
+	int32 CellExportMapOffset;
 	int32 ExportBundleEntriesOffset;
 	int32 DependencyBundleHeadersOffset;
 	int32 DependencyBundleEntriesOffset;
diff --git a/Engine/Source/Runtime/CoreUObject/Public/UObject/PackageFileSummary.h b/Engine/Source/Runtime/CoreUObject/Public/UObject/PackageFileSummary.h
--- a/Engine/Source/Runtime/CoreUObject/Public/UObject/PackageFileSummary.h
+++ b/Engine/Source/Runtime/CoreUObject/Public/UObject/PackageFileSummary.h
@@ -146,6 +146,26 @@ public:
 	*/
 	int32		ImportOffset;
 
+	/**
+	* Number of cells contained in this package
+	*/
+	int32		CellExportCount;
+
+	/**
+	* Location into the file on disk for the CellExportMap data
+	*/
+	int32		CellExportOffset;
+
+	/**
+	* Number of cell imports contained in this package
+	*/
+	int32		CellImportCount;
+
+	/**
+	* Location into the file on disk for the CellImportMap data
+	*/
+	int32		CellImportOffset;
+
 	/**
 	* Location into the file on disk for the DependsMap data
 	*/
