Add transactional on-aborts for `AddRef` and `ReleaseRef`.

#rnx

diff --git a/Engine/Source/Programs/AutoRTFMTests/Private/UObjectTests.cpp b/Engine/Source/Programs/AutoRTFMTests/Private/UObjectTests.cpp
--- a/Engine/Source/Programs/AutoRTFMTests/Private/UObjectTests.cpp
+++ b/Engine/Source/Programs/AutoRTFMTests/Private/UObjectTests.cpp
@@ -818,3 +818,52 @@ TEST_CASE("UObject.FUObjectThreadContext")
 		REQUIRE(0 == FUObjectThreadContext::Get().IsInConstructor);
 	}
 }
+
+TEST_CASE("UObject.AddRef")
+{
+	SECTION("Default")
+	{
+		UMyAutoRTFMTestObject* Object = NewObject<UMyAutoRTFMTestObject>();
+
+		AutoRTFM::Testing::Commit([&]
+		{
+			Object->AddRef();
+		});
+
+		Object->ReleaseRef();
+	}
+}
+
+TEST_CASE("UObject.ReleaseRef")
+{
+	SECTION("Default")
+	{
+		UMyAutoRTFMTestObject* Object = NewObject<UMyAutoRTFMTestObject>();
+
+		Object->AddRef();
+
+		AutoRTFM::Testing::Commit([&]
+		{
+			Object->ReleaseRef();
+		});
+	}
+
+	SECTION("With Cascading Abort")
+	{
+		UMyAutoRTFMTestObject* Object = NewObject<UMyAutoRTFMTestObject>();
+
+		Object->AddRef();
+
+		bool bFirst = true;
+
+		AutoRTFM::Testing::Commit([&]
+		{
+			Object->ReleaseRef();
+
+			if (bFirst)
+			{
+				AutoRTFM::CascadingRetryTransaction([&] { bFirst = false; });
+			}
+		});
+	}
+}
diff --git a/Engine/Source/Runtime/CoreUObject/Public/UObject/UObjectArray.h b/Engine/Source/Runtime/CoreUObject/Public/UObject/UObjectArray.h
--- a/Engine/Source/Runtime/CoreUObject/Public/UObject/UObjectArray.h
+++ b/Engine/Source/Runtime/CoreUObject/Public/UObject/UObjectArray.h
@@ -318,6 +318,12 @@ public:
 				SetRootFlags(EInternalObjectFlags::RefCounted);
 			}
 		};
+
+		// If the transaction is aborted we need to remember to release the reference we added in the open!
+		UE_AUTORTFM_ONABORT(this)
+		{
+			this->ReleaseRef();
+		};
 	}
 
 	void ReleaseRef()
@@ -333,6 +339,19 @@ public:
 				ClearRootFlags(EInternalObjectFlags::RefCounted);
 			}
 		};
+
+		// If the transaction is aborted we need to remember to re-add the reference we released in the open!
+		// Note: this is different to how we handle ref counts in general in UE. Elsewhere we will eagerly
+		// increment ref counts, but delay decrement. This is because `0` is normally the 'destroy the object'
+		// trigger. But for UObject's `0` just clears the `RefCounted` flag, and since GC cannot run between
+		// when we clear the flag and potentially re-add the flag below, the GC couldn't observe the state
+		// that the object wasn't being actively ref-counted and thus could be collected. We do it differently
+		// here because a bunch of the UObject systems depend on the ref count being correct after calls to
+		// add/release.
+		UE_AUTORTFM_ONABORT(this)
+		{
+			this->AddRef();
+		};
 	}
 
 #if STATS || ENABLE_STATNAMEDEVENTS_UOBJECT
