LinkerTables: Deprecate PackagesToScriptSHAMap and remove support for it from SavePackage. It was only necessary for some platforms that we no longer support.
#rnx

diff --git a/Engine/Source/Runtime/CoreUObject/Private/UObject/Class.cpp b/Engine/Source/Runtime/CoreUObject/Private/UObject/Class.cpp
--- a/Engine/Source/Runtime/CoreUObject/Private/UObject/Class.cpp
+++ b/Engine/Source/Runtime/CoreUObject/Private/UObject/Class.cpp
@@ -2410,6 +2410,8 @@ void UStruct::Serialize(FArchive& Ar)
 			int32 iCode = 0;
 			int64 const BytecodeStartOffset = Ar.Tell();
 
+			PRAGMA_DISABLE_DEPRECATION_WARNINGS
+			// ScriptSHAHash is no longer supported; delete this if block when deleting it.
 			if (Ar.IsPersistent() && Ar.GetLinker())
 			{
 				// make sure this is a ULinkerSave
@@ -2442,6 +2444,7 @@ void UStruct::Serialize(FArchive& Ar)
 				LinkerSave->UpdateScriptSHAKey(TempScript);
 			}
 			else
+			PRAGMA_ENABLE_DEPRECATION_WARNINGS
 			{
 				FPropertyProxyArchive PropertyAr(Ar, iCode, this);
 				while (iCode < ScriptBytecodeSize)
diff --git a/Engine/Source/Runtime/CoreUObject/Private/UObject/Linker.cpp b/Engine/Source/Runtime/CoreUObject/Private/UObject/Linker.cpp
--- a/Engine/Source/Runtime/CoreUObject/Private/UObject/Linker.cpp
+++ b/Engine/Source/Runtime/CoreUObject/Private/UObject/Linker.cpp
@@ -168,7 +168,6 @@ FLinker::FLinker(ELinkerType::Type InType, UPackage* InRoot)
 , LinkerRoot( InRoot )
 , FilterClientButNotServer(false)
 , FilterServerButNotClient(false)
-, ScriptSHA(nullptr)
 {
 	check(LinkerRoot);
 
@@ -415,47 +414,34 @@ FString FLinker::GetExportFullName(int32 ExportIndex, const TCHAR* FakeRoot /*=
 	return FLinkerTables::GetExportFullName(FakeRoot ? FakeRoot : LinkerRoot->GetPathName(), ExportIndex, bResolveForcedExports);
 }
 
-/**
- * Tell this linker to start SHA calculations
- */
 void FLinker::StartScriptSHAGeneration()
 {
-	// create it if needed
+	PRAGMA_DISABLE_DEPRECATION_WARNINGS
 	if (ScriptSHA == NULL)
 	{
 		ScriptSHA = new FSHA1;
 	}
-
-	// make sure it's reset
 	ScriptSHA->Reset();
+	PRAGMA_ENABLE_DEPRECATION_WARNINGS
 }
 
-/**
- * If generating a script SHA key, update the key with this script code
- *
- * @param ScriptCode Code to SHAify
- */
 void FLinker::UpdateScriptSHAKey(const TArray<uint8>& ScriptCode)
 {
-	// if we are doing SHA, update it
+	PRAGMA_DISABLE_DEPRECATION_WARNINGS
 	if (ScriptSHA && ScriptCode.Num())
 	{
 		ScriptSHA->Update((uint8*)ScriptCode.GetData(), ScriptCode.Num());
 	}
+	PRAGMA_ENABLE_DEPRECATION_WARNINGS
 }
 
-/**
- * After generating the SHA key for all of the 
- *
- * @param OutKey Storage for the key bytes (20 bytes)
- */
 void FLinker::GetScriptSHAKey(uint8* OutKey)
 {
+	PRAGMA_DISABLE_DEPRECATION_WARNINGS
 	check(ScriptSHA);
-
-	// finish up the calculation, and return it
 	ScriptSHA->Final();
 	ScriptSHA->GetHash(OutKey);
+	PRAGMA_ENABLE_DEPRECATION_WARNINGS
 }
 
 PRAGMA_DISABLE_DEPRECATION_WARNINGS
diff --git a/Engine/Source/Runtime/CoreUObject/Private/UObject/LinkerSave.cpp b/Engine/Source/Runtime/CoreUObject/Private/UObject/LinkerSave.cpp
--- a/Engine/Source/Runtime/CoreUObject/Private/UObject/LinkerSave.cpp
+++ b/Engine/Source/Runtime/CoreUObject/Private/UObject/LinkerSave.cpp
@@ -27,8 +27,9 @@
 	FLinkerSave.
 ----------------------------------------------------------------------------*/
 
-/** A mapping of package name to generated script SHA keys */
+PRAGMA_DISABLE_DEPRECATION_WARNINGS
 TMap<FString, TArray<uint8> > FLinkerSave::PackagesToScriptSHAMap;
+PRAGMA_ENABLE_DEPRECATION_WARNINGS
 
 FLinkerSave::FLinkerSave(UPackage* InParent, const TCHAR* InFilename, bool bForceByteSwapping, bool bInSaveUnversioned)
 :	FLinker(ELinkerType::Save, InParent)
diff --git a/Engine/Source/Runtime/CoreUObject/Private/UObject/SavePackage2.cpp b/Engine/Source/Runtime/CoreUObject/Private/UObject/SavePackage2.cpp
--- a/Engine/Source/Runtime/CoreUObject/Private/UObject/SavePackage2.cpp
+++ b/Engine/Source/Runtime/CoreUObject/Private/UObject/SavePackage2.cpp
@@ -2965,18 +2965,6 @@ ESavePackageResult SaveHarvestedRealms(FSaveContext& SaveContext, ESaveRealm Har
 		return SaveContext.Result;
 	}
 
-	// SHA Generation
-	TArray<uint8>* ScriptSHABytes = nullptr;
-	{
-		// look for this package in the list of packages to generate script SHA for 
-		ScriptSHABytes = FLinkerSave::PackagesToScriptSHAMap.Find(*FPaths::GetBaseFilename(SaveContext.GetFilename()));
-		// if we want to generate the SHA key, start tracking script writes
-		if (ScriptSHABytes)
-		{
-			SaveContext.GetLinker()->StartScriptSHAGeneration();
-		}
-	}
-
 	// Write Exports 
 	{
 		SlowTask.EnterProgressFrame();
@@ -3051,20 +3039,6 @@ ESavePackageResult SaveHarvestedRealms(FSaveContext& SaveContext, ESaveRealm Har
 		SaveContext.GetLinker()->OnPostSaveBulkData();
 	}
 
-	// Get SHA Key
-	{
-		// if we want to generate the SHA key, get it out now that the package has finished saving
-		if (ScriptSHABytes && SaveContext.GetLinker()->ContainsCode())
-		{
-			// make space for the 20 byte key
-			ScriptSHABytes->Empty(20);
-			ScriptSHABytes->AddUninitialized(20);
-
-			// retrieve it
-			SaveContext.GetLinker()->GetScriptSHAKey(ScriptSHABytes->GetData());
-		}
-	}
-
 	IPackageWriter* PackageWriter = SaveContext.GetPackageWriter();
 	SaveContext.Result = WriteAdditionalFiles(SaveContext, SlowTask, VirtualExportsFileOffset /* In/Out */);
 	if (SaveContext.Result != ESavePackageResult::Success)
diff --git a/Engine/Source/Runtime/CoreUObject/Private/UObject/StructScriptLoader.cpp b/Engine/Source/Runtime/CoreUObject/Private/UObject/StructScriptLoader.cpp
--- a/Engine/Source/Runtime/CoreUObject/Private/UObject/StructScriptLoader.cpp
+++ b/Engine/Source/Runtime/CoreUObject/Private/UObject/StructScriptLoader.cpp
@@ -249,6 +249,8 @@ bool FStructScriptLoader::LoadStructWithScript(UStruct* DestScriptContainer, FAr
 	Ar.Seek(ScriptSerializationOffset);
 	if (bIsLinkerLoader)
 	{
+		PRAGMA_DISABLE_DEPRECATION_WARNINGS
+		// ScriptSHAHash is no longer supported; delete this if block when deleting it.
 		auto LinkerLoad = CastChecked<FLinkerLoad>(Ar.GetLinker());
 
 		TArray<uint8> ShaScriptBuffer;
@@ -259,6 +261,7 @@ bool FStructScriptLoader::LoadStructWithScript(UStruct* DestScriptContainer, FAr
 		LinkerLoad->UpdateScriptSHAKey(ShaScriptBuffer);
 
 		Ar.Seek(ScriptSerializationOffset);
+		PRAGMA_ENABLE_DEPRECATION_WARNINGS
 	}
 
 	DestScriptContainer->Script.Empty(BytecodeBufferSize);
diff --git a/Engine/Source/Runtime/CoreUObject/Private/UObject/UObjectGlobals.cpp b/Engine/Source/Runtime/CoreUObject/Private/UObject/UObjectGlobals.cpp
--- a/Engine/Source/Runtime/CoreUObject/Private/UObject/UObjectGlobals.cpp
+++ b/Engine/Source/Runtime/CoreUObject/Private/UObject/UObjectGlobals.cpp
@@ -1936,14 +1936,15 @@ UPackage* LoadPackageInternal(UPackage* InOuter, const FPackagePath& PackagePath
 		// Save the PackagePath we loaded from
 		Result->SetLoadedPath(PackagePath);
 		
-		// is there a script SHA hash for this package?
+		PRAGMA_DISABLE_DEPRECATION_WARNINGS
+		// ScriptSHAHash is no longer supported; delete this block when deleting it.
 		uint8 SavedScriptSHA[20];
 		bool bHasScriptSHAHash = FSHA1::GetFileSHAHash(*Linker->LinkerRoot->GetName(), SavedScriptSHA, false);
 		if (bHasScriptSHAHash)
 		{
-			// if there is, start generating the SHA for any script code in this package
 			Linker->StartScriptSHAGeneration();
 		}
+		PRAGMA_ENABLE_DEPRECATION_WARNINGS
 
 		if (SlowTask)
 		{
@@ -1991,19 +1992,18 @@ UPackage* LoadPackageInternal(UPackage* InOuter, const FPackagePath& PackagePath
 		UE::SetIsEditorLoadingPackage(IsEditorLoadingPackage.GetOriginalValue());
 #endif
 
-		// if we are calculating the script SHA for a package, do the comparison now
+		PRAGMA_DISABLE_DEPRECATION_WARNINGS
+		// ScriptSHAHash is no longer supported; delete this if block when deleting it.
 		if (bHasScriptSHAHash)
 		{
-			// now get the actual hash data
 			uint8 LoadedScriptSHA[20];
 			Linker->GetScriptSHAKey(LoadedScriptSHA);
-
-			// compare SHA hash keys
 			if (FMemory::Memcmp(SavedScriptSHA, LoadedScriptSHA, 20) != 0)
 			{
 				appOnFailSHAVerification(*Linker->GetPackagePath().GetLocalFullPath(), false);
 			}
 		}
+		PRAGMA_ENABLE_DEPRECATION_WARNINGS
 
 		Linker->Flush();
 
diff --git a/Engine/Source/Runtime/CoreUObject/Public/UObject/Linker.h b/Engine/Source/Runtime/CoreUObject/Public/UObject/Linker.h
--- a/Engine/Source/Runtime/CoreUObject/Public/UObject/Linker.h
+++ b/Engine/Source/Runtime/CoreUObject/Public/UObject/Linker.h
@@ -425,8 +425,8 @@ public:
 	/** If true, filter out exports that are for servers but not clients */
 	bool					FilterServerButNotClient;
 
-	/** The SHA1 key generator for this package, if active */
-	class FSHA1*			ScriptSHA;
+	UE_DEPRECATED(5.6, "No longer supported. Contact Epic if you need this functionality.")
+	class FSHA1*			ScriptSHA = nullptr;
 
 	/** Constructor. */
 	FLinker(ELinkerType::Type InType, UPackage* InRoot);
@@ -555,23 +555,11 @@ public:
 		return FString();
 	}
 
-	/**
-	 * Tell this linker to start SHA calculations
-	 */
+	UE_DEPRECATED(5.6, "No longer supported. Contact Epic if you need this functionality.")
 	void StartScriptSHAGeneration();
-
-	/**
-	 * If generating a script SHA key, update the key with this script code
-	 *
-	 * @param ScriptCode Code to SHAify
-	 */
+	UE_DEPRECATED(5.6, "No longer supported. Contact Epic if you need this functionality.")
 	void UpdateScriptSHAKey(const TArray<uint8>& ScriptCode);
-
-	/**
-	 * After generating the SHA key for all of the 
-	 *
-	 * @param OutKey Storage for the key bytes (20 bytes)
-	 */
+	UE_DEPRECATED(5.6, "No longer supported. Contact Epic if you need this functionality.")
 	void GetScriptSHAKey(uint8* OutKey);
 
 	/**
diff --git a/Engine/Source/Runtime/CoreUObject/Public/UObject/LinkerSave.h b/Engine/Source/Runtime/CoreUObject/Public/UObject/LinkerSave.h
--- a/Engine/Source/Runtime/CoreUObject/Public/UObject/LinkerSave.h
+++ b/Engine/Source/Runtime/CoreUObject/Public/UObject/LinkerSave.h
@@ -138,6 +138,7 @@ public:
 	TArray<TUniqueFunction<void(const FPackagePath& PackagePath, FObjectPostSaveContext ObjectSaveContext)>> PostSaveCallbacks;
 
 	/** A mapping of package name to generated script SHA keys */
+	UE_DEPRECATED(5.6, "No longer supported, contact Epic if you need this functionality.")
 	COREUOBJECT_API static TMap<FString, TArray<uint8> > PackagesToScriptSHAMap;
 
 	/** Constructor for file writer */
