Set default value of s.PreInsertPublicExports to true to optimize map creation for packages with very many public exports.

diff --git a/Engine/Source/Runtime/CoreUObject/Private/Serialization/AsyncLoading2.cpp b/Engine/Source/Runtime/CoreUObject/Private/Serialization/AsyncLoading2.cpp
--- a/Engine/Source/Runtime/CoreUObject/Private/Serialization/AsyncLoading2.cpp
+++ b/Engine/Source/Runtime/CoreUObject/Private/Serialization/AsyncLoading2.cpp
@@ -499,7 +499,7 @@ static FAutoConsoleVariableRef CVarDetectRaceDuringLoading(
 );
 #endif // USING_INSTRUMENTATION
 
-static bool GAsyncLoading2_PreInsertPublicExports = false;
+static bool GAsyncLoading2_PreInsertPublicExports = true;
 static FAutoConsoleVariableRef CVar_PreInsertPublicExports(
 	TEXT("s.PreInsertPublicExports"),
 	GAsyncLoading2_PreInsertPublicExports,
