Fixed some log messages in RemoteObjectTransform
- Fixed wrong byte counts for serialized data
- Added new logging for total bytes received/sent

diff --git a/Engine/Source/Runtime/CoreUObject/Private/UObject/RemoteObjectTransfer.cpp b/Engine/Source/Runtime/CoreUObject/Private/UObject/RemoteObjectTransfer.cpp
--- a/Engine/Source/Runtime/CoreUObject/Private/UObject/RemoteObjectTransfer.cpp
+++ b/Engine/Source/Runtime/CoreUObject/Private/UObject/RemoteObjectTransfer.cpp
@@ -262,12 +262,13 @@ UObject* MigrateObjectFromObjectData(FRemoteServerId OwnerServerId, FRemoteObjec
 	TArray<UObject*> ReceivedObjects;
 	int32 RequestedObjectIndex = DeserializeObjectData(ObjectData, ReceivedObjectRemoteIds, ReceivedObjects);
 
-	checkf(ReceivedObjects.Num() > 0, TEXT("PeerId:%s unable to deserialize object data (%d bytes)"), *GetGlobalServerId().ToString(), ObjectData.Bytes.Num());
+	checkf(ReceivedObjects.Num() > 0, TEXT("PeerId:%s unable to deserialize object data (%d bytes)"), *GetGlobalServerId().ToString(), ObjectData.GetNumBytes());
 
 	if (ReceivedObjects.Num() > 0)
 	{
 		if (UE_LOG_ACTIVE(LogRemoteObject, VeryVerbose))
 		{
+			UE_LOG(LogRemoteObject, VeryVerbose, TEXT("PeerId:%s deserialized object data (%d bytes)"), *GetGlobalServerId().ToString(), ObjectData.GetNumBytes());
 			for(int32 SubObjectIndex = 0; SubObjectIndex < ReceivedObjects.Num(); SubObjectIndex++)
 			{
 				FRemoteObjectId SubObjectId = UE::CoreUObject::Private::FObjectHandleUtils::GetRemoteId(ReceivedObjects[SubObjectIndex]);
@@ -540,7 +541,7 @@ void MigrateObjectToRemoteServerWithExplicitPriority(FRemoteTransactionId Reques
 		TSet<UObject*> SentObjects;
 		FRemoteObjectData* ObjectData = SerializeObjectData(Object, SentObjects);
 
-		UE_LOG(LogRemoteObject, VeryVerbose, TEXT("     : MigrateObjectToRemoteServer(%s) serialization complete obj %s (isremote %d)"), *RequestId.ToString(), *ObjectId.ToString(), Handle::IsRemote(Object));
+		UE_LOG(LogRemoteObject, VeryVerbose, TEXT("     : MigrateObjectToRemoteServer(%s) serialization complete obj %s (isremote %d, %d bytes)"), *RequestId.ToString(), *ObjectId.ToString(), Handle::IsRemote(Object), ObjectData->GetNumBytes());
 
 		for(UObject* SO : SentObjects)
 		{
