Handle non-ASCII subpaths in FUtf8String::SerializeWithoutFixup.

diff --git a/Engine/Source/Runtime/Core/Public/UObject/FortniteMainBranchObjectVersions.inl b/Engine/Source/Runtime/Core/Public/UObject/FortniteMainBranchObjectVersions.inl
--- a/Engine/Source/Runtime/Core/Public/UObject/FortniteMainBranchObjectVersions.inl
+++ b/Engine/Source/Runtime/Core/Public/UObject/FortniteMainBranchObjectVersions.inl
@@ -586,5 +586,8 @@ FORTNITE_MAIN_VERSION(MediaPlateHoldoutComponentRemoval, 190)
 // Changed PCG landscape cache default from "serialize at cook" to "never serialize"
 FORTNITE_MAIN_VERSION(PCGLandscapeCacheDefaultSerializationChanged, 191)
 
+// FSoftObjectPath::SubPathString changed to FUtf8String
+FORTNITE_MAIN_VERSION(SoftObjectPathUtf8SubPaths, 192)
+
 // -----<new versions can be added above this line>-------------------------------------------------
 #undef FORTNITE_MAIN_VERSION
\ No newline at end of file
diff --git a/Engine/Source/Runtime/CoreUObject/Private/UObject/SoftObjectPath.cpp b/Engine/Source/Runtime/CoreUObject/Private/UObject/SoftObjectPath.cpp
--- a/Engine/Source/Runtime/CoreUObject/Private/UObject/SoftObjectPath.cpp
+++ b/Engine/Source/Runtime/CoreUObject/Private/UObject/SoftObjectPath.cpp
@@ -9,6 +9,7 @@
 #include "Misc/PackageName.h"
 #include "Misc/StringBuilder.h"
 #include "Serialization/CompactBinaryWriter.h"
+#include "UObject/FortniteMainBranchObjectVersion.h"
 #include "UObject/LinkerLoad.h"
 #include "UObject/UObjectThreadContext.h"
 #include "UObject/CoreRedirects.h"
@@ -436,8 +437,97 @@ void FSoftObjectPath::SerializePath(FArchive& Ar)
 #endif // WITH_EDITOR
 }
 
+void SoftObjectPathLoadSubPathWorkaround(FArchive& Ar, FUtf8String& OutStr)
+{
+	int32 SaveNum = 0;
+	Ar << SaveNum;
+
+	if (SaveNum == 0)
+	{
+		// Empty strings are empty in any representation
+
+		OutStr.Empty();
+		return;
+	}
+
+	bool bLoadUnicodeChar = SaveNum < 0;
+	if (bLoadUnicodeChar)
+	{
+		// If SaveNum is negative, it's either a wide string with non-ANSI elements, or it's a corrupted archive
+
+		// If SaveNum cannot be negated due to integer overflow, Ar is corrupted.
+		if (SaveNum == MIN_int32)
+		{
+			Ar.SetCriticalError();
+			UE_LOG(LogCore, Error, TEXT("Archive is corrupted"));
+			return;
+		}
+
+		SaveNum = -SaveNum;
+	}
+
+	// Protect against network packets allocating too much memory
+	int64 MaxSerializeSize = Ar.GetMaxSerializeSize();
+	if ((MaxSerializeSize > 0) && (SaveNum > MaxSerializeSize))
+	{
+		Ar.SetCriticalError();
+		UE_LOG(LogCore, Error, TEXT("String is too large (Size: %i, Max: %i)"), SaveNum, MaxSerializeSize);
+		return;
+	}
+
+	int32 NumSavedBytes = bLoadUnicodeChar ? SaveNum * sizeof(UCS2CHAR) : SaveNum * sizeof(ANSICHAR);
+
+	// Load the saved bytes into an array
+	TArray<uint8> SavedBytes;
+	SavedBytes.Empty(SaveNum);
+	SavedBytes.AddUninitialized(NumSavedBytes);
+	Ar.Serialize(SavedBytes.GetData(), NumSavedBytes);
+
+	// Do byte swapping if necessary
+	if (bLoadUnicodeChar && Ar.IsByteSwapping())
+	{
+		UCS2CHAR* SavedBytesPtr = (UCS2CHAR*)SavedBytes.GetData();
+		for (int32 CharIndex = 0; CharIndex < SaveNum; ++CharIndex)
+		{
+			SavedBytesPtr[CharIndex] = ByteSwap(SavedBytesPtr[CharIndex]);
+		}
+	}
+
+	// If the last byte is a zero, then it's a wide string or a UTF-8 path ending with NUL - assume the former
+	if (bLoadUnicodeChar && ((UCS2CHAR*)SavedBytes.GetData())[SaveNum - 1] == (UCS2CHAR)0)
+	{
+		FString Temp = (UCS2CHAR*)SavedBytes.GetData();
+
+		// Inline combine any surrogate pairs in the data when loading into a UTF-32 string
+		StringConv::InlineCombineSurrogates(Temp);
+
+		// Since Microsoft's vsnwprintf implementation raises an invalid parameter warning
+		// with a character of 0xffff, scan for it and terminate the string there.
+		// 0xffff isn't an actual Unicode character anyway.
+		int Index = 0;
+		if (Temp.FindChar(0xffff, Index))
+		{
+			Temp[Index] = TEXT('\0');
+			Temp.TrimToNullTerminator();
+		}
+
+		OutStr = FUtf8String(Temp);
+		return;
+	}
+	else if (!bLoadUnicodeChar && ((ANSICHAR*)SavedBytes.GetData())[SaveNum - 1] == (ANSICHAR)0)
+	{
+		OutStr = (ANSICHAR*)SavedBytes.GetData();
+		return;
+	}
+
+	// Assume it was saved as a UTF-8 string without a NUL terminator
+	OutStr = FUtf8String::ConstructFromPtrSize((UTF8CHAR*)SavedBytes.GetData(), SavedBytes.Num());
+}
+
 void FSoftObjectPath::SerializePathWithoutFixup(FArchive& Ar)
 {
+	Ar.UsingCustomVersion(FFortniteMainBranchObjectVersion::GUID);
+
 	if (Ar.IsLoading() && Ar.UEVer() < VER_UE4_ADDED_SOFT_OBJECT_PATH)
 	{
 		FString Path;
@@ -456,12 +546,37 @@ void FSoftObjectPath::SerializePathWithoutFixup(FArchive& Ar)
 		Ar << AssetPathName;
 		AssetPath = WriteToString<FName::StringBufferSize>(AssetPathName).ToView();
 
-		Ar << SubPathString;
+		FWideString SubPathWide(SubPathString);
+		Ar << SubPathWide;
+		SubPathString = FUtf8String(SubPathWide);
 	}
 	else
 	{
 		Ar << AssetPath;
-		Ar << SubPathString;
+
+		if (Ar.CustomVer(FFortniteMainBranchObjectVersion::GUID) < FFortniteMainBranchObjectVersion::SoftObjectPathUtf8SubPaths)
+		{
+			if (Ar.IsLoading())
+			{
+				//**************//
+				//* WORKAROUND *//
+				//**************//
+				// Some packages were saved with UTF-8 strings without a version check, so we need to replicate the combined
+				// FWideString/FUtf8String serialization paths and attempt to guess which one we're loading.
+
+				SoftObjectPathLoadSubPathWorkaround(Ar, SubPathString);
+			}
+			else
+			{
+				FWideString WideSubPathString(SubPathString);
+				Ar << WideSubPathString;
+				SubPathString = FUtf8String(WideSubPathString);
+			}
+		}
+		else
+		{
+			Ar << SubPathString;
+		}
 	}
 }
 
