Added caching for IDO types because most of them are reusable

diff --git a/Engine/Source/Runtime/CoreUObject/Private/UObject/InstanceDataObjectUtils.cpp b/Engine/Source/Runtime/CoreUObject/Private/UObject/InstanceDataObjectUtils.cpp
--- a/Engine/Source/Runtime/CoreUObject/Private/UObject/InstanceDataObjectUtils.cpp
+++ b/Engine/Source/Runtime/CoreUObject/Private/UObject/InstanceDataObjectUtils.cpp
@@ -4,6 +4,9 @@
 
 #if WITH_EDITORONLY_DATA
 
+#include "Async/SharedLock.h"
+#include "Async/SharedMutex.h"
+#include "Async/UniqueLock.h"
 #include "HAL/IConsoleManager.h"
 #include "Logging/StructuredLog.h"
 #include "Misc/ReverseIterate.h"
@@ -477,6 +480,23 @@ namespace UE
 		return Struct->GetBoolMetaData(NAME_IsInstanceDataObjectStruct);
 	}
 
+	template <typename T>
+	static void CleanUpInstanceDataObjectTypeCache(TMap<FBlake3Hash, TWeakObjectPtr<T>>& Cache)
+	{
+		if (Cache.Num() % 64 == 0)
+		{
+			for (auto It = Cache.CreateIterator(); It; ++It)
+			{
+				if (!It->Value.IsValid())
+				{
+					It.RemoveCurrent();
+				}
+			}
+		}
+	}
+
+	static UScriptStruct* CreateInstanceDataObjectStruct(const FPropertyPathNameTree* PropertyTree, const FUnknownEnumNames* EnumNames, UScriptStruct* OwnerStruct, UObject* Outer, const FGuid& Guid, const TCHAR* OriginalName);
+
 	static UStruct* CreateInstanceDataObjectStructRec(const UClass* StructClass, UStruct* TemplateStruct, UObject* Outer, const FPropertyPathNameTree* PropertyTree, const FUnknownEnumNames* EnumNames);
 
 	template <typename StructType>
@@ -496,7 +516,11 @@ namespace UE
 		bool bHasFlags = false;
 
 		// Use the original type name because the template may be a fallback enum or an IDO.
-		FPropertyTypeName EnumTypeName = FindOriginalType(Property);
+		FPropertyTypeName EnumTypeName;
+		{
+			TGuardValue<bool> ImpersonateScope(FUObjectThreadContext::Get().GetSerializeContext()->bImpersonateProperties, true);
+			EnumTypeName = FindOriginalType(Property);
+		}
 		if (EnumTypeName.IsEmpty())
 		{
 			FPropertyTypeNameBuilder Builder;
@@ -556,15 +580,46 @@ namespace UE
 		const FName MaxEnumName = bContainsExistingMax ? TemplateEnum->GetNameByIndex(TemplateEnum->NumEnums() - 1) : MakeFullEnumName("MAX");
 		EnumValueNames.Emplace(MaxEnumName, bHasFlags ? CombinedEnumValues : MaxEnumValue);
 
+		// Construct a key for the enum cache.
+		FBlake3Hash Key;
+		{
+			FBlake3 KeyBuilder;
+			AppendHash(KeyBuilder, EnumTypeName);
+			for (const TPair<FName, int64>& Name : EnumValueNames)
+			{
+				AppendHash(KeyBuilder, Name.Key);
+				KeyBuilder.Update(&Name.Value, sizeof(Name.Value));
+			}
+			KeyBuilder.Update(&bHasFlags, sizeof(bHasFlags));
+			Key = KeyBuilder.Finalize();
+		}
+
+		// Check if a cached enum exists for this key.
+		static TMap<FBlake3Hash, TWeakObjectPtr<UEnum>> EnumCache;
+		static FSharedMutex EnumCacheMutex;
+		if (TSharedLock Lock(EnumCacheMutex); UEnum* Enum = EnumCache.FindRef(Key).Get())
+		{
+			return Enum;
+		}
+
 		// Construct a transient type that impersonates the original type.
 		const FName InstanceDataObjectName(WriteToString<128>(EnumName, TEXTVIEW("_InstanceDataObject")));
-		UEnum* Enum = NewObject<UEnum>(Outer, MakeUniqueObjectName(Outer, UEnum::StaticClass(), InstanceDataObjectName));
-		Enum->SetEnums(EnumValueNames, TemplateEnum->GetCppForm(), bHasFlags ? EEnumFlags::Flags : EEnumFlags::None, /*bAddMaxKeyIfMissing*/ false);
-		Enum->SetMetaData(*WriteToString<32>(NAME_OriginalType), *WriteToString<128>(EnumTypeName));
+		UEnum* NewEnum = NewObject<UEnum>(Outer, MakeUniqueObjectName(Outer, UEnum::StaticClass(), InstanceDataObjectName));
+		NewEnum->SetEnums(EnumValueNames, TemplateEnum->GetCppForm(), bHasFlags ? EEnumFlags::Flags : EEnumFlags::None, /*bAddMaxKeyIfMissing*/ false);
+		NewEnum->SetMetaData(*WriteToString<32>(NAME_OriginalType), *WriteToString<128>(EnumTypeName));
 
 		// TODO: Detect out-of-bounds values and increase the size of the underlying type accordingly.
 
-		return Enum;
+		TUniqueLock Lock(EnumCacheMutex);
+		if (UEnum* Enum = EnumCache.FindRef(Key).Get())
+		{
+			return Enum;
+		}
+
+		CleanUpInstanceDataObjectTypeCache(EnumCache);
+
+		EnumCache.Add(Key, NewEnum);
+		return NewEnum;
 	}
 
 	static FString UnmanglePropertyName(const FName MaybeMangledName, bool& bOutNameWasMangled)
@@ -602,30 +657,29 @@ namespace UE
 			if (!AsStructProperty->Struct->UseNativeSerialization() && !(AsStructProperty->Struct->StructFlags & STRUCT_Immutable))
 			{
 				//@note: Transfer existing metadata over as we build the InstanceDataObject from the struct or it owners, if any, this is useful for testing purposes
-				FString OriginalName;
+				TStringBuilder<256> OriginalName;
 				if (TGuardValue<bool> ImpersonateScope(FUObjectThreadContext::Get().GetSerializeContext()->bImpersonateProperties, true);
 					const FString* OriginalType = FindOriginalTypeName(AsStructProperty))
 				{
-					OriginalName = *OriginalType;
+					OriginalName << *OriginalType;
 				}
 
-				if (OriginalName.IsEmpty())
+				if (OriginalName.Len() == 0)
 				{
 					UE::FPropertyTypeNameBuilder OriginalNameBuilder;
 					OriginalNameBuilder.AddPath(AsStructProperty->Struct);
-					OriginalName = WriteToString<256>(OriginalNameBuilder.Build()).ToView();
+					OriginalName << OriginalNameBuilder.Build();
 				}
 
-				UInstanceDataObjectStruct* Struct = CreateInstanceDataObjectStructRec<UInstanceDataObjectStruct>(AsStructProperty->Struct, Outer, PropertyTree, EnumNames);
+				FGuid StructGuid;
 				if (const FName StructGuidName = PropertyType.GetParameterName(1); !StructGuidName.IsNone())
 				{
-					FGuid::Parse(StructGuidName.ToString(), Struct->Guid);
+					FGuid::Parse(StructGuidName.ToString(), StructGuid);
 				}
-				AsStructProperty->Struct = Struct;
+
+				AsStructProperty->Struct = CreateInstanceDataObjectStruct(PropertyTree, EnumNames, AsStructProperty->Struct, Outer, StructGuid, *OriginalName);
 				AsStructProperty->SetMetaData(NAME_OriginalType, *OriginalName);
 				AsStructProperty->SetMetaData(NAME_PresentAsTypeMetadata, *OriginalName);
-				AsStructProperty->Struct->SetMetaData(NAME_OriginalType, *OriginalName);
-				AsStructProperty->Struct->SetMetaData(NAME_PresentAsTypeMetadata, *OriginalName);
 			}
 		}
 		else if (FByteProperty* AsByteProperty = CastField<FByteProperty>(Property))
@@ -943,7 +997,7 @@ namespace UE
 		FArchive* Archive = nullptr;
 	};
 
-	void CopyTaggedProperties(const UObject* Source, UObject* Dest)
+	TArray<uint8> SaveTaggedProperties(const UObject* Source)
 	{
 		FUObjectSerializeContext* SerializeContext = FUObjectThreadContext::Get().GetSerializeContext();
 
@@ -954,26 +1008,40 @@ namespace UE
 		TGuardValue<bool> ScopedTrackSerializedProperties(SerializeContext->bTrackSerializedProperties, false);
 		TGuardValue<bool> ScopedTrackUnknownProperties(SerializeContext->bTrackUnknownProperties, false);
 		TGuardValue<bool> ScopedTrackUnknownEnumNames(SerializeContext->bTrackUnknownEnumNames, false);
+		TGuardValue<bool> ImpersonatePropertiesScope(SerializeContext->bImpersonateProperties, IsInstanceDataObject(Source));
 
-		TArray<uint8> Buffer;
-		Buffer.Reserve(Source->GetClass()->GetStructureSize());
+		TArray<uint8> Data;
+		FObjectWriter Writer(Data);
+		Writer.ArNoDelta = true;
+		FSerializingDefaultsScope WriterDefaultsScope(Writer, Source);
+		Source->GetClass()->SerializeTaggedProperties(Writer, (uint8*)Source, Source->GetClass(), nullptr);
 
-		{
-			FObjectWriter Writer(Buffer);
-			FSerializingDefaultsScope WriterDefaultsScope(Writer, Source);
-			Writer.ArNoDelta = true;
-			TGuardValue<bool> ImpersonatePropertiesScope(SerializeContext->bImpersonateProperties, IsInstanceDataObject(Source));
-			Source->GetClass()->SerializeTaggedProperties(Writer, (uint8*)Source, Source->GetClass(), nullptr);
-		}
+		return Data;
+	}
 
-		{
-			FObjectReader Reader(Buffer);
-			FSerializingDefaultsScope ReaderDefaultsScope(Reader, Dest);
-			Reader.ArMergeOverrides = true;
-			Reader.ArPreserveArrayElements = true;
-			TGuardValue<bool> ImpersonatePropertiesScope(SerializeContext->bImpersonateProperties, IsInstanceDataObject(Dest));
-			Dest->GetClass()->SerializeTaggedProperties(Reader, (uint8*)Dest, Dest->GetClass(), nullptr);
-		}
+	void LoadTaggedProperties(const TArray<uint8>& Source, UObject* Dest)
+	{
+		FUObjectSerializeContext* SerializeContext = FUObjectThreadContext::Get().GetSerializeContext();
+
+		// Track only initialized properties when copying. This is required to skip uninitialized properties
+		// during saving and to mark initialized properties during loading.
+		const bool bIsCDO = Dest->HasAnyFlags(RF_ClassDefaultObject);
+		TGuardValue<bool> ScopedTrackInitializedProperties(SerializeContext->bTrackInitializedProperties, !bIsCDO);
+		TGuardValue<bool> ScopedTrackSerializedProperties(SerializeContext->bTrackSerializedProperties, false);
+		TGuardValue<bool> ScopedTrackUnknownProperties(SerializeContext->bTrackUnknownProperties, false);
+		TGuardValue<bool> ScopedTrackUnknownEnumNames(SerializeContext->bTrackUnknownEnumNames, false);
+		TGuardValue<bool> ImpersonatePropertiesScope(SerializeContext->bImpersonateProperties, IsInstanceDataObject(Dest));
+
+		FObjectReader Reader(Source);
+		Reader.ArMergeOverrides = true;
+		Reader.ArPreserveArrayElements = true;
+		FSerializingDefaultsScope ReaderDefaultsScope(Reader, Dest);
+		Dest->GetClass()->SerializeTaggedProperties(Reader, (uint8*)Dest, Dest->GetClass(), nullptr);
+	}
+
+	void CopyTaggedProperties(const UObject* Source, UObject* Dest)
+	{
+		LoadTaggedProperties(SaveTaggedProperties(Source), Dest);
 	}
 
 	static void SetClassFlags(UClass* IDOClass, const UClass* OwnerClass)
@@ -989,21 +1057,100 @@ namespace UE
 
 	UClass* CreateInstanceDataObjectClass(const FPropertyPathNameTree* PropertyTree, const FUnknownEnumNames* EnumNames, UClass* OwnerClass, UObject* Outer)
 	{
-		UClass* Result = CreateInstanceDataObjectStructRec<UInstanceDataObjectClass>(OwnerClass, Outer, PropertyTree, EnumNames);
+		FBlake3Hash Key;
+		{
+			FBlake3 KeyBuilder;
+			KeyBuilder.Update(MakeMemoryView(OwnerClass->GetSchemaHash(/*bSkipEditorOnly*/ false).GetBytes()));
+
+			// Hash the index and serial number of the CDO because they will change if it is reinstanced.
+			// The schema hash excludes modifications made by constructors, and those will of course only be run on construction.
+			const UObject* DefaultObject = OwnerClass->GetDefaultObject();
+			const int32 DefaultIndex = GUObjectArray.ObjectToIndex(DefaultObject);
+			const int32 DefaultSerial = GUObjectArray.AllocateSerialNumber(DefaultIndex);
+			KeyBuilder.Update(&DefaultIndex, sizeof(DefaultIndex));
+			KeyBuilder.Update(&DefaultSerial, sizeof(DefaultSerial));
+
+			if (PropertyTree)
+			{
+				AppendHash(KeyBuilder, *PropertyTree);
+			}
+			if (EnumNames)
+			{
+				AppendHash(KeyBuilder, *EnumNames);
+			}
+			Key = KeyBuilder.Finalize();
+		}
+
+		static TMap<FBlake3Hash, TWeakObjectPtr<UInstanceDataObjectClass>> ClassCache;
+		static FSharedMutex ClassCacheMutex;
+		if (TSharedLock Lock(ClassCacheMutex); UClass* Class = ClassCache.FindRef(Key).Get())
+		{
+			return Class;
+		}
+
+		UInstanceDataObjectClass* NewClass = CreateInstanceDataObjectStructRec<UInstanceDataObjectClass>(OwnerClass, Outer, PropertyTree, EnumNames);
 		if (const FString& DisplayName = OwnerClass->GetMetaData(NAME_DisplayName); !DisplayName.IsEmpty())
 		{
-			Result->SetMetaData(NAME_DisplayName, *DisplayName);
+			NewClass->SetMetaData(NAME_DisplayName, *DisplayName);
 		}
 
-		SetClassFlags(Result, OwnerClass);
+		SetClassFlags(NewClass, OwnerClass);
+
+		CopyTaggedProperties(OwnerClass->GetDefaultObject(), NewClass->GetDefaultObject());
 
-		const UObject* OwnerCDO = OwnerClass->GetDefaultObject();
-		UObject* ResultCDO = Result->GetDefaultObject();
-		if (ensure(OwnerCDO && ResultCDO))
+		TUniqueLock Lock(ClassCacheMutex);
+		if (UClass* Class = ClassCache.FindRef(Key).Get())
 		{
-			CopyTaggedProperties(OwnerCDO, ResultCDO);
+			return Class;
 		}
-		return Result;
+
+		CleanUpInstanceDataObjectTypeCache(ClassCache);
+
+		ClassCache.Add(Key, NewClass);
+		return NewClass;
+	}
+
+	UScriptStruct* CreateInstanceDataObjectStruct(const FPropertyPathNameTree* PropertyTree, const FUnknownEnumNames* EnumNames, UScriptStruct* OwnerStruct, UObject* Outer, const FGuid& Guid, const TCHAR* OriginalName)
+	{
+		FBlake3Hash Key;
+		{
+			FBlake3 KeyBuilder;
+			KeyBuilder.Update(MakeMemoryView(OwnerStruct->GetSchemaHash(/*bSkipEditorOnly*/ false).GetBytes()));
+			KeyBuilder.Update(&Guid, sizeof(Guid));
+			KeyBuilder.Update(MakeMemoryView(WriteToUtf8String<256>(OriginalName)));
+			if (PropertyTree)
+			{
+				AppendHash(KeyBuilder, *PropertyTree);
+			}
+			if (EnumNames)
+			{
+				AppendHash(KeyBuilder, *EnumNames);
+			}
+			Key = KeyBuilder.Finalize();
+		}
+
+		static TMap<FBlake3Hash, TWeakObjectPtr<UInstanceDataObjectStruct>> StructCache;
+		static FSharedMutex StructCacheMutex;
+		if (TSharedLock Lock(StructCacheMutex); UScriptStruct* Struct = StructCache.FindRef(Key).Get())
+		{
+			return Struct;
+		}
+
+		UInstanceDataObjectStruct* NewStruct = CreateInstanceDataObjectStructRec<UInstanceDataObjectStruct>(OwnerStruct, Outer, PropertyTree, EnumNames);
+		NewStruct->Guid = Guid;
+		NewStruct->SetMetaData(NAME_OriginalType, OriginalName);
+		NewStruct->SetMetaData(NAME_PresentAsTypeMetadata, OriginalName);
+
+		TUniqueLock Lock(StructCacheMutex);
+		if (UScriptStruct* Struct = StructCache.FindRef(Key).Get())
+		{
+			return Struct;
+		}
+
+		CleanUpInstanceDataObjectTypeCache(StructCache);
+
+		StructCache.Add(Key, NewStruct);
+		return NewStruct;
 	}
 
 	static const FByteProperty* FindSerializedValuesProperty(const UStruct* Struct)
