[AsyncLoading2]

- Removes TMap implementation for FPublicExportMap and enables pre-populating the map upon load from the ExportMap contained in the package header and OptionaSegment (editor only).
- Removes s.PreInsertPublicExports CVAR as it's no longer necessary.
- Adds Num() and Count members to account the number of valid entries in the map (renames Count -> Capacity to better reflect its purpose as well)
- Adds tests for FPublicExportMap under various prefilling scenarios

#rnx

diff --git a/Engine/Source/Runtime/CoreUObject/Private/Serialization/AsyncLoading2.cpp b/Engine/Source/Runtime/CoreUObject/Private/Serialization/AsyncLoading2.cpp
--- a/Engine/Source/Runtime/CoreUObject/Private/Serialization/AsyncLoading2.cpp
+++ b/Engine/Source/Runtime/CoreUObject/Private/Serialization/AsyncLoading2.cpp
@@ -499,13 +499,6 @@ static FAutoConsoleVariableRef CVarDetectRaceDuringLoading(
 );
 #endif // USING_INSTRUMENTATION
 
-static bool GAsyncLoading2_PreInsertPublicExports = true;
-static FAutoConsoleVariableRef CVar_PreInsertPublicExports(
-	TEXT("s.PreInsertPublicExports"),
-	GAsyncLoading2_PreInsertPublicExports,
-	TEXT("After reading the package summary for a package, bulk insert all export hashes for public exports into the loaded package's map"),
-	ECVF_Default);
-
 CSV_DECLARE_CATEGORY_MODULE_EXTERN(CORE_API, Basic);
 CSV_DECLARE_CATEGORY_MODULE_EXTERN(CORE_API, FileIO);
 
@@ -795,15 +788,14 @@ class FLoadedPackageRef
 {
 private:
 	friend class FGlobalImportStore;
-
+	friend class FPublicExportMapBasicTest;
+	friend class FPublicExportMapIterationTest;
+	friend class FPublicExportMapPreInsertPublicExportsTest;
+	
 	class FPublicExportMap
 	{
-		static const int32 InvalidValue = -1;
-// The AGGRESSIVE_MEMORY_SAVING implementation allocates only as many keys and values as necessary to load a package's exports
-// which is space efficient but results in more memory traffic as exports are Store()'d or Remove()'d.
-// The !AGGRESSIVE_MEMORY_SAVING implementation uses more memory but provides O(1) Store() and Remove()
-#if AGGRESSIVE_MEMORY_SAVING && !WITH_EDITOR
 	public:
+		static const int32 InvalidValue = -1;
 
 		FPublicExportMap() = default;
 
@@ -812,9 +804,11 @@ private:
 		FPublicExportMap(FPublicExportMap&& Other)
 		{
 			Allocation = Other.Allocation;
+			Capacity = Other.Capacity;
 			Count = Other.Count;
 			SingleItemValue = Other.SingleItemValue;
 			Other.Allocation = nullptr;
+			Other.Capacity = 0;
 			Other.Count = 0;
 		};
 
@@ -822,39 +816,42 @@ private:
 
 		FPublicExportMap& operator=(FPublicExportMap&& Other)
 		{
-			if (Count > 1)
+			if (Capacity > 1)
 			{
 				FMemory::Free(Allocation);
 			}
 			Allocation = Other.Allocation;
+			Capacity = Other.Capacity;
 			Count = Other.Count;
 			SingleItemValue = Other.SingleItemValue;
 			Other.Allocation = nullptr;
+			Other.Capacity = 0;
 			Other.Count = 0;
 			return *this;
 		}
 
 		~FPublicExportMap()
 		{
-			if (Count > 1)
+			if (Capacity > 1)
 			{
 				FMemory::Free(Allocation);
 			}
 		}
-		void Reserve(int32 NewCount)
+		void Reserve(int32 NewCapacity)
 		{
-			if (NewCount <= Count)
+			if (NewCapacity <= Capacity)
 			{
 				return;
 			}
-			if (NewCount > 1)
+
+			if (NewCapacity > 1)
 			{
 				TArrayView<uint64> OldKeys = GetKeys();
 				TArrayView<int32> OldValues = GetValues();
-				const uint64 OldKeysSize = Count * sizeof(uint64);
-				const uint64 NewKeysSize = NewCount * sizeof(uint64);
-				const uint64 OldValuesSize = Count * sizeof(int32);
-				const uint64 NewValuesSize = NewCount * sizeof(int32);
+				const uint64 OldKeysSize = Capacity * sizeof(uint64);
+				const uint64 NewKeysSize = NewCapacity * sizeof(uint64);
+				const uint64 OldValuesSize = Capacity * sizeof(int32);
+				const uint64 NewValuesSize = NewCapacity * sizeof(int32);
 				const uint64 KeysToAddSize = NewKeysSize - OldKeysSize;
 				const uint64 ValuesToAddSize = NewValuesSize - OldValuesSize;
 
@@ -863,70 +860,118 @@ private:
 				FMemory::Memcpy(NewAllocation + KeysToAddSize, OldKeys.GetData(), OldKeysSize); // Copy old keys
 				FMemory::Memset(NewAllocation + NewKeysSize, 0xFF, ValuesToAddSize); // Insert new values initialized to InvalidValue
 				FMemory::Memcpy(NewAllocation + NewKeysSize + ValuesToAddSize, OldValues.GetData(), OldValuesSize); // Copy old values
-				if (Count > 1)
+
+				if (Capacity > 1)
 				{
 					FMemory::Free(Allocation);
 				}
 				Allocation = NewAllocation;
 			}
-			Count = NewCount;
+			Capacity = NewCapacity;
 		}
 
-		void PreInsertPublicExports(TConstArrayView<FExportMapEntry> Exports)
+		void PreInsertPublicExports(TConstArrayView<TConstArrayView<FExportMapEntry>> ExportLists)
 		{
-			if (Count <= 1) 
+			if (Capacity <= 1)
 			{ 
 				int32 PublicExportCount = 0;
-				for (const FExportMapEntry& Entry : Exports)
+				uint64 LastNewKeySeen = 0;
+				for (TConstArrayView<FExportMapEntry> Exports : ExportLists)
 				{
-					if (Entry.PublicExportHash && !AsyncLoading2_ShouldSkipLoadingExport(Entry.FilterFlags))
+					for (const FExportMapEntry& Entry : Exports)
 					{
-						PublicExportCount++;
+						if (Entry.PublicExportHash && !AsyncLoading2_ShouldSkipLoadingExport(Entry.FilterFlags))
+						{
+							PublicExportCount++;
+							LastNewKeySeen = Entry.PublicExportHash;
+						}
 					}
 				}
 
+				if (PublicExportCount == 0)
+				{
+					return;
+				}
+
 				// Either: 
 				// 	We're empty so we can just insert all the hashes with an associated value of -1 and sort.
 				//  We have a single pair in the SingleItemKey/SingleItemValue slots that we want to keep the mapping for.
-				// So we reserve space for the new keys under the assumption that it is most likely that the old key is also present in Exports.
-				uint64 OldKey = Count == 1 ? SingleItemKey : 0;
-				int32 OldValue = Count == 1 ? SingleItemValue : InvalidValue;
-				Reserve(PublicExportCount);	
-				TArrayView<uint64> Keys = GetKeys();
-				TArrayView<int32> Values = GetValues();
-				int32 Index = 0;
-				for (const FExportMapEntry& Entry : Exports)
+				const uint64 OldKey = Capacity == 1 ? SingleItemKey : 0;
+				const int32 OldValue = Capacity == 1 ? SingleItemValue : InvalidValue;
+
+				// We will be clearing all values during preinsert, thus Count is reset. We will restore
+				// Count when we Store() any valid OldValue below
+				Count = 0;
+
+				if (PublicExportCount == 1)
+				{
+					SingleItemKey = LastNewKeySeen;
+					SingleItemValue = InvalidValue;
+					Capacity = 1;
+				}
+				else
 				{
-					if (Entry.PublicExportHash && !AsyncLoading2_ShouldSkipLoadingExport(Entry.FilterFlags))
+					const int32 NewCapacity = PublicExportCount;
+					const uint64 NewKeysSize = NewCapacity * sizeof(uint64);
+					const uint64 NewValuesSize = NewCapacity * sizeof(int32);
+					uint8* NewAllocation = reinterpret_cast<uint8*>(FMemory::Malloc(NewKeysSize + NewValuesSize));
+
+					TArrayView<uint64> NewKeys = MakeArrayView(reinterpret_cast<uint64*>(NewAllocation), NewCapacity);
+					TArrayView<int32> NewValues = MakeArrayView(reinterpret_cast<int32*>(NewAllocation + sizeof(uint64) * NewCapacity), NewCapacity);
+					int32 Index = 0;
+
+					for (TConstArrayView<FExportMapEntry> Exports : ExportLists)
 					{
-						Keys[Index++] = Entry.PublicExportHash;
+						for (const FExportMapEntry& Entry : Exports)
+						{
+							if (Entry.PublicExportHash && !AsyncLoading2_ShouldSkipLoadingExport(Entry.FilterFlags))
+							{
+								NewKeys[Index++] = Entry.PublicExportHash;
+							}
+						}
 					}
+
+					Algo::Sort(NewKeys);
+
+					// Set all values to -1
+					FMemory::Memset(NewValues.GetData(), 0xFF, NewValues.NumBytes());
+					Allocation = NewAllocation;
+					Capacity = NewCapacity;
 				}
-				Algo::Sort(Keys);
-				// Set all values to -1
-				FMemory::Memset(Values.GetData(), 0xFF, Values.NumBytes());
+
 				// If we had a single object previously, re-store its old value.
 				// If the old key was present in Exports, this will overwrite the -1 we just stored.
 				// If it wasn't in Exports, this will grow the storage, but we consider that unlikely at present.
-				if (OldKey)
+				if (OldValue != InvalidValue)
 				{
 					Store(OldKey, OldValue);
 				}
 			}
 			else
 			{
-				const int32 OldCount = Count;
+				check(Allocation);
+				const int32 OldCapacity = Capacity;
 				TArrayView<uint64> OldKeys = GetKeys();
 				TArrayView<int32> OldValues = GetValues();
-				TArray<uint64> KeysToAdd;
-				for (const FExportMapEntry& Entry : Exports)
+				TArray<uint64, TInlineAllocator<256>> KeysToAdd;
+				int32 MaxKeysToAdd = 0;
+				for (TConstArrayView<FExportMapEntry> Exports : ExportLists)
 				{
-					if (Entry.PublicExportHash && !AsyncLoading2_ShouldSkipLoadingExport(Entry.FilterFlags))
+					MaxKeysToAdd += Exports.Num();
+				}
+				KeysToAdd.Reserve(MaxKeysToAdd);
+
+				for (TConstArrayView<FExportMapEntry> Exports : ExportLists)
+				{
+					for (const FExportMapEntry& Entry : Exports)
 					{
-						int32 Index = Algo::LowerBound(OldKeys, Entry.PublicExportHash);
-						if (Index >= Count || OldKeys[Index] != Entry.PublicExportHash)
+						if (Entry.PublicExportHash && !AsyncLoading2_ShouldSkipLoadingExport(Entry.FilterFlags))
 						{
-							KeysToAdd.Add(Entry.PublicExportHash);
+							int32 Index = Algo::LowerBound(OldKeys, Entry.PublicExportHash);
+							if (Index >= Capacity || OldKeys[Index] != Entry.PublicExportHash)
+							{
+								KeysToAdd.Add(Entry.PublicExportHash);
+							}
 						}
 					}
 				}
@@ -940,15 +985,15 @@ private:
 					Algo::Sort(KeysToAdd);
 					const int32 KeysToAddCount = KeysToAdd.Num();
 
-					// At this point, Count > 1 and KeysToAdd.Num() >= 1, so we don't need to consider the single object case and will always allocate.
-					const int32 NewCount = OldCount + KeysToAddCount;
-					const uint64 NewKeysSize = NewCount * sizeof(uint64);
-					const uint64 NewValuesSize = NewCount * sizeof(int32);
+					// At this point, Capacity > 1 and KeysToAdd.Num() >= 1, so we don't need to consider the single object case and will always allocate.
+					const int32 NewCapacity = OldCapacity + KeysToAddCount;
+					const uint64 NewKeysSize = NewCapacity * sizeof(uint64);
+					const uint64 NewValuesSize = NewCapacity * sizeof(int32);
 					uint8* NewAllocation = reinterpret_cast<uint8*>(FMemory::Malloc(NewKeysSize + NewValuesSize));
 
-					TArrayView<uint64> NewKeys = MakeArrayView(reinterpret_cast<uint64*>(NewAllocation), NewCount);
-					TArrayView<int32> NewValues = MakeArrayView(reinterpret_cast<int32*>(NewAllocation + sizeof(uint64) * NewCount), NewCount);
-					for (int32 InsertIndex = 0, OldIndex = 0, ToAddIndex = 0; InsertIndex < NewCount; ++InsertIndex)
+					TArrayView<uint64> NewKeys = MakeArrayView(reinterpret_cast<uint64*>(NewAllocation), NewCapacity);
+					TArrayView<int32> NewValues = MakeArrayView(reinterpret_cast<int32*>(NewAllocation + sizeof(uint64) * NewCapacity), NewCapacity);
+					for (int32 InsertIndex = 0, OldIndex = 0, ToAddIndex = 0; InsertIndex < NewCapacity; ++InsertIndex)
 					{
 						if (OldKeys.IsValidIndex(OldIndex)
 						&& (!KeysToAdd.IsValidIndex(ToAddIndex) || OldKeys[OldIndex] < KeysToAdd[ToAddIndex]))
@@ -966,26 +1011,35 @@ private:
 					}
 					FMemory::Free(Allocation);
 					Allocation = NewAllocation;
-					Count = NewCount;
+					Capacity = NewCapacity;
 				}
 			}
 		}
 
 		void Store(uint64 ExportHash, int32 ObjectIndex)
 		{
+			checkf(ExportHash, TEXT("Invalid to store a hash of 0 in this map"));
+			checkf(ObjectIndex != InvalidValue, TEXT("Use Remove() instead"));
+
 			TArrayView<uint64> Keys = GetKeys();
 			TArrayView<int32> Values = GetValues();
 			int32 Index = Algo::LowerBound(Keys, ExportHash);
-			if (Index < Count && Keys[Index] == ExportHash)
+			if (Index < Capacity && Keys[Index] == ExportHash)
 			{
 				// Slot already exists so reuse it
-				Values[Index] = ObjectIndex;
+				int32& ExistingValue = Values[Index];
+
+				// If we are storing a valid value on top of an invalid value 
+				// we have increased the number of valid indices
+				Count += ExistingValue == InvalidValue;
+				ExistingValue = ObjectIndex;
 				return;
 			}
-			if (Count == 0 || Keys[0] != 0)
+
+			if (Capacity == 0 || Keys[0] != 0)
 			{
 				// No free slots so we need to add one (will be inserted at the beginning of the array)
-				Reserve(Count + 1);
+				Reserve(Capacity + 1);
 				Keys = GetKeys();
 				Values = GetValues();
 			}
@@ -993,6 +1047,7 @@ private:
 			{
 				--Index; // Update insertion index to one before the lower bound item
 			}
+
 			if (Index > 0)
 			{
 				// Move items down
@@ -1001,17 +1056,23 @@ private:
 			}
 			Keys[Index] = ExportHash;
 			Values[Index] = ObjectIndex;
+			Count++;
 		}
 
 		bool Remove(uint64 ExportHash)
 		{
 			TArrayView<uint64> Keys = GetKeys();
 			int32 Index = Algo::LowerBound(Keys, ExportHash);
-			if (Index < Count && Keys[Index] == ExportHash)
+			if (Index < Capacity && Keys[Index] == ExportHash)
 			{
 				TArrayView<int32> Values = GetValues();
-				Values[Index] = InvalidValue;
-				return true;
+				int32& ExistingValue = Values[Index];
+
+				const bool bRemovedValue = ExistingValue != InvalidValue;
+				Count -= bRemovedValue;
+				ExistingValue = InvalidValue;
+
+				return bRemovedValue;
 			}
 
 			return false;
@@ -1022,7 +1083,7 @@ private:
 		{
 			TArrayView<uint64> Keys = GetKeys();
 			int32 Index = Algo::LowerBound(Keys, ExportHash);
-			if (Index < Count && Keys[Index] == ExportHash)
+			if (Index < Capacity && Keys[Index] == ExportHash)
 			{
 				TArrayView<int32> Values = GetValues();
 				return Values[Index];
@@ -1033,6 +1094,7 @@ private:
 		[[nodiscard]] bool PinForGC(TArray<int32>& OutUnreachableObjectIndices)
 		{
 			OutUnreachableObjectIndices.Reset();
+			int ValidExportCount = 0;
 			for (int32& ObjectIndex : GetValues())
 			{
 				if (ObjectIndex >= 0)
@@ -1043,15 +1105,18 @@ private:
 						UObject* Object = static_cast<UObject*>(ObjectItem->Object);
 						checkf(!ObjectItem->HasAnyFlags(EInternalObjectFlags::LoaderImport), TEXT("%s"), *Object->GetFullName());
 						ObjectItem->SetFlags(EInternalObjectFlags::LoaderImport);
+						ValidExportCount++;
 					}
 					else
 					{
-						OutUnreachableObjectIndices.Reserve(Count);
+						OutUnreachableObjectIndices.Reserve(Capacity);
 						OutUnreachableObjectIndices.Add(ObjectIndex);
 						ObjectIndex = InvalidValue;
+						Count--;
 					}
 				}
 			}
+			check(ValidExportCount == Count);
 			return OutUnreachableObjectIndices.Num() == 0;
 		}
 
@@ -1070,25 +1135,25 @@ private:
 
 		TArrayView<uint64> GetKeys()
 		{
-			if (Count == 1)
+			if (Capacity == 1)
 			{
 				return MakeArrayView(&SingleItemKey, 1);
 			}
 			else
 			{
-				return MakeArrayView(reinterpret_cast<uint64*>(Allocation), Count);
+				return MakeArrayView(reinterpret_cast<uint64*>(Allocation), Capacity);
 			}
 		}
 
 		TArrayView<int32> GetValues()
 		{
-			if (Count == 1)
+			if (Capacity == 1)
 			{
 				return MakeArrayView(&SingleItemValue, 1);
 			}
 			else
 			{
-				return MakeArrayView(reinterpret_cast<int32*>(Allocation + Count * sizeof(uint64)), Count);
+				return MakeArrayView(reinterpret_cast<int32*>(Allocation + Capacity * sizeof(uint64)), Capacity);
 			}
 		}
 
@@ -1146,16 +1211,21 @@ private:
 
 		FValueIterator CreateValueIterator()
 		{
-			if (Count == 1)
+			if (Capacity == 1)
 			{
 				return FValueIterator{ MakeArrayView(&SingleItemValue, 1) };
 			}
 			else
 			{
-				return FValueIterator{ MakeArrayView(reinterpret_cast<int32*>(Allocation + Count * sizeof(uint64)), Count) };
+				return FValueIterator{ MakeArrayView(reinterpret_cast<int32*>(Allocation + Capacity * sizeof(uint64)), Capacity) };
 			}
 		}
 
+		inline int32 Num() const
+		{
+			return Count;
+		}
+
 	private:
 		union
 		{
@@ -1163,138 +1233,13 @@ private:
 			uint64 SingleItemKey;
 		};
 		
-		// Total number of keys stored, not necessarily number of keys that map to valid values
-		// Count == 1 means we use SingleItemKey and SingleItemValue 
-		// Count > 1 means we use Allocation 
+		// Capacity is the total number of keys stored, not necessarily number of keys that map to valid values
+		// Capacity == 1 means we use SingleItemKey and SingleItemValue 
+		// Capacity > 1 means we use Allocation 
+		// Count represents the total valid values (non-invalid values, the objects may still be unreachable)
 		int32 Count = 0;
+		int32 Capacity = 0;
 		int32 SingleItemValue = -1;
-#else
-	public:
-		typedef TMap<uint64, int32> FObjectIndexMap;
-
-		void Reserve(int32 NewCount)
-		{
-			Data.Reserve(NewCount);
-		}
-
-		void PreInsertPublicExports(TConstArrayView<FExportMapEntry> Exports)
-		{
-			// Just behave like Reserve for the TMap version as we have cheap insertion
-			int32 PublicExportCount = 0;
-			for (const FExportMapEntry& Entry : Exports)
-			{
-				if (Entry.PublicExportHash && !AsyncLoading2_ShouldSkipLoadingExport(Entry.FilterFlags))
-				{
-					PublicExportCount++;
-				}
-			}
-			Data.Reserve(PublicExportCount);
-		}
-
-		void Store(uint64 ExportHash, int32 ObjectIndex)
-		{
-			Data.Add(ExportHash, ObjectIndex);
-		}
-
-		bool Remove(uint64 ExportHash)
-		{
-			return !!Data.Remove(ExportHash);
-		}
-
-		int32 Num() const
-		{
-			return Data.Num();
-		}
-
-		int32 Find(uint64 ExportHash)
-		{
-			int32* ObjectIndex = Data.Find(ExportHash);
-			if (ObjectIndex)
-			{
-				return *ObjectIndex;
-			}
-			return InvalidValue;
-		}
-
-		[[nodiscard]] bool PinForGC(TArray<int32>& OutUnreachableObjectIndices)
-		{
-			OutUnreachableObjectIndices.Reset();
-			for (FValueIterator It(Data); It; ++It)
-			{
-				FUObjectItem* ObjectItem = GUObjectArray.IndexToObject(*It);
-				if (!ObjectItem->IsUnreachable())
-				{
-					UObject* Object = static_cast<UObject*>(ObjectItem->Object);
-					checkf(!ObjectItem->HasAnyFlags(EInternalObjectFlags::LoaderImport), TEXT("%s"), *Object->GetFullName());
-					ObjectItem->SetFlags(EInternalObjectFlags::LoaderImport);
-				}
-				else
-				{
-					OutUnreachableObjectIndices.Reserve(Data.Num());
-					OutUnreachableObjectIndices.Add(*It);
-					It.RemoveCurrent();
-				}
-			}
-
-			return OutUnreachableObjectIndices.Num() == 0;
-		}
-
-		void UnpinForGC()
-		{
-			for (FValueIterator It(Data); It; ++It)
-			{
-				UObject* Object = static_cast<UObject*>(GUObjectArray.IndexToObject(*It)->Object);
-				checkf(Object->HasAnyInternalFlags(EInternalObjectFlags::LoaderImport), TEXT("%s"), *Object->GetFullName());
-				Object->AtomicallyClearInternalFlags(EInternalObjectFlags::LoaderImport);
-			}
-		}
-
-		class FValueIterator
-		{
-			FObjectIndexMap::TIterator It;
-		public:
-			FValueIterator(FObjectIndexMap& Data) : It(Data.CreateIterator()) { }
-
-			FORCEINLINE bool operator== (const FObjectIndexMap::TIterator& Other) const
-			{
-				return It == Other;
-			}
-
-			FORCEINLINE bool operator!= (const FObjectIndexMap::TIterator& Other) const
-			{
-				return !(It == Other);
-			}
-
-			FORCEINLINE explicit operator bool() const
-			{
-				return It.operator bool();
-			}
-
-			FORCEINLINE FValueIterator& operator++()
-			{
-				++It;
-				return *this;
-			}
-
-			FORCEINLINE int32 operator*()
-			{
-				return It.Value();
-			}
-
-			FORCEINLINE void RemoveCurrent()
-			{
-				It.RemoveCurrent();
-			}
-		};
-
-		FValueIterator CreateValueIterator()
-		{
-			return FValueIterator { Data };
-		}
-
-	private:
-		FObjectIndexMap Data;
-#endif
 	};
 
 	FPublicExportMap PublicExportMap;
@@ -1440,13 +1385,8 @@ public:
 		return PublicExportMap.CreateValueIterator();
 	}
 
-	void ReserveSpaceForPublicExports(int32 PublicExportCount)
-	{
-		PublicExportMap.Reserve(PublicExportCount);
-	}
-
 	/** Insert all the given exports into the map by their hash, storing a non-present/null value for any objects that are not already in the map */
-	void PreInsertPublicExports(TConstArrayView<FExportMapEntry> Exports)
+	void PreInsertPublicExports(TConstArrayView<TConstArrayView<FExportMapEntry>> Exports)
 	{
 		PublicExportMap.PreInsertPublicExports(Exports);
 	}
@@ -6358,14 +6298,8 @@ EEventLoadNodeExecutionResult FAsyncPackage2::ProcessLinkerLoadPackageSummary(FA
 		FLoadedPackageRef& PackageRef = AsyncLoadingThread.GlobalImportStore.FindPackageRefChecked(Desc.UPackageId, Desc.UPackageName);
 		if (!bLoadHasFailed)
 		{
-			if (GAsyncLoading2_PreInsertPublicExports)
-			{
-				PackageRef.PreInsertPublicExports(LinkerLoadState->LinkerLoadHeaderData.ExportMap);
-			}
-			else
-			{
-				PackageRef.ReserveSpaceForPublicExports(LinkerLoadState->LinkerLoadHeaderData.ExportMap.Num());
-			}
+			PackageRef.PreInsertPublicExports({ LinkerLoadState->LinkerLoadHeaderData.ExportMap });
+
 #if WITH_METADATA
 			// Create metadata object, this needs to happen before any other package wants to use our exports
 			LinkerLoadState->MetaDataIndex = LinkerLoadState->Linker->LoadMetaDataFromExportMap(false);
@@ -6875,37 +6809,16 @@ EEventLoadNodeExecutionResult FAsyncPackage2::Event_ProcessPackageSummary(FAsync
 		if (Package->Desc.bCanBeImported)
 		{
 			FLoadedPackageRef& PackageRef = Package->AsyncLoadingThread.GlobalImportStore.FindPackageRefChecked(Package->Desc.UPackageId, Package->Desc.UPackageName);
-			if (GAsyncLoading2_PreInsertPublicExports)
-			{
-				PackageRef.PreInsertPublicExports(Package->HeaderData.ExportMap);
+#if WITH_EDITOR
+			if (Package->OptionalSegmentHeaderData.IsSet())
+			{ 
+				PackageRef.PreInsertPublicExports({ Package->HeaderData.ExportMap, Package->OptionalSegmentHeaderData->ExportMap } );
 			}
 			else
-			{
-				int32 PublicExportsCount = 0;
-				for (const FExportMapEntry& Export : Package->HeaderData.ExportMap)
-				{
-					if (Export.PublicExportHash)
-					{
-						++PublicExportsCount;
-					}
-				}
-#if WITH_EDITOR
-				if (Package->OptionalSegmentHeaderData.IsSet())
-				{
-					for (const FExportMapEntry& Export : Package->OptionalSegmentHeaderData->ExportMap)
-					{
-						if (Export.PublicExportHash)
-						{
-							++PublicExportsCount;
-						}
-					}
-				}
 #endif
-				if (PublicExportsCount)
-				{
-					PackageRef.ReserveSpaceForPublicExports(PublicExportsCount);
-				}
-		}
+			{ 
+				PackageRef.PreInsertPublicExports({ Package->HeaderData.ExportMap });
+			}
 		}
 
 		check(Package->Desc.PackageIdToLoad == FPackageId::FromName(Package->HeaderData.PackageName));
@@ -10741,3 +10654,434 @@ IAsyncPackageLoader* MakeAsyncPackageLoader2(FIoDispatcher& InIoDispatcher, IAsy
 {
 	return new FAsyncLoadingThread2(InIoDispatcher, InUncookedPackageLoader);
 }
+
+#if WITH_DEV_AUTOMATION_TESTS 
+#include "Misc/AutomationTest.h"
+
+IMPLEMENT_SIMPLE_AUTOMATION_TEST(FPublicExportMapBasicTest, "System.Core.Loading.FPublicExportMap.Basic", EAutomationTestFlags_ApplicationContextMask | EAutomationTestFlags::SmokeFilter);
+bool FPublicExportMapBasicTest::RunTest(const FString& Parameters)
+{
+	typedef FLoadedPackageRef::FPublicExportMap FPublicExportMap;
+
+	FPublicExportMap Map;
+	TestEqual(TEXT("Empty map has no count"), Map.Num(), 0);
+	TestEqual(TEXT("Empty map has no keys"), Map.GetKeys().Num(), 0);
+	TestEqual(TEXT("Empty map has no values"), Map.GetValues().Num(), 0);
+
+	Map.Store(1, 2);
+	TestEqual(TEXT("Find stored key"), Map.Find(1), 2);
+	TestEqual(TEXT("Store increased size"), Map.Num(), 1);
+
+	// May resize
+	Map.Store(2, 3);
+	TestEqual(TEXT("Find stored key"), Map.Find(2), 3);
+	TestEqual(TEXT("Store increased size"), Map.Num(), 2);
+
+	TestEqual(TEXT("Remove existing key"), Map.Remove(1), true);
+	TestEqual(TEXT("Removed decreased count"), Map.Num(), 1);
+	TestEqual(TEXT("Remove removed key fails"), Map.Remove(1), false);
+	TestEqual(TEXT("Failed remove does not change count"), Map.Num(), 1);
+	TestEqual(TEXT("Find removed key fails"), Map.Find(1), FPublicExportMap::InvalidValue);
+
+	TestEqual(TEXT("Remove existing key"), Map.Remove(2), true);
+	TestEqual(TEXT("Removed decreased count"), Map.Num(), 0);
+	TestEqual(TEXT("Remove removed key fails"), Map.Remove(2), false);
+	TestEqual(TEXT("Failed remove does not change count"), Map.Num(), 0);
+	TestEqual(TEXT("Find removed key fails"), Map.Find(2), FPublicExportMap::InvalidValue);
+	return true;
+}
+
+IMPLEMENT_SIMPLE_AUTOMATION_TEST(FPublicExportMapIterationTest, "System.Core.Loading.FPublicExportMap.Iteration", EAutomationTestFlags_ApplicationContextMask | EAutomationTestFlags::SmokeFilter);
+bool FPublicExportMapIterationTest::RunTest(const FString& Parameters)
+{
+	typedef FLoadedPackageRef::FPublicExportMap FPublicExportMap;
+
+	FPublicExportMap Map;
+	TArray<uint64> Keys = { 1, 2, 3, 4, 5, 6, 7, 8 };
+	TArray<int32> Values = { 1, 2, 3, 4, 5, 6, 7, 8 };
+	check(Keys.Num() == Values.Num());
+
+	TestEqual(TEXT("Empty map has no elements"), Map.Num(), 0);
+	for (int32 i = 0; i < Keys.Num(); ++i)
+	{
+		Map.Store(Keys[i], Values[i]);
+	}
+		
+	TSet<uint64> FoundKeys;
+	for (uint64 Key : Map.GetKeys())
+	{
+		FoundKeys.Add(Key);
+	}
+	TSet<int32> FoundValues;
+	for (int32 Value : Map.GetValues())
+	{
+		FoundValues.Add(Value);
+	}
+
+	TestEqual(TEXT("Stored keys matches expected count"), Map.GetKeys().Num(), FoundKeys.Num());
+	for (uint64 Key : Keys)
+	{
+		TestTrue(TEXT("All stored keys can be found"), FoundKeys.Contains(Key));
+	}
+
+	TestEqual(TEXT("Stored values matches expected count"), Map.GetValues().Num(), FoundValues.Num());
+	for (uint64 Key : Keys)
+	{
+		TestTrue(TEXT("All stored values can be found"), FoundKeys.Contains(Key));
+	}
+
+	for (uint64 Key : Keys)
+	{
+		TestEqual(TEXT("Lookup works"), Map.Find(Key), Values[(int32)(Key - 1)]);
+	}
+	return true;
+}
+
+IMPLEMENT_SIMPLE_AUTOMATION_TEST(FPublicExportMapPreInsertPublicExportsTest, "System.Core.Loading.FPublicExportMap.PreInsertPublicExports", EAutomationTestFlags_ApplicationContextMask | EAutomationTestFlags::SmokeFilter);
+bool FPublicExportMapPreInsertPublicExportsTest::RunTest(const FString& Parameters)
+{
+	typedef FLoadedPackageRef::FPublicExportMap FPublicExportMap;
+
+	// Empty map
+	{
+		FPublicExportMap Map;
+		TArray<FExportMapEntry> Exports = { {.PublicExportHash = 1 }, {.PublicExportHash = 2 }, {.PublicExportHash = 3 }, {.PublicExportHash = 4 } };
+		int32 OriginaMapCount = Map.Num();
+
+		TestEqual(TEXT("Empty map has no count"), OriginaMapCount, 0);
+		TestEqual(TEXT("Empty map has no keys"), Map.GetKeys().Num(), OriginaMapCount);
+		TestEqual(TEXT("Empty map has no values"), Map.GetValues().Num(), OriginaMapCount);
+		Map.PreInsertPublicExports({ Exports });
+		TestEqual(TEXT("Empty map has no count"), Map.Num(), OriginaMapCount);
+		TestEqual(TEXT("PreInsert has populated our keys"), Map.GetKeys().Num(), Exports.Num());
+		TestEqual(TEXT("PreInsert has populated our values (should mostly be InvalidValue)"), Map.GetValues().Num(), Exports.Num());
+
+		for (uint64 Key : Map.GetKeys())
+		{
+			TestEqual(TEXT("All preinserted keys are initialized to InvalidValue"), Map.Find(Key), FPublicExportMap::InvalidValue);
+		}
+		TestEqual(TEXT("PreInsertPublicExports does not store new values"), Map.Num(), 0);
+
+		// Overwrite preinserted value
+		TestEqual(TEXT("Find stored key with no value set"), Map.Find(Exports[0].PublicExportHash), FPublicExportMap::InvalidValue);
+		Map.Store(Exports[0].PublicExportHash, (int32)Exports[0].PublicExportHash);
+		TestEqual(TEXT("Find stored key"), Map.Find(Exports[0].PublicExportHash), (int32)Exports[0].PublicExportHash);
+		TestEqual(TEXT("Store increased count"), Map.Num(), OriginaMapCount + 1);
+
+		// Insert new value
+		TestEqual(TEXT("Find missing key"), Map.Find(111), FPublicExportMap::InvalidValue);
+		Map.Store(111, 111);
+		TestEqual(TEXT("Find stored key"), Map.Find(111), 111);
+		TestEqual(TEXT("Store increased count"), Map.Num(), OriginaMapCount + 2);
+
+		TestEqual(TEXT("Remove existing key"), Map.Remove(111), true);
+		TestEqual(TEXT("Remove decreased count"), Map.Num(), OriginaMapCount + 1);
+		TestEqual(TEXT("Remove removed key fails"), Map.Remove(111), false);
+		TestEqual(TEXT("Failed remove did not increase count"), Map.Num(), OriginaMapCount + 1);
+		TestEqual(TEXT("Find removed key returns no value set"), Map.Find(111), FPublicExportMap::InvalidValue);
+	}
+
+	// Empty map, populate a single key
+	{
+		FPublicExportMap Map;
+		TArray<FExportMapEntry> Exports = { {.PublicExportHash = 1 }};
+		int32 OriginaMapCount = Map.Num();
+
+		TestEqual(TEXT("Empty map has no count"), OriginaMapCount, 0);
+		TestEqual(TEXT("Empty map has no keys"), Map.GetKeys().Num(), OriginaMapCount);
+		TestEqual(TEXT("Empty map has no values"), Map.GetValues().Num(), OriginaMapCount);
+		Map.PreInsertPublicExports({ Exports });
+		TestEqual(TEXT("Empty map has no count"), Map.Num(), OriginaMapCount);
+		TestEqual(TEXT("PreInsert has populated our keys"), Map.GetKeys().Num(), Exports.Num());
+		TestEqual(TEXT("PreInsert has populated our values (should mostly be InvalidValue)"), Map.GetValues().Num(), Exports.Num());
+
+		for (uint64 Key : Map.GetKeys())
+		{
+			TestEqual(TEXT("All preinserted keys are initialized to InvalidValue"), Map.Find(Key), FPublicExportMap::InvalidValue);
+		}
+		TestEqual(TEXT("PreInsertPublicExports does not store new values"), Map.Num(), 0);
+
+		// Overwrite preinserted value
+		TestEqual(TEXT("Find stored key with no value set"), Map.Find(Exports[0].PublicExportHash), FPublicExportMap::InvalidValue);
+		Map.Store(Exports[0].PublicExportHash, (int32)Exports[0].PublicExportHash);
+		TestEqual(TEXT("Find stored key"), Map.Find(Exports[0].PublicExportHash), (int32)Exports[0].PublicExportHash);
+		TestEqual(TEXT("Store increased count"), Map.Num(), OriginaMapCount + 1);
+
+		// Insert new value
+		TestEqual(TEXT("Find missing key"), Map.Find(111), FPublicExportMap::InvalidValue);
+		Map.Store(111, 111);
+		TestEqual(TEXT("Find stored key"), Map.Find(111), 111);
+		TestEqual(TEXT("Store increased count"), Map.Num(), OriginaMapCount + 2);
+
+		TestEqual(TEXT("Remove existing key"), Map.Remove(111), true);
+		TestEqual(TEXT("Remove decreased count"), Map.Num(), OriginaMapCount + 1);
+		TestEqual(TEXT("Remove removed key fails"), Map.Remove(111), false);
+		TestEqual(TEXT("Failed remove did not increase count"), Map.Num(), OriginaMapCount + 1);
+		TestEqual(TEXT("Find removed key returns no value set"), Map.Find(111), FPublicExportMap::InvalidValue);
+	}
+
+
+	// One existing entry (preinsert multiple keys, some overlapping)
+	{
+		FPublicExportMap Map;
+		TArray<FExportMapEntry> Exports = { {.PublicExportHash = 1 }, {.PublicExportHash = 2 }, {.PublicExportHash = 3 }, {.PublicExportHash = 4 } };
+
+		uint64 PreExistingKey = 2;
+		int32 PreExistingValue = 2;
+		Map.Store(PreExistingKey, PreExistingValue);
+		int32 OriginaMapCount = Map.Num();
+
+		TestEqual(TEXT("Find stored key"), Map.Find(PreExistingKey), PreExistingValue);
+		TestEqual(TEXT("Store increased count"), OriginaMapCount, 1);
+		TestEqual(TEXT("Store increased keys"), Map.GetKeys().Num(), OriginaMapCount);
+		TestEqual(TEXT("Store increased values"), Map.GetValues().Num(), OriginaMapCount);
+
+		Map.PreInsertPublicExports({ Exports });
+		TestEqual(TEXT("PreInsert doesn't add new values"), Map.Num(), OriginaMapCount);
+		TestEqual(TEXT("PreInsert has populated our keys"), Map.GetKeys().Num(), Exports.Num());
+		TestEqual(TEXT("PreInsert has populated our values (should mostly be InvalidValue)"), Map.GetValues().Num(), Exports.Num());
+
+		bool bFoundPreExistingKey = false;
+		for (uint64 Key : Map.GetKeys())
+		{
+			if (Key == PreExistingKey)
+			{
+				bFoundPreExistingKey = true;
+				TestEqual(TEXT("Pre-existing keys remain intact after PreInsertPublicExports"), Map.Find(Key), PreExistingValue);
+			}
+			else
+			{
+				TestEqual(TEXT("All preinserted keys are initialized to InvalidValue"), Map.Find(Key), FPublicExportMap::InvalidValue);
+			}
+		}
+		TestEqual(TEXT("GetKeys has preexisting key"), bFoundPreExistingKey, true);
+		TestEqual(TEXT("PreInsertPublicExports does not add new values"), Map.Num(), OriginaMapCount);
+
+		// Overwrite preinserted value
+		TestEqual(TEXT("Find stored key with no value set"), Map.Find(Exports[0].PublicExportHash), FPublicExportMap::InvalidValue);
+		Map.Store(Exports[0].PublicExportHash, (int32)Exports[0].PublicExportHash);
+		TestEqual(TEXT("Find stored key"), Map.Find(Exports[0].PublicExportHash), (int32)Exports[0].PublicExportHash);
+		TestEqual(TEXT("Store increased count"), Map.Num(), OriginaMapCount + 1);
+
+		// Insert new value
+		TestEqual(TEXT("Find missing key"), Map.Find(111), FPublicExportMap::InvalidValue);
+		Map.Store(111, 111);
+		TestEqual(TEXT("Find stored key"), Map.Find(111), 111);
+		TestEqual(TEXT("Store increased count"), Map.Num(), OriginaMapCount + 2);
+
+		TestEqual(TEXT("Remove existing key"), Map.Remove(111), true);
+		TestEqual(TEXT("Remove decreased count"), Map.Num(), OriginaMapCount + 1);
+		TestEqual(TEXT("Remove removed key fails"), Map.Remove(111), false);
+		TestEqual(TEXT("Failed remove did not increase count"), Map.Num(), OriginaMapCount + 1);
+		TestEqual(TEXT("Find removed key returns no value set"), Map.Find(111), FPublicExportMap::InvalidValue);
+	}
+
+	// One existing entry (preinsert a single entry overlapping)
+	{
+		FPublicExportMap Map;
+		TArray<FExportMapEntry> Exports = { {.PublicExportHash = 1 }};
+
+		uint64 PreExistingKey = 1;
+		int32 PreExistingValue = 1;
+		Map.Store(PreExistingKey, PreExistingValue);
+		int32 OriginaMapCount = Map.Num();
+
+		TestEqual(TEXT("Find stored key"), Map.Find(PreExistingKey), PreExistingValue);
+		TestEqual(TEXT("Store increased count"), OriginaMapCount, 1);
+		TestEqual(TEXT("Store increased keys"), Map.GetKeys().Num(), OriginaMapCount);
+		TestEqual(TEXT("Store increased values"), Map.GetValues().Num(), OriginaMapCount);
+
+		Map.PreInsertPublicExports({ Exports });
+		TestEqual(TEXT("PreInsert doesn't add new values"), Map.Num(), OriginaMapCount);
+		TestEqual(TEXT("PreInsert has populated our keys"), Map.GetKeys().Num(), Exports.Num());
+		TestEqual(TEXT("PreInsert has populated our values (should mostly be InvalidValue)"), Map.GetValues().Num(), Exports.Num());
+
+		bool bFoundPreExistingKey = false;
+		for (uint64 Key : Map.GetKeys())
+		{
+			if (Key == PreExistingKey)
+			{
+				bFoundPreExistingKey = true;
+				TestEqual(TEXT("Pre-existing keys remain intact after PreInsertPublicExports"), Map.Find(Key), PreExistingValue);
+			}
+			else
+			{
+				TestEqual(TEXT("All preinserted keys are initialized to InvalidValue"), Map.Find(Key), FPublicExportMap::InvalidValue);
+			}
+		}
+		TestEqual(TEXT("GetKeys has preexisting key"), bFoundPreExistingKey, true);
+		TestEqual(TEXT("PreInsertPublicExports does not add new values"), Map.Num(), OriginaMapCount);
+	}
+
+	// One existing entry (preinsert a single entry non-overlapping)
+	{
+		FPublicExportMap Map;
+		TArray<FExportMapEntry> Exports = { {.PublicExportHash = 1 } };
+
+		uint64 PreExistingKey = 2;
+		int32 PreExistingValue = 2;
+		Map.Store(PreExistingKey, PreExistingValue);
+		int32 OriginaMapCount = Map.Num();
+
+		TestEqual(TEXT("Find stored key"), Map.Find(PreExistingKey), PreExistingValue);
+		TestEqual(TEXT("Store increased count"), OriginaMapCount, 1);
+		TestEqual(TEXT("Store increased keys"), Map.GetKeys().Num(), OriginaMapCount);
+		TestEqual(TEXT("Store increased values"), Map.GetValues().Num(), OriginaMapCount);
+
+		Map.PreInsertPublicExports({ Exports });
+		TestEqual(TEXT("PreInsert doesn't add new values"), Map.Num(), OriginaMapCount);
+		TestEqual(TEXT("PreInsert has populated our keys"), Map.GetKeys().Num(), Exports.Num() + OriginaMapCount);
+		TestEqual(TEXT("PreInsert has populated our values (should mostly be InvalidValue)"), Map.GetValues().Num(), Exports.Num() + OriginaMapCount);
+
+		bool bFoundPreExistingKey = false;
+		for (uint64 Key : Map.GetKeys())
+		{
+			if (Key == PreExistingKey)
+			{
+				bFoundPreExistingKey = true;
+				TestEqual(TEXT("Pre-existing keys remain intact after PreInsertPublicExports"), Map.Find(Key), PreExistingValue);
+			}
+			else
+			{
+				TestEqual(TEXT("All preinserted keys are initialized to InvalidValue"), Map.Find(Key), FPublicExportMap::InvalidValue);
+			}
+		}
+		TestEqual(TEXT("GetKeys has preexisting key"), bFoundPreExistingKey, true);
+		TestEqual(TEXT("PreInsertPublicExports does not add new values"), Map.Num(), OriginaMapCount);
+
+		// Overwrite preinserted value
+		TestEqual(TEXT("Find stored key with no value set"), Map.Find(Exports[0].PublicExportHash), FPublicExportMap::InvalidValue);
+		Map.Store(Exports[0].PublicExportHash, (int32)Exports[0].PublicExportHash);
+		TestEqual(TEXT("Find stored key"), Map.Find(Exports[0].PublicExportHash), (int32)Exports[0].PublicExportHash);
+		TestEqual(TEXT("Store increased count"), Map.Num(), OriginaMapCount + 1);
+	}
+
+	// One existing entry (preinsert multiple keys, non-overlapping)
+	{
+		FPublicExportMap Map;
+		TArray<FExportMapEntry> Exports = { {.PublicExportHash = 1 }, {.PublicExportHash = 2 }, {.PublicExportHash = 3 }, {.PublicExportHash = 4 } };
+
+		uint64 PreExistingKey = 7;
+		int32 PreExistingValue = 8;
+		Map.Store(PreExistingKey, PreExistingValue);
+		int32 OriginaMapCount = Map.Num();
+
+		TestEqual(TEXT("Find stored key"), Map.Find(PreExistingKey), PreExistingValue);
+		TestEqual(TEXT("Store increased count"), OriginaMapCount, 1);
+		TestEqual(TEXT("Store increased keys"), Map.GetKeys().Num(), OriginaMapCount);
+		TestEqual(TEXT("Store increased values"), Map.GetValues().Num(), OriginaMapCount);
+
+		Map.PreInsertPublicExports({ Exports });
+		TestEqual(TEXT("PreInsertPublicExports has populated our keys"), Map.GetKeys().Num(), Exports.Num() + OriginaMapCount);
+		TestEqual(TEXT("PreInsert has populated our values (should mostly be InvalidValue)"), Map.GetValues().Num(), Exports.Num() + OriginaMapCount);
+
+		bool bFoundPreExistingKey = false;
+		for (uint64 Key : Map.GetKeys())
+		{
+			if (Key == PreExistingKey)
+			{
+				bFoundPreExistingKey = true;
+				TestEqual(TEXT("Pre-existing keys remain intact after PreInsertPublicExports"), Map.Find(Key), PreExistingValue);
+			}
+			else
+			{
+				TestEqual(TEXT("All preinserted keys are initialized to InvalidValue"), Map.Find(Key), FPublicExportMap::InvalidValue);
+			}
+		}
+		TestEqual(TEXT("GetKeys has preexisting key"), bFoundPreExistingKey, true);
+		TestEqual(TEXT("PreInsertPublicExports does not add new values"), Map.Num(), OriginaMapCount);
+
+		// Overwrite preinserted value
+		TestEqual(TEXT("Find stored key with no value set"), Map.Find(Exports[0].PublicExportHash), FPublicExportMap::InvalidValue);
+		Map.Store(Exports[0].PublicExportHash, (int32)Exports[0].PublicExportHash);
+		TestEqual(TEXT("Find stored key"), Map.Find(Exports[0].PublicExportHash), (int32)Exports[0].PublicExportHash);
+		TestEqual(TEXT("Store increased count"), Map.Num(), OriginaMapCount + 1);
+
+		// Insert new value
+		TestEqual(TEXT("Find missing key"), Map.Find(111), FPublicExportMap::InvalidValue);
+		Map.Store(111, 111);
+		TestEqual(TEXT("Find stored key"), Map.Find(111), 111);
+		TestEqual(TEXT("Store increased count"), Map.Num(), OriginaMapCount + 2);
+
+		TestEqual(TEXT("Remove existing key"), Map.Remove(111), true);
+		TestEqual(TEXT("Remove decreased count"), Map.Num(), OriginaMapCount + 1);
+		TestEqual(TEXT("Remove removed key fails"), Map.Remove(111), false);
+		TestEqual(TEXT("Failed remove did not increase count"), Map.Num(), OriginaMapCount + 1);
+		TestEqual(TEXT("Find removed key returns no value set"), Map.Find(111), FPublicExportMap::InvalidValue);
+	}
+
+	// Many pre-existing keys
+	{
+		FPublicExportMap Map;
+		// Note we are explicitly overlapping key 1 with the preinsert key 1
+		// Note we are mixing ranges of preinserted and preexisting keys so we are forced to sort the keys together
+		uint64 OverlappingKey = 1;
+		TArray<uint64> PreExistingKeys = { 11, 2, 13, 4, 15, 6, 17, 8, OverlappingKey };
+		TArray<int32> PreExistingValues = { 11, 2, 13, 4, 15, 6, 17, 8, (int32)OverlappingKey };
+		check(PreExistingKeys.Num() == PreExistingValues.Num());
+
+		TArray<FExportMapEntry> Exports = { {.PublicExportHash = 12}, {.PublicExportHash = OverlappingKey  }, {.PublicExportHash = 3 }, {.PublicExportHash = 14 } };
+		for (int32 i = 0; i < PreExistingKeys.Num(); ++i)
+		{
+			Map.Store(PreExistingKeys[i], PreExistingValues[i]);
+		}
+		int32 OriginaMapCount = Map.Num();
+
+		TSet<uint64> AllKeys;
+		AllKeys.Append(PreExistingKeys);
+		for (FExportMapEntry& Export : Exports)
+		{
+			AllKeys.Add(Export.PublicExportHash);
+		}
+
+		TestEqual(TEXT("Store increased count"), OriginaMapCount, PreExistingValues.Num());
+		TestEqual(TEXT("Store increased keys"), Map.GetKeys().Num(), OriginaMapCount);
+		TestEqual(TEXT("Store increased values"), Map.GetValues().Num(), OriginaMapCount);
+
+		Map.PreInsertPublicExports({ Exports });
+		TestEqual(TEXT("PreInsertPublicExports has populated our keys"), Map.GetKeys().Num(), AllKeys.Num());
+		TestEqual(TEXT("PreInsert has populated our values (should mostly be InvalidValue)"), Map.GetValues().Num(), AllKeys.Num());
+
+		TSet<uint64> FoundPreExistingKeys;
+		FoundPreExistingKeys.Reserve(PreExistingKeys.Num());
+		for (uint64 Key : Map.GetKeys())
+		{
+			if (PreExistingKeys.Contains(Key))
+			{
+				FoundPreExistingKeys.Add(Key);
+				TestEqual(TEXT("Pre-existing keys remain intact after PreInsertPublicExports"), Map.Find(Key), PreExistingValues[PreExistingKeys.IndexOfByKey(Key)]);
+			}
+			else if (Key == OverlappingKey)
+			{
+				TestEqual(TEXT("Ensure preinsert didn't overwrite the existing value"), Map.Find(Key), (int32)OverlappingKey);
+			}
+			else
+			{
+				TestEqual(TEXT("All preinserted keys are initialized to InvalidValue"), Map.Find(Key), FPublicExportMap::InvalidValue);
+			}
+		}
+		TestEqual(TEXT("All preexisting keys are still in the map"), FoundPreExistingKeys.Difference(TSet<uint64>(PreExistingKeys)).Num(), 0);
+		TestEqual(TEXT("PreInsertPublicExports does not add new values"), Map.Num(), OriginaMapCount);
+
+		// Overwrite preinserted value
+		TestEqual(TEXT("Find stored key with no value set"), Map.Find(Exports[0].PublicExportHash), FPublicExportMap::InvalidValue);
+		Map.Store(Exports[0].PublicExportHash, (int32)Exports[0].PublicExportHash);
+		TestEqual(TEXT("Find stored key"), Map.Find(Exports[0].PublicExportHash), (int32)Exports[0].PublicExportHash);
+		TestEqual(TEXT("Store increased count"), Map.Num(), PreExistingValues.Num() + 1);
+
+		// Insert new value
+		TestEqual(TEXT("Find missing key"), Map.Find(111), FPublicExportMap::InvalidValue);
+		Map.Store(111, 111);
+		TestEqual(TEXT("Find stored key"), Map.Find(111), 111);
+		TestEqual(TEXT("Store increased count"), Map.Num(), PreExistingValues.Num() + 2);
+
+		TestEqual(TEXT("Remove existing key"), Map.Remove(111), true);
+		TestEqual(TEXT("Remove decreased count"), Map.Num(), PreExistingValues.Num() + 1);
+		TestEqual(TEXT("Remove removed key fails"), Map.Remove(111), false);
+		TestEqual(TEXT("Failed remove did not increase count"), Map.Num(), PreExistingValues.Num() + 1);
+		TestEqual(TEXT("Find removed key returns no value set"), Map.Find(111), FPublicExportMap::InvalidValue);
+	}
+	return true;
+}
+
+#endif // WITH_LOW_LEVEL_TESTS
