Fixed an overly-broad warning in FindObjectByTypePath

#rnx

diff --git a/Engine/Source/Runtime/CoreUObject/Private/UObject/PropertyHelper.cpp b/Engine/Source/Runtime/CoreUObject/Private/UObject/PropertyHelper.cpp
--- a/Engine/Source/Runtime/CoreUObject/Private/UObject/PropertyHelper.cpp
+++ b/Engine/Source/Runtime/CoreUObject/Private/UObject/PropertyHelper.cpp
@@ -320,7 +320,7 @@ UObject* FindObjectByTypePath(UClass* Class, FPropertyTypeName TypePath)
 
 		if (UObject* Object = StaticFindFirstObject(Class, *Path + NameIndex, EFindFirstObjectOptions::NativeFirst))
 		{
-			UE_LOGFMT(LogProperty, Warning, "{Class} '{Path}' was found by its short name but not its path.", Class->GetFName(), Path);
+			UE_CLOGFMT(NameIndex, LogProperty, Warning, "{Class} '{Path}' was found by its short name but not its path.", Class->GetFName(), Path);
 			return Object;
 		}
 	}
