LinkerSave/LinkerLoad: Delete deprecated code, and some more deprecations, and refactor existing code to match current FArchive patterns.

Remove 5.0 deprecated functions from Linker,LinkerLoad,LinkerSave.
Change handling of SerializedPropertyChain on FLinkerLoad and FLinkerSave to match the more-performant pattern implemented for FArchiveProxy: virtual overrides forward on to the internal archive rather than setting them on the internal archive in each call to Serialize.
Change FArchive::SetDebugSerializationFlags to be virtual to match the pattern of SetFilterEditorOnly, so it is robustly readable on FArchiveProxy, FLinkerLoad, FLinkerSave and does not need to be forwarded on every call to Serialize.
Update LinkerSave use of LinkerRoot and Saver to uniformly assert they are non-null rather than handling null.
Reorganize FLinkerSave constructor and CreateLinker and FinalizeFile to simplify the 4 possible cases for the Saver used by LinkerSave.
#rnx

diff --git a/Engine/Source/Runtime/Core/Public/Serialization/Archive.h b/Engine/Source/Runtime/Core/Public/Serialization/Archive.h
--- a/Engine/Source/Runtime/Core/Public/Serialization/Archive.h
+++ b/Engine/Source/Runtime/Core/Public/Serialization/Archive.h
@@ -554,7 +554,7 @@ PRAGMA_ENABLE_DEPRECATION_WARNINGS
 	 *
 	 * @param InCustomFlags the new flags to use for custom serialization
 	 */
-	void SetDebugSerializationFlags(uint32 InCustomFlags)
+	virtual void SetDebugSerializationFlags(uint32 InCustomFlags)
 	{
 #if WITH_EDITOR
 		ArDebugSerializationFlags = InCustomFlags;
diff --git a/Engine/Source/Runtime/Core/Public/Serialization/ArchiveProxy.h b/Engine/Source/Runtime/Core/Public/Serialization/ArchiveProxy.h
--- a/Engine/Source/Runtime/Core/Public/Serialization/ArchiveProxy.h
+++ b/Engine/Source/Runtime/Core/Public/Serialization/ArchiveProxy.h
@@ -239,6 +239,12 @@ public:
 		InnerArchive.ForceBlueprintFinalization();
 	}
 
+	virtual void SetDebugSerializationFlags(uint32 InCustomFlags) override
+	{
+		FArchive::SetDebugSerializationFlags(InCustomFlags);
+		InnerArchive.SetDebugSerializationFlags(InCustomFlags);
+	}
+
 	virtual void SetFilterEditorOnly(bool InFilterEditorOnly) override
 	{
 		FArchive::SetFilterEditorOnly(InFilterEditorOnly);
diff --git a/Engine/Source/Runtime/CoreUObject/Private/UObject/Linker.cpp b/Engine/Source/Runtime/CoreUObject/Private/UObject/Linker.cpp
--- a/Engine/Source/Runtime/CoreUObject/Private/UObject/Linker.cpp
+++ b/Engine/Source/Runtime/CoreUObject/Private/UObject/Linker.cpp
@@ -181,22 +181,10 @@ FLinker::FLinker(ELinkerType::Type InType, UPackage* InRoot)
 	}
 }
 
-FLinker::FLinker(ELinkerType::Type InType, UPackage* InRoot, const TCHAR* InFilename)
-: FLinker(InType, InRoot)
-{
-	PRAGMA_DISABLE_DEPRECATION_WARNINGS
-	Filename = InFilename;
-	PRAGMA_ENABLE_DEPRECATION_WARNINGS
-}
-
 /** Returns a descriptor of the PackagePath this Linker is reading from or writing to, usable for an identifier in warning and log messages */
 FString FLinker::GetDebugName() const
 {
-	// UE_DEPRECATED(5.0)
-	UE_LOG(LogLinker, Warning, TEXT("Filename and the default GetDebugName on FLinker is deprecated and will be removed in a future release. Subclasses should override GetDebugName."));
-	PRAGMA_DISABLE_DEPRECATION_WARNINGS
-	return Filename;
-	PRAGMA_ENABLE_DEPRECATION_WARNINGS
+	return FString(TEXT("UnexpectedCallToBaseClassGetDebugName"));
 }
 
 
@@ -631,13 +619,6 @@ static FPackagePath GetPackagePath(UPackage* InOuter, const TCHAR* InPackageName
 		return FPackagePath();
 	}
 }
-		
-COREUOBJECT_API FLinkerLoad* GetPackageLinker(UPackage* InOuter, const TCHAR* InLongPackageName, uint32 LoadFlags,
-	UPackageMap* Sandbox, FGuid* CompatibleGuid, FArchive* InReaderOverride, FUObjectSerializeContext** InOutLoadContext,
-	FLinkerLoad* ImportLinker, const FLinkerInstancingContext* InstancingContext)
-{
-	return GetPackageLinker(InOuter, GetPackagePath(InOuter, InLongPackageName), LoadFlags, Sandbox, InReaderOverride, InOutLoadContext, ImportLinker, InstancingContext);
-}
 
 //
 // Find or create the linker for a package.
@@ -831,21 +812,6 @@ FLinkerLoad* LoadPackageLinker(UPackage* InOuter, const FPackagePath& PackagePat
 	PRAGMA_ENABLE_DEPRECATION_WARNINGS
 }
 
-FLinkerLoad* LoadPackageLinker(UPackage* InOuter, const TCHAR* InLongPackageName, uint32 LoadFlags, UPackageMap* Sandbox, FGuid* CompatibleGuid, FArchive* InReaderOverride, TFunctionRef<void(FLinkerLoad* LoadedLinker)> LinkerLoadedCallback)
-{
-	PRAGMA_DISABLE_DEPRECATION_WARNINGS
-	return LoadPackageLinker(InOuter, GetPackagePath(InOuter, InLongPackageName), LoadFlags, Sandbox, InReaderOverride, LinkerLoadedCallback);
-	PRAGMA_ENABLE_DEPRECATION_WARNINGS
-}
-
-FLinkerLoad* LoadPackageLinker(UPackage* InOuter, const TCHAR* InLongPackageName, uint32 LoadFlags, UPackageMap* Sandbox, FGuid* CompatibleGuid, FArchive* InReaderOverride)
-{
-	PRAGMA_DISABLE_DEPRECATION_WARNINGS
-	return LoadPackageLinker(InOuter, GetPackagePath(InOuter, InLongPackageName), LoadFlags, Sandbox, InReaderOverride, [](FLinkerLoad* InLinker) {});
-	PRAGMA_ENABLE_DEPRECATION_WARNINGS
-}
-
-
 void ResetLinkerExports(UPackage* InPackage)
 {
 	FLinkerManager::Get().ResetLinkerExports(InPackage);
diff --git a/Engine/Source/Runtime/CoreUObject/Private/UObject/LinkerLoad.cpp b/Engine/Source/Runtime/CoreUObject/Private/UObject/LinkerLoad.cpp
--- a/Engine/Source/Runtime/CoreUObject/Private/UObject/LinkerLoad.cpp
+++ b/Engine/Source/Runtime/CoreUObject/Private/UObject/LinkerLoad.cpp
@@ -513,11 +513,6 @@ static FORCEINLINE bool IsCoreUObjectPackage(const FName& PackageName)
 	FLinkerLoad.
 ----------------------------------------------------------------------------*/
 
-FLinkerLoad* FLinkerLoad::CreateLinker(FUObjectSerializeContext* LoadContext, UPackage* Parent, const TCHAR* Filename, uint32 LoadFlags, FArchive* InLoader, const FLinkerInstancingContext* InstancingContext)
-{
-	return CreateLinker(LoadContext, Parent, FPackagePath::FromLocalPath(Filename), LoadFlags, InLoader, InstancingContext);
-}
-
 /**
  * Creates and returns a FLinkerLoad object.
  *
@@ -582,9 +577,6 @@ FLinkerLoad* FLinkerLoad::CreateLinker(FUObjectSerializeContext* LoadContext, UP
 void FLinkerLoad::SetPackagePath(const FPackagePath& InPackagePath)
 {
 	PackagePath = InPackagePath;
-	PRAGMA_DISABLE_DEPRECATION_WARNINGS
-	Filename = PackagePath.GetLocalFullPath();
-	PRAGMA_ENABLE_DEPRECATION_WARNINGS
 }
 
 void FLinkerLoad::SetLoader(FArchive* InLoader, bool bInLoaderNeedsEngineVersionChecks)
@@ -4557,8 +4549,6 @@ UObject* FLinkerLoad::ResolveResource(FPackageIndex Index)
 	return IndexToObject(Index);
 }
 
-// Find the index of a specified object without regard to specific package.
-///@todo - this function is only used through FLinkerLoad::Create to perform package conform compatibility, which is deprecated.
 int32 FLinkerLoad::FindExportIndex( FName ClassName, FName ClassPackage, FName ObjectName, FPackageIndex ExportOuterIndex )
 {
 	int32 iHash = GetHashBucket( ObjectName );
@@ -4612,99 +4602,6 @@ int32 FLinkerLoad::FindExportIndex( FName ClassName, FName ClassPackage, FName O
 	return INDEX_NONE;
 }
 
-UObject* FLinkerLoad::Create( UClass* ObjectClass, FName ObjectName, UObject* Outer, uint32 InLoadFlags, bool Checked )
-{
-	// We no longer handle a NULL outer, which used to mean look in any outer, but we need fully qualified names now
-	// The other case where this was NULL is if you are calling StaticLoadObject on the top-level package, but
-	// you should be using LoadPackage. If for some weird reason you need to load the top-level package with this,
-	// then I believe you'd want to set OuterIndex to 0 when Outer is NULL, but then that could get confused with
-	// loading A.A (they both have OuterIndex of 0, according to Ron)
-	check(Outer);
-
-	int32 OuterIndex = INDEX_NONE;
-
-	// if the outer is the outermost of the package, then we want OuterIndex to be 0, as objects under the top level
-	// will have an OuterIndex to 0
-	if (Outer == Outer->GetOutermost())
-	{
-		OuterIndex = 0;
-	}
-	// otherwise get the linker index of the outer to be the outer index that we look in
-	else
-	{
-		OuterIndex = Outer->GetLinkerIndex();
-		// we _need_ the linker index of the outer to look in, which means that the outer must have been actually 
-		// loaded off disk, and not just CreatePackage'd
-		check(OuterIndex != INDEX_NONE);
-	}
-
-	FPackageIndex OuterPackageIndex;
-
-	if (OuterIndex)
-	{
-		OuterPackageIndex = FPackageIndex::FromExport(OuterIndex);
-	}
-
-	int32 Index = FindExportIndex(ObjectClass->GetFName(), ObjectClass->GetOuter()->GetFName(), ObjectName, OuterPackageIndex);
-	if (Index != INDEX_NONE)
-	{
-		return (InLoadFlags & LOAD_Verify) ? INVALID_OBJECT : CreateExport(Index);
-	}
-
-	// since we didn't find it, see if we can find an object redirector with the same name
-	// Are we allowed to follow redirects?
-	if( !( InLoadFlags & LOAD_NoRedirects ) )
-	{
-		Index = FindExportIndex(UObjectRedirector::StaticClass()->GetFName(), NAME_CoreUObject, ObjectName, OuterPackageIndex);
-		if (Index == INDEX_NONE)
-		{
-			Index = FindExportIndex(UObjectRedirector::StaticClass()->GetFName(), GLongCoreUObjectPackageName, ObjectName, OuterPackageIndex);
-		}
-
-		// if we found a redirector, create it, and move on down the line
-		if (Index != INDEX_NONE)
-		{
-			// create the redirector,
-			UObjectRedirector* Redir = (UObjectRedirector*)CreateExport(Index);
-			Preload(Redir);
-			// if we found what it was point to, then return it
-			if (Redir->DestinationObject && Redir->DestinationObject->IsA(ObjectClass))
-			{
-				// and return the object we are being redirected to
-				return Redir->DestinationObject;
-			}
-		}
-	}
-
-
-// Set this to 1 to find nonqualified names anyway
-#define FIND_OBJECT_NONQUALIFIED 0
-// Set this to 1 if you want to see what it would have found previously. This is useful for fixing up hundreds
-// of now-illegal references in script code.
-#define DEBUG_PRINT_NONQUALIFIED_RESULT 1
-
-#if DEBUG_PRINT_NONQUALIFIED_RESULT || FIND_OBJECT_NONQUALIFIED
-	Index = FindExportIndex(ObjectClass->GetFName(), ObjectClass->GetOuter()->GetFName(), ObjectName, FPackageIndex::FromImport(0));// this is very not legit to be passing INDEX_NONE into this function to mean "ignore"
-	if (Index != INDEX_NONE)
-	{
-#if DEBUG_PRINT_NONQUALIFIED_RESULT
-		UE_ASSET_LOG(LogLinker, Warning, PackagePath, TEXT("Using a non-qualified name (would have) found: %s"), *GetExportFullName(Index));
-#endif
-#if FIND_OBJECT_NONQUALIFIED
-		return (InLoadFlags & LOAD_Verify) ? INVALID_OBJECT : CreateExport(Index);
-#endif
-	}
-#endif
-
-
-	// if we are checking for failure cases, and we failed, throw an error
-	if( Checked )
-	{
-		UE_ASSET_LOG(LogLinker, Warning, PackagePath, TEXT("%s %s not found for creation"), *ObjectClass->GetName(), *ObjectName.ToString());
-	}
-	return NULL;
-}
-
 /**
  * Serialize the object data for the specified object from the unreal package file.  Loads any
  * additional resources required for the object to be in a valid state to receive the loaded
@@ -7278,16 +7175,6 @@ void FLinkerLoad::SetPreloadingEnabled(bool bEnabled)
 {
 	bPreloadingEnabled = bEnabled;
 }
-bool FLinkerLoad::TryGetPreloadedLoader(FArchive*& OutLoader, const TCHAR* FileName)
-{
-	FOpenPackageResult Result;
-	if (IPackageResourceManager::TryTakePreloadableArchive(FPackagePath::FromLocalPath(FileName), Result))
-	{
-		OutLoader = Result.Archive.Release();
-		return true;
-	}
-	return false;
-}
 
 bool FLinkerLoad::TryGetPreloadedLoader(const FPackagePath& InPackagePath, FOpenPackageResult& OutResult)
 {
diff --git a/Engine/Source/Runtime/CoreUObject/Private/UObject/LinkerSave.cpp b/Engine/Source/Runtime/CoreUObject/Private/UObject/LinkerSave.cpp
--- a/Engine/Source/Runtime/CoreUObject/Private/UObject/LinkerSave.cpp
+++ b/Engine/Source/Runtime/CoreUObject/Private/UObject/LinkerSave.cpp
@@ -6,6 +6,7 @@
 #include "Misc/ConfigCacheIni.h"
 #include "Misc/Optional.h"
 #include "Misc/PackageName.h"
+#include "Misc/StringBuilder.h"
 #include "Serialization/BulkData.h"
 #include "Serialization/LargeMemoryWriter.h"
 #include "UObject/Package.h"
@@ -31,183 +32,112 @@ PRAGMA_DISABLE_DEPRECATION_WARNINGS
 TMap<FString, TArray<uint8> > FLinkerSave::PackagesToScriptSHAMap;
 PRAGMA_ENABLE_DEPRECATION_WARNINGS
 
-FLinkerSave::FLinkerSave(UPackage* InParent, const TCHAR* InFilename, bool bForceByteSwapping, bool bInSaveUnversioned)
-:	FLinker(ELinkerType::Save, InParent)
-,	Saver(nullptr)
+FLinkerSave::FLinkerSave(UPackage* InParent)
+	: FLinker(ELinkerType::Save, InParent)
 {
-	SetFilename(InFilename);
-
-	if (FPlatformProperties::HasEditorOnlyData())
-	{
-		// Create file saver.
-		Saver = IFileManager::Get().CreateFileWriter( InFilename, 0 );
-		if( !Saver )
-		{
-			TCHAR LastErrorText[1024];
-			uint32 LastError = FPlatformMisc::GetLastError();
-			if (LastError != 0)
-			{
-				FPlatformMisc::GetSystemErrorMessage(LastErrorText, UE_ARRAY_COUNT(LastErrorText), LastError);
-			}
-			else
-			{
-				FCString::Strcpy(LastErrorText, TEXT("Unknown failure reason."));
-			}
-			UE_LOG(LogLinker, Error, TEXT("Error opening file '%s': %s"), InFilename, LastErrorText);
-			return; // Caller must test this->Saver to detect the failure
-		}
-
-		UPackage* Package = dynamic_cast<UPackage*>(LinkerRoot);
-
-		// Set main summary info.
-		Summary.Tag = PACKAGE_FILE_TAG;
-		Summary.SetToLatestFileVersions(bInSaveUnversioned);
-		Summary.SavedByEngineVersion = FEngineVersion::Current();
-		Summary.CompatibleWithEngineVersion = FEngineVersion::CompatibleWith();
-		Summary.SetPackageFlags(Package ? Package->GetPackageFlags() : 0);
-
-#if USE_STABLE_LOCALIZATION_KEYS
-		if (GIsEditor)
-		{
-			Summary.LocalizationId = TextNamespaceUtil::GetPackageNamespace(LinkerRoot);
-		}
-#endif // USE_STABLE_LOCALIZATION_KEYS
-
-		if (Package)
-		{
-#if WITH_EDITORONLY_DATA
-			Summary.PackageName = Package->GetName();
-#endif
-			Summary.ChunkIDs = Package->GetChunkIDs();
-		}
+	check(InParent);
+}
 
-		// Set status info.
-		this->SetIsSaving(true);
-		this->SetIsPersistent(true);
-		ArForceByteSwapping		= bForceByteSwapping;
+FLinkerSave::FLinkerSave(UPackage* InParent, const TCHAR* InFilename, bool bForceByteSwapping, bool bInSaveUnversioned)
+	: FLinkerSave(InParent)
+{
+	(void) TryAssignFileSaver(FStringView(InFilename), bForceByteSwapping, bInSaveUnversioned);
+}
 
-#if USE_STABLE_LOCALIZATION_KEYS
-		if (GIsEditor)
-		{
-			SetLocalizationNamespace(Summary.LocalizationId);
-		}
-#endif // USE_STABLE_LOCALIZATION_KEYS
-	}
+FLinkerSave::FLinkerSave(UPackage* InParent, FArchive* InSaver, bool bForceByteSwapping, bool bInSaveUnversioned)
+	: FLinkerSave(InParent)
+{
+	AssignSaver(InSaver, bForceByteSwapping, bInSaveUnversioned);
 }
 
+FLinkerSave::FLinkerSave(UPackage* InParent, bool bForceByteSwapping, bool bInSaveUnversioned)
+	: FLinkerSave(InParent)
+{
+	AssignMemorySaver(bForceByteSwapping, bInSaveUnversioned);
+}
 
-FLinkerSave::FLinkerSave(UPackage* InParent, FArchive *InSaver, bool bForceByteSwapping, bool bInSaveUnversioned)
-: FLinker(ELinkerType::Save, InParent)
-, Saver(nullptr)
+void FLinkerSave::AssignSaver(FArchive* InSaver, bool bForceByteSwapping, bool bInSaveUnversioned)
 {
 	SetFilename(TEXT("$$Memory$$"));
-	if (FPlatformProperties::HasEditorOnlyData())
-	{
-		// Create file saver.
-		Saver = InSaver;
-		check(Saver);
-#if WITH_EDITOR
-		ArDebugSerializationFlags = Saver->ArDebugSerializationFlags;
-#endif
-		
-
-		UPackage* Package = dynamic_cast<UPackage*>(LinkerRoot);
+	AssignSaverInternal(InSaver, bForceByteSwapping, bInSaveUnversioned);
+}
 
-		// Set main summary info.
-		Summary.Tag = PACKAGE_FILE_TAG;
-		Summary.SetToLatestFileVersions(bInSaveUnversioned);
-		Summary.SavedByEngineVersion = FEngineVersion::Current();
-		Summary.CompatibleWithEngineVersion = FEngineVersion::CompatibleWith();
-		Summary.SetPackageFlags(Package ? Package->GetPackageFlags() : 0);
+void FLinkerSave::AssignMemorySaver(bool bForceByteSwapping, bool bInSaveUnversioned)
+{
+	check(LinkerRoot); // Must be non-null in constructor
 
-#if USE_STABLE_LOCALIZATION_KEYS
-		if (GIsEditor)
-		{
-			Summary.LocalizationId = TextNamespaceUtil::GetPackageNamespace(LinkerRoot);
-		}
-#endif // USE_STABLE_LOCALIZATION_KEYS
+	SetFilename(TEXT("$$Memory$$"));
+	FArchive* LocalSaver = new FLargeMemoryWriter(0, false, *LinkerRoot->GetLoadedPath().GetDebugName());
+	AssignSaverInternal(LocalSaver, bForceByteSwapping, bInSaveUnversioned);
+}
 
-		if (Package)
+bool FLinkerSave::TryAssignFileSaver(FStringView InFilename, bool bForceByteSwapping, bool bInSaveUnversioned)
+{
+	SetFilename(InFilename);
+	// Create file saver.
+	FArchive* LocalSaver = IFileManager::Get().CreateFileWriter(*WriteToString<256>(InFilename), 0);
+	if (!LocalSaver)
+	{
+		TCHAR LastErrorText[1024];
+		uint32 LastError = FPlatformMisc::GetLastError();
+		if (LastError != 0)
 		{
-#if WITH_EDITORONLY_DATA
-			Summary.PackageName = Package->GetName();
-#endif
-			Summary.ChunkIDs = Package->GetChunkIDs();
+			FPlatformMisc::GetSystemErrorMessage(LastErrorText, UE_ARRAY_COUNT(LastErrorText), LastError);
 		}
-
-		// Set status info.
-		this->SetIsSaving(true);
-		this->SetIsPersistent(true);
-		ArForceByteSwapping = bForceByteSwapping;
-
-#if USE_STABLE_LOCALIZATION_KEYS
-		if (GIsEditor)
+		else
 		{
-			SetLocalizationNamespace(Summary.LocalizationId);
+			FCString::Strcpy(LastErrorText, TEXT("Unknown failure reason."));
 		}
-#endif // USE_STABLE_LOCALIZATION_KEYS
+		UE_LOG(LogLinker, Error, TEXT("Error opening file '%.*s': %s"),
+			InFilename.Len(), InFilename.GetData(), LastErrorText);
+		return false;
 	}
+	AssignSaverInternal(LocalSaver, bForceByteSwapping, bInSaveUnversioned);
+	return true;
 }
 
-FLinkerSave::FLinkerSave(UPackage* InParent, bool bForceByteSwapping, bool bInSaveUnversioned )
-:	FLinker(ELinkerType::Save, InParent)
-,	Saver(nullptr)
+void FLinkerSave::AssignSaverInternal(FArchive* InSaver, bool bForceByteSwapping, bool bInSaveUnversioned)
 {
-	SetFilename(TEXT("$$Memory$$"));
-	if (FPlatformProperties::HasEditorOnlyData())
-	{
-		// Create file saver.
-		Saver = new FLargeMemoryWriter( 0, false, *InParent->GetLoadedPath().GetDebugName() );
-		check(Saver);
+	check(InSaver);
+	check(LinkerRoot); // Must be non-null in constructor
 
-		UPackage* Package = dynamic_cast<UPackage*>(LinkerRoot);
+	Saver = InSaver;
+	UPackage* Package = LinkerRoot;
 
-		// Set main summary info.
-		Summary.Tag = PACKAGE_FILE_TAG;
-		Summary.SetToLatestFileVersions(bInSaveUnversioned);
-		Summary.SavedByEngineVersion = FEngineVersion::Current();
-		Summary.CompatibleWithEngineVersion = FEngineVersion::CompatibleWith();
-		Summary.SetPackageFlags(Package ? Package->GetPackageFlags() : 0);
+	// Set main summary info.
+	Summary.Tag = PACKAGE_FILE_TAG;
+	Summary.SetToLatestFileVersions(bInSaveUnversioned);
+	Summary.SavedByEngineVersion = FEngineVersion::Current();
+	Summary.CompatibleWithEngineVersion = FEngineVersion::CompatibleWith();
+	Summary.SetPackageFlags(Package ? Package->GetPackageFlags() : 0);
 
 #if USE_STABLE_LOCALIZATION_KEYS
-		if (GIsEditor)
-		{
-			Summary.LocalizationId = TextNamespaceUtil::GetPackageNamespace(LinkerRoot);
-		}
+	if (GIsEditor)
+	{
+		Summary.LocalizationId = TextNamespaceUtil::GetPackageNamespace(LinkerRoot);
+		SetLocalizationNamespace(Summary.LocalizationId);
+	}
 #endif // USE_STABLE_LOCALIZATION_KEYS
 
-		if (Package)
-		{
 #if WITH_EDITORONLY_DATA
-			Summary.PackageName = Package->GetName();
+	Summary.PackageName = Package->GetName();
 #endif
-			Summary.ChunkIDs = Package->GetChunkIDs();
-		}
+	Summary.ChunkIDs = Package->GetChunkIDs();
 
-		// Set status info.
-		this->SetIsSaving(true);
-		this->SetIsPersistent(true);
-		ArForceByteSwapping		= bForceByteSwapping;
-
-#if USE_STABLE_LOCALIZATION_KEYS
-		if (GIsEditor)
-		{
-			SetLocalizationNamespace(Summary.LocalizationId);
-		}
-#endif // USE_STABLE_LOCALIZATION_KEYS
-	}
+	// Set FArchive flags.
+	this->SetIsSaving(true);
+	this->SetIsPersistent(true);
+	ArForceByteSwapping = bForceByteSwapping;
+#if WITH_EDITOR
+	ArDebugSerializationFlags = Saver->ArDebugSerializationFlags;
+#endif
 }
 
 bool FLinkerSave::CloseAndDestroySaver()
 {
-	bool bSuccess = true;
-	if (Saver)
-	{
-		// first, do an explicit close to check for archive errors
-		delete Saver;
-	}
+	delete Saver;
 	Saver = nullptr;
-	return bSuccess;
+	return true;
 }
 
 FLinkerSave::~FLinkerSave()
@@ -331,10 +261,6 @@ int64 FLinkerSave::Tell()
 
 void FLinkerSave::Serialize( void* V, int64 Length )
 {
-#if WITH_EDITOR
-	Saver->ArDebugSerializationFlags = ArDebugSerializationFlags;
-	Saver->SetSerializedPropertyChain(GetSerializedPropertyChain(), GetSerializedProperty());
-#endif
 	Saver->Serialize( V, Length );
 }
 
@@ -355,18 +281,12 @@ FString FLinkerSave::GetDebugName() const
 
 const FString& FLinkerSave::GetFilename() const
 {
-	// When the deprecated Filename is removed from FLinker, add a separate Filename variable to FLinkerSave
-	PRAGMA_DISABLE_DEPRECATION_WARNINGS
 	return Filename;
-	PRAGMA_ENABLE_DEPRECATION_WARNINGS
 }
 
 void FLinkerSave::SetFilename(FStringView InFilename)
 {
-	// When the deprecated Filename is removed from FLinker, add a separate Filename variable to FLinkerSave
-	PRAGMA_DISABLE_DEPRECATION_WARNINGS
 	Filename = FString(InFilename);
-	PRAGMA_ENABLE_DEPRECATION_WARNINGS
 }
 
 FString FLinkerSave::GetArchiveName() const
@@ -536,51 +456,47 @@ void FLinkerSave::UsingCustomVersion(const struct FGuid& Guid)
 
 void FLinkerSave::SetUseUnversionedPropertySerialization(bool bInUseUnversioned)
 {
+	check(Saver); // Must be set before calling FArchive functions
+	check(LinkerRoot); // Must be non-null in constructor
+
 	FArchiveUObject::SetUseUnversionedPropertySerialization(bInUseUnversioned);
-	if (Saver)
-	{
-		Saver->SetUseUnversionedPropertySerialization(bInUseUnversioned);
-	}
+	Saver->SetUseUnversionedPropertySerialization(bInUseUnversioned);
 	if (bInUseUnversioned)
 	{
 		Summary.SetPackageFlags(Summary.GetPackageFlags() | PKG_UnversionedProperties);
-		if (LinkerRoot)
-		{
-			LinkerRoot->SetPackageFlags(PKG_UnversionedProperties);
-		}
+		LinkerRoot->SetPackageFlags(PKG_UnversionedProperties);
 	}
 	else
 	{
 		Summary.SetPackageFlags(Summary.GetPackageFlags() & ~PKG_UnversionedProperties);
-		if (LinkerRoot)
-		{
-			LinkerRoot->ClearPackageFlags(PKG_UnversionedProperties);
-		}
+		LinkerRoot->ClearPackageFlags(PKG_UnversionedProperties);
 	}
 }
 
+void FLinkerSave::SetDebugSerializationFlags(uint32 InCustomFlags)
+{
+	check(Saver); // Must be set before calling FArchive functions
+
+	FArchiveUObject::SetDebugSerializationFlags(InCustomFlags);
+	Saver->SetDebugSerializationFlags(InCustomFlags);
+}
+
 void FLinkerSave::SetFilterEditorOnly(bool bInFilterEditorOnly)
 {
+	check(Saver); // Must be set before calling FArchive functions
+	check(LinkerRoot); // Must be non-null in constructor
+
 	FArchiveUObject::SetFilterEditorOnly(bInFilterEditorOnly);
-	if (Saver)
-	{
-		Saver->SetFilterEditorOnly(bInFilterEditorOnly);
-	}
+	Saver->SetFilterEditorOnly(bInFilterEditorOnly);
 	if (bInFilterEditorOnly)
 	{
 		Summary.SetPackageFlags(Summary.GetPackageFlags() | PKG_FilterEditorOnly);
-		if (LinkerRoot)
-		{
-			LinkerRoot->SetPackageFlags(PKG_FilterEditorOnly);
-		}
+		LinkerRoot->SetPackageFlags(PKG_FilterEditorOnly);
 	}
 	else
 	{
 		Summary.SetPackageFlags(Summary.GetPackageFlags() & ~PKG_FilterEditorOnly);
-		if (LinkerRoot)
-		{
-			LinkerRoot->ClearPackageFlags(PKG_FilterEditorOnly);
-		}
+		LinkerRoot->ClearPackageFlags(PKG_FilterEditorOnly);
 	}
 }
 
@@ -913,3 +829,35 @@ void FLinkerSave::OnPostSaveBulkData()
 	SerializedBulkData.Empty();
 #endif //WITH_EDITOR
 }
+
+void FLinkerSave::SetSerializedProperty(FProperty* InProperty)
+{
+	FArchiveUObject::SetSerializedProperty(InProperty);
+	Saver->SetSerializedProperty(InProperty);
+}
+
+void FLinkerSave::SetSerializedPropertyChain(const FArchiveSerializedPropertyChain* InSerializedPropertyChain,
+	class FProperty* InSerializedPropertyOverride)
+{
+	FArchiveUObject::SetSerializedPropertyChain(InSerializedPropertyChain, InSerializedPropertyOverride);
+	Saver->SetSerializedPropertyChain(InSerializedPropertyChain, InSerializedPropertyOverride);
+}
+
+void FLinkerSave::PushSerializedProperty(class FProperty* InProperty, const bool bIsEditorOnlyProperty)
+{
+	FArchive::PushSerializedProperty(InProperty, bIsEditorOnlyProperty);
+	Saver->PushSerializedProperty(InProperty, bIsEditorOnlyProperty);
+}
+
+void FLinkerSave::PopSerializedProperty(class FProperty* InProperty, const bool bIsEditorOnlyProperty)
+{
+	FArchive::PopSerializedProperty(InProperty, bIsEditorOnlyProperty);
+	Saver->PopSerializedProperty(InProperty, bIsEditorOnlyProperty);
+}
+
+#if WITH_EDITORONLY_DATA
+bool FLinkerSave::IsEditorOnlyPropertyOnTheStack() const
+{
+	return Saver->IsEditorOnlyPropertyOnTheStack();
+}
+#endif
diff --git a/Engine/Source/Runtime/CoreUObject/Private/UObject/SavePackage2.cpp b/Engine/Source/Runtime/CoreUObject/Private/UObject/SavePackage2.cpp
--- a/Engine/Source/Runtime/CoreUObject/Private/UObject/SavePackage2.cpp
+++ b/Engine/Source/Runtime/CoreUObject/Private/UObject/SavePackage2.cpp
@@ -1062,36 +1062,66 @@ ESavePackageResult ValidateImports(FSaveContext& SaveContext)
 
 ESavePackageResult CreateLinker(FSaveContext& SaveContext)
 {
+	SCOPED_SAVETIMER(UPackage_Save_CreateLinkerSave);
 	const FString BaseFilename = FPaths::GetBaseFilename(SaveContext.GetFilename());
-	// Make temp file. CreateTempFilename guarantees unique, non-existing filename.
-	// The temp file will be saved in the game save folder to not have to deal with potentially too long paths.
-	// Since the temp filename may include a 32 character GUID as well, limit the user prefix to 32 characters.
+
+	const bool bIsOptionalRealm = SaveContext.GetCurrentHarvestingRealm() == ESaveRealm::Optional;
+	const bool bByteSwapping = SaveContext.IsForceByteSwapping();
+	const bool bSaveUnversioned = SaveContext.IsSaveUnversionedNative();
+	UPackage* const Package = SaveContext.GetPackage();
+	const FName PackageName = Package->GetFName();
+	IPackageWriter* PackageWriter = SaveContext.GetPackageWriter();
+
+	TPimplPtr<FLinkerSave> LinkerSave = MakePimpl<FLinkerSave>(Package);
+	if (SaveContext.IsGenerateSaveError())
 	{
-		SCOPED_SAVETIMER(UPackage_Save_CreateLinkerSave);
+		LinkerSave->SetOutputDevice(SaveContext.GetError());
+	}
+	LinkerSave->SetTransientPropertyOverrides(SaveContext.GetTransientPropertyOverrides());
+	LinkerSave->bUpdatingLoadedPath = SaveContext.IsUpdatingLoadedPath();
+	LinkerSave->bProceduralSave = SaveContext.IsProceduralSave();
+	LinkerSave->bRehydratePayloads = SaveContext.ShouldRehydratePayloads();
 
-		IPackageWriter* PackageWriter = SaveContext.GetPackageWriter();
-		if (PackageWriter || SaveContext.IsSaveToMemory())
+	if (PackageWriter)
+	{
+		check(!SaveContext.IsTextFormat()); // TextFormat not supported with a PackageWriter
+
+		uint16 MultiOutputIndex = bIsOptionalRealm ? 1 : 0;
+		TUniquePtr<FLargeMemoryWriter> ExportsArchive =
+			PackageWriter->CreateLinkerArchive(PackageName, SaveContext.GetAsset(), MultiOutputIndex);
+		LinkerSave->AssignSaver(ExportsArchive.Release(), bByteSwapping, bSaveUnversioned);
+		SaveContext.SetFormatter(MakeUnique<FBinaryArchiveFormatter>(*LinkerSave));
+	}
+	else
+	{
+		check(!SaveContext.IsCooking()); // IsCooking requires a PackageWriter
+		check(!bIsOptionalRealm); // OptionalRealm is only supported when cooking.
+
+		if (SaveContext.IsSaveToMemory())
 		{
-			const bool bIsOptionalRealm = SaveContext.GetCurrentHarvestingRealm() == ESaveRealm::Optional;
+			check(!SaveContext.IsTextFormat()); // TextFormat not supported with SaveToMemory
 
-			// Allocate the linker with a memory writer, forcing byte swapping if wanted.
-			TUniquePtr<FLargeMemoryWriter> ExportsArchive = PackageWriter ?
-				PackageWriter->CreateLinkerArchive(SaveContext.GetPackage()->GetFName(), SaveContext.GetAsset(), bIsOptionalRealm ? 1 : 0) :
-				// The LargeMemoryWriter does not need to be persistent; the LinkerSave wraps it and reports Persistent=true
-				TUniquePtr<FLargeMemoryWriter>(new FLargeMemoryWriter(0, false /* bIsPersistent */, *SaveContext.GetPackage()->GetName()));
-			SaveContext.SetLinker(MakePimpl<FLinkerSave>(SaveContext.GetPackage(), ExportsArchive.Release(),
-				SaveContext.IsForceByteSwapping(), SaveContext.IsSaveUnversionedNative()));
+			TUniquePtr<FLargeMemoryWriter> ExportsArchive(new FLargeMemoryWriter(
+				0 /* PreAllocateBytes */,
+				false /* bPersistent */, // LinkerSave wraps this archive and reports Persistent=true
+				*PackageName.ToString()));
+			LinkerSave->AssignSaver(ExportsArchive.Release(), bByteSwapping, bSaveUnversioned);
+			SaveContext.SetFormatter(MakeUnique<FBinaryArchiveFormatter>(*LinkerSave));
 		}
-		else
+		else if (!SaveContext.IsTextFormat())
 		{
-			// Allocate the linker with a tempfile, forcing byte swapping if wanted.
-			SaveContext.SetTempFilename(FPaths::CreateTempFilename(*FPaths::ProjectSavedDir(), *BaseFilename.Left(32)));
-			SaveContext.SetLinker(MakePimpl<FLinkerSave>(SaveContext.GetPackage(), *SaveContext.GetTempFilename().GetValue(), SaveContext.IsForceByteSwapping(), SaveContext.IsSaveUnversionedNative()));
-			if (!SaveContext.GetLinker()->Saver)
+			// Make a temp filename. CreateTempFilename guarantees unique, non-existing filename.
+			// The temp file will be saved in the game save folder to not have to deal with potentially too long paths.
+			// Since the temp filename may include a 32 character GUID as well, limit the user prefix to 32 characters.
+			FString TempFilename = FPaths::CreateTempFilename(*FPaths::ProjectSavedDir(), *BaseFilename.Left(32));
+			SaveContext.SetTempFilename(TempFilename);
+			if (!LinkerSave->TryAssignFileSaver(TempFilename, bByteSwapping, bSaveUnversioned))
 			{
 				FFormatNamedArguments Arguments;
-				Arguments.Add(TEXT("Name"), FText::FromString(*SaveContext.GetTempFilename()));
-				FText ErrorText = FText::Format(NSLOCTEXT("SavePackage", "CouldNotCreateSaveFile", "Could not create temporary save filename {Name}."), Arguments);
+				Arguments.Add(TEXT("Name"), FText::FromString(*TempFilename));
+				FText ErrorText = FText::Format(
+					NSLOCTEXT("SavePackage", "CouldNotCreateSaveFile", "Could not create temporary save filename {Name}."),
+					Arguments);
 				UE_LOG(LogSavePackage, Error, TEXT("%s"), *ErrorText.ToString());
 				if (SaveContext.IsGenerateSaveError())
 				{
@@ -1099,58 +1129,57 @@ ESavePackageResult CreateLinker(FSaveContext& SaveContext)
 				}
 				return ESavePackageResult::Error;
 			}
+			SaveContext.SetFormatter(MakeUnique<FBinaryArchiveFormatter>(*LinkerSave));
 		}
-
-		if (SaveContext.IsGenerateSaveError())
-		{
-			SaveContext.GetLinker()->SetOutputDevice(SaveContext.GetError());
-		}
-		SaveContext.GetLinker()->SetTransientPropertyOverrides(SaveContext.GetTransientPropertyOverrides());
-		SaveContext.GetLinker()->bUpdatingLoadedPath = SaveContext.IsUpdatingLoadedPath();
-		SaveContext.GetLinker()->bProceduralSave = SaveContext.IsProceduralSave();
-
-		SaveContext.GetLinker()->bRehydratePayloads = SaveContext.ShouldRehydratePayloads();
-		
-		// The package trailer is not supported for text based assets yet
-		if (!SaveContext.IsTextFormat() && !SaveContext.IsProceduralSave())
-		{
-			SaveContext.GetLinker()->PackageTrailerBuilder = MakeUnique<UE::FPackageTrailerBuilder>(SaveContext.GetPackage()->GetName());
-		}
-		else if ((SaveContext.GetSaveArgs().SaveFlags & SAVE_BulkDataByReference) != 0)
+		else // IsTextFormat
 		{
-			if (const FLinkerLoad* LinkerLoad = FLinkerLoad::FindExistingLinkerForPackage(SaveContext.GetPackage()))
-			{
-				const UE::FPackageTrailer* Trailer = LinkerLoad->GetPackageTrailer();
-				if (Trailer && Trailer->GetNumPayloads(UE::EPayloadStorageType::Any) > 0)
-				{
-					SaveContext.GetLinker()->PackageTrailerBuilder = UE::FPackageTrailerBuilder::CreateReferenceToTrailer(*Trailer, SaveContext.GetPackage()->GetName());
-				}
+#if !WITH_TEXT_ARCHIVE_SUPPORT
+			UE_LOG(LogSavePackage, Error, TEXT("Could not save package to text format. Text Archive Support is not compiled into this version of the editor executable."));
+			return ESavePackageResult::Error;
+#else
+			// TODO: Remove the use of the file FArchive in LinkerSave in the text format case; it is unused because
+			// we serialize into the TextFormatArchive instead of into FLinkerSave. But some SavePackage and LinkerSave
+			// functions still expect it to be set.
+			FString TempFilename = FPaths::CreateTempFilename(*FPaths::ProjectSavedDir(), *BaseFilename.Left(32));
+			SaveContext.SetTempFilename(TempFilename);
+			if (!LinkerSave->TryAssignFileSaver(TempFilename, bByteSwapping, bSaveUnversioned))
+			{
+				UE_LOG(LogSavePackage, Error, TEXT("Could not create temporary save filename %s."), *TempFilename);
+				return ESavePackageResult::Error;
 			}
-		}
 
-#if WITH_TEXT_ARCHIVE_SUPPORT
-		if (SaveContext.IsTextFormat())
-		{
-			if (SaveContext.GetTempFilename().IsSet())
-			{
-				SaveContext.SetTextFormatTempFilename(SaveContext.GetTempFilename().GetValue() + FPackageName::GetTextAssetPackageExtension());
-			}
-			else
-			{
-				SaveContext.SetTextFormatTempFilename(FPaths::CreateTempFilename(*FPaths::ProjectSavedDir(), *BaseFilename.Left(32)) + FPackageName::GetTextAssetPackageExtension());
-			}
-			SaveContext.SetTextFormatArchive(TUniquePtr<FArchive>(IFileManager::Get().CreateFileWriter(*SaveContext.GetTextFormatTempFilename().GetValue())));
+			FString TextFormatTempFilename = TempFilename + FPackageName::GetTextAssetPackageExtension();
+			SaveContext.SetTextFormatTempFilename(TextFormatTempFilename);
+			SaveContext.SetTextFormatArchive(TUniquePtr<FArchive>(IFileManager::Get().CreateFileWriter(*TextFormatTempFilename)));
 			TUniquePtr<FJsonArchiveOutputFormatter> OutputFormatter = MakeUnique<FJsonArchiveOutputFormatter>(*SaveContext.GetTextFormatArchive());
-			OutputFormatter->SetObjectIndicesMap(&SaveContext.GetLinker()->ObjectIndicesMap);
+			OutputFormatter->SetObjectIndicesMap(&LinkerSave->ObjectIndicesMap);
 			SaveContext.SetFormatter(MoveTemp(OutputFormatter));
-		}
-		else
 #endif
+		}
+	}
+
+	// Create the PackageTrailer for the Saver types that support it.
+	if ((SaveContext.GetSaveArgs().SaveFlags & SAVE_BulkDataByReference) != 0)
+	{
+		if (const FLinkerLoad* LinkerLoad = FLinkerLoad::FindExistingLinkerForPackage(Package))
 		{
-			SaveContext.SetFormatter(MakeUnique<FBinaryArchiveFormatter>(*(FArchive*)SaveContext.GetLinker()));
+			const UE::FPackageTrailer* Trailer = LinkerLoad->GetPackageTrailer();
+			if (Trailer && Trailer->GetNumPayloads(UE::EPayloadStorageType::Any) > 0)
+			{
+				LinkerSave->PackageTrailerBuilder = UE::FPackageTrailerBuilder::CreateReferenceToTrailer(
+					*Trailer, PackageName.ToString());
+			}
 		}
 	}
+	// PackageTrailer is not supported yet for cook or other procedural saves, unless using SAVE_BulkDataByReference.
+	// PackageTrailer is not supported for text based assets yet.
+	else if (!SaveContext.IsProceduralSave() && !SaveContext.IsTextFormat())
+	{
+		LinkerSave->PackageTrailerBuilder = MakeUnique<UE::FPackageTrailerBuilder>(PackageName.ToString());
+	}
 
+	// Assign Linker to the SaveContext
+	SaveContext.SetLinker(MoveTemp(LinkerSave));
 	SaveContext.SetStructuredArchive(MakeUnique<FStructuredArchive>(*SaveContext.GetFormatter()));
 	return ReturnSuccessOrCancel();
 }
@@ -1171,51 +1200,52 @@ struct FNameEntryIdSortHelper
 ESavePackageResult BuildLinker(FSaveContext& SaveContext)
 {
 	// Setup Linker 
+	FLinkerSave* Linker = SaveContext.GetLinker();
 	{
 		// Use the custom versions we harvested from the dependency harvesting pass
-		SaveContext.GetLinker()->Summary.SetCustomVersionContainer(SaveContext.GetCustomVersions());
+		Linker->Summary.SetCustomVersionContainer(SaveContext.GetCustomVersions());
 
-		SaveContext.GetLinker()->SetPortFlags(UE::SavePackageUtilities::Private::GetSavePackagePortFlags());
+		Linker->SetPortFlags(UE::SavePackageUtilities::Private::GetSavePackagePortFlags());
 		if (SaveContext.IsSaveAutoOptional() && SaveContext.GetCurrentHarvestingRealm() == ESaveRealm::Optional)
 		{
 			// Do not filter editor only data when automatically creating optional data using full uncooked objects
-			SaveContext.GetLinker()->SetFilterEditorOnly(false);
+			Linker->SetFilterEditorOnly(false);
 		}
 		else
 		{
-			SaveContext.GetLinker()->SetFilterEditorOnly(SaveContext.IsFilterEditorOnly());
+			Linker->SetFilterEditorOnly(SaveContext.IsFilterEditorOnly());
 		}
-		SaveContext.GetLinker()->SetSavePackageData(&SaveContext.GetArchiveSavePackageData());
+		Linker->SetSavePackageData(&SaveContext.GetArchiveSavePackageData());
 
 		bool bUseUnversionedProperties = SaveContext.IsSaveUnversionedProperties();
-		SaveContext.GetLinker()->SetUseUnversionedPropertySerialization(bUseUnversionedProperties);
+		Linker->SetUseUnversionedPropertySerialization(bUseUnversionedProperties);
 
 #if WITH_EDITOR
 		if (SaveContext.IsCooking())
 		{
-			SaveContext.GetLinker()->SetDebugSerializationFlags(DSF_EnableCookerWarnings | SaveContext.GetLinker()->GetDebugSerializationFlags());
-			SaveContext.GetLinker()->SetSaveBulkDataToSeparateFiles(true);
+			Linker->SetDebugSerializationFlags(DSF_EnableCookerWarnings | Linker->GetDebugSerializationFlags());
+			Linker->SetSaveBulkDataToSeparateFiles(true);
 
 			if (const ITargetPlatform* TargetPlatform = SaveContext.GetTargetPlatform())
 			{
 				if (TargetPlatform->SupportsFeature(ETargetPlatformFeatures::MemoryMappedFiles))
 				{
-					SaveContext.GetLinker()->SetMemoryMapAlignment(TargetPlatform->GetMemoryMappingAlignment());
+					Linker->SetMemoryMapAlignment(TargetPlatform->GetMemoryMappingAlignment());
 				}
-				SaveContext.GetLinker()->SetFileRegionsEnabled(TargetPlatform->SupportsFeature(ETargetPlatformFeatures::CookFileRegionMetadata));
+				Linker->SetFileRegionsEnabled(TargetPlatform->SupportsFeature(ETargetPlatformFeatures::CookFileRegionMetadata));
 			}
 		}
 		else if (SaveContext.GetPackageWriter() && SaveContext.GetSavePackageContext() &&
 			SaveContext.GetSavePackageContext()->PackageWriterCapabilities.bDeclareRegionForEachAdditionalFile)
 		{
-			SaveContext.GetLinker()->SetFileRegionsEnabled(true);
-			SaveContext.GetLinker()->SetDeclareRegionForEachAdditionalFile(true);
+			Linker->SetFileRegionsEnabled(true);
+			Linker->SetDeclareRegionForEachAdditionalFile(true);
 		}
 #endif
 		if (SaveContext.IsCooking() == false)
 		{
 			const bool bSaveBulkDataByReference = (SaveContext.GetSaveArgs().SaveFlags & SAVE_BulkDataByReference) != 0;
-			SaveContext.GetLinker()->SetSaveBulkDataByReference(bSaveBulkDataByReference);
+			Linker->SetSaveBulkDataByReference(bSaveBulkDataByReference);
 		}
 
 		// Make sure the package has the same version as the linker
@@ -1231,7 +1261,7 @@ ESavePackageResult BuildLinker(FSaveContext& SaveContext)
 		FIoHash SavedHash;
 		FMemory::Memcpy(&SavedHash.GetBytes(), OutputPackageGuid,
 			FMath::Min(sizeof(decltype(SavedHash.GetBytes())), sizeof(*OutputPackageGuid)));
-		SaveContext.GetLinker()->Summary.SetSavedHash(SavedHash);
+		Linker->Summary.SetSavedHash(SavedHash);
 	}
 	else
 	{
@@ -1245,21 +1275,19 @@ ESavePackageResult BuildLinker(FSaveContext& SaveContext)
 				FMath::Min(sizeof(decltype(SavedHash.GetBytes())), sizeof(Guid)));
 			SaveContext.GetPackage()->SetSavedHash(SavedHash);
 		}
-		SaveContext.GetLinker()->Summary.SetSavedHash(SaveContext.GetPackage()->GetSavedHash());
+		Linker->Summary.SetSavedHash(SaveContext.GetPackage()->GetSavedHash());
 	}
-	SaveContext.GetLinker()->Summary.PersistentGuid = SaveContext.GetPackage()->GetPersistentGuid();
+	Linker->Summary.PersistentGuid = SaveContext.GetPackage()->GetPersistentGuid();
 #endif
-	SaveContext.GetLinker()->Summary.Generations = TArray<FGenerationInfo>{ FGenerationInfo(0, 0) };
+	Linker->Summary.Generations = TArray<FGenerationInfo>{ FGenerationInfo(0, 0) };
 	if (SaveContext.IsProceduralSave())
 	{
 		// Procedural saves should be deterministic, so we have to clear the EngineVersion fields
 		// to avoid indeterminism when it changes
-		SaveContext.GetLinker()->Summary.SavedByEngineVersion = FEngineVersion();
-		SaveContext.GetLinker()->Summary.CompatibleWithEngineVersion = FEngineVersion();
+		Linker->Summary.SavedByEngineVersion = FEngineVersion();
+		Linker->Summary.CompatibleWithEngineVersion = FEngineVersion();
 	}
 
-	FLinkerSave* Linker = SaveContext.GetLinker();
-
 	// Build Name Map
 	{
 		SCOPED_SAVETIMER(UPackage_Save_BuildNameMap);
@@ -2676,38 +2704,42 @@ ESavePackageResult FinalizeFile(FStructuredArchive::FRecord& StructuredArchiveRo
 
 	FSavePackageContext* SavePackageContext = SaveContext.GetSavePackageContext();
 	IPackageWriter* PackageWriter = SavePackageContext ? SavePackageContext->PackageWriter : nullptr;
-	if (PackageWriter || SaveContext.IsSaveToMemory())
+	bool bIsOptionalRealm = SaveContext.GetCurrentHarvestingRealm() == ESaveRealm::Optional;
+	FLinkerSave* Linker = SaveContext.GetLinker();
+	if (PackageWriter)
 	{
-		bool bIsOptionalRealm = SaveContext.GetCurrentHarvestingRealm() == ESaveRealm::Optional;
-		FLinkerSave* Linker = SaveContext.GetLinker();
-		UE_LOG(LogSavePackage, VeryVerbose, TEXT("Async saving from memory to '%s'"), SaveContext.GetFilename());
 		FLargeMemoryWriter* Writer = static_cast<FLargeMemoryWriter*>(Linker->Saver);
 
-		if (PackageWriter)
+		IPackageWriter::FPackageInfo PackageInfo;
+		PackageInfo.PackageName = SaveContext.GetPackage()->GetFName();
+		// Adjust LooseFilePath if needed
+		if (bIsOptionalRealm)
 		{
-			IPackageWriter::FPackageInfo PackageInfo;
-			PackageInfo.PackageName = SaveContext.GetPackage()->GetFName();
-			// Adjust LooseFilePath if needed
-			if (bIsOptionalRealm)
-			{
-				// Optional output have the form PackagePath.o.ext
-				PackageInfo.LooseFilePath = FPathViews::ChangeExtension(SaveContext.GetFilename(), TEXT("o.") + FPaths::GetExtension(SaveContext.GetFilename()));
-				PackageInfo.MultiOutputIndex = 1;
-			}
-			else
-			{
-				PackageInfo.LooseFilePath = SaveContext.GetFilename();
-			}
-			PackageInfo.HeaderSize = Linker->Summary.TotalHeaderSize;
-
-			FPackageId PackageId = FPackageId::FromName(PackageInfo.PackageName);
-			PackageInfo.ChunkId = CreateIoChunkId(PackageId.Value(), PackageInfo.MultiOutputIndex, EIoChunkType::ExportBundleData);
-
-			PackageWriter->WritePackageData(PackageInfo, *Writer, Linker->FileRegions);
+			// Optional output have the form PackagePath.o.ext
+			PackageInfo.LooseFilePath = FPathViews::ChangeExtension(SaveContext.GetFilename(), TEXT("o.") + FPaths::GetExtension(SaveContext.GetFilename()));
+			PackageInfo.MultiOutputIndex = 1;
 		}
 		else
 		{
-			checkf(!SaveContext.IsCooking(), TEXT("Cooking requires a PackageWriter"));
+			PackageInfo.LooseFilePath = SaveContext.GetFilename();
+		}
+		PackageInfo.HeaderSize = Linker->Summary.TotalHeaderSize;
+
+		FPackageId PackageId = FPackageId::FromName(PackageInfo.PackageName);
+		PackageInfo.ChunkId = CreateIoChunkId(PackageId.Value(), PackageInfo.MultiOutputIndex, EIoChunkType::ExportBundleData);
+
+		PackageWriter->WritePackageData(PackageInfo, *Writer, Linker->FileRegions);
+		SaveContext.CloseLinkerArchives();
+	}
+	else
+	{
+		check(!SaveContext.IsCooking()); // Cooking requires a PackageWriter
+		check(!bIsOptionalRealm); // OptionalRealm is only supported when cooking.
+
+		if (SaveContext.IsSaveToMemory())
+		{
+			FLargeMemoryWriter* Writer = static_cast<FLargeMemoryWriter*>(Linker->Saver);
+
 			EAsyncWriteOptions WriteOptions(EAsyncWriteOptions::None);
 
 			// Add the uasset file to the list of output files
@@ -2719,51 +2751,90 @@ ESavePackageResult FinalizeFile(FStructuredArchive::FRecord& StructuredArchiveRo
 			{
 				UE::SavePackageUtilities::AsyncWriteFile(WriteOptions, Entry);
 			}
+			SaveContext.CloseLinkerArchives();
 		}
-		SaveContext.CloseLinkerArchives();
-	}
-	else
-	{
-		// Destroy archives used for saving, closing file handle.
-		if (SaveContext.CloseLinkerArchives() == false)
+		else if (!SaveContext.IsTextFormat())
 		{
-			UE_LOG(LogSavePackage, Error, TEXT("Error writing temp file '%s' for '%s'"),
-				SaveContext.GetTempFilename().IsSet() ? *SaveContext.GetTempFilename().GetValue() : TEXT("UNKNOWN"), SaveContext.GetFilename());
-			return ESavePackageResult::Error;
-		}
+			checkf(SaveContext.GetTempFilename().IsSet(),
+				TEXT("The package should've been saved to a tmp file first! (%s)"), SaveContext.GetFilename());
+			const FString* TempFilename = SaveContext.GetTempFilename().GetPtrOrNull();
+
+			// Destroy archives used for saving, closing file handle.
+			if (SaveContext.CloseLinkerArchives() == false)
+			{
+				UE_LOG(LogSavePackage, Error, TEXT("Error writing temp file '%s' for '%s'"),
+					**TempFilename, SaveContext.GetFilename());
+				return ESavePackageResult::Error;
+			}
+
+			// Move file to its real destination
 
-		// Move file to its real destination
-		checkf(SaveContext.GetTempFilename().IsSet(), TEXT("The package should've been saved to a tmp file first! (%s)"), SaveContext.GetFilename());
+			// Add the .uasset file to the list of output files
+			// TODO: Fix the 0 size, it isn't used after this point but needs to be cleaned up.
+			SaveContext.AdditionalPackageFiles.Emplace(SaveContext.GetFilename(), *TempFilename, 0);
 
-		// When saving in text format we will have two temp files, so we need to manually delete the non-textbased one
-		if (SaveContext.IsTextFormat())
+			ESavePackageResult FinalizeResult = UE::SavePackageUtilities::FinalizeTempOutputFiles(
+				SaveContext.GetTargetPackagePath(), SaveContext.AdditionalPackageFiles, SaveContext.GetFinalTimestamp());
+
+			SaveContext.SetTempFilename(TOptional<FString>());
+			TempFilename = nullptr;
+
+			if (FinalizeResult != ESavePackageResult::Success)
+			{
+				if (SaveContext.IsGenerateSaveError())
+				{
+					UE_LOG(LogSavePackage, Error, TEXT("%s"),
+						*FString::Printf(TEXT("Error saving '%s'"), SaveContext.GetFilename()));
+					SaveContext.GetError()->Logf(ELogVerbosity::Warning, TEXT("%s"),
+						*FText::Format(NSLOCTEXT("Core", "SaveWarning", "Error saving '{0}'"),
+							FText::FromString(FString(SaveContext.GetFilename()))).ToString());
+				}
+				else
+				{
+					UE_LOG(LogSavePackage, Warning, TEXT("Error saving '%s'"), SaveContext.GetFilename());
+				}
+				return FinalizeResult;
+			}
+		}
+		else
 		{
+			// Destroy archives used for saving, closing file handle.
+			if (SaveContext.CloseLinkerArchives() == false)
+			{
+				UE_LOG(LogSavePackage, Error, TEXT("Error writing temp file '%s' for '%s'"),
+					SaveContext.GetTempFilename().IsSet() ? *SaveContext.GetTempFilename().GetValue() : TEXT("UNKNOWN"), SaveContext.GetFilename());
+				return ESavePackageResult::Error;
+			}
+
+			// When saving in text format we will have two temp files, so we need to manually delete the non-textbased one
+			check(SaveContext.GetTempFilename().IsSet());
 			check(SaveContext.GetTextFormatTempFilename().IsSet());
 			IFileManager::Get().Delete(*SaveContext.GetTempFilename().GetValue());
 			SaveContext.SetTempFilename(SaveContext.GetTextFormatTempFilename());
 			SaveContext.SetTextFormatTempFilename(TOptional<FString>());
-		}
 
-		// Add the .uasset file to the list of output files (TODO: Fix the 0 size, it isn't used after this point but needs to be cleaned up)
-		SaveContext.AdditionalPackageFiles.Emplace(SaveContext.GetFilename(), SaveContext.GetTempFilename().GetValue(), 0);
+			// Add the .utxt file to the list of output files
+			// TODO: Fix the 0 size, it isn't used after this point but needs to be cleaned up.
+			SaveContext.AdditionalPackageFiles.Emplace(SaveContext.GetFilename(), SaveContext.GetTempFilename().GetValue(), 0);
 
-		ESavePackageResult FinalizeResult = UE::SavePackageUtilities::FinalizeTempOutputFiles(SaveContext.GetTargetPackagePath(), SaveContext.AdditionalPackageFiles,
-			SaveContext.GetFinalTimestamp());
-		
-		SaveContext.SetTempFilename(TOptional<FString>());
+			ESavePackageResult FinalizeResult = UE::SavePackageUtilities::FinalizeTempOutputFiles(SaveContext.GetTargetPackagePath(), SaveContext.AdditionalPackageFiles,
+				SaveContext.GetFinalTimestamp());
 
-		if (FinalizeResult != ESavePackageResult::Success)
-		{
-			if (SaveContext.IsGenerateSaveError())
-			{
-				UE_LOG(LogSavePackage, Error, TEXT("%s"), *FString::Printf(TEXT("Error saving '%s'"), SaveContext.GetFilename()));
-				SaveContext.GetError()->Logf(ELogVerbosity::Warning, TEXT("%s"), *FText::Format(NSLOCTEXT("Core", "SaveWarning", "Error saving '{0}'"), FText::FromString(FString(SaveContext.GetFilename()))).ToString());
-			}
-			else
+			SaveContext.SetTempFilename(TOptional<FString>());
+
+			if (FinalizeResult != ESavePackageResult::Success)
 			{
-				UE_LOG(LogSavePackage, Warning, TEXT("%s"), *FString::Printf(TEXT("Error saving '%s'"), SaveContext.GetFilename()));
+				if (SaveContext.IsGenerateSaveError())
+				{
+					UE_LOG(LogSavePackage, Error, TEXT("%s"), *FString::Printf(TEXT("Error saving '%s'"), SaveContext.GetFilename()));
+					SaveContext.GetError()->Logf(ELogVerbosity::Warning, TEXT("%s"), *FText::Format(NSLOCTEXT("Core", "SaveWarning", "Error saving '{0}'"), FText::FromString(FString(SaveContext.GetFilename()))).ToString());
+				}
+				else
+				{
+					UE_LOG(LogSavePackage, Warning, TEXT("%s"), *FString::Printf(TEXT("Error saving '%s'"), SaveContext.GetFilename()));
+				}
+				return FinalizeResult;
 			}
-			return FinalizeResult;
 		}
 	}
 
diff --git a/Engine/Source/Runtime/CoreUObject/Public/UObject/Linker.h b/Engine/Source/Runtime/CoreUObject/Public/UObject/Linker.h
--- a/Engine/Source/Runtime/CoreUObject/Public/UObject/Linker.h
+++ b/Engine/Source/Runtime/CoreUObject/Public/UObject/Linker.h
@@ -415,10 +415,6 @@ public:
 	/** Raw/bulk data references */
 	TArray<FObjectDataResource> DataResourceMap;
 	
-	/** The name of the file for this package */
-	UE_DEPRECATED(5.0, "Use GetDebugName for logging identifiers. For other purposes, use GetPackagePath on LinkerLoad and GetFilename on LinkerSave.")
-	FString					Filename;
-
 	/** If true, filter out exports that are for clients but not servers */
 	bool					FilterClientButNotServer;
 
@@ -430,8 +426,6 @@ public:
 
 	/** Constructor. */
 	FLinker(ELinkerType::Type InType, UPackage* InRoot);
-	UE_DEPRECATED(5.0, "Linker's filename is deprecated; subclasses should create their own filename if required")
-	FLinker(ELinkerType::Type InType, UPackage* InRoot, const TCHAR* InFilename);
 
 	virtual ~FLinker();
 
@@ -660,11 +654,6 @@ COREUOBJECT_API FLinkerLoad* LoadPackageLinker(UPackage* InOuter, const FPackage
 UE_DEPRECATED(5.4, "Use GetPackageLinker instead")
 COREUOBJECT_API FLinkerLoad* LoadPackageLinker(UPackage* InOuter, const FPackagePath& PackagePath, uint32 LoadFlags = LOAD_None, UPackageMap* Sandbox = nullptr, FArchive* InReaderOverride = nullptr);
 
-UE_DEPRECATED(5.0, "Use version that takes a FPackagePath without a FGuid instead")
-COREUOBJECT_API FLinkerLoad* LoadPackageLinker(UPackage* InOuter, const TCHAR* InLongPackageName, uint32 LoadFlags, UPackageMap* Sandbox, FGuid* CompatibleGuid, FArchive* InReaderOverride, TFunctionRef<void(FLinkerLoad* LoadedLinker)> LinkerLoadedCallback);
-UE_DEPRECATED(5.0, "Use version that takes a FPackagePath without a FGuid instead")
-COREUOBJECT_API FLinkerLoad* LoadPackageLinker(UPackage* InOuter, const TCHAR* InLongPackageName, uint32 LoadFlags = LOAD_None, UPackageMap* Sandbox = nullptr, FGuid* CompatibleGuid = nullptr, FArchive* InReaderOverride = nullptr);
-
 /** 
  * Gets a linker for a package and returns it without loading any objects.
  * @param InOuter Package if known, can be null
@@ -680,9 +669,6 @@ COREUOBJECT_API FLinkerLoad* LoadPackageLinker(UPackage* InOuter, const TCHAR* I
  */
 COREUOBJECT_API FLinkerLoad* GetPackageLinker(UPackage* InOuter, const FPackagePath& PackagePath, uint32 LoadFlags, UPackageMap* Sandbox, FArchive* InReaderOverride = nullptr, FUObjectSerializeContext** InOutLoadContext = nullptr, FLinkerLoad* ImportLinker = nullptr, const FLinkerInstancingContext* InstancingContext = nullptr);
 
-UE_DEPRECATED(5.0, "Use version that takes a FPackagePath without a FGuid instead")
-COREUOBJECT_API FLinkerLoad* GetPackageLinker(UPackage* InOuter, const TCHAR* InLongPackageName, uint32 LoadFlags, UPackageMap* Sandbox, FGuid* CompatibleGuid, FArchive* InReaderOverride = nullptr, FUObjectSerializeContext** InOutLoadContext = nullptr, FLinkerLoad* ImportLinker = nullptr, const FLinkerInstancingContext* InstancingContext = nullptr);
-
 COREUOBJECT_API FString GetPrestreamPackageLinkerName(const TCHAR* InLongPackageName, bool bSkipIfExists = true);
 
 /**
diff --git a/Engine/Source/Runtime/CoreUObject/Public/UObject/LinkerLoad.h b/Engine/Source/Runtime/CoreUObject/Public/UObject/LinkerLoad.h
--- a/Engine/Source/Runtime/CoreUObject/Public/UObject/LinkerLoad.h
+++ b/Engine/Source/Runtime/CoreUObject/Public/UObject/LinkerLoad.h
@@ -636,8 +636,6 @@ public:
 	 * @return	new FLinkerLoad object for Parent/ Filename
 	 */
 	COREUOBJECT_API static FLinkerLoad* CreateLinker(FUObjectSerializeContext* LoadContext, UPackage* Parent, const FPackagePath& PackagePath, uint32 LoadFlags, FArchive* InLoader = nullptr, const FLinkerInstancingContext* InstancingContext = nullptr);
-	UE_DEPRECATED(5.0, "Use version that takes an FPackagePath instead")
-	COREUOBJECT_API static FLinkerLoad* CreateLinker(FUObjectSerializeContext* LoadContext, UPackage* Parent, const TCHAR* Filename, uint32 LoadFlags, FArchive* InLoader = nullptr, const FLinkerInstancingContext* InstancingContext = nullptr);
 
 	void Verify();
 
@@ -704,11 +702,9 @@ public:
 	 */
 	UObject* ResolveResource(FPackageIndex ResourceIndex);
 	
+	UE_DEPRECATED(5.6, "This function was used only for the now-deleted Conform functionality and will be removed. Contact Epic if you need this function.")
 	int32 FindExportIndex( FName ClassName, FName ClassPackage, FName ObjectName, FPackageIndex ExportOuterIndex );
 
-	UE_DEPRECATED(5.0, "Create was only used for the now-deprecated Conform argument to UPackage::Save. Contact Epic via UDN if you still need this function.")
-	UObject* Create( UClass* ObjectClass, FName ObjectName, UObject* Outer, uint32 InLoadFlags, bool Checked );
-
 	/**
 	 * Serialize the object data for the specified object from the unreal package file.  Loads any
 	 * additional resources required for the object to be in a valid state to receive the loaded
@@ -815,8 +811,6 @@ public:
 	COREUOBJECT_API static bool GetPreloadingEnabled();
 	COREUOBJECT_API static void SetPreloadingEnabled(bool bEnabled);
 	COREUOBJECT_API static bool TryGetPreloadedLoader(const FPackagePath& InPackagePath, FOpenPackageResult& OutResult);
-	UE_DEPRECATED(5.0, "Use version that takes a PackagePath instead")
-	COREUOBJECT_API static bool TryGetPreloadedLoader(FArchive*& OutLoader, const TCHAR* FileName);
 	private:
 		static bool bPreloadingEnabled;
 	public:
@@ -1069,7 +1063,7 @@ private:
 		return Loader->TotalSize();
 	}
 
-	// this fixes the warning : 'ULinkerSave::Serialize' hides overloaded virtual function
+	// this fixes the warning : 'ULinkerLoad::Serialize' hides overloaded virtual function
 	using FLinker::Serialize;
 	FORCEINLINE virtual void Serialize(void* V, int64 Length) override
 	{
@@ -1079,8 +1073,6 @@ private:
 
 		checkSlow(FPlatformTLS::GetCurrentThreadId() == OwnerThread);
 #if WITH_EDITOR
-		Loader->SetSerializedProperty(GetSerializedProperty());
-
 		// This handles the case where we're only serializing to a property bag and we don't immediately know where the script
 		// property data starts/ends in the object's data stream. Allows for loading from packages saved with an older version.
 		if (UNLIKELY(bIsLoadingToPropertyBagObject && !bIsSerializingScriptProperties))
@@ -1136,6 +1128,42 @@ private:
 		return *this;
 	}
 
+	virtual void SetSerializedProperty(FProperty* InProperty) override
+	{
+		FArchiveUObject::SetSerializedProperty(InProperty);
+		Loader->SetSerializedProperty(InProperty);
+	}
+	virtual void SetSerializedPropertyChain(const FArchiveSerializedPropertyChain* InSerializedPropertyChain,
+		class FProperty* InSerializedPropertyOverride = nullptr) override
+	{
+		FArchiveUObject::SetSerializedPropertyChain(InSerializedPropertyChain, InSerializedPropertyOverride);
+		Loader->SetSerializedPropertyChain(InSerializedPropertyChain, InSerializedPropertyOverride);
+	}
+	virtual void PushSerializedProperty(class FProperty* InProperty, const bool bIsEditorOnlyProperty) override
+	{
+		FArchive::PushSerializedProperty(InProperty, bIsEditorOnlyProperty);
+		Loader->PushSerializedProperty(InProperty, bIsEditorOnlyProperty);
+	}
+	virtual void PopSerializedProperty(class FProperty* InProperty, const bool bIsEditorOnlyProperty) override
+	{
+		FArchive::PopSerializedProperty(InProperty, bIsEditorOnlyProperty);
+		Loader->PopSerializedProperty(InProperty, bIsEditorOnlyProperty);
+	}
+#if WITH_EDITORONLY_DATA
+	virtual bool IsEditorOnlyPropertyOnTheStack() const override
+	{
+		return Loader->IsEditorOnlyPropertyOnTheStack();
+	}
+#endif
+
+	virtual void SetDebugSerializationFlags(uint32 InCustomFlags) override
+	{
+		FArchiveUObject::SetDebugSerializationFlags(InCustomFlags);
+		Loader->SetDebugSerializationFlags(InCustomFlags);
+	}
+
+
+	// FLinkerLoad functions
 	/**
 	 * Safely verify that an import in the ImportMap points to a good object. This decides whether or not
 	 * a failure to load the object redirector in the wrapper is a fatal error or not (return value)
diff --git a/Engine/Source/Runtime/CoreUObject/Public/UObject/LinkerSave.h b/Engine/Source/Runtime/CoreUObject/Public/UObject/LinkerSave.h
--- a/Engine/Source/Runtime/CoreUObject/Public/UObject/LinkerSave.h
+++ b/Engine/Source/Runtime/CoreUObject/Public/UObject/LinkerSave.h
@@ -56,7 +56,7 @@ public:
 
 	// Variables.
 	/** The archive that actually writes the data to disk. */
-	FArchive* Saver;
+	FArchive* Saver = nullptr;
 
 	FPackageIndex CurrentlySavingExport;
 	UObject* CurrentlySavingExportObject = nullptr;
@@ -142,12 +142,43 @@ public:
 	COREUOBJECT_API static TMap<FString, TArray<uint8> > PackagesToScriptSHAMap;
 
 	/** Constructor for file writer */
+	UE_DEPRECATED(5.6, "Use constructor that takes only Parent and then call AssignFileSaver.")
 	FLinkerSave(UPackage* InParent, const TCHAR* InFilename, bool bForceByteSwapping, bool bInSaveUnversioned = false );
 	/** Constructor for memory writer */
+	UE_DEPRECATED(5.6, "Use constructor that takes only Parent and then call AssignMemorySaver.")
 	FLinkerSave(UPackage* InParent, bool bForceByteSwapping, bool bInSaveUnversioned = false );
 	/** Constructor for custom savers. The linker assumes ownership of the custom saver. */
+	UE_DEPRECATED(5.6, "Use constructor that takes only Parent and then call AssignSaver.")
 	FLinkerSave(UPackage* InParent, FArchive *InSaver, bool bForceByteSwapping, bool bInSaveUnversioned = false);
 
+	/**
+	 * Partially construct the FLinkerSave. No FArchive functions may be called on *this until AssignSaver has been
+	 * called with a non-null FArchive.
+	 */
+	explicit FLinkerSave(UPackage* InParent);
+
+	/**
+	 * Assign the Saver that will receive the binary data for the package and write it to disk/memory/unspecified.
+	 * Must be called only once, after construction and before any calls to FArchive functions on this.
+	 * Takes ownership of the given archive, and deletes it in CloseAndDestroySaver or destructor.
+	 * 
+	 * @param bForceByteSwapping Passed to FArchive::IsByteSwapping
+	 * @param bInSaveUnversioned Passed to FPackageFileSummary::bUnversioned, which if true removes native versioning
+	 *        and custom versions from the saved package. Used for cooked packages, which exactly match native
+	 *        runtime version.
+	 */
+	void AssignSaver(FArchive* InSaver, bool bForceByteSwapping, bool bInSaveUnversioned = false);
+	/** Construct a FMemoryWriter archive and call AssignSaver with it. @see AssignSaver. */
+	void AssignMemorySaver(bool bForceByteSwapping, bool bInSaveUnversioned = false);
+	/**
+	 * Construct a file saver from IPlatformFile and call AssignSaver with it. @see AssignSaver.
+	 * If the file write fails, the Saver will be set to nullptr and it is a fatal error to call FArchive functions on
+	 * *this. Caller must handle returnvalue == false and if so destroy *this without calling FArchive functions.
+	 * 
+	 * @return Whether the creation of the Saver archive succeeded.
+	 */
+	bool TryAssignFileSaver(FStringView InFilename, bool bForceByteSwapping, bool bInSaveUnversioned = false);
+
 	/** Returns the appropriate name index for the source name, or INDEX_NONE if not found in NameIndices */
 	int32 MapName( FNameEntryId Name) const;
 
@@ -178,15 +209,8 @@ public:
 	UE_DEPRECATED(5.5, "GetSerializeContext is not supported. Use FUObjectThreadContext::Get().GetSerializeContext().")
 	FUObjectSerializeContext* GetSerializeContext() override;
 	virtual void UsingCustomVersion(const struct FGuid& Guid) override;
-	/**
-	 * Sets whether tagged property serialization should be replaced by faster unversioned serialization.
-	 * This sets it on itself, the summary, the actual Saver Archive if any and set the proper associated flag on the LinkerRoot
-	 */
 	virtual void SetUseUnversionedPropertySerialization(bool bInUseUnversioned) override;
-	/**
-	 * Sets whether we should be filtering editor only.
-	 * This sets it on itself, the summary, the actual Saver Archive if any and set the proper associated flag on the LinkerRoot
-	 */
+	virtual void SetDebugSerializationFlags(uint32 InCustomFlags) override;
 	virtual void SetFilterEditorOnly(bool bInFilterEditorOnly) override;
 
 	/** Sets the map of overrided properties for each export that should be treated as transient, and nulled out when serializing */
@@ -250,21 +274,25 @@ public:
 	/** Triggered after bulk data payloads has been serialized to disk/package writer.*/
 	void OnPostSaveBulkData();
 
+	virtual void SetSerializedProperty(FProperty* InProperty) override;
+	virtual void SetSerializedPropertyChain(const FArchiveSerializedPropertyChain* InSerializedPropertyChain,
+		class FProperty* InSerializedPropertyOverride = nullptr) override;
+	virtual void PushSerializedProperty(class FProperty* InProperty, const bool bIsEditorOnlyProperty) override;
+	virtual void PopSerializedProperty(class FProperty* InProperty, const bool bIsEditorOnlyProperty) override;
+#if WITH_EDITORONLY_DATA
+	virtual bool IsEditorOnlyPropertyOnTheStack() const override;
+#endif
+
 	// FLinker interface
 	virtual FString GetDebugName() const override;
 
-	// LinkerSave interface
+	// FLinkerSave interface
 	/**
 	 * Closes and deletes the Saver (file, memory or custom writer) which will close any associated file handle.
 	 * Returns false if the owned saver contains errors after closing it, true otherwise.
 	 */
 	bool CloseAndDestroySaver();
 
-	/**
-	 * Sets a flag indicating that this archive contains data required to be gathered for localization.
-	 */
-	void ThisRequiresLocalizationGather();
-
 	/** Get the filename being saved to */
 	const FString& GetFilename() const;
 
@@ -325,6 +353,8 @@ protected:
 
 	void SerializeObjectPointer(const FObjectPtr& Value);
 
+	void AssignSaverInternal(FArchive* InSaver, bool bForceByteSwapping, bool bInSaveUnversioned);
+
 private:
 	/** Optional log output to bubble errors back up. */
 	FOutputDevice* LogOutput = nullptr;
@@ -346,6 +376,9 @@ private:
 	/** Memory mapped bulk data archive. */
 	TMap<FBulkDataCookedIndex, TUniquePtr<FFileRegionMemoryWriter>> MemoryMappedBulkDataAr;
 
+	/** The name of the file for this package */
+	FString Filename;
+
 	const TMap<UObject*, TSet<FProperty*>>* TransientPropertyOverrides = nullptr;
 	/** Alignment for memory mapped data .*/ 
 	int64 MemoryMappingAlignment = -1;
