[AsyncLoading2]

Improve checks for object flags so we have more context in logs when errors occur

#rnx

diff --git a/Engine/Source/Runtime/CoreUObject/Private/Serialization/AsyncLoading2.cpp b/Engine/Source/Runtime/CoreUObject/Private/Serialization/AsyncLoading2.cpp
--- a/Engine/Source/Runtime/CoreUObject/Private/Serialization/AsyncLoading2.cpp
+++ b/Engine/Source/Runtime/CoreUObject/Private/Serialization/AsyncLoading2.cpp
@@ -295,6 +295,12 @@ void FindAllRuntimeScriptPackages(FRuntimeScriptPackages& OutPackages)
 #define ALT2_ENABLE_PACKAGE_DEPENDENCY_DEBUGGING 0
 #endif
 
+#define checkObject(Object, expr) checkf(expr, TEXT("Object='%s' (%p), Flags=%s, InternalFlags=0x%08X") \
+									, (Object ? *Object->GetFullName() : TEXT("null")) \
+									, Object \
+									, (Object ? *LexToString(Object->GetFlags()) : TEXT("None")) \
+									, (Object ? Object->GetInternalFlags() : EInternalObjectFlags::None))
+
 static TSet<FPackageId> GAsyncLoading2_DebugPackageIds;
 static FString GAsyncLoading2_DebugPackageNamesString;
 static TSet<FPackageId> GAsyncLoading2_VerbosePackageIds;
@@ -1103,7 +1109,7 @@ private:
 					if (!ObjectItem->IsUnreachable())
 					{
 						UObject* Object = static_cast<UObject*>(ObjectItem->Object);
-						checkf(!ObjectItem->HasAnyFlags(EInternalObjectFlags::LoaderImport), TEXT("%s"), *Object->GetFullName());
+						checkObject(Object, !ObjectItem->HasAnyFlags(EInternalObjectFlags::LoaderImport));
 						ObjectItem->SetFlags(EInternalObjectFlags::LoaderImport);
 						ValidExportCount++;
 					}
@@ -1127,7 +1133,7 @@ private:
 				if (ObjectIndex >= 0)
 				{
 					UObject* Object = static_cast<UObject*>(GUObjectArray.IndexToObject(ObjectIndex)->Object);
-					checkf(Object->HasAnyInternalFlags(EInternalObjectFlags::LoaderImport), TEXT("%s"), *Object->GetFullName());
+					checkObject(Object, Object->HasAnyInternalFlags(EInternalObjectFlags::LoaderImport));
 					Object->AtomicallyClearInternalFlags(EInternalObjectFlags::LoaderImport);
 				}
 			}
@@ -1442,7 +1448,7 @@ public:
 			bAreAllPublicExportsLoaded = false;
 		}
 #endif
-		checkf(!Package->HasAnyInternalFlags(EInternalObjectFlags::LoaderImport), TEXT("%s"), *Package->GetFullName());
+		checkObject(Package, !Package->HasAnyInternalFlags(EInternalObjectFlags::LoaderImport));
 		Package->SetInternalFlags(EInternalObjectFlags::LoaderImport);
 	}
 
@@ -1456,7 +1462,7 @@ public:
 			return;
 		}
 		PublicExportMap.UnpinForGC();
-		checkf(Package->HasAnyInternalFlags(EInternalObjectFlags::LoaderImport), TEXT("%s"), *Package->GetFullName());
+		checkObject(Package, Package->HasAnyInternalFlags(EInternalObjectFlags::LoaderImport));
 		Package->AtomicallyClearInternalFlags(EInternalObjectFlags::LoaderImport);
 	}
 };
@@ -7587,7 +7593,7 @@ bool FAsyncPackage2::EventDrivenSerializeExport(const FAsyncPackageHeaderData& H
 #if DO_CHECK
 	if (Object->HasAnyFlags(RF_ClassDefaultObject) && Object->GetClass()->HasAnyClassFlags(CLASS_CompiledFromBlueprint))
 	{
-		check(Object->HasAllFlags(RF_NeedPostLoad | RF_WasLoaded));
+		checkObject(Object, Object->HasAllFlags(RF_NeedPostLoad | RF_WasLoaded));
 		//Object->SetFlags(RF_NeedPostLoad | RF_WasLoaded);
 	}
 #endif
@@ -8057,7 +8063,7 @@ EEventLoadNodeExecutionResult FAsyncPackage2::Event_PostLoadExportBundle(FAsyncL
 
 					UObject* Object = Export.Object;
 					check(Object);
-					check(!Object->HasAnyFlags(RF_NeedLoad));
+					checkObject(Object, !Object->HasAnyFlags(RF_NeedLoad));
 					if (!Object->HasAnyFlags(RF_NeedPostLoad))
 					{
 						break;
@@ -8177,7 +8183,7 @@ EEventLoadNodeExecutionResult FAsyncPackage2::Event_DeferredPostLoadExportBundle
 
 					UObject* Object = Export.Object;
 					check(Object);
-					check(!Object->HasAnyFlags(RF_NeedLoad));
+					checkObject(Object, !Object->HasAnyFlags(RF_NeedLoad));
 					if (Object->HasAnyFlags(RF_NeedPostLoad))
 					{
 #if WITH_EDITOR
@@ -9579,7 +9585,7 @@ void FAsyncPackage2::MoveConstructedObjectsToPhase2()
 		UE_MT_SCOPED_READ_ACCESS(ConstructedObjectsAccessDetector)
 		for (UObject* Object : ConstructedObjects)
 		{
-			check(!Object->HasAnyFlags(RF_NeedLoad | RF_NeedInitialization));
+			checkObject(Object, !Object->HasAnyFlags(RF_NeedLoad | RF_NeedInitialization));
 			Object->SetInternalFlags(EInternalObjectFlags::AsyncLoadingPhase2);
 			Object->AtomicallyClearInternalFlags(EInternalObjectFlags::AsyncLoadingPhase1);
 		}
@@ -9600,7 +9606,7 @@ void FAsyncPackage2::MoveConstructedObjectsToPhase2()
 			{
 				if (SubObject && SubObject->HasAnyInternalFlags(EInternalObjectFlags::AsyncLoadingPhase1))
 				{
-					check(!SubObject->HasAnyFlags(RF_NeedLoad | RF_NeedInitialization));
+					checkObject(SubObject, !SubObject->HasAnyFlags(RF_NeedLoad | RF_NeedInitialization));
 					SubObject->SetInternalFlags(EInternalObjectFlags::AsyncLoadingPhase2);
 					SubObject->AtomicallyClearInternalFlags(EInternalObjectFlags::AsyncLoadingPhase1);
 				}
@@ -9610,7 +9616,7 @@ void FAsyncPackage2::MoveConstructedObjectsToPhase2()
 
 		if (Object)
 		{
-			check(!Object->HasAnyFlags(RF_NeedLoad | RF_NeedInitialization));
+			checkObject(Object, !Object->HasAnyFlags(RF_NeedLoad | RF_NeedInitialization));
 			if (Object->HasAnyInternalFlags(EInternalObjectFlags::AsyncLoadingPhase1))
 			{
 				Object->SetInternalFlags(EInternalObjectFlags::AsyncLoadingPhase2);
@@ -9621,7 +9627,7 @@ void FAsyncPackage2::MoveConstructedObjectsToPhase2()
 
 	if (LinkerRoot && LinkerRoot->HasAnyInternalFlags(EInternalObjectFlags::AsyncLoadingPhase1))
 	{
-		check(!LinkerRoot->HasAnyFlags(RF_NeedLoad | RF_NeedInitialization));
+		checkObject(LinkerRoot, !LinkerRoot->HasAnyFlags(RF_NeedLoad | RF_NeedInitialization));
 		LinkerRoot->SetInternalFlags(EInternalObjectFlags::AsyncLoadingPhase2);
 		LinkerRoot->AtomicallyClearInternalFlags(EInternalObjectFlags::AsyncLoadingPhase1);
 	}
@@ -9649,11 +9655,7 @@ void FAsyncPackage2::ClearConstructedObjects()
 		}
 		else
 		{
-			checkf(!Export.Object || !Export.Object->HasAnyInternalFlags(EInternalObjectFlags_AsyncLoading | EInternalObjectFlags::Async),
-				TEXT("Export object: %s (ObjectFlags=%x, InternalObjectFlags=%x)"),
-					*Export.Object->GetFullName(),
-					Export.Object->GetFlags(),
-					Export.Object->GetInternalFlags());
+			checkObject(Export.Object, !Export.Object || !Export.Object->HasAnyInternalFlags(EInternalObjectFlags_AsyncLoading | EInternalObjectFlags::Async));
 		}
 	}
 
