Don't enable additional UObject checks when setting global log level to warning

diff --git a/Engine/Source/Runtime/CoreUObject/Private/UObject/Obj.cpp b/Engine/Source/Runtime/CoreUObject/Private/UObject/Obj.cpp
--- a/Engine/Source/Runtime/CoreUObject/Private/UObject/Obj.cpp
+++ b/Engine/Source/Runtime/CoreUObject/Private/UObject/Obj.cpp
@@ -83,12 +83,16 @@
 
 DEFINE_LOG_CATEGORY(LogObj);
 
-namespace {
-static int32 GPrematureConfigLoadErrorLevel = 0;
-static FAutoConsoleVariableRef CVarPrematureConfigLoadErrorLevel(
-	TEXT("obj.PrematureConfigLoadErrorLevel"),
-	GPrematureConfigLoadErrorLevel,
-	TEXT("Determines the level of warning issued if LoadConfig is called on an object that is not yet sufficiently initialized. 0=Display, 1=Warning, 2=Ensure, 3=Check"));
+namespace
+{
+	static int32 GPrematureConfigLoadErrorLevel = 0;
+	static FAutoConsoleVariableRef CVarPrematureConfigLoadErrorLevel(
+		TEXT("obj.PrematureConfigLoadErrorLevel"),
+		GPrematureConfigLoadErrorLevel,
+		TEXT("Determines the level of warning issued if LoadConfig is called on an object that is not yet sufficiently initialized. 0=Display, 1=Warning, 2=Ensure, 3=Check"));
+
+	// Enable additional (slow) checks in CheckDefaultSubobjects.
+	static constexpr bool bDoCheckSubobjects = false;
 }
 
 /*-----------------------------------------------------------------------------
@@ -2112,17 +2116,13 @@ protected:
 	const UObject* ReferencingObject;
 };
 
-
-// if this is set to fatal, then we don't run any testing since it is time consuming.
-DEFINE_LOG_CATEGORY_STATIC(LogCheckSubobjects, Fatal, All);
-
 #define CompCheck(Pred) \
 	if (!(Pred)) \
 	{ \
 		Result = false; \
 		UE_DEBUG_BREAK(); \
-		UE_LOG(LogCheckSubobjects, Log, TEXT("CompCheck %s failed."), TEXT(#Pred)); \
-	} 
+		UE_LOG(LogObj, Error, TEXT("CompCheck %s failed."), TEXT(#Pred)); \
+	}
 
 bool UObject::CanCheckDefaultSubObjects(bool bForceCheck, bool& bResult) const
 {
@@ -2142,7 +2142,7 @@ bool UObject::CanCheckDefaultSubObjects(bool bForceCheck, bool& bResult) const
 		bCanCheck = false;
 	}
 	// If errors are suppressed, we will not take the time to run this test unless forced to.
-	bCanCheck = bCanCheck && (bForceCheck || UE_LOG_ACTIVE(LogCheckSubobjects, Error));
+	bCanCheck = bCanCheck && (bForceCheck || bDoCheckSubobjects);
 	return bCanCheck;
 }
 
@@ -2158,9 +2158,8 @@ bool UObject::CheckDefaultSubobjects(bool bForceCheck /*= false*/) const
 
 bool UObject::CheckDefaultSubobjectsInternal() const
 {
-	bool Result = true;	
+	bool Result = true;
 
-	CompCheck(this);
 	UClass* ObjClass = GetClass();
 
 	if (ObjClass != UFunction::StaticClass())
@@ -2172,7 +2171,7 @@ bool UObject::CheckDefaultSubobjectsInternal() const
 		for (int32 Index = 0; Index < OtherReferencedSubobjects.Num(); ++Index)
 		{
 			const UObject* TestObject = OtherReferencedSubobjects[Index];
-			UE_LOG(LogCheckSubobjects, Error, TEXT("%s has a reference to default subobject (%s) of %s."), *GetFullName(), *TestObject->GetFullName(), *TestObject->GetOuter()->GetFullName());
+			UE_LOG(LogObj, Error, TEXT("%s has a reference to default subobject (%s) of %s."), *GetFullName(), *TestObject->GetFullName(), *TestObject->GetOuter()->GetFullName());
 		}
 		CompCheck(OtherReferencedSubobjects.Num() == 0);
 	}
