UE: Fix soft object subpath data transformation to utf8 code to create a properly sized char buffer at load time so that the save  works as expected

diff --git a/Engine/Source/Runtime/CoreUObject/Private/UObject/SoftObjectPath.cpp b/Engine/Source/Runtime/CoreUObject/Private/UObject/SoftObjectPath.cpp
--- a/Engine/Source/Runtime/CoreUObject/Private/UObject/SoftObjectPath.cpp
+++ b/Engine/Source/Runtime/CoreUObject/Private/UObject/SoftObjectPath.cpp
@@ -516,7 +516,7 @@ void SoftObjectPathLoadSubPathWorkaround(FArchive& Ar, FUtf8String& OutStr)
 	}
 	else if (!bLoadUnicodeChar && ((ANSICHAR*)SavedBytes.GetData())[SaveNum - 1] == (ANSICHAR)0)
 	{
-		OutStr = (ANSICHAR*)SavedBytes.GetData();
+		OutStr = FUtf8String::ConstructFromPtrSize((ANSICHAR*)SavedBytes.GetData(), SavedBytes.Num());
 		return;
 	}
 
