Add some unit tests to verify that some TWeakFieldPtr usecases compile.

#rnx

- A number of compilation issues with TWeakFieldPtr were fixed in recent refactors. Adding these use cases to the unit tests to catch any future regression.

diff --git a/Engine/Source/Runtime/CoreUObject/Tests/UObject/WeakFieldPtrTest.cpp b/Engine/Source/Runtime/CoreUObject/Tests/UObject/WeakFieldPtrTest.cpp
--- a/Engine/Source/Runtime/CoreUObject/Tests/UObject/WeakFieldPtrTest.cpp
+++ b/Engine/Source/Runtime/CoreUObject/Tests/UObject/WeakFieldPtrTest.cpp
@@ -28,6 +28,49 @@ TEST_CASE_NAMED(TWeakFieldPtrSmokeTest, "CoreUObject::TWeakFieldPtr::Smoke", "[C
 		CHECK(ConstPtr1 == PropPtr);
 		CHECK(!(ConstPtr1 != PropPtr));
 	}
+
+	SECTION("Check that TWeakFieldPtr of different types can be compared in both directions")
+	{
+		FProperty* TestProperty = nullptr;
+
+		TWeakFieldPtr<FProperty> WeakProperty = TestProperty;
+		TWeakFieldPtr<FField> WeakField = TestProperty;
+
+		CHECK(WeakProperty == WeakField);
+		CHECK(WeakField == WeakProperty);
+	}
+
+	SECTION("Check that TWeakFieldPtr can be compared against raw pointers")
+	{
+		FProperty* TestProperty = nullptr;
+		FProperty* ConstTestProperty = nullptr;
+
+		FProperty* TestField = nullptr;
+		FProperty* ConstTestField = nullptr;
+
+		TWeakFieldPtr<FProperty> WeakProperty = TestProperty;
+		TWeakFieldPtr<const FProperty> ConstWeakProperty = ConstTestProperty;
+
+		CHECK(WeakProperty == TestProperty);
+		CHECK_FALSE(WeakProperty != TestProperty);
+		CHECK(WeakProperty == ConstTestProperty);
+		CHECK_FALSE(WeakProperty != ConstTestProperty);
+
+		CHECK(WeakProperty == TestField);
+		CHECK_FALSE(WeakProperty != TestField);
+		CHECK(WeakProperty == ConstTestField);
+		CHECK_FALSE(WeakProperty != ConstTestField);
+
+		CHECK(ConstWeakProperty == TestProperty);
+		CHECK_FALSE(ConstWeakProperty != TestProperty);
+		CHECK(ConstWeakProperty == ConstTestProperty);
+		CHECK_FALSE(ConstWeakProperty != ConstTestProperty);
+
+		CHECK(ConstWeakProperty == TestField);
+		CHECK_FALSE(ConstWeakProperty != TestField);
+		CHECK(ConstWeakProperty == ConstTestField);
+		CHECK_FALSE(ConstWeakProperty != ConstTestField);
+	}
 }
 
 } // UE
