RemoteObjects:  Helping debug live locks

#rnx

diff --git a/Engine/Source/Runtime/CoreUObject/Private/UObject/RemoteExecutor.cpp b/Engine/Source/Runtime/CoreUObject/Private/UObject/RemoteExecutor.cpp
--- a/Engine/Source/Runtime/CoreUObject/Private/UObject/RemoteExecutor.cpp
+++ b/Engine/Source/Runtime/CoreUObject/Private/UObject/RemoteExecutor.cpp
@@ -195,7 +195,24 @@ public:
 
 				FRemoteTransactionId RequestId = ExecutingWork.RequestId;
 				UE::RemoteExecutor::OnTransactionStartingDelegate.Broadcast(RequestId, ExecutingWork.Name);
-				AutoRTFM::ETransactionResult WorkTransactResult = AutoRTFM::Transact([&ExecutingWork]() { ExecutingWork.Work(); });
+				AutoRTFM::ETransactionResult WorkTransactResult = AutoRTFM::Transact([&ExecutingWork]()
+					{
+						// You can set a breakpoint inside the abort handler below to see what caused the abort,
+						// or you can rely on the log output which only catches livelocks
+						const uint32 ExecutionAttempts = ExecutingWork.ExecutionAttempts;
+						UE_AUTORTFM_ONABORT(ExecutionAttempts)
+						{
+							constexpr uint32 MaxExecutionAttemptsBeforeLivelockWarning = 200;
+							if ((ExecutionAttempts >= MaxExecutionAttemptsBeforeLivelockWarning) && (ExecutionAttempts % MaxExecutionAttemptsBeforeLivelockWarning) == 0)
+							{
+								UE_LOG(LogRemoteExec, Display, TEXT("vvv Transaction Aborted %u Times. Dumping Callstack. vvv"), ExecutionAttempts);
+								FDebug::DumpStackTraceToLog(ELogVerbosity::Display);
+								UE_LOG(LogRemoteExec, Display, TEXT("^^^ Transaction Aborted %u Times ^^^"), ExecutionAttempts);
+							}
+						};
+
+						ExecutingWork.Work();
+					});
 				bool bWorkComplete = false;
 				bool bWorkTransactionAborted = false;
 
