Gate Zen package header VCell tables on package version

This restores the ability to read pre-40248290 versioned+cooked data, while losing the ability to read cooked data produced between that change and this one.

diff --git a/Engine/Source/Developer/IoStoreUtilities/Private/PackageStoreOptimizer.cpp b/Engine/Source/Developer/IoStoreUtilities/Private/PackageStoreOptimizer.cpp
--- a/Engine/Source/Developer/IoStoreUtilities/Private/PackageStoreOptimizer.cpp
+++ b/Engine/Source/Developer/IoStoreUtilities/Private/PackageStoreOptimizer.cpp
@@ -836,6 +836,7 @@ void FPackageStoreOptimizer::FinalizePackageHeader(FPackageStorePackage* Package
 	uint64 OffsetBeforeBulkDataMap =
 		sizeof(FZenPackageSummary)
 		+ VersioningInfoSize
+		+ sizeof(FZenPackageCellOffsets)
 		+ NameMapSize
 		+ sizeof(BulkDataPad);
 
@@ -878,6 +879,8 @@ void FPackageStoreOptimizer::FinalizePackageHeader(FPackageStorePackage* Package
 		PackageSummary->bHasVersioningInfo = 0;
 	}
 
+	FZenPackageCellOffsets* CellOffsets = HeaderArchive.SerializeUninitialized<FZenPackageCellOffsets>();
+
 	HeaderArchive.Serialize(NameMapArchive.GetWriterData(), NameMapArchive.Tell());
 
 	HeaderArchive << BulkDataPad;
@@ -905,8 +908,8 @@ void FPackageStoreOptimizer::FinalizePackageHeader(FPackageStorePackage* Package
 	HeaderArchive.Serialize(ImportMapArchive.GetWriterData(), ImportMapArchive.Tell());
 	PackageSummary->ExportMapOffset = HeaderArchive.Tell();
 	HeaderArchive.Serialize(ExportMapArchive.GetWriterData(), ExportMapArchive.Tell());
-	PackageSummary->CellImportMapOffset = HeaderArchive.Tell();
-	PackageSummary->CellExportMapOffset = HeaderArchive.Tell();
+	CellOffsets->CellImportMapOffset = HeaderArchive.Tell();
+	CellOffsets->CellExportMapOffset = HeaderArchive.Tell();
 	// raw arrays of 4-byte aligned items
 	PackageSummary->ExportBundleEntriesOffset = HeaderArchive.Tell();
 	HeaderArchive.Serialize(ExportBundleEntriesArchive.GetWriterData(), ExportBundleEntriesArchive.Tell());
diff --git a/Engine/Source/Editor/UnrealEd/Private/Cooker/CookPackageArtifacts.cpp b/Engine/Source/Editor/UnrealEd/Private/Cooker/CookPackageArtifacts.cpp
--- a/Engine/Source/Editor/UnrealEd/Private/Cooker/CookPackageArtifacts.cpp
+++ b/Engine/Source/Editor/UnrealEd/Private/Cooker/CookPackageArtifacts.cpp
@@ -32,7 +32,7 @@ namespace UE::Cook
 {
 
 // Change TargetDomainVersion to a new guid when all TargetDomain keys in an incremental cook need to be invalidated
-FGuid TargetDomainVersion(0x10E1BAE8, 0x6CA04033, 0xA7E57A37, 0xD967F493);
+FGuid TargetDomainVersion(0xE780E176, 0xE04F4849, 0xA8E8D63D, 0x67B2EFD2);
 
 // Bump PackageArtifactsVersion version when the serialization of PackageArtifacts has changed and we want to add
 // backwards compatibility rather than invalidating everything.
diff --git a/Engine/Source/Runtime/Core/Public/Serialization/BufferWriter.h b/Engine/Source/Runtime/Core/Public/Serialization/BufferWriter.h
--- a/Engine/Source/Runtime/Core/Public/Serialization/BufferWriter.h
+++ b/Engine/Source/Runtime/Core/Public/Serialization/BufferWriter.h
@@ -103,6 +103,15 @@ public:
 	{
 		return WriterPos >= WriterSize;
 	}
+	template <typename Type>
+	Type* SerializeUninitialized()
+	{
+		check(WriterPos >= 0);
+		check((WriterPos + (int64)sizeof(Type)) <= WriterSize);
+		Type* Data = reinterpret_cast<Type*>((uint8*)WriterData + WriterPos);
+		WriterPos += sizeof(Type);
+		return Data;
+	}
 	/**
 	* Returns the name of the Archive.  Useful for getting the name of the package a struct or object
 	* is in when a loading error occurs.
diff --git a/Engine/Source/Runtime/CoreUObject/Private/Serialization/ZenPackageHeader.cpp b/Engine/Source/Runtime/CoreUObject/Private/Serialization/ZenPackageHeader.cpp
--- a/Engine/Source/Runtime/CoreUObject/Private/Serialization/ZenPackageHeader.cpp
+++ b/Engine/Source/Runtime/CoreUObject/Private/Serialization/ZenPackageHeader.cpp
@@ -31,6 +31,19 @@ FZenPackageHeader FZenPackageHeader::MakeView(FMemoryView Memory, FString& OutEr
 		PackageHeader.VersioningInfo.Emplace();
 		PackageHeaderDataReader << PackageHeader.VersioningInfo.GetValue();
 	}
+	const FZenPackageVersioningInfo* LocalVersioningInfo = PackageHeader.VersioningInfo.GetPtrOrNull();
+
+	FZenPackageCellOffsets CellOffsets;
+	if (LocalVersioningInfo == nullptr || LocalVersioningInfo->PackageVersion >= EUnrealEngineObjectUE5Version::VERSE_CELLS)
+	{
+		PackageHeaderDataReader << CellOffsets.CellImportMapOffset;
+		PackageHeaderDataReader << CellOffsets.CellExportMapOffset;
+	}
+	else
+	{
+		CellOffsets.CellImportMapOffset = PackageSummary->ExportBundleEntriesOffset;
+		CellOffsets.CellExportMapOffset = PackageSummary->ExportBundleEntriesOffset;
+	}
 
 	{
 		TRACE_CPUPROFILER_EVENT_SCOPE(LoadPackageNameMap);
@@ -38,7 +51,6 @@ FZenPackageHeader FZenPackageHeader::MakeView(FMemoryView Memory, FString& OutEr
 	}
 	PackageHeader.PackageName = PackageHeader.NameMap.GetName(PackageHeader.PackageSummary->Name);
 	
-	const FZenPackageVersioningInfo* LocalVersioningInfo = PackageHeader.VersioningInfo.GetPtrOrNull();
 	if (LocalVersioningInfo == nullptr || LocalVersioningInfo->PackageVersion >= EUnrealEngineObjectUE5Version::DATA_RESOURCES)
 	{
 		int64 BulkDataMapSize = 0;
@@ -60,7 +72,7 @@ FZenPackageHeader FZenPackageHeader::MakeView(FMemoryView Memory, FString& OutEr
 		(PackageSummary->ExportMapOffset - PackageSummary->ImportMapOffset) / sizeof(FPackageObjectIndex));
 	PackageHeader.ExportMap = TArrayView<const FExportMapEntry>(
 		reinterpret_cast<const FExportMapEntry*>(PackageHeaderDataPtr + PackageSummary->ExportMapOffset),
-		(PackageSummary->CellImportMapOffset - PackageSummary->ExportMapOffset) / sizeof(FExportMapEntry));
+		(CellOffsets.CellImportMapOffset - PackageSummary->ExportMapOffset) / sizeof(FExportMapEntry));
 	PackageHeader.ExportCount = PackageHeader.ExportMap.Num();
 
 	const uint64 ExportBundleEntriesSize = PackageHeader.PackageSummary->DependencyBundleHeadersOffset - PackageHeader.PackageSummary->ExportBundleEntriesOffset;
diff --git a/Engine/Source/Runtime/CoreUObject/Public/Serialization/AsyncLoading2.h b/Engine/Source/Runtime/CoreUObject/Public/Serialization/AsyncLoading2.h
--- a/Engine/Source/Runtime/CoreUObject/Public/Serialization/AsyncLoading2.h
+++ b/Engine/Source/Runtime/CoreUObject/Public/Serialization/AsyncLoading2.h
@@ -284,14 +284,18 @@ struct FZenPackageSummary
 	int32 ImportedPublicExportHashesOffset;
 	int32 ImportMapOffset;
 	int32 ExportMapOffset;
-	int32 CellImportMapOffset;
-	int32 CellExportMapOffset;
 	int32 ExportBundleEntriesOffset;
 	int32 DependencyBundleHeadersOffset;
 	int32 DependencyBundleEntriesOffset;
 	int32 ImportedPackageNamesOffset;
 };
 
+struct FZenPackageCellOffsets
+{
+	int32 CellImportMapOffset;
+	int32 CellExportMapOffset;
+};
+
 /**
  * Export bundle entry.
  */
