[[nodiscard]] for UE smart pointers

#rnx

diff --git a/Engine/Source/Runtime/Core/Public/UObject/StrongObjectPtrTemplates.h b/Engine/Source/Runtime/Core/Public/UObject/StrongObjectPtrTemplates.h
--- a/Engine/Source/Runtime/Core/Public/UObject/StrongObjectPtrTemplates.h
+++ b/Engine/Source/Runtime/Core/Public/UObject/StrongObjectPtrTemplates.h
@@ -27,7 +27,7 @@ class TStrongObjectPtr
 public:
 	using ElementType = ObjectType;
 
-	TStrongObjectPtr(TStrongObjectPtr&& InOther)
+	[[nodiscard]] TStrongObjectPtr(TStrongObjectPtr&& InOther)
 	{
 		Reset();
 		Object = InOther.Object;
@@ -50,18 +50,18 @@ public:
 		Reset();
 	}
 
-	FORCEINLINE_DEBUGGABLE TStrongObjectPtr(TYPE_OF_NULLPTR = nullptr)
+	[[nodiscard]] FORCEINLINE_DEBUGGABLE TStrongObjectPtr(TYPE_OF_NULLPTR = nullptr)
 	{
 		static_assert(TPointerIsConvertibleFromTo<ObjectType, const volatile UObject>::Value, "TStrongObjectPtr can only be constructed with UObject types");
 	}
 
-	FORCEINLINE_DEBUGGABLE explicit TStrongObjectPtr(ObjectType* InObject)
+	[[nodiscard]] FORCEINLINE_DEBUGGABLE explicit TStrongObjectPtr(ObjectType* InObject)
 	{
 		static_assert(TPointerIsConvertibleFromTo<ObjectType, const volatile UObject>::Value, "TStrongObjectPtr can only be constructed with UObject types");
 		Reset(InObject);
 	}
 
-	FORCEINLINE_DEBUGGABLE TStrongObjectPtr(const TStrongObjectPtr& InOther)
+	[[nodiscard]] FORCEINLINE_DEBUGGABLE TStrongObjectPtr(const TStrongObjectPtr& InOther)
 	{
 		Reset(InOther.Get());
 	}
@@ -71,7 +71,7 @@ public:
 		typename OtherReferencerNameProvider
 		UE_REQUIRES(std::is_convertible_v<OtherObjectType*, ObjectType*>)
 	>
-	FORCEINLINE_DEBUGGABLE TStrongObjectPtr(const TStrongObjectPtr<OtherObjectType, OtherReferencerNameProvider>& InOther)
+	[[nodiscard]] FORCEINLINE_DEBUGGABLE TStrongObjectPtr(const TStrongObjectPtr<OtherObjectType, OtherReferencerNameProvider>& InOther)
 	{
 		Reset(InOther.Get());
 	}
@@ -93,29 +93,29 @@ public:
 		return *this;
 	}
 
-	FORCEINLINE_DEBUGGABLE ObjectType& operator*() const
+	[[nodiscard]] FORCEINLINE_DEBUGGABLE ObjectType& operator*() const
 	{
 		check(IsValid());
 		return *(ObjectType*)Get();
 	}
 
-	FORCEINLINE_DEBUGGABLE ObjectType* operator->() const
+	[[nodiscard]] FORCEINLINE_DEBUGGABLE ObjectType* operator->() const
 	{
 		check(IsValid());
 		return (ObjectType*)Get();
 	}
 
-	FORCEINLINE_DEBUGGABLE bool IsValid() const
+	[[nodiscard]] FORCEINLINE_DEBUGGABLE bool IsValid() const
 	{
 		return Object != nullptr;
 	}
 
-	FORCEINLINE_DEBUGGABLE explicit operator bool() const
+	[[nodiscard]] FORCEINLINE_DEBUGGABLE explicit operator bool() const
 	{
 		return IsValid();
 	}
 
-	FORCEINLINE_DEBUGGABLE ObjectType* Get() const
+	[[nodiscard]] FORCEINLINE_DEBUGGABLE ObjectType* Get() const
 	{
 		return (ObjectType*)Object;
 	}
diff --git a/Engine/Source/Runtime/Core/Public/UObject/WeakObjectPtrTemplates.h b/Engine/Source/Runtime/Core/Public/UObject/WeakObjectPtrTemplates.h
--- a/Engine/Source/Runtime/Core/Public/UObject/WeakObjectPtrTemplates.h
+++ b/Engine/Source/Runtime/Core/Public/UObject/WeakObjectPtrTemplates.h
@@ -36,15 +36,15 @@ struct TWeakObjectPtr
 public:
 	using ElementType = T;
 	
-	TWeakObjectPtr() = default;
-	TWeakObjectPtr(const TWeakObjectPtr&) = default;
+	[[nodiscard]] TWeakObjectPtr() = default;
+	[[nodiscard]] TWeakObjectPtr(const TWeakObjectPtr&) = default;
 	TWeakObjectPtr& operator=(const TWeakObjectPtr&) = default;
 	~TWeakObjectPtr() = default;
 
 	/**
 	 * Construct from a null pointer
 	 */
-	FORCEINLINE TWeakObjectPtr(TYPE_OF_NULLPTR) :
+	[[nodiscard]] FORCEINLINE TWeakObjectPtr(TYPE_OF_NULLPTR) :
 		WeakPtr((UObject*)nullptr)
 	{
 	}
@@ -57,7 +57,7 @@ public:
 		typename U
 		UE_REQUIRES(std::is_convertible_v<U, T*>)
 	>
-	FORCEINLINE TWeakObjectPtr(U Object) :
+	[[nodiscard]] FORCEINLINE TWeakObjectPtr(U Object) :
 		WeakPtr((const UObject*)Object)
 	{
 		// This static assert is in here rather than in the body of the class because we want
@@ -73,7 +73,7 @@ public:
 		typename OtherT
 		UE_REQUIRES(std::is_convertible_v<OtherT*, T*>)
 	>
-	FORCEINLINE TWeakObjectPtr(const TWeakObjectPtr<OtherT, TWeakObjectPtrBase>& Other) :
+	[[nodiscard]] FORCEINLINE TWeakObjectPtr(const TWeakObjectPtr<OtherT, TWeakObjectPtrBase>& Other) :
 		WeakPtr(Other.WeakPtr) // we do a C-style cast to private base here to avoid clang 3.6.0 compilation problems with friend declarations
 	{
 	}
@@ -128,7 +128,7 @@ public:
 	 * @param bEvenIfPendingKill if this is true, pendingkill objects are considered valid
 	 * @return nullptr if this object is gone or the weak pointer is explicitly null, otherwise a valid uobject pointer
 	 */
-	FORCEINLINE T* Get(bool bEvenIfPendingKill) const
+	[[nodiscard]] FORCEINLINE T* Get(bool bEvenIfPendingKill) const
 	{
 		return (T*)WeakPtr.Get(bEvenIfPendingKill);
 	}
@@ -136,7 +136,7 @@ public:
 	/**  
 	 * Dereference the weak pointer. This is an optimized version implying bEvenIfPendingKill=false.
 	 */
-	FORCEINLINE T* Get(/*bool bEvenIfPendingKill = false*/) const
+	[[nodiscard]] FORCEINLINE T* Get(/*bool bEvenIfPendingKill = false*/) const
 	{
 		return (T*)WeakPtr.Get();
 	}
@@ -146,7 +146,7 @@ public:
 	 * @param bEvenIfPendingKill if this is true, pendingkill objects are considered valid
 	 * @return nullptr if this object is gone or the weak pointer is explicitly null, otherwise a valid uobject pointer
 	 */
-	FORCEINLINE TStrongObjectPtr<T> Pin(bool bEvenIfPendingKill) const
+	[[nodiscard]] FORCEINLINE TStrongObjectPtr<T> Pin(bool bEvenIfPendingKill) const
 	{
 		TStrongObjectPtr<T> StrongPtr;
 		StrongPtr.Attach((T*)WeakPtr.Pin(bEvenIfPendingKill).Detach());
@@ -156,7 +156,7 @@ public:
 	/**
 	 * Pin the weak pointer as a strong ptr. This is an optimized version implying bEvenIfPendingKill=false.
 	 */
-	FORCEINLINE TStrongObjectPtr<T> Pin(/*bool bEvenIfPendingKill = false*/) const
+	[[nodiscard]] FORCEINLINE TStrongObjectPtr<T> Pin(/*bool bEvenIfPendingKill = false*/) const
 	{
 		TStrongObjectPtr<T> StrongPtr;
 		StrongPtr.Attach((T*)WeakPtr.Pin().Detach());
@@ -169,7 +169,7 @@ public:
 	* @param bEvenIfPendingKill if this is true, pendingkill objects are considered valid
 	* @return nullptr if this object is gone or the weak pointer is explicitly null, otherwise a valid uobject pointer
 	*/
-	FORCEINLINE TStrongObjectPtr<T> TryPin(bool& bOutPinValid, bool bEvenIfPendingKill) const
+	[[nodiscard]] FORCEINLINE TStrongObjectPtr<T> TryPin(bool& bOutPinValid, bool bEvenIfPendingKill) const
 	{
 		TStrongObjectPtr<T> StrongPtr;
 		StrongPtr.Attach((T*)WeakPtr.TryPin(bOutPinValid, bEvenIfPendingKill).Detach());
@@ -179,13 +179,13 @@ public:
 	/**
 	 * Pin the weak pointer as a strong ptr. This is an optimized version implying bEvenIfPendingKill=false.
 	 */
-	FORCEINLINE TStrongObjectPtr<T> TryPin(bool& bOutPinValid /*, bool bEvenIfPendingKill = false */ ) const
+	[[nodiscard]] FORCEINLINE TStrongObjectPtr<T> TryPin(bool& bOutPinValid /*, bool bEvenIfPendingKill = false */ ) const
 	{
 		return TryPin(bOutPinValid, false);
 	}
 
 	/** Deferences the weak pointer even if its marked RF_Unreachable. This is needed to resolve weak pointers during GC (such as ::AddReferenceObjects) */
-	FORCEINLINE T* GetEvenIfUnreachable() const
+	[[nodiscard]] FORCEINLINE T* GetEvenIfUnreachable() const
 	{
 		return (T*)WeakPtr.GetEvenIfUnreachable();
 	}
@@ -193,7 +193,7 @@ public:
 	/**  
 	 * Dereference the weak pointer
 	 */
-	FORCEINLINE T& operator*() const
+	[[nodiscard]] FORCEINLINE T& operator*() const
 	{
 		return *Get();
 	}
@@ -201,7 +201,7 @@ public:
 	/**  
 	 * Dereference the weak pointer
 	 */
-	FORCEINLINE T* operator->() const
+	[[nodiscard]] FORCEINLINE T* operator->() const
 	{
 		return Get();
 	}
@@ -228,7 +228,7 @@ public:
 	 *							due to the presence of the Unreachable flag.
 	 * @return true if Get() would return a valid non-null pointer
 	 */
-	FORCEINLINE bool IsValid(bool bEvenIfPendingKill, bool bThreadsafeTest = false) const
+	[[nodiscard]] FORCEINLINE bool IsValid(bool bEvenIfPendingKill, bool bThreadsafeTest = false) const
 	{
 		return WeakPtr.IsValid(bEvenIfPendingKill, bThreadsafeTest);
 	}
@@ -240,7 +240,7 @@ public:
 	 * due to the Unreachable flag being set. (see bThreadsafeTest above)
 	 * @return true if Get() would return a valid non-null pointer
 	 */
-	FORCEINLINE bool IsValid(/*bool bEvenIfPendingKill = false, bool bThreadsafeTest = false*/) const
+	[[nodiscard]] FORCEINLINE bool IsValid(/*bool bEvenIfPendingKill = false, bool bThreadsafeTest = false*/) const
 	{
 		return WeakPtr.IsValid();
 	}
@@ -251,7 +251,7 @@ public:
 	 * @param bThreadsafeTest set it to true when testing outside of Game Thread. Results in false if WeakObjPtr point to an existing object (no flags checked)
 	 * @return true if this used to point at a real object but no longer does.
 	 */
-	FORCEINLINE bool IsStale(bool bIncludingIfPendingKill = true, bool bThreadsafeTest = false) const
+	[[nodiscard]] FORCEINLINE bool IsStale(bool bIncludingIfPendingKill = true, bool bThreadsafeTest = false) const
 	{
 		return WeakPtr.IsStale(bIncludingIfPendingKill, bThreadsafeTest);
 	}
@@ -260,7 +260,7 @@ public:
 	 * Returns true if this pointer was explicitly assigned to null, was reset, or was never initialized.
 	 * If this returns true, IsValid() and IsStale() will both return false.
 	 */
-	FORCEINLINE bool IsExplicitlyNull() const
+	[[nodiscard]] FORCEINLINE bool IsExplicitlyNull() const
 	{
 		return WeakPtr.IsExplicitlyNull();
 	}
@@ -269,7 +269,7 @@ public:
 	 * Returns true if two weak pointers were originally set to the same object, even if they are now stale
 	 * @param Other weak pointer to compare to
 	 */
-	FORCEINLINE bool HasSameIndexAndSerialNumber(const TWeakObjectPtr& Other) const
+	[[nodiscard]] FORCEINLINE bool HasSameIndexAndSerialNumber(const TWeakObjectPtr& Other) const
 	{
 		return WeakPtr.HasSameIndexAndSerialNumber(Other.WeakPtr);
 	}
@@ -282,7 +282,7 @@ public:
 		typename OtherT
 		UE_REQUIRES(UE_REQUIRES_EXPR((T*)nullptr == (OtherT*)nullptr))
 	>
-	FORCEINLINE bool HasSameIndexAndSerialNumber(const TWeakObjectPtr<OtherT, TWeakObjectPtrBase>& Other) const
+	[[nodiscard]] FORCEINLINE bool HasSameIndexAndSerialNumber(const TWeakObjectPtr<OtherT, TWeakObjectPtrBase>& Other) const
 	{
 		return WeakPtr.HasSameIndexAndSerialNumber(Other.WeakPtr);
 	}
@@ -294,7 +294,7 @@ public:
 	}
 #endif // UE_WITH_REMOTE_OBJECT_HANDLE
 
-	FORCEINLINE bool IsRemote() const
+	[[nodiscard]] FORCEINLINE bool IsRemote() const
 	{
 		return WeakPtr.IsRemote();
 	}
@@ -308,7 +308,7 @@ public:
 	}
 
 	/** Hash function. */
-	FORCEINLINE uint32 GetWeakPtrTypeHash() const
+	[[nodiscard]] FORCEINLINE uint32 GetWeakPtrTypeHash() const
 	{
 		return WeakPtr.GetTypeHash();
 	}
@@ -319,19 +319,19 @@ public:
 	 * @param Other weak pointer to compare to
 	 */
 	template <typename RhsT, typename = decltype((T*)nullptr == (RhsT*)nullptr)>
-	FORCENOINLINE bool operator==(const TWeakObjectPtr<RhsT, TWeakObjectPtrBase>& Rhs) const
+	[[nodiscard]] FORCENOINLINE bool operator==(const TWeakObjectPtr<RhsT, TWeakObjectPtrBase>& Rhs) const
 	{
 		return WeakPtr == Rhs.WeakPtr;
 	}
 
 	template <typename RhsT, typename = decltype((T*)nullptr == (RhsT*)nullptr)>
-	FORCENOINLINE bool operator==(const RhsT* Rhs) const
+	[[nodiscard]] FORCENOINLINE bool operator==(const RhsT* Rhs) const
 	{
 		// NOTE: this constructs a TWeakObjectPtrBase, which has some amount of overhead, so this may not be an efficient operation
 		return WeakPtr == TWeakObjectPtrBase(Rhs);
 	}
 
-	FORCENOINLINE bool operator==(TYPE_OF_NULLPTR) const
+	[[nodiscard]] FORCENOINLINE bool operator==(TYPE_OF_NULLPTR) const
 	{
 		return !IsValid();
 	}
@@ -342,18 +342,18 @@ public:
 	 * @param Other weak pointer to compare to
 	 */
 	template <typename RhsT, typename = decltype((T*)nullptr != (RhsT*)nullptr)>
-	FORCEINLINE bool operator!=(const TWeakObjectPtr<RhsT, TWeakObjectPtrBase>& Rhs) const
+	[[nodiscard]] FORCEINLINE bool operator!=(const TWeakObjectPtr<RhsT, TWeakObjectPtrBase>& Rhs) const
 	{
 		return !(*this == Rhs);
 	}
 
 	template <typename RhsT, typename = decltype((T*)nullptr != (RhsT*)nullptr)>
-	FORCEINLINE bool operator!=(const RhsT* Rhs) const
+	[[nodiscard]] FORCEINLINE bool operator!=(const RhsT* Rhs) const
 	{
 		return !(*this == Rhs);
 	}
 
-	FORCEINLINE bool operator!=(TYPE_OF_NULLPTR) const
+	[[nodiscard]] FORCEINLINE bool operator!=(TYPE_OF_NULLPTR) const
 	{
 		return !(*this == nullptr);
 	}
@@ -381,12 +381,12 @@ TWeakObjectPtr(const TWeakObjectPtr<T>&) -> TWeakObjectPtr<T>;
  * @param Other weak pointer to compare to
  */
 template <typename LhsT, typename RhsT, typename OtherTWeakObjectPtrBase, typename = decltype((LhsT*)nullptr == (RhsT*)nullptr)>
-FORCEINLINE bool operator==(const LhsT* Lhs, const TWeakObjectPtr<RhsT, OtherTWeakObjectPtrBase>& Rhs)
+[[nodiscard]] FORCEINLINE bool operator==(const LhsT* Lhs, const TWeakObjectPtr<RhsT, OtherTWeakObjectPtrBase>& Rhs)
 {
 	return Rhs == Lhs;
 }
 template <typename RhsT, typename OtherTWeakObjectPtrBase>
-FORCEINLINE bool operator==(TYPE_OF_NULLPTR, const TWeakObjectPtr<RhsT, OtherTWeakObjectPtrBase>& Rhs)
+[[nodiscard]] FORCEINLINE bool operator==(TYPE_OF_NULLPTR, const TWeakObjectPtr<RhsT, OtherTWeakObjectPtrBase>& Rhs)
 {
 	return Rhs == nullptr;
 }
@@ -396,12 +396,12 @@ FORCEINLINE bool operator==(TYPE_OF_NULLPTR, const TWeakObjectPtr<RhsT, OtherTWe
  * @param Other weak pointer to compare to
  */
 template <typename LhsT, typename RhsT, typename OtherTWeakObjectPtrBase, typename = decltype((LhsT*)nullptr != (RhsT*)nullptr)>
-FORCEINLINE bool operator!=(const LhsT* Lhs, const TWeakObjectPtr<RhsT, OtherTWeakObjectPtrBase>& Rhs)
+[[nodiscard]] FORCEINLINE bool operator!=(const LhsT* Lhs, const TWeakObjectPtr<RhsT, OtherTWeakObjectPtrBase>& Rhs)
 {
 	return !(Rhs == Lhs);
 }
 template <typename RhsT, typename OtherTWeakObjectPtrBase>
-FORCEINLINE bool operator!=(TYPE_OF_NULLPTR, const TWeakObjectPtr<RhsT, OtherTWeakObjectPtrBase>& Rhs)
+[[nodiscard]] FORCEINLINE bool operator!=(TYPE_OF_NULLPTR, const TWeakObjectPtr<RhsT, OtherTWeakObjectPtrBase>& Rhs)
 {
 	return !(Rhs == nullptr);
 }
@@ -409,7 +409,7 @@ FORCEINLINE bool operator!=(TYPE_OF_NULLPTR, const TWeakObjectPtr<RhsT, OtherTWe
 
 // Helper function which deduces the type of the initializer
 template <typename T>
-FORCEINLINE TWeakObjectPtr<T> MakeWeakObjectPtr(T* Ptr)
+[[nodiscard]] FORCEINLINE TWeakObjectPtr<T> MakeWeakObjectPtr(T* Ptr)
 {
 	return TWeakObjectPtr<T>(Ptr);
 }
@@ -491,7 +491,7 @@ void CopyToWeakArray(DestArrayType& Dest, const SourceArrayType& Src)
 
 /** Hash function. */
 template <typename T>
-FORCEINLINE uint32 GetTypeHash(const TWeakObjectPtr<T>& WeakObjectPtr)
+[[nodiscard]] FORCEINLINE uint32 GetTypeHash(const TWeakObjectPtr<T>& WeakObjectPtr)
 {
 	return WeakObjectPtr.GetWeakPtrTypeHash();
 }
diff --git a/Engine/Source/Runtime/CoreUObject/Public/Templates/NonNullSubclassOf.h b/Engine/Source/Runtime/CoreUObject/Public/Templates/NonNullSubclassOf.h
--- a/Engine/Source/Runtime/CoreUObject/Public/Templates/NonNullSubclassOf.h
+++ b/Engine/Source/Runtime/CoreUObject/Public/Templates/NonNullSubclassOf.h
@@ -21,7 +21,7 @@ class TNonNullSubclassOf : public TSubclassOf<T>
 
 public:
 	/** Default Constructor, defaults to null */
-	FORCEINLINE TNonNullSubclassOf(EDefaultConstructNonNullSubclassOf)
+	[[nodiscard]] FORCEINLINE TNonNullSubclassOf(EDefaultConstructNonNullSubclassOf)
 		: Super(nullptr)
 	{
 	}
@@ -30,7 +30,7 @@ public:
 	template <
 		typename PtrType
 		UE_REQUIRES(std::is_same_v<PtrType, UClass> || std::is_same_v<PtrType, FFieldClass>)>
-	FORCEINLINE TNonNullSubclassOf(PtrType* From)
+	[[nodiscard]] FORCEINLINE TNonNullSubclassOf(PtrType* From)
 		: Super(From)
 	{
 		checkf(From, TEXT("Initializing TNonNullSubclassOf with null"));
@@ -41,7 +41,7 @@ public:
 		typename U
 		UE_REQUIRES(std::is_convertible_v<U*, T*>)
 	>
-	FORCEINLINE TNonNullSubclassOf(const TSubclassOf<U>& From)
+	[[nodiscard]] FORCEINLINE TNonNullSubclassOf(const TSubclassOf<U>& From)
 		: Super(From)
 	{
 	}
@@ -73,11 +73,12 @@ public:
 	///////////////////////////////////////////////////////////
 	constexpr static bool bHasIntrusiveUnsetOptionalState = true;
 	using IntrusiveUnsetOptionalStateType = TNonNullSubclassOf;
-	FORCEINLINE explicit TNonNullSubclassOf(FIntrusiveUnsetOptionalState)
+	[[nodiscard]] FORCEINLINE explicit TNonNullSubclassOf(FIntrusiveUnsetOptionalState)
 		: Super(nullptr)
 	{
 	}
-	FORCEINLINE bool operator==(FIntrusiveUnsetOptionalState) const
+	
+	[[nodiscard]] FORCEINLINE bool operator==(FIntrusiveUnsetOptionalState) const
 	{
 		return Super::Get() == nullptr;
 	}
diff --git a/Engine/Source/Runtime/CoreUObject/Public/Templates/SubclassOf.h b/Engine/Source/Runtime/CoreUObject/Public/Templates/SubclassOf.h
--- a/Engine/Source/Runtime/CoreUObject/Public/Templates/SubclassOf.h
+++ b/Engine/Source/Runtime/CoreUObject/Public/Templates/SubclassOf.h
@@ -35,15 +35,15 @@ private:
 public:
 	using ElementType = T;
 
-	TSubclassOf() = default;
-	TSubclassOf(TSubclassOf&&) = default;
-	TSubclassOf(const TSubclassOf&) = default;
+	[[nodiscard]] TSubclassOf() = default;
+	[[nodiscard]] TSubclassOf(TSubclassOf&&) = default;
+	[[nodiscard]] TSubclassOf(const TSubclassOf&) = default;
 	TSubclassOf& operator=(TSubclassOf&&) = default;
 	TSubclassOf& operator=(const TSubclassOf&) = default;
 	~TSubclassOf() = default;
 
 	/** Constructor that takes a UClass*. */
-	FORCEINLINE TSubclassOf(UClass* From)
+	[[nodiscard]] FORCEINLINE TSubclassOf(UClass* From)
 		: Class(From)
 	{
 	}
@@ -56,7 +56,7 @@ public:
 			std::is_convertible_v<U, UClass*>
 		)
 	>
-	FORCEINLINE TSubclassOf(U&& From)
+	[[nodiscard]] FORCEINLINE TSubclassOf(U&& From)
 		: Class(From)
 	{
 	}
@@ -66,7 +66,7 @@ public:
 		typename OtherT
 		UE_REQUIRES(std::is_convertible_v<OtherT*, T*>)
 	>
-	FORCEINLINE TSubclassOf(const TSubclassOf<OtherT>& Other)
+	[[nodiscard]] FORCEINLINE TSubclassOf(const TSubclassOf<OtherT>& Other)
 		: Class(Other.Class)
 	{
 		IWYU_MARKUP_IMPLICIT_CAST(OtherT, T);
@@ -106,7 +106,7 @@ public:
 	}
 
 	/** Dereference back into a UClass*, does runtime type checking. */
-	FORCEINLINE UClass* operator*() const
+	[[nodiscard]] FORCEINLINE UClass* operator*() const
 	{
 		if (!Class || !Class->IsChildOf(T::StaticClass()))
 		{
@@ -116,19 +116,19 @@ public:
 	}
 
 	/** Dereference back into a UClass*, does runtime type checking. */
-	FORCEINLINE UClass* Get() const
+	[[nodiscard]] FORCEINLINE UClass* Get() const
 	{
 		return **this;
 	}
 
 	/** Dereference back into a UClass*, does runtime type checking. */
-	FORCEINLINE UClass* operator->() const
+	[[nodiscard]] FORCEINLINE UClass* operator->() const
 	{
 		return **this;
 	}
 
 	/** Implicit conversion to UClass*, does runtime type checking. */
-	FORCEINLINE operator UClass*() const
+	[[nodiscard]] FORCEINLINE operator UClass*() const
 	{
 		return **this;
 	}
@@ -138,7 +138,7 @@ public:
 	 *
 	 * @return the CDO, or null if class is null
 	 */
-	FORCEINLINE T* GetDefaultObject() const
+	[[nodiscard]] FORCEINLINE T* GetDefaultObject() const
 	{
 		UObject* Result = nullptr;
 		if (Class)
@@ -164,7 +164,7 @@ public:
 		return GetTypeHash(SubclassOf.Class);
 	}
 
-	TObjectPtr<UClass>& GetGCPtr()
+	[[nodiscard]] TObjectPtr<UClass>& GetGCPtr()
 	{
 		return Class;
 	}
@@ -172,7 +172,7 @@ public:
 #if DO_CHECK
 	// This is a DEVELOPMENT ONLY debugging function and should not be relied upon. Client
 	// systems should never require unsafe access to the referenced UClass
-	UClass* DebugAccessRawClassPtr() const
+	[[nodiscard]] UClass* DebugAccessRawClassPtr() const
 	{
 		return Class;
 	}
diff --git a/Engine/Source/Runtime/CoreUObject/Public/Templates/SubclassOfField.h b/Engine/Source/Runtime/CoreUObject/Public/Templates/SubclassOfField.h
--- a/Engine/Source/Runtime/CoreUObject/Public/Templates/SubclassOfField.h
+++ b/Engine/Source/Runtime/CoreUObject/Public/Templates/SubclassOfField.h
@@ -32,15 +32,15 @@ private:
 	friend class TSubclassOfField;
 
 public:
-	TSubclassOfField() = default;
-	TSubclassOfField(TSubclassOfField&&) = default;
-	TSubclassOfField(const TSubclassOfField&) = default;
+	[[nodiscard]] TSubclassOfField() = default;
+	[[nodiscard]] TSubclassOfField(TSubclassOfField&&) = default;
+	[[nodiscard]] TSubclassOfField(const TSubclassOfField&) = default;
 	TSubclassOfField& operator=(TSubclassOfField&&) = default;
 	TSubclassOfField& operator=(const TSubclassOfField&) = default;
 	~TSubclassOfField() = default;
 
 	/** Constructor that takes a FFieldClass*. */
-	FORCEINLINE TSubclassOfField(FFieldClass* From)
+	[[nodiscard]] FORCEINLINE TSubclassOfField(FFieldClass* From)
 		: Class(From)
 	{
 	}
@@ -53,7 +53,7 @@ public:
 			std::is_convertible_v<U, FFieldClass*>
 		)
 	>
-	FORCEINLINE TSubclassOfField(U&& From)
+	[[nodiscard]] FORCEINLINE TSubclassOfField(U&& From)
 		: Class(From)
 	{
 	}
@@ -63,7 +63,7 @@ public:
 		typename OtherT
 		UE_REQUIRES(std::is_convertible_v<OtherT*, T*>)
 	>
-	FORCEINLINE TSubclassOfField(const TSubclassOfField<OtherT>& Other)
+	[[nodiscard]] FORCEINLINE TSubclassOfField(const TSubclassOfField<OtherT>& Other)
 		: Class(Other.Class)
 	{
 	}
@@ -101,7 +101,7 @@ public:
 	}
 
 	/** Dereference back into a FFieldClass*, does runtime type checking. */
-	FORCEINLINE FFieldClass* operator*() const
+	[[nodiscard]] FORCEINLINE FFieldClass* operator*() const
 	{
 		if (!Class || !Class->IsChildOf(T::StaticClass()))
 		{
@@ -111,19 +111,19 @@ public:
 	}
 
 	/** Dereference back into a FFieldClass*, does runtime type checking. */
-	FORCEINLINE FFieldClass* Get() const
+	[[nodiscard]] FORCEINLINE FFieldClass* Get() const
 	{
 		return **this;
 	}
 
 	/** Dereference back into a FFieldClass*, does runtime type checking. */
-	FORCEINLINE FFieldClass* operator->() const
+	[[nodiscard]] FORCEINLINE FFieldClass* operator->() const
 	{
 		return **this;
 	}
 
 	/** Implicit conversion to FFieldClass*, does runtime type checking. */
-	FORCEINLINE operator FFieldClass*() const
+	[[nodiscard]] FORCEINLINE operator FFieldClass*() const
 	{
 		return **this;
 	}
@@ -133,7 +133,7 @@ public:
 	 *
 	 * @return the CDO, or null if class is null
 	 */
-	FORCEINLINE T* GetDefaultObject() const
+	[[nodiscard]] FORCEINLINE T* GetDefaultObject() const
 	{
 		FField* Result = nullptr;
 		if (Class)
@@ -149,7 +149,7 @@ public:
 		Ar << Class;
 	}
 
-	friend uint32 GetTypeHash(const TSubclassOfField& SubclassOf)
+	[[nodiscard]] friend uint32 GetTypeHash(const TSubclassOfField& SubclassOf)
 	{
 		return GetTypeHash(SubclassOf.Class);
 	}
diff --git a/Engine/Source/Runtime/CoreUObject/Public/UObject/ObjectPtr.h b/Engine/Source/Runtime/CoreUObject/Public/UObject/ObjectPtr.h
--- a/Engine/Source/Runtime/CoreUObject/Public/UObject/ObjectPtr.h
+++ b/Engine/Source/Runtime/CoreUObject/Public/UObject/ObjectPtr.h
@@ -62,21 +62,21 @@ struct TObjectPtr;
 struct FObjectPtr
 {
 public:
-	FObjectPtr()
+	[[nodiscard]] FObjectPtr()
 		: Handle(UE::CoreUObject::Private::MakeObjectHandle(nullptr))
 	{
 	}
 
-	explicit FObjectPtr(ENoInit)
+	[[nodiscard]] explicit FObjectPtr(ENoInit)
 	{
 	}
 
-	FORCEINLINE FObjectPtr(TYPE_OF_NULLPTR)
+	[[nodiscard]] FORCEINLINE FObjectPtr(TYPE_OF_NULLPTR)
 		: Handle(UE::CoreUObject::Private::MakeObjectHandle(nullptr))
 	{
 	}
 
-	explicit FORCEINLINE FObjectPtr(UObject* Object)
+	[[nodiscard]] explicit FORCEINLINE FObjectPtr(UObject* Object)
 		: Handle(UE::CoreUObject::Private::MakeObjectHandle(Object))
 	{
 #if UE_OBJECT_PTR_GC_BARRIER
@@ -85,13 +85,13 @@ public:
 	}
 
 	UE_OBJPTR_DEPRECATED(5.0, "Construction with incomplete type pointer is deprecated.  Please update this code to use MakeObjectPtrUnsafe.")
-	explicit FORCEINLINE FObjectPtr(void* IncompleteObject)
+	[[nodiscard]] explicit FORCEINLINE FObjectPtr(void* IncompleteObject)
 		: Handle(UE::CoreUObject::Private::MakeObjectHandle(reinterpret_cast<UObject*>(IncompleteObject)))
 	{
 	}
 
 #if UE_WITH_OBJECT_HANDLE_LATE_RESOLVE || UE_WITH_REMOTE_OBJECT_HANDLE
-	explicit FORCEINLINE FObjectPtr(FObjectHandle Handle)
+	[[nodiscard]] explicit FORCEINLINE FObjectPtr(FObjectHandle Handle)
 		: Handle(Handle)
 	{
 #if UE_OBJECT_PTR_GC_BARRIER
@@ -100,7 +100,7 @@ public:
 	}
 #endif
 	
-	explicit FORCEINLINE FObjectPtr(FRemoteObjectId ObjectId)
+	[[nodiscard]] explicit FORCEINLINE FObjectPtr(FRemoteObjectId ObjectId)
 #if UE_WITH_REMOTE_OBJECT_HANDLE
 		: FObjectPtr(UE::CoreUObject::Private::FRemoteObjectHandlePrivate::FromIdNoResolve(ObjectId))
 #else
@@ -109,25 +109,25 @@ public:
 	{
 	}
 
-	FORCEINLINE UObject* Get() const
+	[[nodiscard]] FORCEINLINE UObject* Get() const
 	{
 		return UE::CoreUObject::Private::ResolveObjectHandle(Handle);
 	}
 
-	FORCEINLINE UClass* GetClass() const
+	[[nodiscard]] FORCEINLINE UClass* GetClass() const
 	{
 		return UE::CoreUObject::Private::ResolveObjectHandleClass(Handle);
 	}
 
 
 #if UE_OBJECT_PTR_GC_BARRIER
-	FObjectPtr(FObjectPtr&& InOther)
+	[[nodiscard]] FObjectPtr(FObjectPtr&& InOther)
 		: Handle(MoveTemp(InOther.Handle))
 	{
 		ConditionallyMarkAsReachable(*this);
 	}
 
-	FObjectPtr(const FObjectPtr& InOther)
+	[[nodiscard]] FObjectPtr(const FObjectPtr& InOther)
 		: Handle(InOther.Handle)
 	{
 		ConditionallyMarkAsReachable(*this);
@@ -147,8 +147,8 @@ public:
 		return *this;
 	}
 #else
-	FObjectPtr(FObjectPtr&&) = default;
-	FObjectPtr(const FObjectPtr&) = default;
+	[[nodiscard]] FObjectPtr(FObjectPtr&&) = default;
+	[[nodiscard]] FObjectPtr(const FObjectPtr&) = default;
 	FObjectPtr& operator=(FObjectPtr&&) = default;
 	FObjectPtr& operator=(const FObjectPtr&) = default;
 #endif // UE_OBJECT_PTR_GC_BARRIER
@@ -175,93 +175,87 @@ public:
 		return *this;
 	}
 
-	FORCEINLINE bool operator==(FObjectPtr Other) const { return (Handle == Other.Handle); }
-	FORCEINLINE bool operator!=(FObjectPtr Other) const { return (Handle != Other.Handle); }
+	[[nodiscard]] FORCEINLINE bool operator==(FObjectPtr Other) const { return (Handle == Other.Handle); }
+	[[nodiscard]] FORCEINLINE bool operator!=(FObjectPtr Other) const { return (Handle != Other.Handle); }
 
 	UE_OBJPTR_DEPRECATED(5.0, "Use of ToTObjectPtr is unsafe and is deprecated.")
-	FORCEINLINE TObjectPtr<UObject>& ToTObjectPtr();
+	[[nodiscard]] FORCEINLINE TObjectPtr<UObject>& ToTObjectPtr();
 	UE_OBJPTR_DEPRECATED(5.0, "Use of ToTObjectPtr is unsafe and is deprecated.")
-	FORCEINLINE const TObjectPtr<UObject>& ToTObjectPtr() const;
+	[[nodiscard]] FORCEINLINE const TObjectPtr<UObject>& ToTObjectPtr() const;
 
-	FORCEINLINE UObject* operator->() const { return Get(); }
-	FORCEINLINE UObject& operator*() const { return *Get(); }
+	[[nodiscard]] FORCEINLINE UObject* operator->() const { return Get(); }
+	[[nodiscard]] FORCEINLINE UObject& operator*() const { return *Get(); }
 
-	UE_DEPRECATED(5.1, "IsNull is deprecated, please use operator bool instead.")
-	FORCEINLINE bool IsNull() const { return UE::CoreUObject::Private::ResolveObjectHandleNoRead(Handle) == nullptr; }
-	
-	UE_DEPRECATED(5.1, "IsNullNoResolve is deprecated, please use operator bool instead.")
-	FORCEINLINE bool IsNullNoResolve() const { return IsObjectHandleNull(Handle); }
-	
-	FORCEINLINE bool operator!() const { return IsObjectHandleNull(Handle); }
-	explicit FORCEINLINE operator bool() const { return !IsObjectHandleNull(Handle); }
+	[[nodiscard]] FORCEINLINE bool operator!() const { return IsObjectHandleNull(Handle); }
+	[[nodiscard]] explicit FORCEINLINE operator bool() const { return !IsObjectHandleNull(Handle); }
 
-	FORCEINLINE bool IsResolved() const { return IsObjectHandleResolved(Handle); }
+	[[nodiscard]] FORCEINLINE bool IsResolved() const { return IsObjectHandleResolved(Handle); }
 
 	/**
 	 * FObjectPtr::IsRemote is only used when UE_WITH_REMOTE_OBJECT_HANDLE is true, and is mutually exclusive with other uses of unresolved FObjectPtr.
 	 * When using UE_WITH_REMOTE_OBJECT_HANDLE, unresolved FObjectPtrs mean the object is not locally available and is present on another server, see @RemoteObject.h
 	*/
-	FORCEINLINE bool IsRemote() const { return !IsResolved(); }
+	[[nodiscard]] FORCEINLINE bool IsRemote() const { return !IsResolved(); }
 
 	// Gets the PathName of the object without resolving the object reference.
 	// @TODO OBJPTR: Deprecate this.
-	FString GetPath() const { return GetPathName(); }
+	[[nodiscard]] FString GetPath() const { return GetPathName(); }
 
 #if UE_WITH_OBJECT_HANDLE_LATE_RESOLVE
 	// Gets the PathName of the object without resolving the object reference.
-	COREUOBJECT_API FString GetPathName() const;
+	[[nodiscard]] COREUOBJECT_API FString GetPathName() const;
 
 	// Gets the FName of the object without resolving the object reference.
-	COREUOBJECT_API FName GetFName() const;
+	[[nodiscard]] COREUOBJECT_API FName GetFName() const;
 
 	// Gets the string name of the object without resolving the object reference.
-	FORCEINLINE FString GetName() const
+	[[nodiscard]] FORCEINLINE FString GetName() const
 	{
 		return GetFName().ToString();
 	}
 
 	// Gets the Outer UObject for this Object without resolving the object reference.
-	COREUOBJECT_API FObjectPtr GetOuter() const;
+	[[nodiscard]] COREUOBJECT_API FObjectPtr GetOuter() const;
 
 	// Gets the package for this Object without resolving the object reference.
-	COREUOBJECT_API FObjectPtr GetPackage() const;
+	[[nodiscard]] COREUOBJECT_API FObjectPtr GetPackage() const;
 
 	// Reports whether this Object is within an Outer or Package. In some cases may need to resolve the object reference.
-	COREUOBJECT_API bool IsIn(FObjectPtr SomeOuter) const;
+	[[nodiscard]] COREUOBJECT_API bool IsIn(FObjectPtr SomeOuter) const;
 
 	/** Returns the full name for the object in the form: Class ObjectPath */
-	COREUOBJECT_API FString GetFullName(EObjectFullNameFlags Flags = EObjectFullNameFlags::None) const;
+	[[nodiscard]] COREUOBJECT_API FString GetFullName(EObjectFullNameFlags Flags = EObjectFullNameFlags::None) const;
 #else
-	FString GetPathName() const
+	[[nodiscard]] FString GetPathName() const
 	{
 		const UObject* ResolvedObject = Get();
 		return ResolvedObject ? UE::CoreUObject::Private::GetPathName(ResolvedObject) : TEXT("None");
 	}
 
-	FName GetFName() const
+	[[nodiscard]] FName GetFName() const
 	{
 		const UObject* ResolvedObject = Get();
 		return ResolvedObject ? UE::CoreUObject::Private::GetFName(ResolvedObject) : NAME_None;
 	}
 
-	FString GetName() const
+	[[nodiscard]] FString GetName() const
 	{
 		return GetFName().ToString();
 	}
 
-	FObjectPtr GetOuter() const
+	[[nodiscard]] FObjectPtr GetOuter() const
 	{
 		const UObject* ResolvedObject = Get();
 		return ResolvedObject ? FObjectPtr(UE::CoreUObject::Private::GetOuter(ResolvedObject)) : FObjectPtr(nullptr);
 	}
 
-	FObjectPtr GetPackage() const
+	[[nodiscard]] FObjectPtr GetPackage() const
 	{
 		const UObject* ResolvedObject = Get();
 		return ResolvedObject ? FObjectPtr(UE::CoreUObject::Private::GetPackage(ResolvedObject)) : FObjectPtr(nullptr);
 	}
 
-	bool IsIn(FObjectPtr SomeOuter) const
+	[[nodiscard]] bool IsIn(FObjectPtr SomeOuter) const
 	{
 		const UObject* ResolvedObject = Get();
 		const UObject* ResolvedSomeOuter = SomeOuter.Get();
@@ -273,25 +267,25 @@ public:
 	 * Returns the fully qualified pathname for this object as well as the name of the class, in the format:
 	 * 'ClassName Outermost.[Outer:]Name'.
 	 */
-	FString GetFullName(EObjectFullNameFlags Flags = EObjectFullNameFlags::None) const
+	[[nodiscard]] FString GetFullName(EObjectFullNameFlags Flags = EObjectFullNameFlags::None) const
 	{
 		// UObjectBaseUtility::GetFullName is safe to call on null objects.
 		return UE::CoreUObject::Private::GetFullName(Get(), nullptr, Flags);
 	}
 #endif
 
-	FORCEINLINE FObjectHandle GetHandle() const { return Handle; }
-	FORCEINLINE FObjectHandle& GetHandleRef() const { return Handle; }
+	[[nodiscard]] FORCEINLINE FObjectHandle GetHandle() const { return Handle; }
+	[[nodiscard]] FORCEINLINE FObjectHandle& GetHandleRef() const { return Handle; }
 
-	COREUOBJECT_API bool IsA(const UClass* SomeBase) const;
+	[[nodiscard]] COREUOBJECT_API bool IsA(const UClass* SomeBase) const;
 
 	template <typename T>
-	FORCEINLINE bool IsA() const
+	[[nodiscard]] FORCEINLINE bool IsA() const
 	{
 		return IsA(T::StaticClass());
 	}
 
-	FORCEINLINE FRemoteObjectId GetRemoteId() const
+	[[nodiscard]] FORCEINLINE FRemoteObjectId GetRemoteId() const
 	{
 #if UE_WITH_REMOTE_OBJECT_HANDLE
 		return GetHandle().GetRemoteId();
@@ -301,7 +295,7 @@ public:
 	}
 
 private:
-	friend FORCEINLINE uint32 GetTypeHash(const FObjectPtr& Object)
+	[[nodiscard]] friend FORCEINLINE uint32 GetTypeHash(const FObjectPtr& Object)
 	{
 		return GetTypeHash(Object.Handle);
 	}
@@ -519,36 +513,36 @@ struct TObjectPtr
 public:
 	using ElementType = T;
 
-	TObjectPtr()
+	[[nodiscard]] TObjectPtr()
 		: ObjectPtr()
 	{
 	}
 
 #if UE_OBJECT_PTR_GC_BARRIER
-	TObjectPtr(TObjectPtr<T>&& Other)
+	[[nodiscard]] TObjectPtr(TObjectPtr<T>&& Other)
 		: ObjectPtr(MoveTemp(Other.ObjectPtr))
 	{
 	}
-	TObjectPtr(const TObjectPtr<T>& Other)
+	[[nodiscard]] TObjectPtr(const TObjectPtr<T>& Other)
 		: ObjectPtr(Other.ObjectPtr)
 	{
 	}
 #else
-	TObjectPtr(TObjectPtr<T>&& Other) = default;
+	[[nodiscard]] TObjectPtr(TObjectPtr<T>&& Other) = default;
 	TObjectPtr(const TObjectPtr<T>& Other) = default;
 #endif // UE_OBJECT_PTR_GC_BARRIER
 
-	explicit FORCEINLINE TObjectPtr(ENoInit)
+	[[nodiscard]] explicit FORCEINLINE TObjectPtr(ENoInit)
 		: ObjectPtr(NoInit)
 	{
 	}
 
-	FORCEINLINE TObjectPtr(TYPE_OF_NULLPTR)
+	[[nodiscard]] FORCEINLINE TObjectPtr(TYPE_OF_NULLPTR)
 		: ObjectPtr(nullptr)
 	{
 	}
 
-	explicit FORCEINLINE TObjectPtr(FObjectPtr ObjPtr)
+	[[nodiscard]] explicit FORCEINLINE TObjectPtr(FObjectPtr ObjPtr)
 		: ObjectPtr(ObjPtr)
 	{
 	}
@@ -557,7 +551,7 @@ public:
 		typename U
 		UE_REQUIRES(std::is_convertible_v<U*, T*>)
 	>
-	FORCEINLINE TObjectPtr(const TObjectPtr<U>& Other)
+	[[nodiscard]] FORCEINLINE TObjectPtr(const TObjectPtr<U>& Other)
 		: ObjectPtr(Other.ObjectPtr)
 	{
 		IWYU_MARKUP_IMPLICIT_CAST(U, T);
@@ -567,18 +561,18 @@ public:
 		typename U
 		UE_REQUIRES(!TIsTObjectPtr_V<std::decay_t<U>> && std::is_convertible_v<U, T*>)
 	>
-	FORCEINLINE TObjectPtr(const U& Object)
+	[[nodiscard]] FORCEINLINE TObjectPtr(const U& Object)
 		: ObjectPtr(const_cast<std::remove_const_t<T>*>(ImplicitConv<T*>(Object)))
 	{
 	}
 
 	UE_DEPRECATED(5.6, "Constructing a TObjectPtr from a reference is deprecated - use a pointer instead.")
-	FORCEINLINE TObjectPtr(T& Object)
+	[[nodiscard]] FORCEINLINE TObjectPtr(T& Object)
 		: ObjectPtr(const_cast<std::remove_const_t<T>*>(&Object))
 	{
 	}
 
-	explicit FORCEINLINE TObjectPtr(TPrivateObjectPtr<T>&& PrivatePtr)
+	[[nodiscard]] explicit FORCEINLINE TObjectPtr(TPrivateObjectPtr<T>&& PrivatePtr)
 		: ObjectPtr(const_cast<UObject*>(PrivatePtr.Pointer))
 	{
 	}
@@ -637,7 +631,7 @@ public:
 		typename U,
 		typename Base = std::common_type_t<T*, U*>
 	>
-	FORCEINLINE bool operator==(const TObjectPtr<U>& Other) const
+	[[nodiscard]] FORCEINLINE bool operator==(const TObjectPtr<U>& Other) const
 	{
 #if UE_WITH_OBJECT_HANDLE_TYPE_SAFETY
 		// Do a NULL test first before comparing the underlying handles, in case either side is
@@ -655,7 +649,7 @@ public:
 	}
 
 	// Equality/Inequality comparisons against nullptr
-	FORCEINLINE bool operator==(TYPE_OF_NULLPTR) const
+	[[nodiscard]] FORCEINLINE bool operator==(TYPE_OF_NULLPTR) const
 	{
 		return ObjectPtr_Private::IsObjectPtrNull<T>(ObjectPtr);
 	}
@@ -665,14 +659,14 @@ public:
 		typename U,
 		typename = decltype(ObjectPtr_Private::IsObjectPtrEqual(std::declval<const TObjectPtr<T>&>(), std::declval<const U&>()))
 	>
-	FORCEINLINE bool operator==(const U& Other) const
+	[[nodiscard]] FORCEINLINE bool operator==(const U& Other) const
 	{
 		return ObjectPtr_Private::IsObjectPtrEqual(*this, Other);
 	}
 
 #if __cplusplus < 202002L
 	template <typename U>
-	FORCEINLINE auto operator!=(const U& Other) const -> decltype(!(*this == Other))
+	[[nodiscard]] FORCEINLINE auto operator!=(const U& Other) const -> decltype(!(*this == Other))
 	{
 		return !(*this == Other);
 	}
@@ -684,62 +678,111 @@ public:
 	//			a reference to the wrong type of object which we'll just send back static_casted as the wrong type.  Doing
 	//			a check or checkSlow here could catch this, but it would be better if the check could happen elsewhere that
 	//			isn't called as frequently.
-	FORCEINLINE T* Get() const { return ObjectPtr_Private::Get<T>(ObjectPtr); }
-	FORCEINLINE UClass* GetClass() const { return ObjectPtr.GetClass(); }
-	FORCEINLINE TObjectPtr<UObject> GetOuter() const
+	[[nodiscard]] FORCEINLINE T* Get() const { return ObjectPtr_Private::Get<T>(ObjectPtr); }
+	[[nodiscard]] FORCEINLINE UClass* GetClass() const { return ObjectPtr.GetClass(); }
+	[[nodiscard]] FORCEINLINE TObjectPtr<UObject> GetOuter() const
 	{ 
 		FObjectPtr Ptr = ObjectPtr.GetOuter();
 		return TObjectPtr<UObject>(Ptr);
 	}
 
-	FORCEINLINE TObjectPtr<UPackage> GetPackage() const
+	[[nodiscard]] FORCEINLINE TObjectPtr<UPackage> GetPackage() const
 	{
 		FObjectPtr Ptr = ObjectPtr.GetPackage();
 		return TObjectPtr<UPackage>(Ptr);
 	}
 
-	FORCEINLINE bool IsIn(FObjectPtr SomeOuter) const
+	[[nodiscard]] FORCEINLINE bool IsIn(FObjectPtr SomeOuter) const
 	{
 		return ObjectPtr.IsIn(SomeOuter);
 	}
 
 	template <typename U>
-	FORCEINLINE bool IsIn(TObjectPtr<U> SomeOuter) const
+	[[nodiscard]] FORCEINLINE bool IsIn(TObjectPtr<U> SomeOuter) const
 	{
 		return ObjectPtr.IsIn(SomeOuter.ObjectPtr);
 	}
 
-	FORCEINLINE operator T* () const { return Get(); }
+	[[nodiscard]] FORCEINLINE operator T* () const 
+	{
+		return Get();
+	}
 	template <typename U>
 	UE_OBJPTR_DEPRECATED(5.0, "Explicit cast to other raw pointer types is deprecated.  Please use the Cast API or get the raw pointer with ToRawPtr and cast that instead.")
-	explicit FORCEINLINE operator U* () const { return (U*)Get(); }
-	explicit FORCEINLINE operator UPTRINT() const { return (UPTRINT)Get(); }
-	FORCEINLINE T* operator->() const { return Get(); }
-	FORCEINLINE T& operator*() const { return *Get(); }
+	[[nodiscard]] explicit FORCEINLINE operator U* () const 
+	{
+		return (U*)Get();
+	}
+	[[nodiscard]] explicit FORCEINLINE operator UPTRINT() const 
+	{
+		return (UPTRINT)Get(); 
+	}
+	[[nodiscard]] FORCEINLINE T* operator->() const 
+	{
+		return Get();
+	}
+	[[nodiscard]] FORCEINLINE T& operator*() const 
+	{
+		return *Get();
+	}
 
 	UE_OBJPTR_DEPRECATED(5.0, "Conversion to a mutable pointer is deprecated.  Please pass a TObjectPtr<T>& instead so that assignment can be tracked accurately.")
-	explicit FORCEINLINE operator T*& () { return GetInternalRef(); }
-
-	UE_DEPRECATED(5.1, "IsNull is deprecated, please use operator bool instead.  if (!MyObjectPtr) { ... }")
-	FORCEINLINE bool IsNull() const { return ObjectPtr_Private::IsObjectPtrNull<T>(ObjectPtr); }
-
-	UE_DEPRECATED(5.1, "IsNullNoResolve is deprecated, please use operator bool instead.  if (!MyObjectPtr) { ... }")
-	FORCEINLINE bool IsNullNoResolve() const { return ObjectPtr_Private::IsObjectPtrNull<T>(ObjectPtr); }
+	[[nodiscard]] explicit FORCEINLINE operator T*& () 
+	{
+		return GetInternalRef();
+	}
 
-	FORCEINLINE bool operator!() const { return ObjectPtr_Private::IsObjectPtrNull<T>(ObjectPtr); }
-	explicit FORCEINLINE operator bool() const { return !ObjectPtr_Private::IsObjectPtrNull<T>(ObjectPtr); }
-	FORCEINLINE bool IsResolved() const { return ObjectPtr.IsResolved(); }
-	FORCEINLINE bool IsRemote() const { return ObjectPtr.IsRemote(); }
-	FORCEINLINE FString GetPath() const { return ObjectPtr.GetPath(); }
-	FORCEINLINE FString GetPathName() const { return ObjectPtr.GetPathName(); }
-	FORCEINLINE FName GetFName() const { return ObjectPtr.GetFName(); }
-	FORCEINLINE FString GetName() const { return ObjectPtr.GetName(); }
-	FORCEINLINE FString GetFullName(EObjectFullNameFlags Flags = EObjectFullNameFlags::None) const { return ObjectPtr.GetFullName(Flags); }
-	FORCEINLINE FObjectHandle GetHandle() const { return ObjectPtr.GetHandle(); }
-	FORCEINLINE bool IsA(const UClass* SomeBase) const { return ObjectPtr.IsA(SomeBase); }
-	template <typename U> FORCEINLINE bool IsA() const { return ObjectPtr.IsA<U>(); }
+	[[nodiscard]] FORCEINLINE bool operator!() const 
+	{
+		return ObjectPtr_Private::IsObjectPtrNull<T>(ObjectPtr);
+	}
+	[[nodiscard]] explicit FORCEINLINE operator bool() const 
+	{
+		return !ObjectPtr_Private::IsObjectPtrNull<T>(ObjectPtr);
+	}
+	[[nodiscard]] FORCEINLINE bool IsResolved() const
+	{
+		return ObjectPtr.IsResolved();
+	}
+	[[nodiscard]] FORCEINLINE bool IsRemote() const 
+	{
+		return ObjectPtr.IsRemote();
+	}
+	[[nodiscard]] FORCEINLINE FString GetPath() const 
+	{
+		return ObjectPtr.GetPath();
+	}
+	[[nodiscard]] FORCEINLINE FString GetPathName() const
+	{
+		return ObjectPtr.GetPathName();
+	}
+	[[nodiscard]] FORCEINLINE FName GetFName() const
+	{
+		return ObjectPtr.GetFName();
+	}
+	[[nodiscard]] FORCEINLINE FString GetName() const
+	{
+		return ObjectPtr.GetName();
+	}
+	[[nodiscard]] FORCEINLINE FString GetFullName(EObjectFullNameFlags Flags = EObjectFullNameFlags::None) const 
+	{
+		return ObjectPtr.GetFullName(Flags);
+	}
+	[[nodiscard]] FORCEINLINE FObjectHandle GetHandle() const 
+	{
+		return ObjectPtr.GetHandle();
+	}
+	[[nodiscard]] FORCEINLINE bool IsA(const UClass* SomeBase) const 
+	{
+		return ObjectPtr.IsA(SomeBase);
+	}
+	template <typename U> 
+	[[nodiscard]] FORCEINLINE bool IsA() const 
+	{
+		return ObjectPtr.IsA<U>();
+	}
 
-	FORCEINLINE uint32 GetPtrTypeHash() const
+	[[nodiscard]] FORCEINLINE uint32 GetPtrTypeHash() const
 	{
 		return GetTypeHash(ObjectPtr);
 	}
@@ -750,7 +793,7 @@ public:
 	}
 
 #if UE_WITH_REMOTE_OBJECT_HANDLE
-	FORCEINLINE FRemoteObjectId GetRemoteId() const
+	[[nodiscard]] FORCEINLINE FRemoteObjectId GetRemoteId() const
 	{
 		return ObjectPtr.GetRemoteId();
 	}
@@ -903,7 +946,7 @@ namespace ObjectPtr_Private
 }
 
 template <typename T>
-FORCEINLINE uint32 GetTypeHash(const TObjectPtr<T>& InObjectPtr)
+[[nodiscard]] FORCEINLINE uint32 GetTypeHash(const TObjectPtr<T>& InObjectPtr)
 {
 	return ObjectPtr_Private::Friend::GetPtrTypeHash(InObjectPtr);
 }
@@ -927,7 +970,7 @@ template <
 	typename U,
 	decltype(ObjectPtr_Private::IsObjectPtrEqual(std::declval<const TObjectPtr<T>&>(), std::declval<const U&>()))* = nullptr
 >
-FORCEINLINE bool operator==(const U& Other, const TObjectPtr<T>& Ptr)
+[[nodiscard]] FORCEINLINE bool operator==(const U& Other, const TObjectPtr<T>& Ptr)
 {
 	return ObjectPtr_Private::IsObjectPtrEqual(Ptr, Other);
 }
@@ -936,7 +979,7 @@ template <
 	typename U,
 	decltype(ObjectPtr_Private::IsObjectPtrEqual(std::declval<const TObjectPtr<T>&>(), std::declval<const U&>()))* = nullptr
 >
-FORCEINLINE bool operator!=(const U& Other, const TObjectPtr<T>& Ptr)
+[[nodiscard]] FORCEINLINE bool operator!=(const U& Other, const TObjectPtr<T>& Ptr)
 {
 	return !ObjectPtr_Private::IsObjectPtrEqual(Ptr, Other);
 }
@@ -965,26 +1008,26 @@ private:
 
 /** Used to allow the caller to provide a pointer to an incomplete type of T that has explicitly cast to a UObject. */
 template <typename T>
-TPrivateObjectPtr<T> MakeObjectPtrUnsafe(const UObject* Obj)
+[[nodiscard]] TPrivateObjectPtr<T> MakeObjectPtrUnsafe(const UObject* Obj)
 {
 	return TPrivateObjectPtr<T>{Obj};
 }
 
 template <typename T>
-TObjectPtr<T> ToObjectPtr(T* Obj)
+[[nodiscard]] TObjectPtr<T> ToObjectPtr(T* Obj)
 {
 	return TObjectPtr<T>{Obj};
 }
 
 template <typename T>
-FORCEINLINE T* ToRawPtr(const TObjectPtr<T>& Ptr)
+[[nodiscard]] FORCEINLINE T* ToRawPtr(const TObjectPtr<T>& Ptr)
 {
 	// NOTE: This is specifically not getting a reference to the internal pointer.
 	return Ptr.Get();
 }
 
 template <typename T>
-FORCEINLINE T* ToRawPtr(T* Ptr)
+[[nodiscard]] FORCEINLINE T* ToRawPtr(T* Ptr)
 {
 	return Ptr;
 }
@@ -992,14 +1035,14 @@ FORCEINLINE T* ToRawPtr(T* Ptr)
 #if !UE_DEPRECATE_MUTABLE_TOBJECTPTR
 template <typename T, SIZE_T Size>
 UE_OBJPTR_DEPRECATED(5.3, "Mutable ToRawPtrArrayUnsafe() is deprecated. Use MutableView() or TArray<TObjectPtr<...>> instead.")
-FORCEINLINE T**
+[[nodiscard]] FORCEINLINE T**
 ToRawPtrArrayUnsafe(TObjectPtr<T>(&ArrayOfPtr)[Size])
 {
 #if UE_WITH_OBJECT_HANDLE_LATE_RESOLVE || UE_WITH_OBJECT_HANDLE_TRACKING
 	for (TObjectPtr<T>& Item : ArrayOfPtr)
 	{
 		// NOTE: Relying on the fact that the TObjectPtr will cache the resolved pointer in place after calling Get.
-		Item.Get();
+		(void)Item.Get();
 		check(Item.IsResolved());
 	}
 #endif
@@ -1009,7 +1052,7 @@ ToRawPtrArrayUnsafe(TObjectPtr<T>(&ArrayOfPtr)[Size])
 #endif
 
 template <typename T, SIZE_T Size>
-FORCEINLINE const T* const *
+[[nodiscard]] FORCEINLINE const T* const *
 ToRawPtrArray(const TObjectPtr<T>(&ArrayOfPtr)[Size])
 {
 #if UE_WITH_OBJECT_HANDLE_LATE_RESOLVE || UE_WITH_OBJECT_HANDLE_TRACKING
@@ -1021,7 +1064,7 @@ ToRawPtrArray(const TObjectPtr<T>(&ArrayOfPtr)[Size])
 }
 
 template <typename T>
-FORCEINLINE T** ToRawPtrArrayUnsafe(T** ArrayOfPtr)
+[[nodiscard]] FORCEINLINE T** ToRawPtrArrayUnsafe(T** ArrayOfPtr)
 {
 	return ArrayOfPtr;
 }
@@ -1033,10 +1076,10 @@ template <
 	UE_REQUIRES(TIsTArray_V<ArrayTypeNoRef>)
 >
 #if UE_DEPRECATE_MUTABLE_TOBJECTPTR
-const auto&
+[[nodiscard]] const auto&
 #else
 UE_OBJPTR_DEPRECATED(5.3, "Mutable ToRawPtrTArrayUnsafe() is deprecated. Use MutableView() or TArray<TObjectPtr<...>> instead.")
-decltype(auto)
+[[nodiscard]] decltype(auto)
 #endif
 ToRawPtrTArrayUnsafe(ArrayType&& Array)
 {
@@ -1069,7 +1112,7 @@ struct TContainerElementTypeCompatibility<TObjectPtr<T>>
 #if UE_WITH_OBJECT_HANDLE_LATE_RESOLVE || UE_WITH_OBJECT_HANDLE_TRACKING
 		while (Iter != IterEnd)
 		{
-			Operator(Iter).Get();
+			(void)Operator(Iter).Get();
 			check(Operator(Iter).IsResolved());
 			++Iter;
 		}
@@ -1179,7 +1222,7 @@ struct TCallTraits<TObjectPtr<T>> : public TCallTraitsBase<TObjectPtr<T>>
 
 
 template <typename T>
-FORCEINLINE TWeakObjectPtr<T> MakeWeakObjectPtr(TObjectPtr<T> Ptr)
+[[nodiscard]] FORCEINLINE TWeakObjectPtr<T> MakeWeakObjectPtr(TObjectPtr<T> Ptr)
 {
 #if UE_WITH_REMOTE_OBJECT_HANDLE
 	if (Ptr.IsRemote())
@@ -1261,7 +1304,7 @@ inline void Exchange(TArray<T*>& A, TArray<TObjectPtr<T>>& B)
  * Returns a pointer to a valid object if the Test object passes IsValid() tests, otherwise null
  */
 template <typename T>
-T* GetValid(const TObjectPtr<T>& Test)
+[[nodiscard]] T* GetValid(const TObjectPtr<T>& Test)
 {
 	T* TestPtr = ToRawPtr(Test);
 	return IsValid(TestPtr) ? TestPtr : nullptr;
@@ -1350,7 +1393,7 @@ namespace UE::Core::Private // private facilities; not for direct use
 		static FORCEINLINE void PerformDecayActions(const TObjectPtr<T>& Value)
 		{
 #if UE_WITH_OBJECT_HANDLE_LATE_RESOLVE || UE_WITH_OBJECT_HANDLE_TRACKING
-			Value.Get();
+			(void)Value.Get();
 			check(Value.IsResolved());
 #endif
 		}
@@ -1407,7 +1450,7 @@ namespace UE::Core::Private // private facilities; not for direct use
 		static FORCEINLINE void PerformDecayActions(const TNonNullPtr<TObjectPtr<T>>& Value)
 		{
 #if UE_WITH_OBJECT_HANDLE_LATE_RESOLVE || UE_WITH_OBJECT_HANDLE_TRACKING
-			Value.GetRef().Get();
+			(void)Value.GetRef().Get();
 #endif
 		}
 	};
@@ -1747,13 +1790,13 @@ template <typename T,
 
 /* cast utilities; best avoided, use mostly for compatibility with existing code */
 template <typename To, typename From>
-FORCEINLINE To* StaticCastPtr(const TObjectPtr<From>& P)
+[[nodiscard]] FORCEINLINE To* StaticCastPtr(const TObjectPtr<From>& P)
 {
 	return static_cast<To*>(P.Get());
 }
 
 template <typename T>
-FORCEINLINE decltype(auto) ConstCast(const TObjectPtr<T>& P)
+[[nodiscard]] FORCEINLINE decltype(auto) ConstCast(const TObjectPtr<T>& P)
 {
 	return reinterpret_cast<TObjectPtr<std::remove_cv_t<T>>&>(const_cast<TObjectPtr<T>&>(P));
 }
@@ -1764,7 +1807,7 @@ class TNonNullPtr<TObjectPtr<ObjectType>>
 {
 public:
 	
-	FORCEINLINE TNonNullPtr(EDefaultConstructNonNullPtr)
+	[[nodiscard]] FORCEINLINE TNonNullPtr(EDefaultConstructNonNullPtr)
 		: Object(nullptr)
 	{	
 	}	
@@ -1772,7 +1815,7 @@ public:
 	/**
 	 * nullptr constructor - not allowed.
 	 */
-	FORCEINLINE TNonNullPtr(TYPE_OF_NULLPTR)
+	[[nodiscard]] FORCEINLINE TNonNullPtr(TYPE_OF_NULLPTR)
 	{
 		// Essentially static_assert(false), but this way prevents GCC/Clang from crying wolf by merely inspecting the function body
 		static_assert(sizeof(ObjectType) == 0, "Tried to initialize TNonNullPtr with a null pointer!");
@@ -1781,7 +1824,7 @@ public:
 	/**
 	 * Constructs a non-null pointer from the provided pointer. Must not be nullptr.
 	 */
-	FORCEINLINE TNonNullPtr(TObjectPtr<ObjectType> InObject)
+	[[nodiscard]] FORCEINLINE TNonNullPtr(TObjectPtr<ObjectType> InObject)
 		: Object(InObject)
 	{
 		ensureAlwaysMsgf(InObject, TEXT("Tried to initialize TNonNullPtr with a null pointer!"));
@@ -1794,7 +1837,7 @@ public:
 		typename OtherObjectType
 		UE_REQUIRES(std::is_convertible_v<OtherObjectType*, ObjectType*>)
 	>
-	FORCEINLINE TNonNullPtr(const TNonNullPtr<OtherObjectType>& Other)
+	[[nodiscard]] FORCEINLINE TNonNullPtr(const TNonNullPtr<OtherObjectType>& Other)
 		: Object(Other.Object)
 	{
 	}
@@ -1835,7 +1878,7 @@ public:
 	/**
 	 * Returns the internal pointer
 	 */
-	FORCEINLINE operator ObjectType*() const
+	[[nodiscard]] FORCEINLINE operator ObjectType*() const
 	{
 		ensureMsgf(Object, TEXT("Tried to access null pointer!"));
 		return Object;
@@ -1844,7 +1887,7 @@ public:
 	/**
 	 * Returns the internal pointer
 	 */
-	FORCEINLINE ObjectType* Get() const
+	[[nodiscard]] FORCEINLINE ObjectType* Get() const
 	{
 		ensureMsgf(Object, TEXT("Tried to access null pointer!"));
 		return Object;
@@ -1853,7 +1896,7 @@ public:
 	/**
 	 * Dereference operator returns a reference to the object this pointer points to
 	 */
-	FORCEINLINE ObjectType& operator*() const
+	[[nodiscard]] FORCEINLINE ObjectType& operator*() const
 	{
 		ensureMsgf(Object, TEXT("Tried to access null pointer!"));
 		return *Object;
@@ -1862,25 +1905,25 @@ public:
 	/**
 	 * Arrow operator returns a pointer to this pointer's object
 	 */
-	FORCEINLINE ObjectType* operator->() const
+	[[nodiscard]] FORCEINLINE ObjectType* operator->() const
 	{
 		ensureMsgf(Object, TEXT("Tried to access null pointer!"));
 		return Object;
 	}
 
-	FORCEINLINE const TObjectPtr<ObjectType>& GetRef() const
+	[[nodiscard]] FORCEINLINE const TObjectPtr<ObjectType>& GetRef() const
 	{
 		ensureMsgf(Object, TEXT("Tried to access null pointer!"));		
 		return Object;
 	}
 	
-	FORCEINLINE TObjectPtr<ObjectType>& GetRef() 
+	[[nodiscard]] FORCEINLINE TObjectPtr<ObjectType>& GetRef() 
 	{
 		ensureMsgf(Object, TEXT("Tried to access null pointer!"));		
 		return Object;
 	}	
 
-	FORCEINLINE bool IsInitialized() const
+	[[nodiscard]] FORCEINLINE bool IsInitialized() const
 	{
 		return Object != nullptr;
 	}
diff --git a/Engine/Source/Runtime/CoreUObject/Public/UObject/SoftObjectPtr.h b/Engine/Source/Runtime/CoreUObject/Public/UObject/SoftObjectPtr.h
--- a/Engine/Source/Runtime/CoreUObject/Public/UObject/SoftObjectPtr.h
+++ b/Engine/Source/Runtime/CoreUObject/Public/UObject/SoftObjectPtr.h
@@ -47,28 +47,28 @@ private:
 	using Super = TPersistentObjectPtr<FSoftObjectPath>;
 
 public:	
-	FORCEINLINE FSoftObjectPtr() = default;
-	FORCEINLINE FSoftObjectPtr(const FSoftObjectPtr& Other) = default;
-	FORCEINLINE FSoftObjectPtr(FSoftObjectPtr&& Other) = default;
+	[[nodiscard]] FORCEINLINE FSoftObjectPtr() = default;
+	[[nodiscard]] FORCEINLINE FSoftObjectPtr(const FSoftObjectPtr& Other) = default;
+	[[nodiscard]] FORCEINLINE FSoftObjectPtr(FSoftObjectPtr&& Other) = default;
 	FORCEINLINE ~FSoftObjectPtr() = default;
 	FORCEINLINE FSoftObjectPtr& operator=(const FSoftObjectPtr& Other) = default;
 	FORCEINLINE FSoftObjectPtr& operator=(FSoftObjectPtr&& Other) = default;
 
-	explicit FORCEINLINE FSoftObjectPtr(const FSoftObjectPath& ObjectPath)
+	[[nodiscard]] explicit FORCEINLINE FSoftObjectPtr(const FSoftObjectPath& ObjectPath)
 		: Super(ObjectPath)
 	{
 	}
 
-	explicit FORCEINLINE FSoftObjectPtr(FObjectPtr Object)
+	[[nodiscard]] explicit FORCEINLINE FSoftObjectPtr(FObjectPtr Object)
 	{
 		(*this)=Object;
 	}
-	explicit FORCEINLINE FSoftObjectPtr(const UObject* Object)
+	[[nodiscard]] explicit FORCEINLINE FSoftObjectPtr(const UObject* Object)
 		: FSoftObjectPtr(FObjectPtr(const_cast<UObject*>(Object)))
 	{
 	}
 	template <typename T>
-	explicit FORCEINLINE FSoftObjectPtr(TObjectPtr<T> Object)
+	[[nodiscard]] explicit FORCEINLINE FSoftObjectPtr(TObjectPtr<T> Object)
 		: FSoftObjectPtr(FObjectPtr(Object))
 	{
 		// This needs to be a template instead of TObjectPtr<const UObject> because C++ does derived-to-base
@@ -93,38 +93,38 @@ public:
 	}
 
 	/** Returns the StringObjectPath that is wrapped by this SoftObjectPtr */
-	FORCEINLINE const FSoftObjectPath& ToSoftObjectPath() const
+	[[nodiscard]] FORCEINLINE const FSoftObjectPath& ToSoftObjectPath() const
 	{
 		return GetUniqueID();
 	}
 
 	/** Returns string representation of reference, in form /package/path.assetname */
-	FORCEINLINE FString ToString() const
+	[[nodiscard]] FORCEINLINE FString ToString() const
 	{
 		return ToSoftObjectPath().ToString();
 	}
 
 	/** Returns /package/path string, leaving off the asset name */
-	FORCEINLINE FString GetLongPackageName() const
+	[[nodiscard]] FORCEINLINE FString GetLongPackageName() const
 	{
 		return ToSoftObjectPath().GetLongPackageName();
 	}
 	
 	/** Returns /package/path name, leaving off the asset name */
-	FORCEINLINE FName GetLongPackageFName() const
+	[[nodiscard]] FORCEINLINE FName GetLongPackageFName() const
 	{
 		return ToSoftObjectPath().GetLongPackageFName();
 	}
 
 	/** Returns assetname string, leaving off the /package/path. part */
-	FORCEINLINE FString GetAssetName() const
+	[[nodiscard]] FORCEINLINE FString GetAssetName() const
 	{
 		return ToSoftObjectPath().GetAssetName();
 	}
 
 #if WITH_EDITOR
 	/** Overridden to deal with PIE lookups */
-	FORCEINLINE UObject* Get() const
+	[[nodiscard]] FORCEINLINE UObject* Get() const
 	{
 		if (UE::GetPlayInEditorID() != INDEX_NONE)
 		{
@@ -178,9 +178,9 @@ struct TSoftObjectPtr
 public:
 	using ElementType = T;
 	
-	FORCEINLINE TSoftObjectPtr() = default;
-	FORCEINLINE TSoftObjectPtr(const TSoftObjectPtr& Other) = default;
-	FORCEINLINE TSoftObjectPtr(TSoftObjectPtr&& Other) = default;
+	[[nodiscard]] FORCEINLINE TSoftObjectPtr() = default;
+	[[nodiscard]] FORCEINLINE TSoftObjectPtr(const TSoftObjectPtr& Other) = default;
+	[[nodiscard]] FORCEINLINE TSoftObjectPtr(TSoftObjectPtr&& Other) = default;
 	FORCEINLINE ~TSoftObjectPtr() = default;
 	FORCEINLINE TSoftObjectPtr& operator=(const TSoftObjectPtr& Other) = default;
 	FORCEINLINE TSoftObjectPtr& operator=(TSoftObjectPtr&& Other) = default;
@@ -190,7 +190,7 @@ public:
 		class U
 		UE_REQUIRES(std::is_convertible_v<U*, T*>)
 	>
-	FORCEINLINE TSoftObjectPtr(const TSoftObjectPtr<U>& Other)
+	[[nodiscard]] FORCEINLINE TSoftObjectPtr(const TSoftObjectPtr<U>& Other)
 		: SoftObjectPtr(Other.SoftObjectPtr)
 	{
 	}
@@ -199,7 +199,7 @@ public:
 		UE_REQUIRES(!std::is_convertible_v<U*, T*>)
 	>
 	UE_SOFTOBJECTPTR_CONVERSION_DEPRECATED(5.5, "Constructing TSoftObjectPtr from an incompatible pointer type has been deprecated.")
-	FORCEINLINE TSoftObjectPtr(const TSoftObjectPtr<U>& Other)
+	[[nodiscard]] FORCEINLINE TSoftObjectPtr(const TSoftObjectPtr<U>& Other)
 		: SoftObjectPtr(Other.SoftObjectPtr)
 	{
 	}
@@ -209,7 +209,7 @@ public:
 		class U
 		UE_REQUIRES(std::is_convertible_v<U*, T*>)
 	>
-	FORCEINLINE TSoftObjectPtr(TSoftObjectPtr<U>&& Other)
+	[[nodiscard]] FORCEINLINE TSoftObjectPtr(TSoftObjectPtr<U>&& Other)
 		: SoftObjectPtr(MoveTemp(Other.SoftObjectPtr))
 	{
 	}
@@ -218,7 +218,7 @@ public:
 		UE_REQUIRES(!std::is_convertible_v<U*, T*>)
 	>
 	UE_SOFTOBJECTPTR_CONVERSION_DEPRECATED(5.5, "Constructing TSoftObjectPtr from an incompatible pointer type has been deprecated.")
-	FORCEINLINE TSoftObjectPtr(TSoftObjectPtr<U>&& Other)
+	[[nodiscard]] FORCEINLINE TSoftObjectPtr(TSoftObjectPtr<U>&& Other)
 		: SoftObjectPtr(MoveTemp(Other.SoftObjectPtr))
 	{
 	}
@@ -228,7 +228,7 @@ public:
 		typename U
 		UE_REQUIRES(std::is_convertible_v<U*, T*>)
 	>
-	FORCEINLINE TSoftObjectPtr(U* Object)
+	[[nodiscard]] FORCEINLINE TSoftObjectPtr(U* Object)
 		: SoftObjectPtr(Object)
 	{
 	}
@@ -237,7 +237,7 @@ public:
 		UE_REQUIRES(!std::is_convertible_v<U*, T*>)
 	>
 	UE_SOFTOBJECTPTR_CONVERSION_DEPRECATED(5.5, "Constructing TSoftObjectPtr from an incompatible pointer type has been deprecated.")
-	FORCEINLINE TSoftObjectPtr(U* Object)
+	[[nodiscard]] FORCEINLINE TSoftObjectPtr(U* Object)
 		: SoftObjectPtr(Object)
 	{
 	}
@@ -247,7 +247,7 @@ public:
 		typename U
 		UE_REQUIRES(std::is_convertible_v<U*, T*>)
 	>
-	FORCEINLINE TSoftObjectPtr(const TObjectPtr<U> Object)
+	[[nodiscard]] FORCEINLINE TSoftObjectPtr(const TObjectPtr<U> Object)
 		: SoftObjectPtr(Object.Get())
 	{
 	}
@@ -256,13 +256,13 @@ public:
 		UE_REQUIRES(!std::is_convertible_v<U*, T*>)
 	>
 	UE_SOFTOBJECTPTR_CONVERSION_DEPRECATED(5.5, "Constructing TSoftObjectPtr from an incompatible pointer type has been deprecated.")
-	FORCEINLINE TSoftObjectPtr(const TObjectPtr<U> Object)
+	[[nodiscard]] FORCEINLINE TSoftObjectPtr(const TObjectPtr<U> Object)
 		: SoftObjectPtr(Object.Get())
 	{
 	}
 
 	/** Construct from a nullptr */
-	FORCEINLINE TSoftObjectPtr(TYPE_OF_NULLPTR)
+	[[nodiscard]] FORCEINLINE TSoftObjectPtr(TYPE_OF_NULLPTR)
 		: SoftObjectPtr(nullptr)
 	{
 	}
@@ -272,7 +272,7 @@ public:
 		typename SoftObjectPathType
 		UE_REQUIRES(std::is_same_v<SoftObjectPathType, FSoftObjectPath>)
 	>
-	explicit FORCEINLINE TSoftObjectPtr(SoftObjectPathType ObjectPath)
+	[[nodiscard]] explicit FORCEINLINE TSoftObjectPtr(SoftObjectPathType ObjectPath)
 		: SoftObjectPtr(MoveTemp(ObjectPath))
 	{
 		// The reason for the strange templated/constrained signature is to prevent this type hole:
@@ -283,7 +283,7 @@ public:
 		// The real fix is to make the FSoftObjectPath(const UObject*) constructor explicit, but that will involve substantially more changes.
 	}
 	UE_SOFTOBJECTPTR_CONVERSION_DEPRECATED(5.5, "Constructing TSoftObjectPtr from an FString has been deprecated - instead, explicitly construct a FSoftObjectPath.")
-	explicit TSoftObjectPtr(const FString& Path)
+	[[nodiscard]] explicit TSoftObjectPtr(const FString& Path)
 		: SoftObjectPtr(FSoftObjectPath(Path))
 	{
 	}
@@ -404,7 +404,7 @@ public:
 	 *
 	 * @param Other soft pointer to compare to
 	 */
-	FORCEINLINE bool operator==(const TSoftObjectPtr& Rhs) const
+	[[nodiscard]] FORCEINLINE bool operator==(const TSoftObjectPtr& Rhs) const
 	{
 		return SoftObjectPtr == Rhs.SoftObjectPtr;
 	}
@@ -412,7 +412,7 @@ public:
 		typename U
 		UE_REQUIRES(TModels_V<CEqualityComparableWith, T*, U*>)
 	>
-	FORCEINLINE bool operator==(const TSoftObjectPtr<U>& Rhs) const
+	[[nodiscard]] FORCEINLINE bool operator==(const TSoftObjectPtr<U>& Rhs) const
 	{
 		return SoftObjectPtr == Rhs.SoftObjectPtr;
 	}
@@ -421,7 +421,7 @@ public:
 		UE_REQUIRES(!TModels_V<CEqualityComparableWith, T*, U*>)
 	>
 	UE_SOFTOBJECTPTR_CONVERSION_DEPRECATED(5.5, "Comparing TSoftObjectPtrs of incompatible pointer types has been deprecated.")
-	FORCEINLINE bool operator==(const TSoftObjectPtr<U>& Rhs) const
+	[[nodiscard]] FORCEINLINE bool operator==(const TSoftObjectPtr<U>& Rhs) const
 	{
 		return SoftObjectPtr == Rhs.SoftObjectPtr;
 	}
@@ -429,7 +429,7 @@ public:
 		typename U
 		UE_REQUIRES(TModels_V<CEqualityComparableWith, T*, U*>)
 	>
-	FORCEINLINE bool operator==(U* Rhs) const
+	[[nodiscard]] FORCEINLINE bool operator==(U* Rhs) const
 	{
 		return SoftObjectPtr == TSoftObjectPtr<U>(Rhs).SoftObjectPtr;
 	}
@@ -438,11 +438,11 @@ public:
 		UE_REQUIRES(!TModels_V<CEqualityComparableWith, T*, U*>)
 	>
 	UE_SOFTOBJECTPTR_CONVERSION_DEPRECATED(5.5, "Comparing TSoftObjectPtrs of incompatible pointer types has been deprecated.")
-	FORCEINLINE bool operator==(U* Rhs) const
+	[[nodiscard]] FORCEINLINE bool operator==(U* Rhs) const
 	{
 		return SoftObjectPtr == TSoftObjectPtr<U>(Rhs).SoftObjectPtr;
 	}
-	FORCEINLINE bool operator==(TYPE_OF_NULLPTR) const
+	[[nodiscard]] FORCEINLINE bool operator==(TYPE_OF_NULLPTR) const
 	{
 		return SoftObjectPtr == nullptr;
 	}
@@ -452,16 +452,16 @@ public:
 	 *
 	 * @return nullptr if this object is gone or the lazy pointer was null, otherwise a valid UObject pointer
 	 */
-	T* Get() const;
+	[[nodiscard]] T* Get() const;
 
 	/** Dereference the soft pointer */
-	FORCEINLINE T& operator*() const
+	[[nodiscard]] FORCEINLINE T& operator*() const
 	{
 		return *Get();
 	}
 
 	/** Dereference the soft pointer */
-	FORCEINLINE T* operator->() const
+	[[nodiscard]] FORCEINLINE T* operator->() const
 	{
 		return Get();
 	}
@@ -491,7 +491,7 @@ public:
 	 *
 	 * @return true if Get() would return a valid non-null pointer
 	 */
-	FORCEINLINE bool IsValid() const
+	[[nodiscard]] FORCEINLINE bool IsValid() const
 	{
 		// This does the runtime type check
 		return Get() != nullptr;
@@ -502,7 +502,7 @@ public:
 	 *
 	 * @return true if this does not point to a real object, but could possibly
 	 */
-	FORCEINLINE bool IsPending() const
+	[[nodiscard]] FORCEINLINE bool IsPending() const
 	{
 		return SoftObjectPtr.IsPending();
 	}
@@ -512,55 +512,55 @@ public:
 	 *
 	 * @return true if this is explicitly pointing to no object
 	 */
-	FORCEINLINE bool IsNull() const
+	[[nodiscard]] FORCEINLINE bool IsNull() const
 	{
 		return SoftObjectPtr.IsNull();
 	}
 
 	/** Returns the StringObjectPath that is wrapped by this SoftObjectPtr */
-	FORCEINLINE const FSoftObjectPath& GetUniqueID() const
+	[[nodiscard]] FORCEINLINE const FSoftObjectPath& GetUniqueID() const
 	{
 		return SoftObjectPtr.GetUniqueID();
 	}
 
 	/** Returns the StringObjectPath that is wrapped by this SoftObjectPtr */
-	FORCEINLINE const FSoftObjectPath& ToSoftObjectPath() const
+	[[nodiscard]] FORCEINLINE const FSoftObjectPath& ToSoftObjectPath() const
 	{
 		return SoftObjectPtr.GetUniqueID();
 	}
 
 	/** Returns string representation of reference, in form /package/path.assetname */
-	FORCEINLINE FString ToString() const
+	[[nodiscard]] FORCEINLINE FString ToString() const
 	{
 		return ToSoftObjectPath().ToString();
 	}
 
 	/** Returns /package/path string, leaving off the asset name */
-	FORCEINLINE FString GetLongPackageName() const
+	[[nodiscard]] FORCEINLINE FString GetLongPackageName() const
 	{
 		return ToSoftObjectPath().GetLongPackageName();
 	}
 	
 	/** Returns /package/path name, leaving off the asset name */
-	FORCEINLINE FName GetLongPackageFName() const
+	[[nodiscard]] FORCEINLINE FName GetLongPackageFName() const
 	{
 		return ToSoftObjectPath().GetLongPackageFName();
 	}
 
 	/** Returns assetname string, leaving off the /package/path part */
-	FORCEINLINE FString GetAssetName() const
+	[[nodiscard]] FORCEINLINE FString GetAssetName() const
 	{
 		return ToSoftObjectPath().GetAssetName();
 	}
 
 	/** Dereference soft pointer to see if it points somewhere valid */
-	FORCEINLINE explicit operator bool() const
+	[[nodiscard]] FORCEINLINE explicit operator bool() const
 	{
 		return IsValid();
 	}
 
 	/** Hash function */
-	FORCEINLINE uint32 GetPtrTypeHash() const
+	[[nodiscard]] FORCEINLINE uint32 GetPtrTypeHash() const
 	{
 		return GetTypeHash(static_cast<const TPersistentObjectPtr<FSoftObjectPath>&>(SoftObjectPtr));
 	}
@@ -583,7 +583,7 @@ template <
 	typename RhsType
 	UE_REQUIRES(TModels_V<CEqualityComparableWith, LhsType*, RhsType*>)
 >
-FORCEINLINE bool operator==(LhsType* Lhs, const TSoftObjectPtr<RhsType>& Rhs)
+[[nodiscard]] FORCEINLINE bool operator==(LhsType* Lhs, const TSoftObjectPtr<RhsType>& Rhs)
 {
 	return Rhs == Lhs;
 }
@@ -593,12 +593,12 @@ template <
 	UE_REQUIRES(!TModels_V<CEqualityComparableWith, LhsType*, RhsType*>)
 >
 UE_SOFTOBJECTPTR_CONVERSION_DEPRECATED(5.5, "Comparing TSoftObjectPtrs of incompatible pointer types has been deprecated.")
-FORCEINLINE bool operator==(LhsType* Lhs, const TSoftObjectPtr<RhsType>& Rhs)
+[[nodiscard]] FORCEINLINE bool operator==(LhsType* Lhs, const TSoftObjectPtr<RhsType>& Rhs)
 {
 	return Rhs == Lhs;
 }
 template <typename RhsType>
-FORCEINLINE bool operator==(TYPE_OF_NULLPTR Lhs, const TSoftObjectPtr<RhsType>& Rhs)
+[[nodiscard]] FORCEINLINE bool operator==(TYPE_OF_NULLPTR Lhs, const TSoftObjectPtr<RhsType>& Rhs)
 {
 	return Rhs == Lhs;
 }
@@ -607,7 +607,7 @@ template <
 	typename RhsType
 	UE_REQUIRES(TModels_V<CEqualityComparableWith, LhsType*, RhsType*>)
 >
-FORCEINLINE bool operator!=(const TSoftObjectPtr<LhsType>& Lhs, const TSoftObjectPtr<RhsType>& Rhs)
+[[nodiscard]] FORCEINLINE bool operator!=(const TSoftObjectPtr<LhsType>& Lhs, const TSoftObjectPtr<RhsType>& Rhs)
 {
 	return !(Lhs == Rhs);
 }
@@ -617,7 +617,7 @@ template <
 	UE_REQUIRES(!TModels_V<CEqualityComparableWith, LhsType*, RhsType*>)
 >
 UE_SOFTOBJECTPTR_CONVERSION_DEPRECATED(5.5, "Comparing TSoftObjectPtrs of incompatible pointer types has been deprecated.")
-FORCEINLINE bool operator!=(const TSoftObjectPtr<LhsType>& Lhs, const TSoftObjectPtr<RhsType>& Rhs)
+[[nodiscard]] FORCEINLINE bool operator!=(const TSoftObjectPtr<LhsType>& Lhs, const TSoftObjectPtr<RhsType>& Rhs)
 {
 	return !(Lhs == Rhs);
 }
@@ -626,7 +626,7 @@ template <
 	typename RhsType
 	UE_REQUIRES(TModels_V<CEqualityComparableWith, LhsType*, RhsType*>)
 >
-FORCEINLINE bool operator!=(const TSoftObjectPtr<LhsType>& Lhs, RhsType* Rhs)
+[[nodiscard]] FORCEINLINE bool operator!=(const TSoftObjectPtr<LhsType>& Lhs, RhsType* Rhs)
 {
 	return !(Lhs == Rhs);
 }
@@ -636,12 +636,12 @@ template <
 	UE_REQUIRES(!TModels_V<CEqualityComparableWith, LhsType*, RhsType*>)
 >
 UE_SOFTOBJECTPTR_CONVERSION_DEPRECATED(5.5, "Comparing TSoftObjectPtrs of incompatible pointer types has been deprecated.")
-FORCEINLINE bool operator!=(const TSoftObjectPtr<LhsType>& Lhs, RhsType* Rhs)
+[[nodiscard]] FORCEINLINE bool operator!=(const TSoftObjectPtr<LhsType>& Lhs, RhsType* Rhs)
 {
 	return !(Lhs == Rhs);
 }
 template <typename LhsType>
-FORCEINLINE bool operator!=(const TSoftObjectPtr<LhsType>& Lhs, TYPE_OF_NULLPTR Rhs)
+[[nodiscard]] FORCEINLINE bool operator!=(const TSoftObjectPtr<LhsType>& Lhs, TYPE_OF_NULLPTR Rhs)
 {
 	return !(Lhs == Rhs);
 }
@@ -650,7 +650,7 @@ template <
 	typename RhsType
 	UE_REQUIRES(TModels_V<CEqualityComparableWith, LhsType*, RhsType*>)
 >
-FORCEINLINE bool operator!=(LhsType* Lhs, const TSoftObjectPtr<RhsType>& Rhs)
+[[nodiscard]] FORCEINLINE bool operator!=(LhsType* Lhs, const TSoftObjectPtr<RhsType>& Rhs)
 {
 	return !(Lhs == Rhs);
 }
@@ -660,12 +660,12 @@ template <
 	UE_REQUIRES(!TModels_V<CEqualityComparableWith, LhsType*, RhsType*>)
 >
 UE_SOFTOBJECTPTR_CONVERSION_DEPRECATED(5.5, "Comparing TSoftObjectPtrs of incompatible pointer types has been deprecated.")
-FORCEINLINE bool operator!=(LhsType* Lhs, const TSoftObjectPtr<RhsType>& Rhs)
+[[nodiscard]] FORCEINLINE bool operator!=(LhsType* Lhs, const TSoftObjectPtr<RhsType>& Rhs)
 {
 	return !(Lhs == Rhs);
 }
 template <typename RhsType>
-FORCEINLINE bool operator!=(TYPE_OF_NULLPTR Lhs, const TSoftObjectPtr<RhsType>& Rhs)
+[[nodiscard]] FORCEINLINE bool operator!=(TYPE_OF_NULLPTR Lhs, const TSoftObjectPtr<RhsType>& Rhs)
 {
 	return !(Lhs == Rhs);
 }
@@ -673,7 +673,7 @@ FORCEINLINE bool operator!=(TYPE_OF_NULLPTR Lhs, const TSoftObjectPtr<RhsType>&
 
 /** Hash function */
 template<class T>
-FORCEINLINE uint32 GetTypeHash(const TSoftObjectPtr<T>& Ptr)
+[[nodiscard]] FORCEINLINE uint32 GetTypeHash(const TSoftObjectPtr<T>& Ptr)
 {
 	return Ptr.GetPtrTypeHash();
 }
@@ -698,14 +698,14 @@ struct TCallTraits<TSoftObjectPtr<T>> : public TCallTraitsBase<TSoftObjectPtr<T>
 
 /** Utility to create a TSoftObjectPtr without specifying the type */
 template <class T>
-TSoftObjectPtr<T> MakeSoftObjectPtr(T* Object)
+[[nodiscard]] TSoftObjectPtr<T> MakeSoftObjectPtr(T* Object)
 {
 	static_assert(std::is_base_of_v<UObject, T>, "Type must derive from UObject");
 	return TSoftObjectPtr<T>(Object);
 }
 
 template <class T>
-TSoftObjectPtr<T> MakeSoftObjectPtr(TObjectPtr<T> Object)
+[[nodiscard]] TSoftObjectPtr<T> MakeSoftObjectPtr(TObjectPtr<T> Object)
 {
 	static_assert(std::is_base_of_v<UObject, T>, "Type must derive from UObject");
 	return TSoftObjectPtr<T>(ToRawPtr(Object));
@@ -723,9 +723,9 @@ class TSoftClassPtr
 public:
 	using ElementType = TClass;
 	
-	FORCEINLINE TSoftClassPtr() = default;
-	FORCEINLINE TSoftClassPtr(const TSoftClassPtr& Other) = default;
-	FORCEINLINE TSoftClassPtr(TSoftClassPtr&& Other) = default;
+	[[nodiscard]] FORCEINLINE TSoftClassPtr() = default;
+	[[nodiscard]] FORCEINLINE TSoftClassPtr(const TSoftClassPtr& Other) = default;
+	[[nodiscard]] FORCEINLINE TSoftClassPtr(TSoftClassPtr&& Other) = default;
 	FORCEINLINE ~TSoftClassPtr() = default;
 	FORCEINLINE TSoftClassPtr& operator=(const TSoftClassPtr& Other) = default;
 	FORCEINLINE TSoftClassPtr& operator=(TSoftClassPtr&& Other) = default;
@@ -735,19 +735,19 @@ public:
 		class TClassA
 		UE_REQUIRES(std::is_convertible_v<TClassA*, TClass*>)
 	>
-	FORCEINLINE TSoftClassPtr(const TSoftClassPtr<TClassA>& Other)
+	[[nodiscard]] FORCEINLINE TSoftClassPtr(const TSoftClassPtr<TClassA>& Other)
 		: SoftObjectPtr(Other.SoftObjectPtr)
 	{
 	}
 
 	/** Construct from a class already in memory */
-	FORCEINLINE TSoftClassPtr(const UClass* From)
+	[[nodiscard]] FORCEINLINE TSoftClassPtr(const UClass* From)
 		: SoftObjectPtr(From)
 	{
 	}
 
 	/** Construct from a soft object path */
-	explicit FORCEINLINE TSoftClassPtr(const FSoftObjectPath& ObjectPath)
+	[[nodiscard]] explicit FORCEINLINE TSoftClassPtr(const FSoftObjectPath& ObjectPath)
 		: SoftObjectPtr(ObjectPath)
 	{
 	}
@@ -804,7 +804,7 @@ public:
 	 *
 	 * @param Other soft pointer to compare to 
 	 */
-	FORCEINLINE bool operator==(const TSoftClassPtr& Rhs) const
+	[[nodiscard]] FORCEINLINE bool operator==(const TSoftClassPtr& Rhs) const
 	{
 		return SoftObjectPtr == Rhs.SoftObjectPtr;
 	}
@@ -815,7 +815,7 @@ public:
 	 *
 	 * @param Other soft pointer to compare to
 	 */
-	FORCEINLINE bool operator!=(const TSoftClassPtr& Rhs) const
+	[[nodiscard]] FORCEINLINE bool operator!=(const TSoftClassPtr& Rhs) const
 	{
 		return SoftObjectPtr != Rhs.SoftObjectPtr;
 	}
@@ -826,7 +826,7 @@ public:
 	 *
 	 * @return nullptr if this object is gone or the soft pointer was null, otherwise a valid UClass pointer
 	 */
-	FORCEINLINE UClass* Get() const
+	[[nodiscard]] FORCEINLINE UClass* Get() const
 	{
 		UClass* Class = dynamic_cast<UClass*>(SoftObjectPtr.Get());
 		if (!Class || !Class->IsChildOf(TClass::StaticClass()))
@@ -837,13 +837,13 @@ public:
 	}
 
 	/** Dereference the soft pointer */
-	FORCEINLINE UClass& operator*() const
+	[[nodiscard]] FORCEINLINE UClass& operator*() const
 	{
 		return *Get();
 	}
 
 	/** Dereference the soft pointer */
-	FORCEINLINE UClass* operator->() const
+	[[nodiscard]] FORCEINLINE UClass* operator->() const
 	{
 		return Get();
 	}
@@ -853,7 +853,7 @@ public:
 	 *
 	 * @return true if Get() would return a valid non-null pointer
 	 */
-	FORCEINLINE bool IsValid() const
+	[[nodiscard]] FORCEINLINE bool IsValid() const
 	{
 		// This also does the UClass type check
 		return Get() != nullptr;
@@ -864,7 +864,7 @@ public:
 	 *
 	 * @return true if this does not point to a real object, but could possibly
 	 */
-	FORCEINLINE bool IsPending() const
+	[[nodiscard]] FORCEINLINE bool IsPending() const
 	{
 		return SoftObjectPtr.IsPending();
 	}
@@ -874,55 +874,55 @@ public:
 	 *
 	 * @return true if this is explicitly pointing to no object
 	 */
-	FORCEINLINE bool IsNull() const
+	[[nodiscard]] FORCEINLINE bool IsNull() const
 	{
 		return SoftObjectPtr.IsNull();
 	}
 
 	/** Returns the StringObjectPath that is wrapped by this SoftObjectPtr */
-	FORCEINLINE const FSoftObjectPath& GetUniqueID() const
+	[[nodiscard]] FORCEINLINE const FSoftObjectPath& GetUniqueID() const
 	{
 		return SoftObjectPtr.GetUniqueID();
 	}
 
 	/** Returns the StringObjectPath that is wrapped by this SoftObjectPtr */
-	FORCEINLINE const FSoftObjectPath& ToSoftObjectPath() const
+	[[nodiscard]] FORCEINLINE const FSoftObjectPath& ToSoftObjectPath() const
 	{
 		return SoftObjectPtr.GetUniqueID();
 	}
 
 	/** Returns string representation of reference, in form /package/path.assetname  */
-	FORCEINLINE FString ToString() const
+	[[nodiscard]] FORCEINLINE FString ToString() const
 	{
 		return ToSoftObjectPath().ToString();
 	}
 
 	/** Returns /package/path string, leaving off the asset name */
-	FORCEINLINE FString GetLongPackageName() const
+	[[nodiscard]] FORCEINLINE FString GetLongPackageName() const
 	{
 		return ToSoftObjectPath().GetLongPackageName();
 	}
 	
 	/** Returns /package/path name, leaving off the asset name */
-	FORCEINLINE FName GetLongPackageFName() const
+	[[nodiscard]] FORCEINLINE FName GetLongPackageFName() const
 	{
 		return ToSoftObjectPath().GetLongPackageFName();
 	}
 
 	/** Returns assetname string, leaving off the /package/path part */
-	FORCEINLINE FString GetAssetName() const
+	[[nodiscard]] FORCEINLINE FString GetAssetName() const
 	{
 		return ToSoftObjectPath().GetAssetName();
 	}
 
 	/** Dereference soft pointer to see if it points somewhere valid */
-	FORCEINLINE explicit operator bool() const
+	[[nodiscard]] FORCEINLINE explicit operator bool() const
 	{
 		return IsValid();
 	}
 
 	/** Hash function */
-	FORCEINLINE uint32 GetPtrTypeHash() const
+	[[nodiscard]] FORCEINLINE uint32 GetPtrTypeHash() const
 	{
 		return GetTypeHash(static_cast<const TPersistentObjectPtr<FSoftObjectPath>&>(SoftObjectPtr));
 	}
@@ -972,14 +972,14 @@ struct TCallTraits<TSoftClassPtr<T>> : public TCallTraitsBase<TSoftClassPtr<T>>
 
 /** Utility to create a TSoftObjectPtr without specifying the type */
 template <class T>
-TSoftClassPtr<std::remove_cv_t<T>> MakeSoftClassPtr(T* Object)
+[[nodiscard]] TSoftClassPtr<std::remove_cv_t<T>> MakeSoftClassPtr(T* Object)
 {
 	static_assert(std::is_base_of_v<UClass, T>, "Type must derive from UClass");
 	return TSoftClassPtr<std::remove_cv_t<T>>(Object);
 }
 
 template <class T>
-TSoftClassPtr<std::remove_cv_t<T>> MakeSoftClassPtr(TObjectPtr<T> Object)
+[[nodiscard]] TSoftClassPtr<std::remove_cv_t<T>> MakeSoftClassPtr(TObjectPtr<T> Object)
 {
 	static_assert(std::is_base_of_v<UClass, T>, "Type must derive from UClass");
 	return TSoftClassPtr<std::remove_cv_t<T>>(ToRawPtr(Object));
@@ -1014,13 +1014,13 @@ T* TSoftObjectPtr<T>::Get() const
 
 #if !PLATFORM_COMPILER_HAS_GENERATED_COMPARISON_OPERATORS
 template<class TClass>
-FORCEINLINE bool operator==(const UClass* Lhs, const TSoftClassPtr<TClass>& Rhs)
+[[nodiscard]] FORCEINLINE bool operator==(const UClass* Lhs, const TSoftClassPtr<TClass>& Rhs)
 {
 	return Rhs == Lhs;
 }
 
 template<class TClass>
-FORCEINLINE bool operator!=(const UClass* Lhs, const TSoftClassPtr<TClass>& Rhs)
+[[nodiscard]] FORCEINLINE bool operator!=(const UClass* Lhs, const TSoftClassPtr<TClass>& Rhs)
 {
 	return Rhs != Lhs;
 }
@@ -1028,7 +1028,7 @@ FORCEINLINE bool operator!=(const UClass* Lhs, const TSoftClassPtr<TClass>& Rhs)
 
 /** Hash function */
 template<class TClass>
-FORCEINLINE uint32 GetTypeHash(const TSoftClassPtr<TClass>& Ptr)
+[[nodiscard]] FORCEINLINE uint32 GetTypeHash(const TSoftClassPtr<TClass>& Ptr)
 {
 	return Ptr.GetPtrTypeHash();
 }
