Fix [[nodiscard]] in unit tests
#rnx

diff --git a/Engine/Source/Runtime/CoreUObject/Tests/ObjectHandleTest.cpp b/Engine/Source/Runtime/CoreUObject/Tests/ObjectHandleTest.cpp
--- a/Engine/Source/Runtime/CoreUObject/Tests/ObjectHandleTest.cpp
+++ b/Engine/Source/Runtime/CoreUObject/Tests/ObjectHandleTest.cpp
@@ -45,7 +45,7 @@ protected:
 		FObjectRef TargetRef(FName(PackageName), NAME_None, NAME_None, UE::CoreUObject::Private::FObjectPathId(ObjectName));
 		UObject* ResolvedObject = TargetRef.Resolve();
 		FObjectPtr Ptr(ResolvedObject);
-		Ptr.Get();
+		(void)Ptr.Get();
 		TEST_TRUE(TEXT("expected not null"), ResolvedObject != nullptr);
 		ObjectRefMetrics.TestNumResolves(TEXT("NumResolves should be incremented by one after a resolve attempt"), 1);
 		ObjectRefMetrics.TestNumReads(TEXT("NumReads should be incremented by one after a resolve attempt"), 1, bExpectSubRefReads /*bAllowAdditionalReads*/);
@@ -672,10 +672,10 @@ TEST_CASE_METHOD(FObjectHandleTestBase, "CoreUObject::FObjectHandle::Type Safety
 	CHECK(ToRawPtr(TestLateResolveUnsafeInvalidObjectPtr) == nullptr);
 	{
 		SKIP(); // comment this out to run the tests below (they should assert as they will not be resolved)
-		ToRawPtrTArrayUnsafe(TestLateResolveSafeObjectPtrArray);
-		ToRawPtrTArrayUnsafe(TestLateResolveSafeConstObjectPtrArray);
-		ToRawPtrTArrayUnsafe(TestLateResolveUnsafeObjectPtrArray);
-		ToRawPtrTArrayUnsafe(TestLateResolveUnsafeConstObjectPtrArray);
+		(void)ToRawPtrTArrayUnsafe(TestLateResolveSafeObjectPtrArray);
+		(void)ToRawPtrTArrayUnsafe(TestLateResolveSafeConstObjectPtrArray);
+		(void)ToRawPtrTArrayUnsafe(TestLateResolveUnsafeObjectPtrArray);
+		(void)ToRawPtrTArrayUnsafe(TestLateResolveUnsafeConstObjectPtrArray);
 	}
 #endif
 }
diff --git a/Engine/Source/Runtime/CoreUObject/Tests/ObjectPtrTest.cpp b/Engine/Source/Runtime/CoreUObject/Tests/ObjectPtrTest.cpp
--- a/Engine/Source/Runtime/CoreUObject/Tests/ObjectPtrTest.cpp
+++ b/Engine/Source/Runtime/CoreUObject/Tests/ObjectPtrTest.cpp
@@ -294,7 +294,7 @@ TEST_CASE_METHOD(FObjectPtrTestBase, "CoreUObject::TObjectPtr::Hash Consistency"
 		ObjectRefMetrics.TestNumResolves(TEXT("Unexpected resolve count after hashing an FObjectPtr"), 0);
 		ObjectRefMetrics.TestNumFailedResolves(TEXT("Unexpected resolve failure after hashing an FObjectPtr"), 0);
 
-		TestPublicWrappedObjectPtr.Get();
+		(void)TestPublicWrappedObjectPtr.Get();
 
 		ObjectRefMetrics.TestNumResolves(TEXT("Unexpected resolve count after resolving an FObjectPtr"), 1);
 		ObjectRefMetrics.TestNumFailedResolves(TEXT("Unexpected resolve failure after resolving an FObjectPtr"), 0);
