Use the same name between outer array/set properties and their inner value property.

This makes the property naming consistent with UObject properties.

diff --git a/Engine/Source/Runtime/CoreUObject/Private/StructUtils/PropertyBag.cpp b/Engine/Source/Runtime/CoreUObject/Private/StructUtils/PropertyBag.cpp
--- a/Engine/Source/Runtime/CoreUObject/Private/StructUtils/PropertyBag.cpp
+++ b/Engine/Source/Runtime/CoreUObject/Private/StructUtils/PropertyBag.cpp
@@ -263,7 +263,6 @@ namespace UE::StructUtils::Private
 			// support for nested containers, i.e. : TArray<TArray<float>>
 			FFieldVariant PropertyOwner = PropertyScope;
 			FProperty** ValuePropertyPtr = &Prop;
-			FString ParentName;
 
 			// Create the container list
 			for (EPropertyBagContainerType BagContainerType : Desc.ContainerTypes)
@@ -273,23 +272,19 @@ namespace UE::StructUtils::Private
 				case EPropertyBagContainerType::Array:
 					{
 						// create an array property as a container for the tail
-						const FString PropName = ParentName.IsEmpty() ? Desc.Name.ToString() : ParentName + TEXT("_InnerArray");
-						FArrayProperty* ArrayProperty = new FArrayProperty(PropertyOwner, FName(PropName), RF_Public);
+						FArrayProperty* ArrayProperty = new FArrayProperty(PropertyOwner, Desc.Name, RF_Public);
 						*ValuePropertyPtr = ArrayProperty;
 						ValuePropertyPtr = &ArrayProperty->Inner;
 						PropertyOwner = ArrayProperty;
-						ParentName = PropName;
 						break;
 					}
 				case EPropertyBagContainerType::Set:
 					{
 						// create a set property as a container for the tail
-						const FString PropName = ParentName.IsEmpty() ? Desc.Name.ToString() : ParentName + TEXT("_InnerSet");
-						FSetProperty* SetProperty = new FSetProperty(PropertyOwner, FName(PropName), RF_Public);
+						FSetProperty* SetProperty = new FSetProperty(PropertyOwner, Desc.Name, RF_Public);
 						*ValuePropertyPtr = SetProperty;
 						ValuePropertyPtr = &SetProperty->ElementProp;
 						PropertyOwner = SetProperty;
-						ParentName = PropName;
 						break;
 					}
 				default:
@@ -300,7 +295,6 @@ namespace UE::StructUtils::Private
 
 			// finally create the tail type
 			FPropertyBagPropertyDesc InnerDesc = Desc;
-			InnerDesc.Name = FName(ParentName + TEXT("_InnerType"));
 			InnerDesc.ContainerTypes.Reset();
 			*ValuePropertyPtr = CreatePropertyFromDesc(InnerDesc, PropertyOwner);
 				
