Made data members private in FUObjectItem

diff --git a/Engine/Plugins/Experimental/EditorDataStorage/Source/TedsCore/Private/TypedElementDatabaseCompatibility.cpp b/Engine/Plugins/Experimental/EditorDataStorage/Source/TedsCore/Private/TypedElementDatabaseCompatibility.cpp
--- a/Engine/Plugins/Experimental/EditorDataStorage/Source/TedsCore/Private/TypedElementDatabaseCompatibility.cpp
+++ b/Engine/Plugins/Experimental/EditorDataStorage/Source/TedsCore/Private/TypedElementDatabaseCompatibility.cpp
@@ -1113,13 +1113,13 @@ void UEditorDataStorageCompatibility::OnPostGcUnreachableAnalysis()
 				[this](RowHandle Row, const FTypedElementUObjectIdColumn& ObjectId)
 				{
 					FUObjectItem* Description = GUObjectArray.IndexToObject(ObjectId.Id);
-					if (ensureMsgf(Description && Description->SerialNumber == ObjectId.SerialNumber,
+					if (ensureMsgf(Description && Description->GetSerialNumber() == ObjectId.SerialNumber,
 						TEXT("The UObject found in TEDS no longer exists. TEDS was likely not informed in an earlier GC pass.")))
 						// Unable to provide additional information such as the UObject's name as the UObject will not be valid.
 					{
 						if (Description->HasAnyFlags(EInternalObjectFlags::Garbage | EInternalObjectFlags::Unreachable))
 						{
-							if (UObject* Object = static_cast<UObject*>(Description->Object)) // No need to delete if this isn't a full UObject.
+							if (UObject* Object = static_cast<UObject*>(Description->GetObject())) // No need to delete if this isn't a full UObject.
 							{
 								QueuedCommands.AddCommand(FRemoveCompatibleUObject{ .Object = Object, .ObjectRow = Row });
 							}
@@ -1137,7 +1137,7 @@ void UEditorDataStorageCompatibility::OnPostGcUnreachableAnalysis()
 				[&DeletedObjects](RowHandle Row, const FTypedElementUObjectIdColumn& ObjectId)
 				{
 					FUObjectItem* Description = GUObjectArray.IndexToObject(ObjectId.Id);
-					if (ensureMsgf(Description && Description->SerialNumber == ObjectId.SerialNumber,
+					if (ensureMsgf(Description && Description->GetSerialNumber() == ObjectId.SerialNumber,
 						TEXT("The UObject found in TEDS no longer exists. TEDS was likely not informed in an earlier GC pass.")))
 						// Unable to provide additional information such as the UObject's name as the UObject will not be valid.
 					{
@@ -1150,7 +1150,7 @@ void UEditorDataStorageCompatibility::OnPostGcUnreachableAnalysis()
 
 			for (TPair<FUObjectItem*, RowHandle>& ObjectItem : DeletedObjects)
 			{
-				if (UObject* Object = static_cast<UObject*>(ObjectItem.Key->Object)) // No need to delete if this isn't a full UObject.
+				if (UObject* Object = static_cast<UObject*>(ObjectItem.Key->GetObject())) // No need to delete if this isn't a full UObject.
 				{
 					RemoveCompatibleObjectExplicitTransactionable<false>(Object, ObjectItem.Value);
 				}
diff --git a/Engine/Plugins/MemoryUsageQueries/Source/MemoryUsageQueries/Private/MemoryUsageQueries.cpp b/Engine/Plugins/MemoryUsageQueries/Source/MemoryUsageQueries/Private/MemoryUsageQueries.cpp
--- a/Engine/Plugins/MemoryUsageQueries/Source/MemoryUsageQueries/Private/MemoryUsageQueries.cpp
+++ b/Engine/Plugins/MemoryUsageQueries/Source/MemoryUsageQueries/Private/MemoryUsageQueries.cpp
@@ -346,7 +346,7 @@ void FMemoryUsageReferenceProcessor::Init(const TArray<FName>& PackageNames)
 	for (FRawObjectIterator It(true); It; ++It)
 	{
 		FUObjectItem* ObjectItem = *It;
-		UObject* Object = (UObject*)ObjectItem->Object;
+		UObject* Object = (UObject*)ObjectItem->GetObject();
 
 		if (ObjectItem->IsUnreachable())
 		{
@@ -447,7 +447,7 @@ void FMemoryUsageReferenceProcessor::GetUnreachablePackages(TSet<FName>& OutUnre
 	{
 		if (ReachableFull[i] && !ReachableExcluded[i])
 		{
-			UObject* Obj = static_cast<UObject*>(GUObjectArray.IndexToObjectUnsafeForGC(i)->Object);
+			UObject* Obj = static_cast<UObject*>(GUObjectArray.IndexToObjectUnsafeForGC(i)->GetObject());
 			if (Obj && Obj->IsA<UPackage>())
 			{
 				OutUnreachablePackages.Add(Obj->GetFName());
diff --git a/Engine/Plugins/Runtime/RigVM/Source/RigVM/Private/RigVMCore/RigVMRegistry.cpp b/Engine/Plugins/Runtime/RigVM/Source/RigVM/Private/RigVMCore/RigVMRegistry.cpp
--- a/Engine/Plugins/Runtime/RigVM/Source/RigVM/Private/RigVMCore/RigVMRegistry.cpp
+++ b/Engine/Plugins/Runtime/RigVM/Source/RigVM/Private/RigVMCore/RigVMRegistry.cpp
@@ -81,7 +81,7 @@ void FRigVMRegistry_NoLock::AddReferencedObjects(FReferenceCollector& Collector)
 						{
 							if(const FUObjectItem* Item = GUObjectArray.IndexToObject(ObjectIndex))
 							{
-								if(Item->Object == Type.Type.CPPTypeObject)
+								if(Item->GetObject() == Type.Type.CPPTypeObject)
 								{
 									Collector.AddReferencedObject(Type.Type.CPPTypeObject);
 								}
diff --git a/Engine/Source/Editor/UnrealEd/Private/Cooker/CookGarbageCollect.cpp b/Engine/Source/Editor/UnrealEd/Private/Cooker/CookGarbageCollect.cpp
--- a/Engine/Source/Editor/UnrealEd/Private/Cooker/CookGarbageCollect.cpp
+++ b/Engine/Source/Editor/UnrealEd/Private/Cooker/CookGarbageCollect.cpp
@@ -288,7 +288,7 @@ void ConstructSoftGCPackageToObjectList(TArray<UObject*>& PackageToObjectListBuf
 			for (int32 ObjectIndex = 0; ObjectIndex < NumObjects && (FirstObjectIndex + ObjectIndex) < MaxNumberOfObjects; ++ObjectIndex)
 			{
 				FUObjectItem& ObjectItem = GUObjectArray.GetObjectItemArrayUnsafe()[FirstObjectIndex + ObjectIndex];
-				if (!ObjectItem.Object)
+				if (!ObjectItem.GetObject())
 				{
 					continue;
 				}
@@ -296,7 +296,7 @@ void ConstructSoftGCPackageToObjectList(TArray<UObject*>& PackageToObjectListBuf
 				{
 					continue;
 				}
-				UObject* Object = static_cast<UObject*>(ObjectItem.Object);
+				UObject* Object = static_cast<UObject*>(ObjectItem.GetObject());
 				if (!Object->HasAnyFlags(RF_Public))
 				{
 					continue;
diff --git a/Engine/Source/Editor/UnrealEd/Private/Kismet2/ReloadUtilities.cpp b/Engine/Source/Editor/UnrealEd/Private/Kismet2/ReloadUtilities.cpp
--- a/Engine/Source/Editor/UnrealEd/Private/Kismet2/ReloadUtilities.cpp
+++ b/Engine/Source/Editor/UnrealEd/Private/Kismet2/ReloadUtilities.cpp
@@ -1039,7 +1039,7 @@ void FReload::Reinstance()
 	TArray<UScriptStruct*> ScriptStructs;
 	for (FRawObjectIterator It; It; ++It)
 	{
-		if (UFunction* Function = Cast<UFunction>(static_cast<UObject*>(It->Object)))
+		if (UFunction* Function = Cast<UFunction>(static_cast<UObject*>(It->GetObject())))
 		{
 			if (FNativeFuncPtr NewFunction = FunctionRemap.FindRef(Function->GetNativeFunc()))
 			{
@@ -1047,7 +1047,7 @@ void FReload::Reinstance()
 				Function->SetNativeFunc(NewFunction);
 			}
 		} 
-		else if (UScriptStruct* ScriptStruct = Cast<UScriptStruct>(static_cast<UObject*>(It->Object)))
+		else if (UScriptStruct* ScriptStruct = Cast<UScriptStruct>(static_cast<UObject*>(It->GetObject())))
 		{
 			// IMPORTANT: The order here is very important.  There is a chance that GetCppStructOps will fail for 
 			// types that are in a package we don't care about.  So we have to filter on the package first.
diff --git a/Engine/Source/Runtime/CoreUObject/Private/Serialization/AsyncLoading2.cpp b/Engine/Source/Runtime/CoreUObject/Private/Serialization/AsyncLoading2.cpp
--- a/Engine/Source/Runtime/CoreUObject/Private/Serialization/AsyncLoading2.cpp
+++ b/Engine/Source/Runtime/CoreUObject/Private/Serialization/AsyncLoading2.cpp
@@ -1108,7 +1108,7 @@ private:
 					FUObjectItem* ObjectItem = GUObjectArray.IndexToObject(ObjectIndex);
 					if (!ObjectItem->IsUnreachable())
 					{
-						UObject* Object = static_cast<UObject*>(ObjectItem->Object);
+						UObject* Object = static_cast<UObject*>(ObjectItem->GetObject());
 						checkObject(Object, !ObjectItem->HasAnyFlags(EInternalObjectFlags::LoaderImport));
 						ObjectItem->SetFlags(EInternalObjectFlags::LoaderImport);
 						ValidExportCount++;
@@ -1132,7 +1132,7 @@ private:
 			{
 				if (ObjectIndex >= 0)
 				{
-					UObject* Object = static_cast<UObject*>(GUObjectArray.IndexToObject(ObjectIndex)->Object);
+					UObject* Object = static_cast<UObject*>(GUObjectArray.IndexToObject(ObjectIndex)->GetObject());
 					checkObject(Object, Object->HasAnyInternalFlags(EInternalObjectFlags::LoaderImport));
 					Object->AtomicallyClearInternalFlags(EInternalObjectFlags::LoaderImport);
 				}
@@ -1305,7 +1305,7 @@ public:
 	{
 		if (HasPackage())
 		{
-			return static_cast<UPackage*>(GUObjectArray.IndexToObject(PackageObjectIndex)->Object);
+			return static_cast<UPackage*>(GUObjectArray.IndexToObject(PackageObjectIndex)->GetObject());
 		}
 		return nullptr;
 	}
@@ -1418,7 +1418,7 @@ public:
 		int32 ObjectIndex = PublicExportMap.Find(ExportHash);
 		if (ObjectIndex >= 0)
 		{
-			return static_cast<UObject*>(GUObjectArray.IndexToObject(ObjectIndex)->Object);
+			return static_cast<UObject*>(GUObjectArray.IndexToObject(ObjectIndex)->GetObject());
 		}
 		return nullptr;
 	}
@@ -1870,7 +1870,7 @@ public:
 		for (auto It = PackageRef.GetPublicExportObjectIndices(); It; ++It)
 		{
 			const int32 ObjectIndex = *It;
-			UObject* Object = static_cast<UObject*>(GUObjectArray.IndexToObject(ObjectIndex)->Object);
+			UObject* Object = static_cast<UObject*>(GUObjectArray.IndexToObject(ObjectIndex)->GetObject());
 			ensureMsgf(!Object->HasAnyInternalFlags(EInternalObjectFlags::LoaderImport) || GUObjectArray.IsDisregardForGC(Object),
 					TEXT("FGlobalImportStore::VerifyPackageForRemoval: The loaded public export object '%s' with flags (ObjectFlags=%x, InternalObjectFlags=%x) and id %s is probably still referenced by the loader."),
 					*Object->GetFullName(),
@@ -9184,7 +9184,7 @@ static void VerifyObjectLoadFlagsWhenFinishedLoading()
 		[AsyncFlags,LoadIntermediateFlags](int32 ObjectIndex)
 	{
 		FUObjectItem* ObjectItem = &GUObjectArray.GetObjectItemArrayUnsafe()[ObjectIndex];
-		if (UObject* Obj = static_cast<UObject*>(ObjectItem->Object))
+		if (UObject* Obj = static_cast<UObject*>(ObjectItem->GetObject()))
 		{
 			const EInternalObjectFlags InternalFlags = Obj->GetInternalFlags();
 			const EObjectFlags Flags = Obj->GetFlags();
@@ -9244,7 +9244,7 @@ void FAsyncLoadingThread2::CollectUnreachableObjects(
 		ParallelFor(TEXT("CollectUnreachableObjectsTask"), UnreachableObjectItems.Num(), 2048,
 			[&UnreachableObjectItems, &OutUnreachableObjects](int32 Index)
 		{
-			UObject* Object = static_cast<UObject*>(UnreachableObjectItems[Index]->Object);
+			UObject* Object = static_cast<UObject*>(UnreachableObjectItems[Index]->GetObject());
 
 			FUnreachableObject& Item = OutUnreachableObjects[Index];
 			Item.ObjectIndex = GUObjectArray.ObjectToIndex(Object);
@@ -9278,7 +9278,7 @@ void FAsyncLoadingThread2::CollectUnreachableObjects(
 		ParallelFor(TEXT("VerifyUnreachableObjectsDebugTask"), UnreachableObjectItems.Num(), 512,
 			[this, &UnreachableObjectItems](int32 Index)
 		{
-			UObject* Object = static_cast<UObject*>(UnreachableObjectItems[Index]->Object);
+			UObject* Object = static_cast<UObject*>(UnreachableObjectItems[Index]->GetObject());
 			UObject* Outer = UE::CoreUObject::Private::FObjectHandleUtils::GetNonAccessTrackedOuterNoResolve(Object);
 			if (!Outer)
 			{
diff --git a/Engine/Source/Runtime/CoreUObject/Private/Serialization/AsyncPackageLoader.cpp b/Engine/Source/Runtime/CoreUObject/Private/Serialization/AsyncPackageLoader.cpp
--- a/Engine/Source/Runtime/CoreUObject/Private/Serialization/AsyncPackageLoader.cpp
+++ b/Engine/Source/Runtime/CoreUObject/Private/Serialization/AsyncPackageLoader.cpp
@@ -737,7 +737,7 @@ void CancelAsyncLoading()
 	for (int32 ObjectIndex = 0; ObjectIndex < GUObjectArray.GetObjectArrayNum(); ++ObjectIndex)
 	{
 		FUObjectItem* ObjectItem = &GUObjectArray.GetObjectItemArrayUnsafe()[ObjectIndex];
-		if (UObject* Obj = static_cast<UObject*>(ObjectItem->Object))
+		if (UObject* Obj = static_cast<UObject*>(ObjectItem->GetObject()))
 		{
 			check(!Obj->HasAnyInternalFlags(AsyncFlags));
 		}
diff --git a/Engine/Source/Runtime/CoreUObject/Private/UObject/Class.cpp b/Engine/Source/Runtime/CoreUObject/Private/UObject/Class.cpp
--- a/Engine/Source/Runtime/CoreUObject/Private/UObject/Class.cpp
+++ b/Engine/Source/Runtime/CoreUObject/Private/UObject/Class.cpp
@@ -6674,7 +6674,7 @@ bool UClass::HotReloadPrivateStaticClass(
 		int32 CountClass = 0;
 		for ( FRawObjectIterator It; It; ++It )
 		{
-			UObject* Target = static_cast<UObject*>(It->Object);
+			UObject* Target = static_cast<UObject*>(It->GetObject());
 			if (OldVTable == *(void**)Target)
 			{
 				*(void**)Target = NewVTable;
@@ -6897,7 +6897,7 @@ void UClass::AssembleReferenceTokenStreams()
 	// not taken into account for garbage collection but have instances that are.
 	for (FRawObjectIterator It(false); It; ++It) // GetDefaultObject can create a new class, that need to be handled as well, so we cannot use TObjectIterator
 	{
-		if (UClass* Class = Cast<UClass>((UObject*)(It->Object)))
+		if (UClass* Class = Cast<UClass>((UObject*)(It->GetObject())))
 		{
 			// Force the default object to be created (except when we're in the middle of exit purge -
 			// this may happen if we exited PreInit early because of error).
diff --git a/Engine/Source/Runtime/CoreUObject/Private/UObject/GarbageCollection.cpp b/Engine/Source/Runtime/CoreUObject/Private/UObject/GarbageCollection.cpp
--- a/Engine/Source/Runtime/CoreUObject/Private/UObject/GarbageCollection.cpp
+++ b/Engine/Source/Runtime/CoreUObject/Private/UObject/GarbageCollection.cpp
@@ -739,7 +739,7 @@ void OnDisregardForGCSetDisabled(int32 NumObjects)
 		{
 			int32 ObjectIndex = ThreadState.Index++;
 			FUObjectItem* RootItem = &GUObjectArray.GetObjectItemArrayUnsafe()[ObjectIndex];	
-			if (RootItem->Object)
+			if (RootItem->GetObject())
 			{
 				FGCFlags::SetReachable_ForGC(RootItem);
 				if (RootItem->HasAnyFlags(EInternalObjectFlags_RootFlags))
@@ -813,7 +813,7 @@ public:
 				FUObjectItem* ObjectItem = UnreachableObject.ObjectItem;
 				check(ObjectItem->IsUnreachable());
 
-				UObject* Object = (UObject*)ObjectItem->Object;
+				UObject* Object = (UObject*)ObjectItem->GetObject();
 				check(Object->HasAllFlags(RF_FinishDestroyed | RF_BeginDestroyed));
 
 				GUObjectArray.FreeUObjectIndex(Object);
@@ -1025,7 +1025,7 @@ static bool MarkClusterMutableObjectsAsReachable(FUObjectCluster& Cluster, Conta
 					if (FGCFlags::MarkAsReachableInterlocked_ForGC(ReferencedMutableObjectItem))
 					{
 						// Needs doing because this is either a normal unclustered object (clustered objects are never unreachable) or a cluster root
-						ObjectsToSerialize.Add(static_cast<UObject*>(ReferencedMutableObjectItem->Object));
+						ObjectsToSerialize.Add(static_cast<UObject*>(ReferencedMutableObjectItem->GetObject()));
 
 						// So is this a cluster root maybe?
 						if (ReferencedMutableObjectItem->GetOwnerIndex() < 0)
@@ -1126,7 +1126,7 @@ static FORCENOINLINE void MarkReferencedClustersAsReachable(int32 ClusterIndex,
 		for (int32 ClusterObjectIndex : Cluster.Objects)
 		{
 			FUObjectItem* ClusterObjectItem = GUObjectArray.IndexToObjectUnsafeForGC(ClusterObjectIndex);
-			UObject* ClusterObject = static_cast<UObject*>(ClusterObjectItem->Object);
+			UObject* ClusterObject = static_cast<UObject*>(ClusterObjectItem->GetObject());
 			ObjectsToSerialize.Add(ClusterObject);
 		}
 		Cluster.bNeedsDissolving = true;
@@ -4203,7 +4203,7 @@ public:
 				while (ThreadState.Index <= ThreadState.LastIndex)
 				{
 					FUObjectItem* RootItem = &GUObjectArray.GetObjectItemArrayUnsafe()[RootsArray[ThreadState.Index++]];
-					UObject* Object = static_cast<UObject*>(RootItem->Object);
+					UObject* Object = static_cast<UObject*>(RootItem->GetObject());
 
 					// IsValidLowLevel is extremely slow in this loop so only do it in debug
 					checkSlow(Object->IsValidLowLevel());					
@@ -4236,7 +4236,7 @@ public:
 				while (ThreadState.Index <= ThreadState.LastIndex)
 				{
 					FUObjectItem* ObjectItem = &GUObjectArray.GetObjectItemArrayUnsafe()[ThreadState.Index++];
-					UObject* Object = static_cast<UObject*>(ObjectItem->Object);
+					UObject* Object = static_cast<UObject*>(ObjectItem->GetObject());
 					if (Object &&
 						!ObjectItem->HasAnyFlags(EInternalObjectFlags_RootFlags) && // It may be counter intuitive to reject roots but these are tracked with GRoots and have already been marked and added
 						!(bWithGarbageElimination && ObjectItem->IsGarbage()) && Object->HasAnyFlags(KeepFlags)) // Garbage elimination works regardless of KeepFlags
@@ -4278,7 +4278,7 @@ public:
 				for (; ThreadState.Index <= ThreadState.LastIndex; ++ThreadState.Index)
 				{
 					FUObjectItem* ObjectItem = &GUObjectArray.GetObjectItemArrayUnsafe()[ThreadState.Index];
-					if (ObjectItem->Object)
+					if (ObjectItem->GetObject())
 					{
 						FGCFlags::SetMaybeUnreachable_ForGC(ObjectItem);
 					}
@@ -4750,7 +4750,7 @@ bool IncrementalDestroyGarbage(bool bUseTimeLimit, double TimeLimit)
 				check(!FGCFlags::IsReachable_ForGC(ObjectItem) && FGCFlags::IsMaybeUnreachable_ForGC(ObjectItem));
 				if (ObjectItem->IsUnreachable())
 				{
-					UObject* Object = static_cast<UObject*>(ObjectItem->Object);
+					UObject* Object = static_cast<UObject*>(ObjectItem->GetObject());
 					// Object should always have had BeginDestroy called on it and never already be destroyed
 					check( Object->HasAnyFlags( RF_BeginDestroyed ) && !Object->HasAnyFlags( RF_FinishDestroyed ) );
 
@@ -6040,7 +6040,7 @@ bool UnhashUnreachableObjects(bool bUseTimeLimit, double TimeLimit)
 
 			FUObjectItem* ObjectItem = GUnreachableObjects[GUnrechableObjectIndex++].ObjectItem;
 			{
-				UObject* Object = static_cast<UObject*>(ObjectItem->Object);
+				UObject* Object = static_cast<UObject*>(ObjectItem->GetObject());
 				FScopedCBDProfile Profile(Object);
 				// Begin the object's asynchronous destruction.
 #if UE_WITH_REMOTE_OBJECT_HANDLE
@@ -6200,7 +6200,7 @@ void PurgeAllUObjectsOnExit()
 
 		FUObjectItem* ObjItem = *It;
 		checkSlow(ObjItem);
-		UObject* Obj = static_cast<UObject*>(ObjItem->Object);
+		UObject* Obj = static_cast<UObject*>(ObjItem->GetObject());
 		if (Obj)
 		{
 			// Skip Structures, properties, etc.. They could be still necessary while GC.
@@ -6309,7 +6309,7 @@ FORCEINLINE static void MarkObjectItemAsReachable(FUObjectItem* ObjectItem)
 	else
 #endif
 	{
-		checkf(GIsIncrementalReachabilityPending, TEXT("%s is marked as MaybeUnreachable but Incremental Reachability Analysis is not in progress"), *static_cast<UObject*>(ObjectItem->Object)->GetFullName());
+		checkf(GIsIncrementalReachabilityPending, TEXT("%s is marked as MaybeUnreachable but Incremental Reachability Analysis is not in progress"), *static_cast<UObject*>(ObjectItem->GetObject())->GetFullName());
 	}
 	if (FGCFlags::MarkAsReachableInterlocked_ForGC(ObjectItem))
 	{
@@ -6317,7 +6317,7 @@ FORCEINLINE static void MarkObjectItemAsReachable(FUObjectItem* ObjectItem)
 		{
 			// This object became reachable so add it to a list of new objects to process in the next iteration of incremental GC because
 			// we need to mark objects it's referencing as reachable too
-			GReachableObjects.Push(static_cast<UObject*>(ObjectItem->Object));
+			GReachableObjects.Push(static_cast<UObject*>(ObjectItem->GetObject()));
 		}
 		else
 		{
diff --git a/Engine/Source/Runtime/CoreUObject/Private/UObject/GarbageCollectionVerification.cpp b/Engine/Source/Runtime/CoreUObject/Private/UObject/GarbageCollectionVerification.cpp
--- a/Engine/Source/Runtime/CoreUObject/Private/UObject/GarbageCollectionVerification.cpp
+++ b/Engine/Source/Runtime/CoreUObject/Private/UObject/GarbageCollectionVerification.cpp
@@ -141,9 +141,9 @@ void VerifyGCAssumptions()
 		for (int32 ObjectIndex = 0; ObjectIndex < NumObjects && (FirstObjectIndex + ObjectIndex) < MaxNumberOfObjects; ++ObjectIndex)
 		{
 			FUObjectItem& ObjectItem = GUObjectArray.GetObjectItemArrayUnsafe()[FirstObjectIndex + ObjectIndex];
-			if (ObjectItem.Object && ObjectItem.Object != FGCObject::GGCObjectReferencer)
+			if (ObjectItem.GetObject() && ObjectItem.GetObject() != FGCObject::GGCObjectReferencer)
 			{
-				ObjectsToSerialize.Add(static_cast<UObject*>(ObjectItem.Object));
+				ObjectsToSerialize.Add(static_cast<UObject*>(ObjectItem.GetObject()));
 			}
 		}
 		
@@ -184,8 +184,8 @@ public:
 		Cluster = GUObjectClusters.GetObjectCluster(CurrentObject);
 		check(Cluster);
 		FUObjectItem* RootItem = GUObjectArray.IndexToObject(Cluster->RootIndex);
-		check(RootItem && RootItem->Object);
-		ClusterRootObject = static_cast<UObject*>(RootItem->Object);
+		check(RootItem && RootItem->GetObject());
+		ClusterRootObject = static_cast<UObject*>(RootItem->GetObject());
 	}
 	/**
 	* Handles UObject reference from the token stream. Performance is critical here so we're FORCEINLINING this function.
@@ -242,8 +242,8 @@ public:
 					const FUObjectItem* ClusterRootObjectItem = GUObjectArray.ObjectToObjectItem(ClusterRootObject);
 					const int32 OtherClusterRootIndex = GUObjectArray.ObjectToIndex(Object);
 					const FUObjectItem* OtherClusterRootItem = GUObjectArray.IndexToObjectUnsafeForGC(OtherClusterRootIndex);
-					check(OtherClusterRootItem && OtherClusterRootItem->Object);
-					UObject* OtherClusterRootObject = static_cast<UObject*>(OtherClusterRootItem->Object);
+					check(OtherClusterRootItem && OtherClusterRootItem->GetObject());
+					UObject* OtherClusterRootObject = static_cast<UObject*>(OtherClusterRootItem->GetObject());
 					UE_CLOG(
 						OtherClusterRootIndex != Cluster->RootIndex  // Same cluster is legal 
 					&&	!Cluster->ReferencedClusters.Contains(OtherClusterRootIndex)  // cluster-cluster reference is legal
@@ -284,8 +284,8 @@ public:
 				const int32 OtherClusterRootIndex = ObjectItem->GetOwnerIndex();
 				check(OtherClusterRootIndex > 0);
 				const FUObjectItem* OtherClusterRootItem = GUObjectArray.IndexToObjectUnsafeForGC(OtherClusterRootIndex);
-				check(OtherClusterRootItem && OtherClusterRootItem->Object);
-				UObject* OtherClusterRootObject = static_cast<UObject*>(OtherClusterRootItem->Object);
+				check(OtherClusterRootItem && OtherClusterRootItem->GetObject());
+				UObject* OtherClusterRootObject = static_cast<UObject*>(OtherClusterRootItem->GetObject());
 				UE_CLOG(
 						OtherClusterRootIndex != Cluster->RootIndex  // Same cluster is legal 
 					&&	!Cluster->ReferencedClusters.Contains(OtherClusterRootIndex)  // Cluster-cluster reference 
@@ -362,15 +362,15 @@ void VerifyClustersAssumptions()
 				{
 					FUObjectItem* RootItem = GUObjectArray.IndexToObject(Cluster.RootIndex);
 					check(RootItem);
-					check(RootItem->Object);
-					ObjectsToSerialize.Add(static_cast<UObject*>(RootItem->Object));
+					check(RootItem->GetObject());
+					ObjectsToSerialize.Add(static_cast<UObject*>(RootItem->GetObject()));
 				}
 				for (int32 ObjectIndex : Cluster.Objects)
 				{
 					FUObjectItem* ObjectItem = GUObjectArray.IndexToObject(ObjectIndex);
 					check(ObjectItem);
-					check(ObjectItem->Object);
-					ObjectsToSerialize.Add(static_cast<UObject*>(ObjectItem->Object));
+					check(ObjectItem->GetObject());
+					ObjectsToSerialize.Add(static_cast<UObject*>(ObjectItem->GetObject()));
 				}
 
 				UE::GC::FWorkerContext Context;
@@ -402,9 +402,9 @@ void VerifyObjectFlags()
 		for (int32 ObjectIndex = 0; ObjectIndex < NumObjects && (FirstObjectIndex + ObjectIndex) < MaxNumberOfObjects; ++ObjectIndex)
 		{
 			FUObjectItem& ObjectItem = GUObjectArray.GetObjectItemArrayUnsafe()[FirstObjectIndex + ObjectIndex];
-			if (ObjectItem.Object)
+			if (ObjectItem.GetObject())
 			{
-				UObject* Object = (UObject*)ObjectItem.Object;
+				UObject* Object = (UObject*)ObjectItem.GetObject();
 				bool bHasObjectFlag = Object->HasAnyFlags(RF_MirroredGarbage);
 				bool bHasInternalFlag = ObjectItem.HasAnyFlags(EInternalObjectFlags::Garbage);
 				if (bHasObjectFlag != bHasInternalFlag)
@@ -536,9 +536,9 @@ void VerifyNoUnreachableObjects(int32 NumUnreachable)
 		for (int32 ObjectIndex = 0; ObjectIndex < NumObjects && (FirstObjectIndex + ObjectIndex) < GUObjectArray.GetObjectArrayNum(); ++ObjectIndex)
 		{
 			FUObjectItem& ObjectItem = GUObjectArray.GetObjectItemArrayUnsafe()[FirstObjectIndex + ObjectIndex];
-			if (ObjectItem.Object)
+			if (ObjectItem.GetObject())
 			{
-				UObject* Object = static_cast<UObject*>(ObjectItem.Object);
+				UObject* Object = static_cast<UObject*>(ObjectItem.GetObject());
 				if (!ObjectItem.HasAnyFlags(EInternalObjectFlags::Unreachable))
 				{
 					if (FGCFlags::IsMaybeUnreachable_ForGC(&ObjectItem))
diff --git a/Engine/Source/Runtime/CoreUObject/Private/UObject/Obj.cpp b/Engine/Source/Runtime/CoreUObject/Private/UObject/Obj.cpp
--- a/Engine/Source/Runtime/CoreUObject/Private/UObject/Obj.cpp
+++ b/Engine/Source/Runtime/CoreUObject/Private/UObject/Obj.cpp
@@ -3849,7 +3849,7 @@ static void StaticShutdownAfterError()
 
 		for( FRawObjectIterator It; It; ++It )
 		{
-			UObject* Object = static_cast<UObject*>(It->Object);
+			UObject* Object = static_cast<UObject*>(It->GetObject());
 			Object->ShutdownAfterError();
 		}
 	}
diff --git a/Engine/Source/Runtime/CoreUObject/Private/UObject/ReferenceChainSearch.cpp b/Engine/Source/Runtime/CoreUObject/Private/UObject/ReferenceChainSearch.cpp
--- a/Engine/Source/Runtime/CoreUObject/Private/UObject/ReferenceChainSearch.cpp
+++ b/Engine/Source/Runtime/CoreUObject/Private/UObject/ReferenceChainSearch.cpp
@@ -64,7 +64,7 @@ namespace UE::ReferenceChainSearch
 		}
 		ObjectType VertexToObject(FVertex Vertex) const
 		{
-			return static_cast<UObject*>(GUObjectArray.IndexToObject(Vertex)->Object);
+			return static_cast<UObject*>(GUObjectArray.IndexToObject(Vertex)->GetObject());
 		}
 
 		bool IsIn(ConstObjectType Outer, ConstObjectType PotentialInner) const
@@ -419,7 +419,7 @@ namespace UE::ReferenceChainSearch
 						for (int32 ObjectIndex = 0; ObjectIndex < NumObjects; ++ObjectIndex)
 						{
 							FUObjectItem& ObjectItem = GUObjectArray.GetObjectItemArrayUnsafe()[FirstObjectIndex + ObjectIndex];
-							UObject* Object = static_cast<UObject*>(ObjectItem.Object);
+							UObject* Object = static_cast<UObject*>(ObjectItem.GetObject());
 							if (Object && !ObjectItem.IsUnreachable() && !This->ShouldSkipReferencer(ThreadIndex, Object))
 							{
 								ObjectsToSerialize.Add(Object);
@@ -439,7 +439,7 @@ namespace UE::ReferenceChainSearch
 				for (FRawObjectIterator It; It; ++It)
 				{
 					FUObjectItem* ObjItem = *It;
-					UObject* Object = static_cast<UObject*>(ObjItem->Object);
+					UObject* Object = static_cast<UObject*>(ObjItem->GetObject());
 
 					// // We can't ask the iterator for only GC objects because that would skip the GC Object referencer
 					// if (bGCOnly && GUObjectArray.IsDisregardForGC(Object) && (Object != FGCObject::GGCObjectReferencer))
diff --git a/Engine/Source/Runtime/CoreUObject/Private/UObject/ReferencerFinder.cpp b/Engine/Source/Runtime/CoreUObject/Private/UObject/ReferencerFinder.cpp
--- a/Engine/Source/Runtime/CoreUObject/Private/UObject/ReferencerFinder.cpp
+++ b/Engine/Source/Runtime/CoreUObject/Private/UObject/ReferencerFinder.cpp
@@ -200,11 +200,11 @@ static TArray<UObject*> GetAllReferencersIncludingWeak(
 		{
 			FUObjectItem& ObjectItem = GUObjectArray.GetObjectItemArrayUnsafe()[Index];
 			// Skip any objects still being deserialized because it is not safe to access them until they are constructed correctly.
-			if (!ObjectItem.Object || ObjectItem.IsUnreachable() || ObjectItem.HasAnyFlags(EInternalObjectFlags::AsyncLoadingPhase1))
+			if (!ObjectItem.GetObject() || ObjectItem.IsUnreachable() || ObjectItem.HasAnyFlags(EInternalObjectFlags::AsyncLoadingPhase1))
 			{
 				continue;
 			}
-			UObject* PotentialReferencer = static_cast<UObject*>(ObjectItem.Object);
+			UObject* PotentialReferencer = static_cast<UObject*>(ObjectItem.GetObject());
 			
 			if(!IsValid(PotentialReferencer))
 			{
@@ -281,9 +281,9 @@ static TArray<UObject*> GetAllReferencersExcludingWeak(
 		for (int32 Index = FirstObjectIndex; Index < EndIndex; ++Index)
 		{
 			FUObjectItem& ObjectItem = GUObjectArray.GetObjectItemArrayUnsafe()[Index];
-			if (ObjectItem.Object && !ObjectItem.IsUnreachable())
+			if (ObjectItem.GetObject() && !ObjectItem.IsUnreachable())
 			{
-				UObject* PotentialReferencer = static_cast<UObject*>(ObjectItem.Object);
+				UObject* PotentialReferencer = static_cast<UObject*>(ObjectItem.GetObject());
 				if (ObjectsToIgnore && ObjectsToIgnore->Contains(PotentialReferencer))
 				{
 					continue;
diff --git a/Engine/Source/Runtime/CoreUObject/Private/UObject/RemoteObject.cpp b/Engine/Source/Runtime/CoreUObject/Private/UObject/RemoteObject.cpp
--- a/Engine/Source/Runtime/CoreUObject/Private/UObject/RemoteObject.cpp
+++ b/Engine/Source/Runtime/CoreUObject/Private/UObject/RemoteObject.cpp
@@ -84,7 +84,7 @@ public:
 		Stub->ResidentServerId = DestinationServerId;
 		FUObjectItem* ObjectItem = GUObjectArray.ObjectToObjectItem(Object);
 		checkf(ObjectItem, TEXT("Attempting to get a serial number for an object that does not exist in the global UObject array (it's possible GUObjectArray is not initialized yet, ObjectIndex=%d)"), GUObjectArray.ObjectToIndex(Object));
-		Stub->SerialNumber = ObjectItem->SerialNumber;
+		Stub->SerialNumber = ObjectItem->GetSerialNumber();
 		Stub->Name = Object->GetFName();
 
 		return Stub;
diff --git a/Engine/Source/Runtime/CoreUObject/Private/UObject/UObjectArchetype.cpp b/Engine/Source/Runtime/CoreUObject/Private/UObject/UObjectArchetype.cpp
--- a/Engine/Source/Runtime/CoreUObject/Private/UObject/UObjectArchetype.cpp
+++ b/Engine/Source/Runtime/CoreUObject/Private/UObject/UObjectArchetype.cpp
@@ -248,7 +248,7 @@ UObject* GetArchetypeImpl(const UObject* InObject, const FObjectArchetypeHelper:
 	{
 		FUObjectItem* ArchetypeItem = GUObjectArray.IndexToObject(ArchetypeIndex);
 		check(ArchetypeItem != nullptr);
-		Archetype = static_cast<UObject*>(ArchetypeItem->Object);
+		Archetype = static_cast<UObject*>(ArchetypeItem->GetObject());
 #if UE_VERIFY_CACHED_ARCHETYPE
 		UObject* ExpectedArchetype = GetArchetypeFromRequiredInfoImpl(InObject->GetClass(), InObject->GetOuter(), InObject->GetFName(), InObject->GetFlags(), bUseUpToDateClass, Policy);
 		if (ExpectedArchetype != Archetype)
diff --git a/Engine/Source/Runtime/CoreUObject/Private/UObject/UObjectClusters.cpp b/Engine/Source/Runtime/CoreUObject/Private/UObject/UObjectClusters.cpp
--- a/Engine/Source/Runtime/CoreUObject/Private/UObject/UObjectClusters.cpp
+++ b/Engine/Source/Runtime/CoreUObject/Private/UObject/UObjectClusters.cpp
@@ -199,7 +199,7 @@ void FUObjectClusterContainer::DissolveClusterAndMarkObjectsAsUnreachable(FUObje
 	}
 
 #if !UE_GCCLUSTER_VERBOSE_LOGGING
-	UObject* ClusterRootObject = static_cast<UObject*>(RootObjectItem->Object);
+	UObject* ClusterRootObject = static_cast<UObject*>(RootObjectItem->GetObject());
 #endif
 	ClusterRootObject->OnClusterMarkedAsPendingKill();
 
@@ -308,7 +308,7 @@ void DumpClusterToLog(const FUObjectCluster& Cluster, bool bHierarchy, bool bInd
 #endif
 
 	FUObjectItem* RootItem = GUObjectArray.IndexToObjectUnsafeForGC(Cluster.RootIndex);
-	UObject* RootObject = static_cast<UObject*>(RootItem->Object);
+	UObject* RootObject = static_cast<UObject*>(RootItem->GetObject());
 #if WITH_VERSE_VM || defined(__INTELLISENSE__)
 	FString ExtraDetail = FString::Printf(TEXT(", MutableCells: %d"), Cluster.MutableCells.Num());
 #else
@@ -324,7 +324,7 @@ void DumpClusterToLog(const FUObjectCluster& Cluster, bool bHierarchy, bool bInd
 			if (!bIndexOnly)
 			{
 				FUObjectItem* ObjectItem = GUObjectArray.IndexToObjectUnsafeForGC(ObjectIndex);
-				UObject* Object = static_cast<UObject*>(ObjectItem->Object);
+				UObject* Object = static_cast<UObject*>(ObjectItem->GetObject());
 				UE_LOG(LogObj, Display, TEXT("    [%.4d]: %s (Index: %d)"), Index++, *Object->GetFullName(), ObjectIndex);
 			}
 			else
@@ -340,7 +340,7 @@ void DumpClusterToLog(const FUObjectCluster& Cluster, bool bHierarchy, bool bInd
 				if (!bIndexOnly)
 				{
 					FUObjectItem* ClusterRootItem = GUObjectArray.IndexToObjectUnsafeForGC(ClusterRootIndex);
-					UObject* ClusterRootObject = static_cast<UObject*>(ClusterRootItem->Object);
+					UObject* ClusterRootObject = static_cast<UObject*>(ClusterRootItem->GetObject());
 					UE_LOG(LogObj, Display, TEXT("    -> %s (Index: %d)"), *ClusterRootObject->GetFullName(), ClusterRootIndex);
 				}
 				else
@@ -361,7 +361,7 @@ void DumpClusterToLog(const FUObjectCluster& Cluster, bool bHierarchy, bool bInd
 				if (!bIndexOnly)
 				{
 					FUObjectItem* ObjectItem = GUObjectArray.IndexToObjectUnsafeForGC(ObjectIndex);
-					UObject* Object = static_cast<UObject*>(ObjectItem->Object);
+					UObject* Object = static_cast<UObject*>(ObjectItem->GetObject());
 					UE_LOG(LogObj, Display, TEXT("    => %s (Index: %d)"), *Object->GetFullName(), ObjectIndex);
 				}
 				else
@@ -409,7 +409,7 @@ void ListClusters(const TArray<FString>& Args)
 		{
 			Algo::SortBy(AllClusters, [](FUObjectCluster* A)
 			{
-				return GUObjectArray.IndexToObject(A->RootIndex)->Object->GetFName();
+				return GUObjectArray.IndexToObject(A->RootIndex)->GetObject()->GetFName();
 			}, FNameLexicalLess());
 		}
 		else if (Arg == TEXT("SortByObjectCount"))
@@ -496,7 +496,7 @@ void FindStaleClusters(const TArray<FString>& Args)
 		{
 			TotalNumClusters++;
 
-			UObject* ClusterRootObject = static_cast<UObject*>(ObjectItem->Object);
+			UObject* ClusterRootObject = static_cast<UObject*>(ObjectItem->GetObject());
 			FReferenceChainSearch SearchRefs(ClusterRootObject, EReferenceChainSearchMode::ExternalOnly);
 			
 			if (SearchRefs.GetReferenceChains().Num() == 0)
@@ -515,13 +515,13 @@ void DumpRefsToCluster(FUObjectCluster* Cluster)
 {
 	FUObjectItem* RootItem = GUObjectArray.IndexToObject(Cluster->RootIndex);
 
-	UE_LOG(LogObj, Display, TEXT("Dumping references to objects in cluster %s"), *static_cast<UObject*>(RootItem->Object)->GetFullName());
+	UE_LOG(LogObj, Display, TEXT("Dumping references to objects in cluster %s"), *static_cast<UObject*>(RootItem->GetObject())->GetFullName());
 
 	bool bIsReferenced = false;
 	for (int32 ObjectIndex : Cluster->Objects)
 	{
 		FUObjectItem* ObjectItem = GUObjectArray.IndexToObject(ObjectIndex);
-		FReferenceChainSearch SearchRefs(static_cast<UObject*>(ObjectItem->Object), EReferenceChainSearchMode::ExternalOnly | EReferenceChainSearchMode::Shortest);
+		FReferenceChainSearch SearchRefs(static_cast<UObject*>(ObjectItem->GetObject()), EReferenceChainSearchMode::ExternalOnly | EReferenceChainSearchMode::Shortest);
 		if (SearchRefs.GetReferenceChains().Num())
 		{
 			bIsReferenced = true;
@@ -530,7 +530,7 @@ void DumpRefsToCluster(FUObjectCluster* Cluster)
 	}
 	if (!bIsReferenced)
 	{
-		UE_LOG(LogObj, Display, TEXT("Cluster %s is not currently referenced by anything."), *static_cast<UObject*>(RootItem->Object)->GetFullName());
+		UE_LOG(LogObj, Display, TEXT("Cluster %s is not currently referenced by anything."), *static_cast<UObject*>(RootItem->GetObject())->GetFullName());
 	}
 }
 
@@ -556,7 +556,7 @@ void DumpRefsToCluster(const TArray<FString>& Args)
 		FUObjectItem* RootItem = GUObjectArray.IndexToObject(RootIndex);
 		if (RootItem->HasAnyFlags(EInternalObjectFlags::ClusterRoot))
 		{
-			FUObjectCluster* Cluster = GUObjectClusters.GetObjectCluster(static_cast<UObject*>(RootItem->Object));
+			FUObjectCluster* Cluster = GUObjectClusters.GetObjectCluster(static_cast<UObject*>(RootItem->GetObject()));
 			DumpRefsToCluster(Cluster);
 		}
 	}
@@ -635,7 +635,7 @@ public:
 
 	UObject* GetClusterRoot()
 	{
-		return static_cast<UObject*>(GUObjectArray.IndexToObject(Cluster.RootIndex)->Object);
+		return static_cast<UObject*>(GUObjectArray.IndexToObject(Cluster.RootIndex)->GetObject());
 	}
 
 	/**
@@ -882,7 +882,7 @@ void UObjectBaseUtility::AddToCluster(UObjectBaseUtility* ClusterRootOrObjectFro
 		const int32 ClusterRootIndex = Cluster->RootIndex;
 		if (!bAddAsMutableObject)
 		{
-			UObject* ClusterRootObject = static_cast<UObject*>(GUObjectArray.IndexToObjectUnsafeForGC(Cluster->RootIndex)->Object);
+			UObject* ClusterRootObject = static_cast<UObject*>(GUObjectArray.IndexToObjectUnsafeForGC(Cluster->RootIndex)->GetObject());
 			FClusterReferenceProcessor Processor(ClusterRootIndex, *Cluster, ClusterRootObject->GetOutermost());
 			UObject* ThisObject = static_cast<UObject*>(this);
 			UObject** ThisObjPtr = &ThisObject;
diff --git a/Engine/Source/Runtime/CoreUObject/Private/UObject/UObjectGlobals.cpp b/Engine/Source/Runtime/CoreUObject/Private/UObject/UObjectGlobals.cpp
--- a/Engine/Source/Runtime/CoreUObject/Private/UObject/UObjectGlobals.cpp
+++ b/Engine/Source/Runtime/CoreUObject/Private/UObject/UObjectGlobals.cpp
@@ -5487,7 +5487,7 @@ FORCENOINLINE static void CollectFieldPathReference(FReferenceCollector& Collect
 	if (FUObjectItem* FieldOwnerItem = FGCInternals::GetResolvedOwner(FieldPath))
 	{
 		// Handle reference synchronously and update field path if reference was nulled out
-		UObject* Owner = static_cast<UObject*>(FieldOwnerItem->Object);		
+		UObject* Owner = static_cast<UObject*>(FieldOwnerItem->GetObject());
 		if (CollectStackReference(Collector, Property, Owner, Referencer))
 		{
 			FGCInternals::ClearCachedField(FieldPath);
diff --git a/Engine/Source/Runtime/CoreUObject/Private/UObject/UObjectHash.cpp b/Engine/Source/Runtime/CoreUObject/Private/UObject/UObjectHash.cpp
--- a/Engine/Source/Runtime/CoreUObject/Private/UObject/UObjectHash.cpp
+++ b/Engine/Source/Runtime/CoreUObject/Private/UObject/UObjectHash.cpp
@@ -767,7 +767,7 @@ UObject* StaticFindObjectFastInternalThreadSafe(FUObjectHashTables& ThreadHash,
 		for (TMultiMap<int32, uint32>::TConstKeyIterator HashIt(ThreadHash.HashOuter, Hash); HashIt; ++HashIt)
 		{
 			uint32 InternalIndex = HashIt.Value();
-			UObject* Object = static_cast<UObject*>(GUObjectArray.IndexToObject(InternalIndex)->Object);
+			UObject* Object = static_cast<UObject*>(GUObjectArray.IndexToObject(InternalIndex)->GetObject());
 			if
 				/* check that the name matches the name we're searching for */
 				((Object->GetFName() == ObjectName)
@@ -1905,7 +1905,7 @@ void LogHashStatisticsInternal(TMultiMap<int32, uint32>& Hash, FOutputDevice& Ar
 	int32 Count = 0;
 	for (TMultiMap<int32, uint32>::TConstKeyIterator HashIt(Hash, MaxBin); HashIt && Count < 30; ++HashIt)
 	{
-		UObject* Object = static_cast<UObject*>(GUObjectArray.IndexToObject(HashIt.Value())->Object);
+		UObject* Object = static_cast<UObject*>(GUObjectArray.IndexToObject(HashIt.Value())->GetObject());
 		Ar.Logf(TEXT("\tObject is %s (%s)"), *Object->GetName(), *Object->GetFullName());
 		Count++;
 	}
diff --git a/Engine/Source/Runtime/CoreUObject/Private/UObject/WeakObjectPtr.cpp b/Engine/Source/Runtime/CoreUObject/Private/UObject/WeakObjectPtr.cpp
--- a/Engine/Source/Runtime/CoreUObject/Private/UObject/WeakObjectPtr.cpp
+++ b/Engine/Source/Runtime/CoreUObject/Private/UObject/WeakObjectPtr.cpp
@@ -126,7 +126,7 @@ UObject* FWeakObjectPtr::GetEvenIfUnreachable() const
 	if (Internal_IsValid(true, true))
 	{
 		FUObjectItem* ObjectItem = GUObjectArray.IndexToObject(GetObjectIndex_Private(), true);
-		Result = static_cast<UObject*>(ObjectItem->Object);
+		Result = static_cast<UObject*>(ObjectItem->GetObject());
 	}
 	return Result;
 }
@@ -149,7 +149,7 @@ TStrongObjectPtr<UObject> FWeakObjectPtr::PinEvenIfUnreachable() const
 	if (Internal_IsValid(true, true))
 	{
 		FUObjectItem* ObjectItem = GUObjectArray.IndexToObject(GetObjectIndex_Private(), true);
-		Result = static_cast<UObject*>(ObjectItem->Object);
+		Result = static_cast<UObject*>(ObjectItem->GetObject());
 	}
 	return TStrongObjectPtr<UObject>(Result);
 }
@@ -158,7 +158,7 @@ TStrongObjectPtr<UObject> FWeakObjectPtr::Internal_Pin(bool bEvenIfGarbage) cons
 {
 	FGCScopeGuard GCScopeGuard;
 	FUObjectItem* const ObjectItem = Internal_GetObjectItem();
-	return TStrongObjectPtr<UObject>(((ObjectItem != nullptr) && GUObjectArray.IsValid(ObjectItem, bEvenIfGarbage)) ? (UObject*)ObjectItem->Object : nullptr);
+	return TStrongObjectPtr<UObject>(((ObjectItem != nullptr) && GUObjectArray.IsValid(ObjectItem, bEvenIfGarbage)) ? (UObject*)ObjectItem->GetObject() : nullptr);
 }
 
 class TStrongObjectPtr<UObject> FWeakObjectPtr::TryPin(bool& bOutPinValid, bool bEvenIfGarbage) const
@@ -183,7 +183,7 @@ bool FWeakObjectPtr::TryPinEvenIfUnreachable(class TStrongObjectPtr<UObject>& Ou
 	if (Internal_IsValid(true, true))
 	{
 		FUObjectItem* ObjectItem = GUObjectArray.IndexToObject(GetObjectIndex_Private(), true);
-		Result = static_cast<UObject*>(ObjectItem->Object);
+		Result = static_cast<UObject*>(ObjectItem->GetObject());
 	}
 	OutResult = TStrongObjectPtr<UObject>(Result);
 	return true;
@@ -198,7 +198,7 @@ TStrongObjectPtr<UObject> FWeakObjectPtr::Internal_TryPin(bool& bOutPinValid, bo
 	if (bOutPinValid)
 	{
 		FUObjectItem* const ObjectItem = Internal_GetObjectItem();
-		Result = TStrongObjectPtr<UObject>(((ObjectItem != nullptr) && GUObjectArray.IsValid(ObjectItem, bEvenIfGarbage)) ? (UObject*)ObjectItem->Object : nullptr);
+		Result = TStrongObjectPtr<UObject>(((ObjectItem != nullptr) && GUObjectArray.IsValid(ObjectItem, bEvenIfGarbage)) ? (UObject*)ObjectItem->GetObject() : nullptr);
 	}
 
 	return Result;
diff --git a/Engine/Source/Runtime/CoreUObject/Public/UObject/FastReferenceCollector.h b/Engine/Source/Runtime/CoreUObject/Public/UObject/FastReferenceCollector.h
--- a/Engine/Source/Runtime/CoreUObject/Public/UObject/FastReferenceCollector.h
+++ b/Engine/Source/Runtime/CoreUObject/Public/UObject/FastReferenceCollector.h
@@ -561,7 +561,7 @@ FORCEINLINE_DEBUGGABLE void VisitFieldPath(DispatcherType& Dispatcher, FFieldPat
 {
 	if (FUObjectItem* FieldOwnerItem = FGCInternals::GetResolvedOwner(FieldPath))
 	{
-		UObject* OwnerObject = static_cast<UObject*>(FieldOwnerItem->Object);
+		UObject* OwnerObject = static_cast<UObject*>(FieldOwnerItem->GetObject());
 		UObject* PreviousOwner = OwnerObject;
 		Dispatcher.HandleReferenceDirectly(Dispatcher.Context.GetReferencingObject(), OwnerObject, FMemberId(MemberIdx), Origin, true);
 								
diff --git a/Engine/Source/Runtime/CoreUObject/Public/UObject/GCObjectInfo.h b/Engine/Source/Runtime/CoreUObject/Public/UObject/GCObjectInfo.h
--- a/Engine/Source/Runtime/CoreUObject/Public/UObject/GCObjectInfo.h
+++ b/Engine/Source/Runtime/CoreUObject/Public/UObject/GCObjectInfo.h
@@ -33,7 +33,7 @@ public:
 		: Name(Obj->GetFName())
 		, Flags(Obj->GetFlags())
 		, InternalFlags(Obj->GetInternalFlags())
-		, ClusterRootIndex(GUObjectArray.ObjectToObjectItem(Obj)->ClusterRootIndex)
+		, ClusterRootIndex(GUObjectArray.ObjectToObjectItem(Obj)->GetOwnerIndex())
 		, bDisregardForGC(GUObjectArray.IsDisregardForGC(Obj))
 		, RefCount(GUObjectArray.ObjectToObjectItem(Obj)->GetRefCount())
 	{
diff --git a/Engine/Source/Runtime/CoreUObject/Public/UObject/UObjectArray.h b/Engine/Source/Runtime/CoreUObject/Public/UObject/UObjectArray.h
--- a/Engine/Source/Runtime/CoreUObject/Public/UObject/UObjectArray.h
+++ b/Engine/Source/Runtime/CoreUObject/Public/UObject/UObjectArray.h
@@ -47,12 +47,11 @@ struct
 	friend class FUObjectArray;
 	friend class UE::GC::Private::FGCFlags;
 
+private:
 	// Pointer to the allocated object
 	class UObjectBase* Object;
-private:
 	// Internal flags. These can only be changed via Set* and Clear* functions
 	int32 Flags;
-public:
 	// UObject Owner Cluster Index
 	int32 ClusterRootIndex;	
 	// Weak Object Pointer Serial number associated with the object
@@ -60,8 +59,9 @@ public:
 	// RefCount associated with the object preventing its destruction.
 	int32 RefCount;
 
+public:
 #if UE_WITH_REMOTE_OBJECT_HANDLE
-	// Globally unqiue id of this object
+	// Globally unique id of this object
 	FRemoteObjectId RemoteId;
 #endif
 
@@ -102,6 +102,11 @@ public:
 	FUObjectItem& operator=(FUObjectItem&&) = delete;
 	FUObjectItem& operator=(const FUObjectItem&) = delete;
 
+	FORCEINLINE class UObjectBase* GetObject() const
+	{
+		return Object;
+	}
+
 	FORCEINLINE void SetOwnerIndex(int32 OwnerIndex)
 	{
 		ClusterRootIndex = OwnerIndex;
@@ -504,7 +509,7 @@ public:
 		check(Result == NumElements);
 		++NumElements;
 		FPlatformMisc::MemoryBarrier();
-		check(Objects[Result].Object == nullptr);
+		check(Objects[Result].GetObject() == nullptr);
 		return Result;
 	}
 
@@ -515,7 +520,7 @@ public:
 		check(Result == (NumElements + Count - 1));
 		NumElements += Count;
 		FPlatformMisc::MemoryBarrier();
-		check(Objects[Result].Object == nullptr);
+		check(Objects[Result].GetObject() == nullptr);
 		return Result;
 	}
 
@@ -1538,7 +1543,7 @@ struct FIndexToObject
 	static FORCEINLINE class UObjectBase* IndexToObject(int32 Index, bool bEvenIfGarbage)
 	{
 		FUObjectItem* ObjectItem = GUObjectArray.IndexToObject(Index, bEvenIfGarbage);
-		return ObjectItem ? ObjectItem->Object : nullptr;
+		return ObjectItem ? ObjectItem->GetObject() : nullptr;
 	}
 };
 
diff --git a/Engine/Source/Runtime/CoreUObject/Public/UObject/UObjectIterator.h b/Engine/Source/Runtime/CoreUObject/Public/UObject/UObjectIterator.h
--- a/Engine/Source/Runtime/CoreUObject/Public/UObject/UObjectIterator.h
+++ b/Engine/Source/Runtime/CoreUObject/Public/UObject/UObjectIterator.h
@@ -191,7 +191,7 @@ public:
 	{
 		// casting UObjectBase to UObject for clients
 		FUObjectItem* ObjectItem = GetObject();
-		return (UObject*)(ObjectItem ? ObjectItem->Object : nullptr);
+		return (UObject*)(ObjectItem ? ObjectItem->GetObject() : nullptr);
 	}
 	/**
 	 * Iterator dereference
@@ -200,7 +200,7 @@ public:
 	FORCEINLINE UObject* operator->() const
 	{
 		FUObjectItem* ObjectItem = GetObject();
-		return (UObject*)(ObjectItem ? ObjectItem->Object : nullptr);
+		return (UObject*)(ObjectItem ? ObjectItem->GetObject() : nullptr);
 	}
 private:
 	/** Class to restrict results to */
diff --git a/Engine/Source/Runtime/CoreUObject/Public/UObject/WeakObjectPtr.h b/Engine/Source/Runtime/CoreUObject/Public/UObject/WeakObjectPtr.h
--- a/Engine/Source/Runtime/CoreUObject/Public/UObject/WeakObjectPtr.h
+++ b/Engine/Source/Runtime/CoreUObject/Public/UObject/WeakObjectPtr.h
@@ -534,7 +534,7 @@ private:
 	FORCEINLINE_DEBUGGABLE UObject* Internal_Get(bool bEvenIfGarbage) const
 	{
 		FUObjectItem* const ObjectItem = Internal_GetObjectItem();
-		return ((ObjectItem != nullptr) && GUObjectArray.IsValid(ObjectItem, bEvenIfGarbage)) ? (UObject*)ObjectItem->Object : nullptr;
+		return ((ObjectItem != nullptr) && GUObjectArray.IsValid(ObjectItem, bEvenIfGarbage)) ? (UObject*)ObjectItem->GetObject() : nullptr;
 	}
 
 	COREUOBJECT_API TStrongObjectPtr<UObject> Internal_Pin(bool bEvenIfGarbage) const;
diff --git a/Engine/Source/Runtime/Engine/Private/Collision/ChaosVDSQTraceHelper.cpp b/Engine/Source/Runtime/Engine/Private/Collision/ChaosVDSQTraceHelper.cpp
--- a/Engine/Source/Runtime/Engine/Private/Collision/ChaosVDSQTraceHelper.cpp
+++ b/Engine/Source/Runtime/Engine/Private/Collision/ChaosVDSQTraceHelper.cpp
@@ -27,7 +27,7 @@ const TObject* GetObjetPointerFromID(int32 ID)
 {
 	constexpr  bool bEvenIfGarbage = false;
 	const FUObjectItem* ObjectItem = GUObjectArray.IndexToValidObject(ID, bEvenIfGarbage);
-	if (const TObject* CastedObject = Cast<TObject>(ObjectItem ? static_cast<const UObject*>(ObjectItem->Object) : nullptr))
+	if (const TObject* CastedObject = Cast<TObject>(ObjectItem ? static_cast<const UObject*>(ObjectItem->GetObject()) : nullptr))
 	{
 		return CastedObject;
 	}
