Fixed code that attempts to directly assign Object on FUobjectItem as it has now became a private member

diff --git a/Engine/Source/Programs/AutoRTFMTests/Private/UECoreTests.cpp b/Engine/Source/Programs/AutoRTFMTests/Private/UECoreTests.cpp
--- a/Engine/Source/Programs/AutoRTFMTests/Private/UECoreTests.cpp
+++ b/Engine/Source/Programs/AutoRTFMTests/Private/UECoreTests.cpp
@@ -807,7 +807,7 @@ TEST_CASE("UECore.FUObjectItem")
 	SECTION("CreateStatID First In Open")
 	{
 		FUObjectItem Item;
-		Item.Object = NewObject<UMyAutoRTFMTestObject>();
+		Item.SetObject(NewObject<UMyAutoRTFMTestObject>());
 		Item.CreateStatID();
 
 		PROFILER_CHAR* const StatIDStringStorage = Item.StatIDStringStorage;
@@ -835,7 +835,7 @@ TEST_CASE("UECore.FUObjectItem")
 	SECTION("CreateStatID First In Closed")
 	{
 		FUObjectItem Item;
-		Item.Object = NewObject<UMyAutoRTFMTestObject>();
+		Item.SetObject(NewObject<UMyAutoRTFMTestObject>());
 		REQUIRE(nullptr == Item.StatIDStringStorage);
 		REQUIRE(!Item.StatID.IsValidStat());
 
@@ -864,7 +864,7 @@ TEST_CASE("UECore.FUObjectItem")
 		AutoRTFM::Testing::Abort([&]
 		{
 			FUObjectItem Item;
-			Item.Object = NewObject<UMyAutoRTFMTestObject>();
+			Item.SetObject(NewObject<UMyAutoRTFMTestObject>());
 			Item.CreateStatID();
 
 			AutoRTFM::Open([&]
@@ -879,7 +879,7 @@ TEST_CASE("UECore.FUObjectItem")
 		AutoRTFM::Testing::Commit([&]
 		{
 			FUObjectItem Item;
-			Item.Object = NewObject<UMyAutoRTFMTestObject>();
+			Item.SetObject(NewObject<UMyAutoRTFMTestObject>());
 			Item.CreateStatID();
 
 			AutoRTFM::Open([&]
@@ -894,7 +894,7 @@ TEST_CASE("UECore.FUObjectItem")
 	{
 		{
 			FUObjectItem Item;
-			Item.Object = NewObject<UMyAutoRTFMTestObject>();
+			Item.SetObject(NewObject<UMyAutoRTFMTestObject>());
 			REQUIRE(nullptr == Item.StatIDStringStorage);
 			REQUIRE(!Item.StatID.IsValidStat());
 
@@ -924,7 +924,7 @@ TEST_CASE("UECore.FUObjectItem")
 
 		{
 			FUObjectItem Item;
-			Item.Object = NewObject<UMyAutoRTFMTestObject>();
+			Item.SetObject(NewObject<UMyAutoRTFMTestObject>());
 			REQUIRE(nullptr == Item.StatIDStringStorage);
 			REQUIRE(!Item.StatID.IsValidStat());
 
diff --git a/Engine/Source/Runtime/CoreUObject/Private/UObject/GarbageCollection.cpp b/Engine/Source/Runtime/CoreUObject/Private/UObject/GarbageCollection.cpp
--- a/Engine/Source/Runtime/CoreUObject/Private/UObject/GarbageCollection.cpp
+++ b/Engine/Source/Runtime/CoreUObject/Private/UObject/GarbageCollection.cpp
@@ -6304,7 +6304,7 @@ FORCEINLINE static void MarkObjectItemAsReachable(FUObjectItem* ObjectItem)
 	if constexpr (bIsVerse)
 	{
 		// When verse VM is enabled, this method is also used to report that a UObject is being referenced inside of a VCell
-		checkf(GIsFrankenGCCollecting, TEXT("%s is marked as MaybeUnreachable but Reachability Analysis is not in progress"), *static_cast<UObject*>(ObjectItem->Object)->GetFullName());
+		checkf(GIsFrankenGCCollecting, TEXT("%s is marked as MaybeUnreachable but Reachability Analysis is not in progress"), *static_cast<UObject*>(ObjectItem->GetObject()->GetFullName());
 	}
 	else
 #endif
diff --git a/Engine/Source/Runtime/CoreUObject/Private/UObject/UObjectArray.cpp b/Engine/Source/Runtime/CoreUObject/Private/UObject/UObjectArray.cpp
--- a/Engine/Source/Runtime/CoreUObject/Private/UObject/UObjectArray.cpp
+++ b/Engine/Source/Runtime/CoreUObject/Private/UObject/UObjectArray.cpp
@@ -283,8 +283,8 @@ void FUObjectArray::AllocateUObjectIndex(UObjectBase* Object, EInternalObjectFla
 	}
 	// Add to global table.
 	FUObjectItem* ObjectItem = IndexToObject(Index);
-	UE_CLOG(ObjectItem->Object != nullptr, LogUObjectArray, Fatal, TEXT("Attempting to add %s at index %d but another object (0x%016llx) exists at that index!"), *Object->GetFName().ToString(), Index, (int64)(PTRINT)ObjectItem->Object);
-	ObjectItem->Object = Object;
+	UE_CLOG(ObjectItem->GetObject() != nullptr, LogUObjectArray, Fatal, TEXT("Attempting to add %s at index %d but another object (0x%016llx) exists at that index!"), *Object->GetFName().ToString(), Index, (int64)(PTRINT)ObjectItem->GetObject());
+	ObjectItem->SetObject(Object);
 	// At this point all not-compiled-in objects are not fully constructed yet and this is the earliest we can mark them as such
 	ObjectItem->Flags = (int32)EInternalObjectFlags::PendingConstruction;
 	if (!(IsOpenForDisregardForGC() & GUObjectArray.DisregardForGCEnabled())) //-V792
@@ -372,7 +372,7 @@ void FUObjectArray::FreeUObjectIndex(UObjectBase* Object)
 
 	int32 Index = Object->InternalIndex;
 	FUObjectItem* ObjectItem = IndexToObject(Index);
-	UE_CLOG(ObjectItem->Object != Object, LogUObjectArray, Fatal, TEXT("Removing object (0x%016llx) at index %d but the index points to a different object (0x%016llx)!"), (int64)(PTRINT)Object, Index, (int64)(PTRINT)ObjectItem->Object);
+	UE_CLOG(ObjectItem->GetObject() != Object, LogUObjectArray, Fatal, TEXT("Removing object (0x%016llx) at index %d but the index points to a different object (0x%016llx)!"), (int64)(PTRINT)Object, Index, (int64)(PTRINT)ObjectItem->GetObject());
 
 	// Can't destroy a refcounted object
 	check((ObjectItem->RefCount == 0 && (ObjectItem->GetFlags() & EInternalObjectFlags::RefCounted) == EInternalObjectFlags::None) || GExitPurge);
@@ -392,7 +392,7 @@ void FUObjectArray::FreeUObjectIndex(UObjectBase* Object)
 		ObjectItem->RefCount = 0;
 	};
 
-	ObjectItem->Object = nullptr;
+	ObjectItem->SetObject(nullptr);
 	ObjectItem->ClusterRootIndex = 0;
 	ObjectItem->SerialNumber = 0;
 #if UE_WITH_REMOTE_OBJECT_HANDLE
@@ -561,7 +561,7 @@ void FUObjectArray::DumpUObjectCountsToLog() const
 	for (int32 ObjectIndex = 0; ObjectIndex < GetObjectArrayNum(); ++ObjectIndex)
 	{
 		const FUObjectItem& ObjectItem = GetObjectItemArrayUnsafe()[ObjectIndex];
-		UObject* Object = (UObject*)ObjectItem.Object;
+		UObject* Object = (UObject*)ObjectItem.GetObject();
 		if (Object && Object->IsA(UClass::StaticClass()))
 		{
 			NumClasses++;
@@ -574,9 +574,9 @@ void FUObjectArray::DumpUObjectCountsToLog() const
 	for (int32 ObjectIndex = 0; ObjectIndex < GetObjectArrayNum(); ++ObjectIndex)
 	{
 		const FUObjectItem& ObjectItem = GetObjectItemArrayUnsafe()[ObjectIndex];
-		if (ObjectItem.Object)
+		if (ObjectItem.GetObject())
 		{
-			UObject* Object = (UObject*)ObjectItem.Object;
+			UObject* Object = (UObject*)ObjectItem.GetObject();
 			UClass* ObjectClass = Object->GetClass();
 			FClassEntry& ClassEntry = ClassCountMap.FindOrAdd(ObjectClass);
 			ClassEntry.Class = ObjectClass;
diff --git a/Engine/Source/Runtime/CoreUObject/Public/UObject/UObjectArray.h b/Engine/Source/Runtime/CoreUObject/Public/UObject/UObjectArray.h
--- a/Engine/Source/Runtime/CoreUObject/Public/UObject/UObjectArray.h
+++ b/Engine/Source/Runtime/CoreUObject/Public/UObject/UObjectArray.h
@@ -107,6 +107,11 @@ public:
 		return Object;
 	}
 
+	FORCEINLINE void SetObject(class UObjectBase* InObject)
+	{
+		Object = InObject;
+	}
+
 	FORCEINLINE void SetOwnerIndex(int32 OwnerIndex)
 	{
 		ClusterRootIndex = OwnerIndex;
@@ -992,11 +997,11 @@ public:
 	FORCEINLINE FUObjectItem* IndexToObject(int32 Index, bool bEvenIfGarbage)
 	{
 		FUObjectItem* ObjectItem = IndexToObject(Index);
-		if (ObjectItem && ObjectItem->Object)
+		if (ObjectItem && ObjectItem->GetObject())
 		{
 			if (!bEvenIfGarbage && ObjectItem->HasAnyFlags(EInternalObjectFlags::Garbage))
 			{
-				ObjectItem = nullptr;;
+				ObjectItem = nullptr;
 			}
 		}
 		return ObjectItem;
@@ -1317,7 +1322,7 @@ public:
 			while(++Index < Array.GetObjectArrayNum())
 			{
 				NextObject = const_cast<FUObjectItem*>(&Array.ObjObjects[Index]);
-				if (NextObject->Object)
+				if (NextObject->GetObject())
 				{
 					CurrentObject = NextObject;
 					return true;
