Fixed missing ) in previous change

diff --git a/Engine/Source/Runtime/CoreUObject/Private/UObject/GarbageCollection.cpp b/Engine/Source/Runtime/CoreUObject/Private/UObject/GarbageCollection.cpp
--- a/Engine/Source/Runtime/CoreUObject/Private/UObject/GarbageCollection.cpp
+++ b/Engine/Source/Runtime/CoreUObject/Private/UObject/GarbageCollection.cpp
@@ -6304,7 +6304,7 @@ FORCEINLINE static void MarkObjectItemAsReachable(FUObjectItem* ObjectItem)
 	if constexpr (bIsVerse)
 	{
 		// When verse VM is enabled, this method is also used to report that a UObject is being referenced inside of a VCell
-		checkf(GIsFrankenGCCollecting, TEXT("%s is marked as MaybeUnreachable but Reachability Analysis is not in progress"), *static_cast<UObject*>(ObjectItem->GetObject()->GetFullName());
+		checkf(GIsFrankenGCCollecting, TEXT("%s is marked as MaybeUnreachable but Reachability Analysis is not in progress"), *static_cast<UObject*>(ObjectItem->GetObject())->GetFullName());
 	}
 	else
 #endif
