UPackage: Const-correctness: Make UPackage::GetSavedHash const.
#rnx

diff --git a/Engine/Source/Runtime/CoreUObject/Private/UObject/Package.cpp b/Engine/Source/Runtime/CoreUObject/Private/UObject/Package.cpp
--- a/Engine/Source/Runtime/CoreUObject/Private/UObject/Package.cpp
+++ b/Engine/Source/Runtime/CoreUObject/Private/UObject/Package.cpp
@@ -407,7 +407,7 @@ bool UPackage::Rename(const TCHAR* InName, UObject* NewOuter, ERenameFlags Flags
 
 #if WITH_EDITORONLY_DATA
 
-FIoHash UPackage::GetSavedHash()
+FIoHash UPackage::GetSavedHash() const
 {
 #if UE_STRIP_DEPRECATED_PROPERTIES
 	return FIoHash();
diff --git a/Engine/Source/Runtime/CoreUObject/Public/UObject/Package.h b/Engine/Source/Runtime/CoreUObject/Public/UObject/Package.h
--- a/Engine/Source/Runtime/CoreUObject/Public/UObject/Package.h
+++ b/Engine/Source/Runtime/CoreUObject/Public/UObject/Package.h
@@ -841,7 +841,7 @@ public:
 	}
 
 	/** Hash of the package's .uasset/.umap file when it was last saved by the editor. */
-	COREUOBJECT_API FIoHash GetSavedHash();
+	COREUOBJECT_API FIoHash GetSavedHash() const;
 	COREUOBJECT_API void SetSavedHash(const FIoHash& InSavedHash);
 #endif
 
