Deprecate some mutable TObjectPtr usage mechanisms.

diff --git a/Engine/Source/Runtime/Core/Public/Misc/CoreDefines.h b/Engine/Source/Runtime/Core/Public/Misc/CoreDefines.h
--- a/Engine/Source/Runtime/Core/Public/Misc/CoreDefines.h
+++ b/Engine/Source/Runtime/Core/Public/Misc/CoreDefines.h
@@ -67,7 +67,7 @@ Can be overridden by setting to 1 or 0 in the project's .Target.cs files.
 #endif
 
 #ifndef UE_DEPRECATE_MUTABLE_TOBJECTPTR
-#define UE_DEPRECATE_MUTABLE_TOBJECTPTR 0
+#define UE_DEPRECATE_MUTABLE_TOBJECTPTR 1
 #endif
 
 #ifndef UE_REFERENCE_COLLECTOR_REQUIRE_OBJECTPTR
diff --git a/Engine/Source/Runtime/CoreUObject/Public/UObject/ObjectPtr.h b/Engine/Source/Runtime/CoreUObject/Public/UObject/ObjectPtr.h
--- a/Engine/Source/Runtime/CoreUObject/Public/UObject/ObjectPtr.h
+++ b/Engine/Source/Runtime/CoreUObject/Public/UObject/ObjectPtr.h
@@ -1034,7 +1034,7 @@ template <typename T>
 
 #if !UE_DEPRECATE_MUTABLE_TOBJECTPTR
 template <typename T, SIZE_T Size>
-UE_OBJPTR_DEPRECATED(5.3, "Mutable ToRawPtrArrayUnsafe() is deprecated. Use MutableView() or TArray<TObjectPtr<...>> instead.")
+UE_DEPRECATED(5.6, "Mutable ToRawPtrArrayUnsafe() is deprecated. Use MutableView() or TArray<TObjectPtr<...>> instead.")
 [[nodiscard]] FORCEINLINE T**
 ToRawPtrArrayUnsafe(TObjectPtr<T>(&ArrayOfPtr)[Size])
 {
@@ -1078,7 +1078,7 @@ template <
 #if UE_DEPRECATE_MUTABLE_TOBJECTPTR
 [[nodiscard]] const auto&
 #else
-UE_OBJPTR_DEPRECATED(5.3, "Mutable ToRawPtrTArrayUnsafe() is deprecated. Use MutableView() or TArray<TObjectPtr<...>> instead.")
+UE_DEPRECATED(5.6, "Mutable ToRawPtrTArrayUnsafe() is deprecated. Use MutableView() or TArray<TObjectPtr<...>> instead.")
 [[nodiscard]] decltype(auto)
 #endif
 ToRawPtrTArrayUnsafe(ArrayType&& Array)
@@ -1260,7 +1260,7 @@ inline void Swap(T*& A, TObjectPtr<T>& B)
 
 /** Swap variants between TArray<TObjectPtr<T>> and TArray<T*> */
 #if !UE_DEPRECATE_MUTABLE_TOBJECTPTR
-UE_OBJPTR_DEPRECATED(5.3, "Swap between TObjectPtr arrays and raw pointer arrays is deprecated. Swap TArray<TObjectPtr<...>> values instead.")
+UE_DEPRECATED(5.6, "Swap between TObjectPtr arrays and raw pointer arrays is deprecated. Swap TArray<TObjectPtr<...>> values instead.")
 template <typename T>
 inline void Swap(TArray<TObjectPtr<T>>& A, TArray<T*>& B)
 {
@@ -1287,7 +1287,7 @@ inline void Exchange(T*& A, TObjectPtr<T>& B)
 
 #if !UE_DEPRECATE_MUTABLE_TOBJECTPTR
 /** Exchange variants between TArray<TObjectPtr<T>> and TArray<T*> */
-UE_OBJPTR_DEPRECATED(5.3, "Exchange between TObjectPtr arrays and raw pointer arrays is deprecated. Exchange TArray<TObjectPtr<...>> values instead.")
+UE_DEPRECATED(5.6, "Exchange between TObjectPtr arrays and raw pointer arrays is deprecated. Exchange TArray<TObjectPtr<...>> values instead.")
 template <typename T>
 inline void Exchange(TArray<TObjectPtr<T>>& A, TArray<T*>& B)
 {
