AssetRegistry: Remove deprecation handling of FAssetPackageData::PackageGuid and instead store the entire PackageSavedHash read from the PackageHeader. Change GetPackageSavedHash to be WITH_EDITORONLY_DATA to match the UPackage::GetSavedHash API. The PackageSavedHash will always be serialized when serializing FAssetPackageData, but the game will not store a copy of it and will write out zeroes if it ever serializes AssetPackageData.
[RN] Minor,Core

diff --git a/Engine/Source/Runtime/AssetRegistry/Private/AssetRegistryState.cpp b/Engine/Source/Runtime/AssetRegistry/Private/AssetRegistryState.cpp
--- a/Engine/Source/Runtime/AssetRegistry/Private/AssetRegistryState.cpp
+++ b/Engine/Source/Runtime/AssetRegistry/Private/AssetRegistryState.cpp
@@ -4020,7 +4020,11 @@ void FAssetRegistryState::Dump(const TArray<FString>& Arguments, TArray<FString>
 			PageBuffer.Append(TEXT("	"));
 			Key.AppendString(PageBuffer);
 			PageBuffer.Append(TEXT(" : "));
+#if WITH_EDITORONLY_DATA
 			PageBuffer << PackageData->GetPackageSavedHash();
+#else
+			PageBuffer << FIoHash();
+#endif
 			PageBuffer.Appendf(TEXT(" : %d bytes"), PackageData->DiskSize);
 			AddLine();
 		}
diff --git a/Engine/Source/Runtime/CoreUObject/Private/AssetRegistry/AssetData.cpp b/Engine/Source/Runtime/CoreUObject/Private/AssetRegistry/AssetData.cpp
--- a/Engine/Source/Runtime/CoreUObject/Private/AssetRegistry/AssetData.cpp
+++ b/Engine/Source/Runtime/CoreUObject/Private/AssetRegistry/AssetData.cpp
@@ -930,9 +930,33 @@ bool FPackageCustomVersion::TryRead(const FCbFieldView& Field)
 void FAssetPackageData::SerializeForCacheInternal(FArchive& Ar, FAssetPackageData& PackageData, FAssetRegistryVersion::Type Version)
 {
 	Ar << PackageData.DiskSize;
-	PRAGMA_DISABLE_DEPRECATION_WARNINGS
-	Ar << PackageData.PackageGuid;
-	PRAGMA_ENABLE_DEPRECATION_WARNINGS
+
+	// The serialization format for FAssetPackageData is shared between editor and runtime, so that PackageDatas
+	// written by one can be read by the other. But the PackageSavedHash is excluded from the runtime data and API to
+	// save memory. Whenever serializing an AssetPackageData at runtime, we read a placeholder and write zeroes.
+	// Caveat: this means a round trip from editor -> runtime -> editor will lose the PackageSavedHash data and replace
+	// it with zeroes.
+	if (Version < FAssetRegistryVersion::PackageSavedHash)
+	{
+		FGuid LegacyGuid;
+		Ar << LegacyGuid;
+#if WITH_EDITORONLY_DATA
+		if (Ar.IsLoading())
+		{
+			FMemory::Memcpy(&PackageData.PackageSavedHash, &LegacyGuid,
+				FMath::Min(sizeof(decltype(PackageData.PackageSavedHash.GetBytes())), sizeof(LegacyGuid)));
+		}
+#endif
+	}
+	else
+	{
+#if WITH_EDITORONLY_DATA
+		Ar << PackageData.PackageSavedHash;
+#else
+		FIoHash PlaceholderSavedHash;
+		Ar << PlaceholderSavedHash;
+#endif
+	}
 	if (Version >= FAssetRegistryVersion::AddedCookedMD5Hash)
 	{
 		Ar << PackageData.CookedHash;
@@ -1000,21 +1024,18 @@ void FAssetPackageData::SerializeForCacheOldVersion(FArchive& Ar, FAssetRegistry
 
 FIoHash FAssetPackageData::GetPackageSavedHash() const
 {
-	FIoHash Result;
-	PRAGMA_DISABLE_DEPRECATION_WARNINGS;
-	FMemory::Memcpy(&Result.GetBytes(), &this->PackageGuid,
-		FMath::Min(sizeof(decltype(Result.GetBytes())), sizeof(PackageGuid)));
-	PRAGMA_ENABLE_DEPRECATION_WARNINGS;
-	return Result;
+#if WITH_EDITORONLY_DATA
+	return PackageSavedHash;
+#else
+	return FIoHash();
+#endif
 }
 
 void FAssetPackageData::SetPackageSavedHash(const FIoHash& InHash)
 {
-	PRAGMA_DISABLE_DEPRECATION_WARNINGS;
-	PackageGuid = FGuid();
-	FMemory::Memcpy(&PackageGuid, &InHash.GetBytes(),
-		FMath::Min(sizeof(PackageGuid), sizeof(decltype(InHash.GetBytes()))));
-	PRAGMA_ENABLE_DEPRECATION_WARNINGS;
+#if WITH_EDITORONLY_DATA
+	PackageSavedHash = InHash;
+#endif
 }
 
 inline FCbWriter& operator<<(FCbWriter& Writer, const TPair<FIoChunkId, FIoHash>& Value)
@@ -1042,13 +1063,20 @@ void FAssetPackageData::NetworkWrite(FCbWriter& Writer) const
 	{
 		Writer << CookedHash;
 	}
-	FIoHash PackageSavedHash = GetPackageSavedHash();
+
+	// The network serialization format for FAssetPackageData is shared between editor and runtime for robustness
+	// just like the persistent format; see the comment in SerializeForCacheInternal for notes.
+#if WITH_EDITORONLY_DATA
 	bool bPackageSavedHash = !PackageSavedHash.IsZero();
 	Writer << bPackageSavedHash;
 	if (bPackageSavedHash)
 	{
 		Writer << PackageSavedHash;
 	}
+#else
+	bool bPackageSavedHash = false;
+	Writer << bPackageSavedHash;
+#endif
 	Writer << ChunkHashes.Array();
 	Writer << ImportedClasses;
 	Writer << DiskSize;
@@ -1076,13 +1104,23 @@ bool FAssetPackageData::TryNetworkRead(FCbFieldView Field)
 		CookedHash = FMD5Hash();
 	}
 	bool bPackageSavedHash = false;
-	FIoHash PackageSavedHash;
 	bOk = LoadFromCompactBinary(*Iter++, bPackageSavedHash) & bOk;
+#if WITH_EDITORONLY_DATA
 	if (bPackageSavedHash)
 	{
 		bOk = LoadFromCompactBinary(*Iter++, PackageSavedHash) & bOk;
 	}
-	SetPackageSavedHash(PackageSavedHash);
+	else
+	{
+		PackageSavedHash = FIoHash::Zero;
+	}
+#else
+	if (bPackageSavedHash)
+	{
+		FIoHash UnusedPackageSavedHash;
+		bOk = LoadFromCompactBinary(*Iter++, UnusedPackageSavedHash) & bOk;
+	}
+#endif
 	TArray<TPair<FIoChunkId, FIoHash>> ChunkHashesArray;
 	if (LoadFromCompactBinary(*Iter++, ChunkHashesArray))
 	{
diff --git a/Engine/Source/Runtime/CoreUObject/Public/AssetRegistry/AssetData.h b/Engine/Source/Runtime/CoreUObject/Public/AssetRegistry/AssetData.h
--- a/Engine/Source/Runtime/CoreUObject/Public/AssetRegistry/AssetData.h
+++ b/Engine/Source/Runtime/CoreUObject/Public/AssetRegistry/AssetData.h
@@ -100,6 +100,7 @@ struct FAssetRegistryVersion
 		AssetPackageDataHasExtension,		// Added Extension to AssetPackageData.
 		AssetPackageDataHasPackageLocation,	// Added PackageLocation to AssetPackageData.
 		MarshalledTextAsUTF8String,			// Replaced 2 byte wide string with UTF8 String
+		PackageSavedHash,					// Replaced FAssetPackageData::PackageGuid with PackageSavedHash
 
 		// -----<new versions can be added above this line>-------------------------------------------------
 		VersionPlusOne,
@@ -914,14 +915,9 @@ private:
 }
 
 /** A class to hold data about a package on disk, this data is updated on save/load and is not updated when an asset changes in memory */
-PRAGMA_DISABLE_DEPRECATION_WARNINGS // Silence deprecation warnings for deprecated PackageGuid member in implicit constructors
 class FAssetPackageData
 {
 public:
-	/** Guid of the source package, uniquely identifies an asset package */
-	UE_DEPRECATED(5.4, "Use GetPackageSavedHash instead.")
-	FGuid PackageGuid;
-
 	/** MD5 of the cooked package on disk, for tracking nondeterministic changes */
 	FMD5Hash CookedHash;
 
@@ -945,6 +941,9 @@ public:
 	int32 FileVersionLicenseeUE;
 
 private:
+#if WITH_EDITORONLY_DATA
+	FIoHash PackageSavedHash;
+#endif
 	UE::AssetRegistry::FPackageCustomVersionsHandle CustomVersions;
 	/** Bit storage for flags */
 	uint32 Flags;
@@ -1005,8 +1004,15 @@ public:
 	}
 
 	/** Hash of the package's .uasset/.umap file when it was last saved by the editor. */
+#if WITH_EDITORONLY_DATA
 	COREUOBJECT_API FIoHash GetPackageSavedHash() const;
 	COREUOBJECT_API void SetPackageSavedHash(const FIoHash& InHash);
+#else
+	UE_DEPRECATED(5.6, "No longer supported in non-editor configurations, to save memory.")
+	COREUOBJECT_API FIoHash GetPackageSavedHash() const;
+	UE_DEPRECATED(5.6, "No longer supported in non-editor configurations, to save memory.")
+	COREUOBJECT_API void SetPackageSavedHash(const FIoHash& InHash);
+#endif
 
 	// Note: these functions should only be used for live communication between processes running the same version of the engine.
 	// There is no versioning support
@@ -1043,7 +1049,6 @@ private:
 		FLAG_PACKAGE_LOCATION_MASK					= uint32(FPackageName::EPackageLocationFilter::ValidMask) << FLAG_PACKAGE_LOCATION_TO_ENUM_SHIFT_SIZE
 	};
 };
-PRAGMA_ENABLE_DEPRECATION_WARNINGS
 
 
 /**
