[Asyncloading2]
- Extend error message for missing outers/class/super objects so we can better identify why things are going wrong.
-- In particular, it's useful to know the name of an Export, not just it's index, as well as if it was filtered or failed to load.

#rnx

diff --git a/Engine/Source/Runtime/CoreUObject/Private/Serialization/AsyncLoading2.cpp b/Engine/Source/Runtime/CoreUObject/Private/Serialization/AsyncLoading2.cpp
--- a/Engine/Source/Runtime/CoreUObject/Private/Serialization/AsyncLoading2.cpp
+++ b/Engine/Source/Runtime/CoreUObject/Private/Serialization/AsyncLoading2.cpp
@@ -7072,10 +7072,18 @@ UObject* FAsyncPackage2::EventDrivenIndexToObject(const FAsyncPackageHeaderData&
 	}
 	if (Index.IsExport())
 	{
-		Result = Header.ExportsView[Index.ToExport()].Object;
-		UE_CLOG(!Result, LogStreaming, Warning, TEXT("Missing Dependency, missing export 0x%llX in package %s"),
-			Index.Value(),
-			*Desc.PackagePathToLoad.GetPackageFName().ToString());
+		const FExportObject& Export = Header.ExportsView[Index.ToExport()];
+		Result = Export.Object;
+		UE_CLOG(!Result, LogStreaming, Warning, TEXT("Missing Dependency, missing export (%s) in package %s. \n\tExport Details - Index: 0x%llX, Super: '%s', Template: '%s', LoadFailed: %s, Filtered: %s (%s), FoundInMemory: %s"),
+			*Header.NameMap.GetName(Header.ExportMap[Index.ToExport()].ObjectName).ToString(),
+			*Desc.PackagePathToLoad.GetPackageFName().ToString(),
+			Index.Value(), 
+			(Export.SuperObject ? *Export.SuperObject->GetPathName() : TEXT("null")),
+			(Export.TemplateObject ? *Export.TemplateObject->GetPathName() : TEXT("null")),
+			(Export.bExportLoadFailed ? TEXT("true") : TEXT("false")),
+			(Export.bFiltered ? TEXT("true") : TEXT("false")),
+			LexToString(Header.ExportMap[Index.ToExport()].FilterFlags),
+			(Export.bWasFoundInMemory ? TEXT("true") : TEXT("false")));
 	}
 	else if (Index.IsImport())
 	{
@@ -7264,6 +7272,26 @@ UObject* FAsyncPackage2::ConditionalSerializeImport(const FAsyncPackageHeaderDat
 	return SourcePackage->ConditionalSerializeExport(*SourcePackageHeader, ExportIndex);
 }
 
+static FName GetFNameFromPackageObjectIndex(const FAsyncPackageHeaderData& Header, FPackageImportStore& ImportStore, FPackageObjectIndex Index)
+{
+	FName Out(NAME_None);
+	if (Index.IsExport())
+	{
+		Header.NameMap.TryGetName(Header.ExportMap[Index.ToExport()].ObjectName, Out);
+	}
+	else
+	{
+		// No need to check if the Index is invalid, that is handled internally
+		const UObject* Import = ImportStore.FindOrGetImportObject(Header, Index);
+		if (Import)
+		{
+			Out = Import->GetFName();
+		}
+	}
+
+	return Out;
+}
+
 void FAsyncPackage2::EventDrivenCreateExport(const FAsyncPackageHeaderData& Header, int32 LocalExportIndex)
 {
 	TRACE_CPUPROFILER_EVENT_SCOPE(CreateExport);
@@ -7303,13 +7331,13 @@ void FAsyncPackage2::EventDrivenCreateExport(const FAsyncPackageHeaderData& Head
 
 	if (!LoadClass)
 	{
-		UE_ASYNC_PACKAGE_LOG(Error, Desc, TEXT("CreateExport"), TEXT("Could not find class object for %s"), *ObjectName.ToString());
+		UE_ASYNC_PACKAGE_LOG(Error, Desc, TEXT("CreateExport"), TEXT("Could not find class object (0x%llX - '%s') for %s"), Export.ClassIndex.Value(), *GetFNameFromPackageObjectIndex(Header, ImportStore, Export.ClassIndex).ToString(), *ObjectName.ToString());
 		ExportObject.bExportLoadFailed = true;
 		return;
 	}
 	if (!ThisParent)
 	{
-		UE_ASYNC_PACKAGE_LOG(Error, Desc, TEXT("CreateExport"), TEXT("Could not find outer object for %s"), *ObjectName.ToString());
+		UE_ASYNC_PACKAGE_LOG(Error, Desc, TEXT("CreateExport"), TEXT("Could not find outer (0x%llX - '%s') object for %s"), Export.OuterIndex.Value(), *GetFNameFromPackageObjectIndex(Header, ImportStore, Export.OuterIndex).ToString(), *ObjectName.ToString());
 		ExportObject.bExportLoadFailed = true;
 		return;
 	}
@@ -7319,7 +7347,7 @@ void FAsyncPackage2::EventDrivenCreateExport(const FAsyncPackageHeaderData& Head
 		ExportObject.SuperObject = EventDrivenIndexToObject(Header, Export.SuperIndex, false);
 		if (!ExportObject.SuperObject)
 		{
-			UE_ASYNC_PACKAGE_LOG(Error, Desc, TEXT("CreateExport"), TEXT("Could not find SuperStruct object for %s"), *ObjectName.ToString());
+			UE_ASYNC_PACKAGE_LOG(Error, Desc, TEXT("CreateExport"), TEXT("Could not find SuperStruct object (0x%llX - '%s') for %s"), Export.SuperIndex.Value(), *GetFNameFromPackageObjectIndex(Header, ImportStore, Export.SuperIndex).ToString(), *ObjectName.ToString());
 			ExportObject.bExportLoadFailed = true;
 			return;
 		}
diff --git a/Engine/Source/Runtime/CoreUObject/Public/Serialization/AsyncLoading2.h b/Engine/Source/Runtime/CoreUObject/Public/Serialization/AsyncLoading2.h
--- a/Engine/Source/Runtime/CoreUObject/Public/Serialization/AsyncLoading2.h
+++ b/Engine/Source/Runtime/CoreUObject/Public/Serialization/AsyncLoading2.h
@@ -242,6 +242,17 @@ enum class EExportFilterFlags : uint8
 	NotForClient,
 	NotForServer
 };
+static const TCHAR* LexToString(EExportFilterFlags Flags)
+{
+	switch (Flags)
+	{
+		case EExportFilterFlags::NotForServer: return TEXT("NotForServer");
+		case EExportFilterFlags::NotForClient: return TEXT("NotForClient");
+		case EExportFilterFlags::None: return TEXT("None");
+		default: checkNoEntry();
+	}
+	return TEXT("");
+}
 
 enum class EZenPackageVersion : uint32
 {
