Disable UPS when duplicating for PIE, as we are seeing increased instability in StaticDuplicate associated with PIE world duplication. It's possible that user logic running during PIE world duplication is triggering some serialization assymetry, which is not supported by UPS.

#rnx

diff --git a/Engine/Source/Runtime/CoreUObject/Private/Serialization/DuplicateDataReader.cpp b/Engine/Source/Runtime/CoreUObject/Private/Serialization/DuplicateDataReader.cpp
--- a/Engine/Source/Runtime/CoreUObject/Private/Serialization/DuplicateDataReader.cpp
+++ b/Engine/Source/Runtime/CoreUObject/Private/Serialization/DuplicateDataReader.cpp
@@ -27,7 +27,7 @@ FDuplicateDataReader::FDuplicateDataReader( class FUObjectAnnotationSparse<FDupl
 {
 	this->SetIsLoading(true);
 	this->SetIsPersistent(true);
-	this->SetUseUnversionedPropertySerialization(true);
+	this->SetUseUnversionedPropertySerialization(((InPortFlags & PPF_DuplicateForPIE) == 0));
 	this->ArNoIntraPropertyDelta = true;
 	ArPortFlags |= PPF_Duplicate | InPortFlags;
 
diff --git a/Engine/Source/Runtime/CoreUObject/Private/Serialization/DuplicateDataWriter.cpp b/Engine/Source/Runtime/CoreUObject/Private/Serialization/DuplicateDataWriter.cpp
--- a/Engine/Source/Runtime/CoreUObject/Private/Serialization/DuplicateDataWriter.cpp
+++ b/Engine/Source/Runtime/CoreUObject/Private/Serialization/DuplicateDataWriter.cpp
@@ -38,7 +38,7 @@ FDuplicateDataWriter::FDuplicateDataWriter( FUObjectAnnotationSparse<FDuplicated
 {
 	this->SetIsSaving(true);
 	this->SetIsPersistent(true);
-	this->SetUseUnversionedPropertySerialization(true);
+	this->SetUseUnversionedPropertySerialization(((InPortFlags & PPF_DuplicateForPIE) == 0));
 	this->ArNoIntraPropertyDelta = true;
 	ArAllowLazyLoading	= false;
 	ArPortFlags |= PPF_Duplicate | InPortFlags;
