SaveConfig now handles the case where there is no existing branch for the given Filename by adding a temporary branch and then removing it from the cache after the flush.

This fixes the export functionality in Editor Preferences.

diff --git a/Engine/Source/Runtime/CoreUObject/Private/UObject/Obj.cpp b/Engine/Source/Runtime/CoreUObject/Private/UObject/Obj.cpp
--- a/Engine/Source/Runtime/CoreUObject/Private/UObject/Obj.cpp
+++ b/Engine/Source/Runtime/CoreUObject/Private/UObject/Obj.cpp
@@ -3424,6 +3424,13 @@ void UObject::SaveConfig(uint64 RequiredPropertyFlags, const TCHAR* InFilename,
 	// if a filename was specified, save all properties to that file no matter which class they come from
 	const FString Filename = InFilename ? InFilename : GetConfigFilename(this);
 
+	// If there is no existing branch for the filename and file operations are enabled, then add a new temporary branch and remove it after the flush.
+	const bool bIsTemporaryBranch = Config->FindBranch(*Filename, Filename) == nullptr && !Config->AreFileOperationsDisabled();
+	if (bIsTemporaryBranch)
+	{
+		Config->AddNewBranch(Filename);
+	}
+
 	// Determine whether the file we are writing is a default file config.
 	const bool bIsADefaultIniWrite = Filename == GetDefaultConfigFilename() || Filename == GetGlobalUserConfigFilename();
 
@@ -3590,10 +3597,10 @@ void UObject::SaveConfig(uint64 RequiredPropertyFlags, const TCHAR* InFilename,
 		}
 	}
 
-	// only write out the config file if this is GConfig
-	if (Config == GConfig)
+	// Only write out the config file if this is GConfig or if we are saving a temporary branch.
+	if (Config == GConfig || bIsTemporaryBranch)
 	{
-		Config->Flush(0, *Filename);
+		Config->Flush(bIsTemporaryBranch/*bRemoveFromCache*/, *Filename);
 	}
 }
 
