Fixed FValidateAssetsDetails::ValidationMessages not being populated when using FValidateAssetsSettings::bCollectPerAssetDetails

diff --git a/Engine/Plugins/Editor/DataValidation/Source/DataValidation/Private/EditorValidatorSubsystem.cpp b/Engine/Plugins/Editor/DataValidation/Source/DataValidation/Private/EditorValidatorSubsystem.cpp
--- a/Engine/Plugins/Editor/DataValidation/Source/DataValidation/Private/EditorValidatorSubsystem.cpp
+++ b/Engine/Plugins/Editor/DataValidation/Source/DataValidation/Private/EditorValidatorSubsystem.cpp
@@ -719,7 +719,7 @@ EDataValidationResult UEditorValidatorSubsystem::ValidateAssetsInternal(
 			Details.PackageName = Data.PackageName;
 			Details.AssetName = Data.AssetName;
 			Details.Result = AssetResult;
-			ValidationContext.SplitIssues(Details.ValidationWarnings, Details.ValidationErrors);
+			ValidationContext.SplitIssues(Details.ValidationWarnings, Details.ValidationErrors, &Details.ValidationMessages);
 
 			Details.ExternalObjects.Reserve(ValidationExternalObjects.Num());
 			for (const FAssetData& ExtData : ValidationExternalObjects)
@@ -893,20 +893,6 @@ void UEditorValidatorSubsystem::ValidateChangelistPreSubmit(
 	{
 		OutValidationWarnings = Details->ValidationWarnings;
 		OutValidationErrors = Details->ValidationErrors;
-		
-		for (const TSharedRef<FTokenizedMessage>& Message : Details->ValidationMessages)
-		{
-			switch (Message->GetSeverity())
-			{
-				case EMessageSeverity::Error:
-					OutValidationErrors.Add(Message->ToText());
-					break;
-				case EMessageSeverity::Warning:
-				case EMessageSeverity::PerformanceWarning:
-					OutValidationWarnings.Add(Message->ToText());
-					break;
-			}
-		}
 	}
 }
 
diff --git a/Engine/Source/Runtime/CoreUObject/Private/Misc/DataValidation.cpp b/Engine/Source/Runtime/CoreUObject/Private/Misc/DataValidation.cpp
--- a/Engine/Source/Runtime/CoreUObject/Private/Misc/DataValidation.cpp
+++ b/Engine/Source/Runtime/CoreUObject/Private/Misc/DataValidation.cpp
@@ -74,23 +74,38 @@ TSharedRef<FTokenizedMessage> FDataValidationContext::AddMessage(EMessageSeverit
 	return Message;
 }
 
-void FDataValidationContext::SplitIssues(TArray<FText>& Warnings, TArray<FText>& Errors) const
+void FDataValidationContext::SplitIssues(TArray<FText>& Warnings, TArray<FText>& Errors, TArray<TSharedRef<FTokenizedMessage>>* Messages) const
 {
 	for (const FIssue& Issue : GetIssues())
 	{
-		FText Message = Issue.Message;
-		if (Issue.TokenizedMessage.IsValid())
+		if (Issue.Severity == EMessageSeverity::Error || Issue.Severity == EMessageSeverity::Warning)
 		{
-			Message = Issue.TokenizedMessage->ToText();
+			FText Message = Issue.Message;
+			if (Issue.TokenizedMessage)
+			{
+				Message = Issue.TokenizedMessage->ToText();
+			}
+
+			if (Issue.Severity == EMessageSeverity::Error)
+			{
+				Errors.Add(MoveTemp(Message));
+			}
+			else if (Issue.Severity == EMessageSeverity::Warning)
+			{
+				Warnings.Add(MoveTemp(Message));
+			}
 		}
 
-		if (Issue.Severity == EMessageSeverity::Error)
+		if (Messages)
 		{
-			Errors.Add(MoveTemp(Message));
-		}
-		else if (Issue.Severity == EMessageSeverity::Warning)
-		{	
-			Warnings.Add(MoveTemp(Message));
+			if (Issue.TokenizedMessage)
+			{
+				Messages->Add(Issue.TokenizedMessage.ToSharedRef());
+			}
+			else
+			{
+				Messages->Add(FTokenizedMessage::Create(Issue.Severity, Issue.Message));
+			}
 		}
 	}
 }
diff --git a/Engine/Source/Runtime/CoreUObject/Public/Misc/DataValidation.h b/Engine/Source/Runtime/CoreUObject/Public/Misc/DataValidation.h
--- a/Engine/Source/Runtime/CoreUObject/Public/Misc/DataValidation.h
+++ b/Engine/Source/Runtime/CoreUObject/Public/Misc/DataValidation.h
@@ -106,7 +106,7 @@ public:
 	uint32 GetNumWarnings() const { return NumWarnings; }
 	uint32 GetNumErrors() const { return NumErrors; }
 
-	COREUOBJECT_API void SplitIssues(TArray<FText>& Warnings, TArray<FText>& Errors) const;
+	COREUOBJECT_API void SplitIssues(TArray<FText>& Warnings, TArray<FText>& Errors, TArray<TSharedRef<FTokenizedMessage>>* Messages = nullptr) const;
 	
 private:
 	TArray<FIssue> Issues;
