Updated the FSoftObjectPath no export type to account for SubPathString being changed from an FString to an FUtf8String.

#tests Verified stability in cooked client/server and PIE.

diff --git a/Engine/Source/Runtime/CoreUObject/Public/UObject/NoExportTypes.h b/Engine/Source/Runtime/CoreUObject/Public/UObject/NoExportTypes.h
--- a/Engine/Source/Runtime/CoreUObject/Public/UObject/NoExportTypes.h
+++ b/Engine/Source/Runtime/CoreUObject/Public/UObject/NoExportTypes.h
@@ -2280,7 +2280,7 @@ struct FSoftObjectPath
 
 	/** Optional FString for subobject within an asset */
 	UPROPERTY()
-	FString SubPathString;
+	FUtf8String SubPathString;
 };
 
 /** 
