Removed trailing NULs from FSoftObjectPath subpaths and fixed a small inefficiency with an array reserve.

diff --git a/Engine/Source/Runtime/Core/Public/UObject/FortniteMainBranchObjectVersions.inl b/Engine/Source/Runtime/Core/Public/UObject/FortniteMainBranchObjectVersions.inl
--- a/Engine/Source/Runtime/Core/Public/UObject/FortniteMainBranchObjectVersions.inl
+++ b/Engine/Source/Runtime/Core/Public/UObject/FortniteMainBranchObjectVersions.inl
@@ -589,5 +589,8 @@ FORTNITE_MAIN_VERSION(PCGLandscapeCacheDefaultSerializationChanged, 191)
 // FSoftObjectPath::SubPathString changed to FUtf8String
 FORTNITE_MAIN_VERSION(SoftObjectPathUtf8SubPaths, 192)
 
+// FSoftObjectPath::SubPathString could be saved with trailing NULs and need truncating
+FORTNITE_MAIN_VERSION(SoftObjectPathTrailingNULsMaintained, 193)
+
 // -----<new versions can be added above this line>-------------------------------------------------
 #undef FORTNITE_MAIN_VERSION
\ No newline at end of file
diff --git a/Engine/Source/Runtime/CoreUObject/Private/UObject/SoftObjectPath.cpp b/Engine/Source/Runtime/CoreUObject/Private/UObject/SoftObjectPath.cpp
--- a/Engine/Source/Runtime/CoreUObject/Private/UObject/SoftObjectPath.cpp
+++ b/Engine/Source/Runtime/CoreUObject/Private/UObject/SoftObjectPath.cpp
@@ -479,7 +479,7 @@ void SoftObjectPathLoadSubPathWorkaround(FArchive& Ar, FUtf8String& OutStr)
 
 	// Load the saved bytes into an array
 	TArray<uint8> SavedBytes;
-	SavedBytes.Empty(SaveNum);
+	SavedBytes.Empty(NumSavedBytes);
 	SavedBytes.AddUninitialized(NumSavedBytes);
 	Ar.Serialize(SavedBytes.GetData(), NumSavedBytes);
 
@@ -577,6 +577,14 @@ void FSoftObjectPath::SerializePathWithoutFixup(FArchive& Ar)
 		{
 			Ar << SubPathString;
 		}
+		if (Ar.CustomVer(FFortniteMainBranchObjectVersion::GUID) < FFortniteMainBranchObjectVersion::SoftObjectPathTrailingNULsMaintained)
+		{
+			int32 NonZeroChar = SubPathString.FindLastCharByPredicate([](UTF8CHAR Ch){ return Ch != '\0'; });
+			if (NonZeroChar != INDEX_NONE)
+			{
+				SubPathString.LeftInline(NonZeroChar + 1);
+			}
+		}
 	}
 }
 
