[HH] Disable hitch detection while reporting garbage refs.

#rnx

diff --git a/Engine/Source/Runtime/CoreUObject/Private/UObject/GarbageCollection.cpp b/Engine/Source/Runtime/CoreUObject/Private/UObject/GarbageCollection.cpp
--- a/Engine/Source/Runtime/CoreUObject/Private/UObject/GarbageCollection.cpp
+++ b/Engine/Source/Runtime/CoreUObject/Private/UObject/GarbageCollection.cpp
@@ -36,6 +36,7 @@
 #include "ProfilingDebugging/CsvProfiler.h"
 #include "HAL/Runnable.h"
 #include "HAL/RunnableThread.h"
+#include "HAL/ThreadHeartBeat.h"
 #include "UObject/FieldPathProperty.h"
 #include "UObject/GarbageCollectionHistory.h"
 #include "UObject/GarbageCollectionInternalFlags.h"
@@ -5356,6 +5357,10 @@ static void DumpGarbageReferencers(TConstArrayView<TUniquePtr<FWorkerContext>> C
 		return;
 	}
 
+	// Suspend HitchDetector while logging garbage refs
+	FSlowHeartBeatScope SuspendHeartBeat;
+	FDisableHitchDetectorScope SuspendGameThreadHitch;
+
 	// Dump all garbage references in detail
 	if (GGarbageReferenceTrackingEnabled == 1 || GGarbageReferenceTrackingEnabled == 3)
 	{
