Fix Horde issue 813190

diff --git a/Engine/Source/Runtime/CoreUObject/Private/UObject/ObjectResource.cpp b/Engine/Source/Runtime/CoreUObject/Private/UObject/ObjectResource.cpp
--- a/Engine/Source/Runtime/CoreUObject/Private/UObject/ObjectResource.cpp
+++ b/Engine/Source/Runtime/CoreUObject/Private/UObject/ObjectResource.cpp
@@ -4,7 +4,7 @@
 #include "UObject/Class.h"
 #include "Templates/Casts.h"
 #if WITH_VERSE_VM || defined(__INTELLISENSE__)
-#include "VerseVM/VVMCell.h"
+#include "VerseVM/Inline/VVMCellInline.h"
 #include "VerseVM/VVMUniqueString.h"
 #endif
 
