SavePackage: In EditorSaves, calculate the SavedPackageHash from the hash of the final package bytes.
ObjectVersion change: LegacyFileVersion in FPackageFileSummary bumped from -8 to -9, and EUnrealEngineObjectUE5Version bumped to PACKAGE_SAVED_HASH.

Remove expired deprecation handling for FPackageFileSummary.Guid; add the new private FIoHash SavedHash to replace the functionality it was implementing.
Remove expired deprecation handling for UPackage.GetGuid,MakeNewGuid,SetGuid,SavePackage, and replace private UPackage.Guid with FIoHash SavedHash.
Remove expired deprecation handling for SAVE_ flags, and deprecate SAVE_KeepGuid; the guid is now always recalculated (except for text format where it is not yet implemented and is always kept).
Remove expired deprecation handling (and add one new deprecationthat we missed in 5.1) for FSavePackageArgs.
[RN] Minor, Core

diff --git a/Engine/Source/Editor/UnrealEd/Classes/Commandlets/BaseIteratePackagesCommandlet.h b/Engine/Source/Editor/UnrealEd/Classes/Commandlets/BaseIteratePackagesCommandlet.h
--- a/Engine/Source/Editor/UnrealEd/Classes/Commandlets/BaseIteratePackagesCommandlet.h
+++ b/Engine/Source/Editor/UnrealEd/Classes/Commandlets/BaseIteratePackagesCommandlet.h
@@ -138,6 +138,7 @@ protected:
 	/** Ignore package version changelist **/
 	bool bIgnoreChangelist;
 
+	UE_DEPRECATED(5.6, "KeepPackageGUIDOnSave is no longer supported, because package save no longer records arbitrary guids but instead records a hash of the package's bytes - PackageSavedHash.")
 	bool bKeepPackageGUIDOnSave;
 
 	/** Filter packages based on a collection **/
diff --git a/Engine/Source/Editor/UnrealEd/Private/Commandlets/BaseIteratePackagesCommandlet.cpp b/Engine/Source/Editor/UnrealEd/Private/Commandlets/BaseIteratePackagesCommandlet.cpp
--- a/Engine/Source/Editor/UnrealEd/Private/Commandlets/BaseIteratePackagesCommandlet.cpp
+++ b/Engine/Source/Editor/UnrealEd/Private/Commandlets/BaseIteratePackagesCommandlet.cpp
@@ -452,7 +452,7 @@ void UBaseIteratePackagesCommandlet::SavePackages(const TArray<UPackage*>& Packa
 		FString FileName;
 		if (FPackageName::TryConvertLongPackageNameToFilename(LongPackageName, FileName, Extension))
 		{
-			ESaveFlags SaveFlags = bKeepPackageGUIDOnSave ? SAVE_KeepGUID : SAVE_None;
+			ESaveFlags SaveFlags = SAVE_None;
 			SavePackageHelper(PackageToSave, FileName, RF_Standalone, GWarn, SaveFlags);
 		}
 		else
@@ -877,7 +877,7 @@ void UBaseIteratePackagesCommandlet::LoadAndSaveOnePackage(const FString& Filena
 							*Filename,Linker->Summary.GetFileVersionUE().ToValue(), Linker->Summary.GetFileVersionLicenseeUE(), GPackageFileUEVersion.ToValue(), VER_LATEST_ENGINE_LICENSEEUE4 );
 					}
 
-					ESaveFlags SaveFlags = bKeepPackageGUIDOnSave ? SAVE_KeepGUID : SAVE_None;
+					ESaveFlags SaveFlags = SAVE_None;
 					
 					if (bIsReadOnly == false || SourceControlQueue == nullptr)
 					{
@@ -1082,7 +1082,12 @@ int32 UBaseIteratePackagesCommandlet::Main( const FString& Params )
 	WorldInitialisationValues.AllowAudioPlayback(false);
 	WorldInitialisationValues.CreatePhysicsScene(true);
 
-	bKeepPackageGUIDOnSave = Switches.Contains(TEXT("KeepPackageGUIDOnSave"));
+	// UE_DEPRECATED(5.6)
+	if (Switches.Contains(TEXT("KeepPackageGUIDOnSave")))
+	{
+		UE_LOG(LogIteratePackagesCommandlet, Warning,
+			TEXT("KeepPackageGUIDOnSave is no longer supported, because package save no longer records arbitrary guids but instead records a hash of the package's bytes - PackageSavedHash."));
+	}
 	
 	/** check for filtering packages by collection. **/
 	FString FilterByCollection;
diff --git a/Engine/Source/Editor/UnrealEd/Private/Commandlets/ContentCommandlets.cpp b/Engine/Source/Editor/UnrealEd/Private/Commandlets/ContentCommandlets.cpp
--- a/Engine/Source/Editor/UnrealEd/Private/Commandlets/ContentCommandlets.cpp
+++ b/Engine/Source/Editor/UnrealEd/Private/Commandlets/ContentCommandlets.cpp
@@ -964,8 +964,7 @@ void UResavePackagesCommandlet::LoadAndSaveOnePackage(const FString& Filename)
 							*Filename,Linker->Summary.GetFileVersionUE().ToValue(), Linker->Summary.GetFileVersionLicenseeUE(), GPackageFileUEVersion.ToValue(), VER_LATEST_ENGINE_LICENSEEUE4 );
 					}
 
-					const static bool bKeepPackageGUIDOnSave = FParse::Param(FCommandLine::Get(), TEXT("KeepPackageGUIDOnSave"));
-					ESaveFlags SaveFlags = bKeepPackageGUIDOnSave ? SAVE_KeepGUID : SAVE_None;
+					ESaveFlags SaveFlags = SAVE_None;
 					
 					if (bIsReadOnly == false || SourceControlQueue == nullptr)
 					{
@@ -1170,6 +1169,14 @@ int32 UResavePackagesCommandlet::Main( const FString& Params )
 	/** determine if we are building navigation data for the map packages on the pass. **/
 	bShouldBuildNavigationData = Switches.Contains(TEXT("BuildNavigationData"));
 
+	// UE_DEPRECATED(5.6)
+	bool bKeepPackageGUIDOnSave = FParse::Param(FCommandLine::Get(), TEXT("KeepPackageGUIDOnSave"));
+	if (bKeepPackageGUIDOnSave)
+	{
+		UE_LOG(LogContentCommandlet, Warning,
+			TEXT("KeepPackageGUIDOnSave is no longer supported, because package save no longer records arbitrary guids but instead records a hash of the package's bytes - PackageSavedHash."));
+	}
+
 	/** check for filtering packages by collection. **/
 	FString FilterByCollection;
 	FParse::Value(*Params, TEXT("FilterByCollection="), FilterByCollection);
diff --git a/Engine/Source/Editor/UnrealEd/Private/Commandlets/TextAssetCommandlet.cpp b/Engine/Source/Editor/UnrealEd/Private/Commandlets/TextAssetCommandlet.cpp
--- a/Engine/Source/Editor/UnrealEd/Private/Commandlets/TextAssetCommandlet.cpp
+++ b/Engine/Source/Editor/UnrealEd/Private/Commandlets/TextAssetCommandlet.cpp
@@ -509,7 +509,7 @@ bool UTextAssetCommandlet::DoTextAssetProcessing(const FProcessingArgs& InArgs)
 				// Firstly, do a resave of the package
 				UPackage* OriginalPackage = LoadPackage(nullptr, *SourceLongPackageName, LOAD_None);
 				IFileManager::Get().Delete(*SourceFilename, false, true, true);
-				SavePackageHelper(OriginalPackage, SourceFilename, RF_Standalone, GWarn, SAVE_KeepGUID);
+				SavePackageHelper(OriginalPackage, SourceFilename, RF_Standalone, GWarn, SAVE_None);
 				CollectGarbage(RF_NoFlags, true);
 
 				// Make a copy of the resaved source package which we can use as the base revision for each test
@@ -616,7 +616,7 @@ bool UTextAssetCommandlet::DoTextAssetProcessing(const FProcessingArgs& InArgs)
 						
 						{
 							TRACE_CPUPROFILER_EVENT_SCOPE(SavePackage); 
-							SavePackageHelper(Package, *WorkingFilenames[Bucket], RF_Standalone, GWarn, SAVE_KeepGUID);
+							SavePackageHelper(Package, *WorkingFilenames[Bucket], RF_Standalone, GWarn, SAVE_None);
 						}
 						
 						{
@@ -755,7 +755,7 @@ bool UTextAssetCommandlet::DoTextAssetProcessing(const FProcessingArgs& InArgs)
 						SCOPE_SECONDS_COUNTER(Timer);
 						TRACE_CPUPROFILER_EVENT_SCOPE(UTextAssetCommandlet::SavePackage);
 						IFileManager::Get().Delete(*DestinationFilename, false, true, true);
-						bSaveSuccessful = SavePackageHelper(Package, *DestinationFilename, RF_Standalone, GWarn, SAVE_KeepGUID);
+						bSaveSuccessful = SavePackageHelper(Package, *DestinationFilename, RF_Standalone, GWarn, SAVE_None);
 					}
 					TotalPackageSaveTime += Timer;
 					IterationPackageSaveTime += Timer;
diff --git a/Engine/Source/Editor/UnrealEd/Private/Cooker/CookFunctionLibrary.cpp b/Engine/Source/Editor/UnrealEd/Private/Cooker/CookFunctionLibrary.cpp
--- a/Engine/Source/Editor/UnrealEd/Private/Cooker/CookFunctionLibrary.cpp
+++ b/Engine/Source/Editor/UnrealEd/Private/Cooker/CookFunctionLibrary.cpp
@@ -170,7 +170,7 @@ void UCookFunctionLibrary::CookAsset(UObject* Object, const FString& ForPlatform
 	SaveArgs.TopLevelFlags = RF_Public;
 	SaveArgs.bForceByteSwapping = (!TargetPlatform->IsLittleEndian()) ^ (!PLATFORM_LITTLE_ENDIAN);;
 	SaveArgs.bWarnOfLongFilename = false;
-	SaveArgs.SaveFlags = SAVE_AllowTimeout|SAVE_KeepGUID;
+	SaveArgs.SaveFlags = SAVE_AllowTimeout;
 	SaveArgs.ArchiveCookData = &CookData;
 	SaveArgs.bSlowTask = true;
 	SaveArgs.SavePackageContext = &SaveContext;
diff --git a/Engine/Source/Editor/UnrealEd/Private/Cooker/CookGenerationHelper.cpp b/Engine/Source/Editor/UnrealEd/Private/Cooker/CookGenerationHelper.cpp
--- a/Engine/Source/Editor/UnrealEd/Private/Cooker/CookGenerationHelper.cpp
+++ b/Engine/Source/Editor/UnrealEd/Private/Cooker/CookGenerationHelper.cpp
@@ -2507,18 +2507,6 @@ void FCookGenerationInfo::CreatePackageHash()
 		}
 	}
 	PackageHash = FIoHash(Blake3.Finalize());
-	// We store the PackageHash as a FIoHash, but UPackage and FAssetPackageData store it as a FGuid, which is smaller,
-	// so we have to remove any data which doesn't fit into FGuid. This can be removed when we remove the deprecated
-	// Guid storage on UPackage.
-	PRAGMA_DISABLE_DEPRECATION_WARNINGS;
-	constexpr int SizeDifference = sizeof(PackageHash) - sizeof(decltype(DeclVal<UPackage>().GetGuid()));
-	if (SizeDifference > 0)
-	{
-		FMemory::Memset(((uint8*)&PackageHash.GetBytes())
-			+ (sizeof(decltype(PackageHash.GetBytes())) - SizeDifference),
-			0, SizeDifference);
-	}
-	PRAGMA_ENABLE_DEPRECATION_WARNINGS;
 }
 
 void FCookGenerationInfo::IterativeCookValidateOrClear(FGenerationHelper& GenerationHelper,
diff --git a/Engine/Source/Editor/UnrealEd/Private/Cooker/CookSavePackage.cpp b/Engine/Source/Editor/UnrealEd/Private/Cooker/CookSavePackage.cpp
--- a/Engine/Source/Editor/UnrealEd/Private/Cooker/CookSavePackage.cpp
+++ b/Engine/Source/Editor/UnrealEd/Private/Cooker/CookSavePackage.cpp
@@ -238,7 +238,7 @@ void FSaveCookedPackageContext::SetupPackage()
 			UE_LOG(LogCook, Fatal, TEXT("Package %s marked as reloading for cook was requested to save"), *PackageName);
 		}
 
-		SaveFlags = SAVE_KeepGUID | SAVE_Async
+		SaveFlags = SAVE_Async
 			| (COTFS.IsCookFlagSet(ECookInitializationFlags::Unversioned) ? SAVE_Unversioned : 0);
 		SaveFlags |= COTFS.IsCookFlagSet(ECookInitializationFlags::CookEditorOptional) ? SAVE_Optional : SAVE_None;
 
diff --git a/Engine/Source/Editor/UnrealEd/Private/EditorDomain/EditorDomainUtils.cpp b/Engine/Source/Editor/UnrealEd/Private/EditorDomain/EditorDomainUtils.cpp
--- a/Engine/Source/Editor/UnrealEd/Private/EditorDomain/EditorDomainUtils.cpp
+++ b/Engine/Source/Editor/UnrealEd/Private/EditorDomain/EditorDomainUtils.cpp
@@ -1742,9 +1742,6 @@ bool TrySavePackage(UPackage* Package)
 	SaveArgs.SaveFlags = SaveFlags;
 	SaveArgs.bSlowTask = false;
 	SaveArgs.SavePackageContext = &SavePackageContext;
-	PRAGMA_DISABLE_DEPRECATION_WARNINGS;
-	SaveArgs.OutputPackageGuid.Emplace(IoHashToGuid(PackageDigest.Hash));
-	PRAGMA_ENABLE_DEPRECATION_WARNINGS;
 	FSavePackageResultStruct Result = GEditor->Save(Package, nullptr, TEXT("EditorDomainPackageWriter"), SaveArgs);
 	if (Result.Result != ESavePackageResult::Success)
 	{
diff --git a/Engine/Source/Runtime/Core/Public/UObject/ObjectVersion.h b/Engine/Source/Runtime/Core/Public/UObject/ObjectVersion.h
--- a/Engine/Source/Runtime/Core/Public/UObject/ObjectVersion.h
+++ b/Engine/Source/Runtime/Core/Public/UObject/ObjectVersion.h
@@ -93,6 +93,9 @@ enum class EUnrealEngineObjectUE5Version : uint32
 	// Added VCells to the object graph
 	VERSE_CELLS,
 
+	// Changed PackageFileSummary to write FIoHash PackageSavedHash instead of FGuid Guid
+	PACKAGE_SAVED_HASH,
+
 	// -----<new versions can be added before this line>-------------------------------------------------
 	// - this needs to be the last line (see note below)
 	AUTOMATIC_VERSION_PLUS_ONE,
diff --git a/Engine/Source/Runtime/CoreUObject/Private/UObject/Package.cpp b/Engine/Source/Runtime/CoreUObject/Private/UObject/Package.cpp
--- a/Engine/Source/Runtime/CoreUObject/Private/UObject/Package.cpp
+++ b/Engine/Source/Runtime/CoreUObject/Private/UObject/Package.cpp
@@ -409,27 +409,12 @@ bool UPackage::Rename(const TCHAR* InName, UObject* NewOuter, ERenameFlags Flags
 
 FIoHash UPackage::GetSavedHash() const
 {
-#if UE_STRIP_DEPRECATED_PROPERTIES
-	return FIoHash();
-#else
-	FIoHash Result;
-	PRAGMA_DISABLE_DEPRECATION_WARNINGS;
-	FMemory::Memcpy(&Result.GetBytes(), &Guid,
-		FMath::Min(sizeof(decltype(Result.GetBytes())), sizeof(Guid)));
-	PRAGMA_ENABLE_DEPRECATION_WARNINGS;
-	return Result;
-#endif
+	return SavedHash;
 }
 
 void UPackage::SetSavedHash(const FIoHash& InSavedHash)
 {
-#if !UE_STRIP_DEPRECATED_PROPERTIES
-	PRAGMA_DISABLE_DEPRECATION_WARNINGS;
-	Guid = FGuid();
-	FMemory::Memcpy(&Guid, &InSavedHash.GetBytes(),
-		FMath::Min(sizeof(Guid), sizeof(decltype(InSavedHash.GetBytes()))));
-	PRAGMA_ENABLE_DEPRECATION_WARNINGS;
-#endif
+	SavedHash = InSavedHash;
 }
 
 #endif
diff --git a/Engine/Source/Runtime/CoreUObject/Private/UObject/PackageFileSummary.cpp b/Engine/Source/Runtime/CoreUObject/Private/UObject/PackageFileSummary.cpp
--- a/Engine/Source/Runtime/CoreUObject/Private/UObject/PackageFileSummary.cpp
+++ b/Engine/Source/Runtime/CoreUObject/Private/UObject/PackageFileSummary.cpp
@@ -66,6 +66,7 @@ void operator<<(FStructuredArchive::FSlot Slot, FPackageFileSummary& Sum)
 	}
 
 	FStructuredArchive::FRecord Record = Slot.EnterRecord();
+	int64 StartOffset = BaseArchive.Tell();
 
 	if (bCanStartSerializing)
 	{
@@ -103,8 +104,14 @@ void operator<<(FStructuredArchive::FSlot Slot, FPackageFileSummary& Sum)
 		*		-6 indicates optimizations to how custom versions are being serialized
 		*		-7 indicates the texture allocation info has been removed from the summary
 		*		-8 indicates that the UE5 version has been added to the summary
+		*		-9 indicates a contractual change in when early exits are required based on FileVersionTooNew. At or
+		*		   after this LegacyFileVersion, we support changing the PackageFileSummary serialization format for
+		*		   all bytes serialized after FileVersionLicensee, and that format change can be conditional on any
+		*		   of the versions parsed before that point. All packageloaders that understand the -9
+		*		   legacyfileformat are required to early exit without further serialization at that point if
+		*		   FileVersionTooNew is true.
 		*/
-		const int32 CurrentLegacyFileVersion = -8;
+		const int32 CurrentLegacyFileVersion = -9;
 		int32 LegacyFileVersion = CurrentLegacyFileVersion;
 		Record << SA_VALUE(TEXT("LegacyFileVersion"), LegacyFileVersion);
 
@@ -136,11 +143,6 @@ void operator<<(FStructuredArchive::FSlot Slot, FPackageFileSummary& Sum)
 			
 				Record << SA_VALUE(TEXT("FileVersionLicenseeUE4"), Sum.FileVersionLicenseeUE);
 
-				if (LegacyFileVersion <= -2)
-				{
-					Sum.CustomVersionContainer.Serialize(Record.EnterField(TEXT("CustomVersions")), GetCustomVersionFormatForArchive(LegacyFileVersion));
-				}
-
 				if (!Sum.FileVersionUE.FileVersionUE4 && !Sum.FileVersionUE.FileVersionUE5 && !Sum.FileVersionLicenseeUE)
 				{
 					// Record that the summary was unversioned when it was loaded off disk
@@ -161,12 +163,35 @@ void operator<<(FStructuredArchive::FSlot Slot, FPackageFileSummary& Sum)
 					Sum.FileVersionLicenseeUE = GPackageFileLicenseeUEVersion;
 					Sum.CustomVersionContainer = FCurrentCustomVersions::GetAll();
 				}
+
+				if (Sum.FileVersionUE < VER_UE4_OLDEST_LOADABLE_PACKAGE || Sum.IsFileVersionTooNew())
+				{
+					// We early exit as soon as possible if the fileversion is too old or too new so that changes can
+					// be made to the expected format after this point, and still have older versions of the editor
+					// gracefully abort the load of the newer package rather than crashing. They might crash on format
+					// differences if they try to serialize arrays from the wrong offset.
+					Sum.FileVersionUE.Reset();
+					Sum.FileVersionLicenseeUE = 0;
+					return;
+				}
+
+				if (Sum.GetFileVersionUE() >= EUnrealEngineObjectUE5Version::PACKAGE_SAVED_HASH)
+				{
+					Record << SA_VALUE(TEXT("SavedHash"), Sum.SavedHash);
+					Record << SA_VALUE(TEXT("TotalHeaderSize"), Sum.TotalHeaderSize);
+				}
+
+				if (LegacyFileVersion <= -2)
+				{
+					Sum.CustomVersionContainer.Serialize(Record.EnterField(TEXT("CustomVersions")), GetCustomVersionFormatForArchive(LegacyFileVersion));
+				}
 			}
 			else
 			{
 				// This is probably an old UE3 file, make sure that the linker will fail to load with it.
 				Sum.FileVersionUE.Reset();
 				Sum.FileVersionLicenseeUE = 0;
+				return;
 			}
 		}
 		else
@@ -178,9 +203,6 @@ void operator<<(FStructuredArchive::FSlot Slot, FPackageFileSummary& Sum)
 				Record << SA_VALUE(TEXT("FileVersionUE4"), Zero); // VersionUE4
 				Record << SA_VALUE(TEXT("FileVersionUE5"), Zero); // VersionUE5
 				Record << SA_VALUE(TEXT("FileVersionLicenseeUE4"), Zero); // VersionLicenseeUE4
-
-				FCustomVersionContainer NoCustomVersions;
-				NoCustomVersions.Serialize(Record.EnterField(TEXT("CustomVersions")));
 			}
 			else
 			{
@@ -190,12 +212,44 @@ void operator<<(FStructuredArchive::FSlot Slot, FPackageFileSummary& Sum)
 				Record << SA_VALUE(TEXT("FileVersionUE4"), Sum.FileVersionUE.FileVersionUE4);
 				Record << SA_VALUE(TEXT("FileVersionUE5"), Sum.FileVersionUE.FileVersionUE5);
 				Record << SA_VALUE(TEXT("FileVersionLicenseeUE4"), Sum.FileVersionLicenseeUE);
+			}
+
+#if WITH_EDITORONLY_DATA
+			if (BaseArchive.IsSaving() && !BaseArchive.IsTextFormat())
+			{
+				// GetSavedHashRelativeOffset must be kept in sync with the bytes written during serialize;
+				// we use it during saving to seek back to the offset of the SavedHash and rewrite it after
+				// calculating it.
+				int64 CurrentRelativeOffset = BaseArchive.Tell() - StartOffset;
+				int64 ExpectedRelativeOffset = Sum.GetSavedHashRelativeOffset();
+				checkf(CurrentRelativeOffset == ExpectedRelativeOffset,
+					TEXT("Expected (CurrentRelativeOffset) %" INT64_FMT " == %" INT64_FMT " (ExpectedRelativeOffset)"),
+					CurrentRelativeOffset, ExpectedRelativeOffset);
+			}
+#endif
+			// Saved FileVersion might not be current FileVersion, when called from AssetHeaderPatcher
+			if (Sum.GetFileVersionUE() >= EUnrealEngineObjectUE5Version::PACKAGE_SAVED_HASH)
+			{
+				Record << SA_VALUE(TEXT("SavedHash"), Sum.SavedHash);
+				Record << SA_VALUE(TEXT("TotalHeaderSize"), Sum.TotalHeaderSize);
+			}
 
-				// Serialise custom version map.
+			// Serialise custom version map.
+			if (Sum.bUnversioned)
+			{
+				FCustomVersionContainer NoCustomVersions;
+				NoCustomVersions.Serialize(Record.EnterField(TEXT("CustomVersions")));
+			}
+			else
+			{
 				Sum.CustomVersionContainer.Serialize(Record.EnterField(TEXT("CustomVersions")));
 			}
 		}
-		Record << SA_VALUE(TEXT("TotalHeaderSize"), Sum.TotalHeaderSize);
+
+		if (Sum.GetFileVersionUE() < EUnrealEngineObjectUE5Version::PACKAGE_SAVED_HASH)
+		{
+			Record << SA_VALUE(TEXT("TotalHeaderSize"), Sum.TotalHeaderSize);
+		}
 		Record << SA_VALUE(TEXT("PackageName"), Sum.PackageName); //@note used to be FolderName, now unused and deprecated
 
 		if (BaseArchive.IsCooking())
@@ -253,11 +307,6 @@ void operator<<(FStructuredArchive::FSlot Slot, FPackageFileSummary& Sum)
 
 		Record << SA_VALUE(TEXT("DependsOffset"), Sum.DependsOffset);
 
-		if (BaseArchive.IsLoading() && (Sum.FileVersionUE < VER_UE4_OLDEST_LOADABLE_PACKAGE || Sum.IsFileVersionTooNew()))
-		{
-			return; // we can't safely load more than this because the below was different in older files.
-		}
-
 		if (Sum.FileVersionUE >= VER_UE4_ADD_STRING_ASSET_REFERENCES_MAP)
 		{
 			Record << SA_VALUE(TEXT("SoftPackageReferencesCount"), Sum.SoftPackageReferencesCount) << SA_VALUE(TEXT("SoftPackageReferencesOffset"), Sum.SoftPackageReferencesOffset);
@@ -270,9 +319,22 @@ void operator<<(FStructuredArchive::FSlot Slot, FPackageFileSummary& Sum)
 
 		Record << SA_VALUE(TEXT("ThumbnailTableOffset"), Sum.ThumbnailTableOffset);
 
-		PRAGMA_DISABLE_DEPRECATION_WARNINGS
-		Record << SA_VALUE(TEXT("Guid"), Sum.Guid);
-		PRAGMA_ENABLE_DEPRECATION_WARNINGS
+		if (Sum.GetFileVersionUE() < EUnrealEngineObjectUE5Version::PACKAGE_SAVED_HASH)
+		{
+			FGuid LegacyGuid;
+			if (!BaseArchive.IsLoading())
+			{
+				FMemory::Memcpy(&LegacyGuid, &Sum.SavedHash.GetBytes(),
+					FMath::Min(sizeof(decltype(Sum.SavedHash.GetBytes())), sizeof(LegacyGuid)));
+			}
+			Record << SA_VALUE(TEXT("Guid"), LegacyGuid);
+			if (BaseArchive.IsLoading())
+			{
+				Sum.SavedHash = FIoHash::Zero;
+				FMemory::Memcpy(&Sum.SavedHash.GetBytes(), &LegacyGuid,
+					FMath::Min(sizeof(decltype(Sum.SavedHash.GetBytes())), sizeof(LegacyGuid)));
+			}
+		}
 
 #if WITH_EDITORONLY_DATA
 		if (!BaseArchive.IsFilterEditorOnly())
@@ -465,21 +527,25 @@ void operator<<(FStructuredArchive::FSlot Slot, FPackageFileSummary& Sum)
 #if WITH_EDITORONLY_DATA
 FIoHash FPackageFileSummary::GetSavedHash() const
 {
-	FIoHash Result;
-	PRAGMA_DISABLE_DEPRECATION_WARNINGS;
-	FMemory::Memcpy(&Result.GetBytes(), &Guid,
-		FMath::Min(sizeof(decltype(Result.GetBytes())), sizeof(Guid)));
-	PRAGMA_ENABLE_DEPRECATION_WARNINGS;
-	return Result;
+	return SavedHash;
 }
 
 void FPackageFileSummary::SetSavedHash(const FIoHash& InSavedHash)
 {
-	PRAGMA_DISABLE_DEPRECATION_WARNINGS;
-	Guid = FGuid();
-	FMemory::Memcpy(&Guid, &InSavedHash.GetBytes(),
-		FMath::Min(sizeof(Guid), sizeof(decltype(InSavedHash.GetBytes()))));
-	PRAGMA_ENABLE_DEPRECATION_WARNINGS;
+	SavedHash = InSavedHash;
+}
+
+int64 FPackageFileSummary::GetSavedHashRelativeOffset() const
+{
+	// This offset is constructed by copying the code from void operator<<(FStructuredArchive::FSlot Slot, FPackageFileSummary& Sum)
+	// with LegacyFileVersion == CurrentLegacyFileVersion
+	return
+		sizeof(Tag)
+		+ sizeof(int32) // LegacyFileVersion
+		+ sizeof(int32) // LegacyUE3Version
+		+ sizeof(FileVersionUE.FileVersionUE4)
+		+ sizeof(FileVersionUE.FileVersionUE5)
+		+ sizeof(FileVersionLicenseeUE);
 }
 #endif
 
diff --git a/Engine/Source/Runtime/CoreUObject/Private/UObject/SavePackage.cpp b/Engine/Source/Runtime/CoreUObject/Private/UObject/SavePackage.cpp
--- a/Engine/Source/Runtime/CoreUObject/Private/UObject/SavePackage.cpp
+++ b/Engine/Source/Runtime/CoreUObject/Private/UObject/SavePackage.cpp
@@ -20,24 +20,6 @@ FSavePackageResultStruct UPackage::Save(UPackage* InOuter, UObject* InAsset, con
 	return UPackage::Save2(InOuter, InAsset, Filename, SaveArgs);
 }
 
-bool UPackage::SavePackage(UPackage* InOuter, UObject* Base, EObjectFlags TopLevelFlags, const TCHAR* Filename,
-	FOutputDevice* Error, FLinkerNull* Conform, bool bForceByteSwapping, bool bWarnOfLongFilename, uint32 SaveFlags,
-	const ITargetPlatform* TargetPlatform, const FDateTime& FinalTimeStamp, bool bSlowTask)
-{
-	// CookData should only be nonzero if we are cooking.
-	TOptional<FArchiveCookContext> CookContext;
-	TOptional<FArchiveCookData> CookData;
-	if (TargetPlatform != nullptr)
-	{
-		CookContext.Emplace(InOuter, UE::Cook::ECookType::Unknown,
-			UE::Cook::ECookingDLC::Unknown, TargetPlatform, nullptr);
-		CookData.Emplace(*TargetPlatform, *CookContext);
-	}
-	FSavePackageArgs SaveArgs = { nullptr /* deprecated target platform */, CookData.GetPtrOrNull(), TopLevelFlags, SaveFlags, bForceByteSwapping,
-		bWarnOfLongFilename, bSlowTask, FinalTimeStamp, Error };
-	return SavePackage(InOuter, Base, Filename, SaveArgs);
-}
-
 bool UPackage::SavePackage(UPackage* InOuter, UObject* InAsset, const TCHAR* Filename, const FSavePackageArgs& SaveArgs)
 {
 	const FSavePackageResultStruct Result = Save(InOuter, InAsset, Filename, SaveArgs);
diff --git a/Engine/Source/Runtime/CoreUObject/Private/UObject/SavePackage/SaveContext.h b/Engine/Source/Runtime/CoreUObject/Private/UObject/SavePackage/SaveContext.h
--- a/Engine/Source/Runtime/CoreUObject/Private/UObject/SavePackage/SaveContext.h
+++ b/Engine/Source/Runtime/CoreUObject/Private/UObject/SavePackage/SaveContext.h
@@ -5,6 +5,7 @@
 #include "AssetRegistry/AssetData.h"
 #include "HAL/FileManager.h"
 #include "HAL/IConsoleManager.h"
+#include "IO/IoHash.h"
 #include "Misc/AssetRegistryInterface.h"
 #include "Misc/Optional.h"
 #include "Misc/PackageName.h"
@@ -814,11 +815,6 @@ public:
 		return !(SaveArgs.SaveFlags & SAVE_NoError);
 	}
 
-	bool IsKeepGuid() const
-	{
-		return !!(SaveArgs.SaveFlags & SAVE_KeepGUID);
-	}
-
 	bool IsKeepDirty() const
 	{
 		return !!(SaveArgs.SaveFlags & SAVE_KeepDirty);
@@ -1362,6 +1358,7 @@ public:
 	EPropertyLocalizationGathererResultFlags GatherableTextResultFlags = EPropertyLocalizationGathererResultFlags::Empty;
 
 	//@note FH: Most of these public members should be moved to harvested realm class 
+	FIoHash PackageSavedHash;
 	int64 PackageHeaderAndExportSize = 0;
 	int64 TotalPackageSizeUncompressed = 0;
 	int32 OffsetAfterPackageFileSummary = 0;
diff --git a/Engine/Source/Runtime/CoreUObject/Private/UObject/SavePackage2.cpp b/Engine/Source/Runtime/CoreUObject/Private/UObject/SavePackage2.cpp
--- a/Engine/Source/Runtime/CoreUObject/Private/UObject/SavePackage2.cpp
+++ b/Engine/Source/Runtime/CoreUObject/Private/UObject/SavePackage2.cpp
@@ -1278,30 +1278,6 @@ ESavePackageResult BuildLinker(FSaveContext& SaveContext)
 	}
 
 #if WITH_EDITORONLY_DATA
-	PRAGMA_DISABLE_DEPRECATION_WARNINGS;
-	const FGuid* OutputPackageGuid = SaveContext.GetSaveArgs().OutputPackageGuid.GetPtrOrNull();
-	PRAGMA_ENABLE_DEPRECATION_WARNINGS;
-	if (OutputPackageGuid)
-	{
-		FIoHash SavedHash;
-		FMemory::Memcpy(&SavedHash.GetBytes(), OutputPackageGuid,
-			FMath::Min(sizeof(decltype(SavedHash.GetBytes())), sizeof(*OutputPackageGuid)));
-		Linker->Summary.SetSavedHash(SavedHash);
-	}
-	else
-	{
-		if (!SaveContext.IsKeepGuid())
-		{
-			// TODO: Change this to instead be calculated later after the save from a hash of the saved data.
-			// Remove IsKeepGuid and OutputPackageGuid when that occurs.
-			FGuid Guid = FGuid::NewGuid();
-			FIoHash SavedHash;
-			FMemory::Memcpy(&SavedHash.GetBytes(), &Guid,
-				FMath::Min(sizeof(decltype(SavedHash.GetBytes())), sizeof(Guid)));
-			SaveContext.GetPackage()->SetSavedHash(SavedHash);
-		}
-		Linker->Summary.SetSavedHash(SaveContext.GetPackage()->GetSavedHash());
-	}
 	Linker->Summary.PersistentGuid = SaveContext.GetPackage()->GetPersistentGuid();
 #endif
 	Linker->Summary.Generations = TArray<FGenerationInfo>{ FGenerationInfo(0, 0) };
@@ -2984,7 +2960,25 @@ ESavePackageResult UpdatePackageHeader(FStructuredArchive::FRecord& StructuredAr
 
 		if (!SaveContext.IsTextFormat())
 		{
+			// Seek to beginning of the file to overwrite the previously-written placeholder Summary.
 			Linker->Seek(0);
+			// For the upcoming calculation of the SavedHash, set a placeholder SavedHash of 0.
+			// We overwrite this in FinalizeFile.
+#if WITH_EDITORONLY_DATA
+			Linker->Summary.SetSavedHash(FIoHash());
+#endif
+		}
+		else
+		{
+			// Calculating a SaveHash for text-based format is not yet implemented, and it's not clear what it should
+			// be; should it be a hash of the text file? Should we save the file to binary to calculate the SaveHash
+			// and then save it to text and write in the calculated SaveHash?
+			// For now, just copy the original SaveHash from the package.
+#if WITH_EDITORONLY_DATA
+			FIoHash SavedHash = SaveContext.GetPackage()->GetSavedHash();
+			SaveContext.PackageSavedHash = SavedHash;
+			Linker->Summary.SetSavedHash(SavedHash);
+#endif
 		}
 		{
 			StructuredArchiveRoot.EnterField(TEXT("Summary")) << Linker->Summary;
@@ -3002,10 +2996,109 @@ ESavePackageResult UpdatePackageHeader(FStructuredArchive::FRecord& StructuredAr
 	return ReturnSuccessOrCancel();
 }
 
+static bool TryCalculateSavedHashAndWriteItBackToMemoryWriter(FLargeMemoryWriter* Writer,
+	FPackageFileSummary& Summary, FStringView PackageDebugName, FIoHash& OutSavedHash)
+{
+#if WITH_EDITORONLY_DATA
+	// We only hash the Header+Exports archive up to PayloadTocOffset. The package trailer is written after that point,
+	// and can be stored in other files rather than appended to the end of the Header+Exports file depending on user
+	// settings. The hashes for the contents of the packageheader are embedded in the BulkDatas that store their bytes
+	// in it, so its bytes are incorporated into our hash of the Header+Exports bytes which include those hashes.
+	//
+	// Our SavedHash does not yet include the data in the CreatePayloadSidecarFile file, that feature is not currently
+	// used and we expect to remove it.
+	//
+	// Our SaveHash does not include data in the separate .ubulk file created when cooking; this is okay for the same
+	// reason as the packagetrailer: the hash of the bytes of the bulk data are included in the bulkdata structs that
+	// we do serialize.
+	int64 HashEnd = Summary.PayloadTocOffset > 0 ? Summary.PayloadTocOffset : Writer->TotalSize();
+
+	OutSavedHash = FIoHash::HashBuffer(Writer->GetData(), HashEnd);
+	Summary.SetSavedHash(OutSavedHash);
+
+	int64 SavedHashOffset = Summary.GetSavedHashRelativeOffset();
+	if (SavedHashOffset < 0 
+		|| SavedHashOffset + static_cast<int64>(sizeof(FIoHash)) > Writer->TotalSize())
+	{
+		UE_LOG(LogSavePackage, Error,
+			TEXT("Error writing PackageSaveHash for '%.*s': GetSavedHashRelativeOffset out of range: %" INT64_FMT " not in (0, %" INT64_FMT ")."),
+			PackageDebugName.Len(), PackageDebugName.GetData(), SavedHashOffset, Writer->TotalSize());
+		return false;
+	}
+
+	int64 SavedWriteOffset = Writer->Tell();
+	Writer->Seek(SavedHashOffset);
+	*Writer << OutSavedHash;
+	Writer->Seek(SavedWriteOffset);
+	return true;
+#else
+	OutSavedHash = FIoHash::Zero;
+	return true;
+#endif
+}
+
+static bool TryCalculateSavedHashAndWriteItBackToFile(const FString& Filename,
+	FPackageFileSummary& Summary, bool bByteSwapping, FStringView PackageDebugName, FIoHash& OutSavedHash)
+{
+#if WITH_EDITORONLY_DATA
+	bool bWroteSavedHash = false;
+
+	// See comment on PayloadTocOffset in TryCalculateSavedHashAndWriteItBackToMemoryWriter.
+
+	int64 HashEnd = Summary.PayloadTocOffset > 0 ? Summary.PayloadTocOffset : -1;
+	if (!UE::SavePackageUtilities::TryHashFile(Filename, OutSavedHash, 0 /* RangeStart */, HashEnd))
+	{
+		UE_LOG(LogSavePackage, Error,
+			TEXT("Error writing PackageSaveHash for '%.*s': Could not hash file %s."),
+			PackageDebugName.Len(), PackageDebugName.GetData(), *Filename);
+		return false;
+	}
+	Summary.SetSavedHash(OutSavedHash);
+
+	TUniquePtr<FArchive> ReopenedFile(IFileManager::Get().CreateFileWriter(*Filename, FILEWRITE_Append));
+	if (!ReopenedFile)
+	{
+		UE_LOG(LogSavePackage, Error,
+			TEXT("Error writing PackageSaveHash for '%.*s': Could not open for writing file %s."),
+			PackageDebugName.Len(), PackageDebugName.GetData(), *Filename);
+		return false;
+	}
+
+	int64 SavedHashOffset = Summary.GetSavedHashRelativeOffset();
+	if (SavedHashOffset < 0
+		|| SavedHashOffset + static_cast<int64>(sizeof(FIoHash)) > ReopenedFile->TotalSize())
+	{
+		UE_LOG(LogSavePackage, Error,
+			TEXT("Error writing PackageSaveHash for '%.*s': GetSavedHashRelativeOffset out of range: %" INT64_FMT " not in (0, %" INT64_FMT ")."),
+			PackageDebugName.Len(), PackageDebugName.GetData(), SavedHashOffset, ReopenedFile->TotalSize());
+		return false;
+	}
+
+	ReopenedFile->SetByteSwapping(bByteSwapping);
+	ReopenedFile->Seek(SavedHashOffset);
+	*ReopenedFile << OutSavedHash;
+
+	if (ReopenedFile->IsError())
+	{
+		UE_LOG(LogSavePackage, Error,
+			TEXT("Error writing PackageSaveHash for '%.*s': Write failed to file %s."),
+			PackageDebugName.Len(), PackageDebugName.GetData(), *Filename);
+		return false;
+	}
+
+	return true;
+#else
+	OutSavedHash = FIoHash::Zero;
+	return true;
+#endif
+}
+
 ESavePackageResult FinalizeFile(FStructuredArchive::FRecord& StructuredArchiveRoot, FSaveContext& SaveContext)
 {
 	SCOPED_SAVETIMER(UPackage_Save_FinalizeFile);
 
+	FIoHash SavedHash;
+
 	FSavePackageContext* SavePackageContext = SaveContext.GetSavePackageContext();
 	IPackageWriter* PackageWriter = SavePackageContext ? SavePackageContext->PackageWriter : nullptr;
 	bool bIsOptionalRealm = SaveContext.GetCurrentHarvestingRealm() == ESaveRealm::Optional;
@@ -3013,6 +3106,11 @@ ESavePackageResult FinalizeFile(FStructuredArchive::FRecord& StructuredArchiveRo
 	if (PackageWriter)
 	{
 		FLargeMemoryWriter* Writer = static_cast<FLargeMemoryWriter*>(Linker->Saver);
+		if (!TryCalculateSavedHashAndWriteItBackToMemoryWriter(Writer, Linker->Summary,
+			SaveContext.GetFilename(), SavedHash))
+		{
+			return ESavePackageResult::Error;
+		}
 
 		IPackageWriter::FPackageInfo PackageInfo;
 		PackageInfo.PackageName = SaveContext.GetPackage()->GetFName();
@@ -3043,6 +3141,11 @@ ESavePackageResult FinalizeFile(FStructuredArchive::FRecord& StructuredArchiveRo
 		if (SaveContext.IsSaveToMemory())
 		{
 			FLargeMemoryWriter* Writer = static_cast<FLargeMemoryWriter*>(Linker->Saver);
+			if (!TryCalculateSavedHashAndWriteItBackToMemoryWriter(Writer, Linker->Summary,
+				SaveContext.GetFilename(), SavedHash))
+			{
+				return ESavePackageResult::Error;
+			}
 
 			EAsyncWriteOptions WriteOptions(EAsyncWriteOptions::None);
 
@@ -3071,6 +3174,12 @@ ESavePackageResult FinalizeFile(FStructuredArchive::FRecord& StructuredArchiveRo
 				return ESavePackageResult::Error;
 			}
 
+			if (!TryCalculateSavedHashAndWriteItBackToFile(*TempFilename, Linker->Summary,
+				Linker->IsByteSwapping(), SaveContext.GetFilename(), SavedHash))
+			{
+				return ESavePackageResult::Error;
+			}
+
 			// Move file to its real destination
 
 			// Add the .uasset file to the list of output files
@@ -3110,6 +3219,9 @@ ESavePackageResult FinalizeFile(FStructuredArchive::FRecord& StructuredArchiveRo
 				return ESavePackageResult::Error;
 			}
 
+			// SaveHash was written when writing the PackageSummary in UpdatePackageHeader.
+			SavedHash = SaveContext.PackageSavedHash;
+
 			// When saving in text format we will have two temp files, so we need to manually delete the non-textbased one
 			check(SaveContext.GetTempFilename().IsSet());
 			check(SaveContext.GetTextFormatTempFilename().IsSet());
@@ -3142,6 +3254,11 @@ ESavePackageResult FinalizeFile(FStructuredArchive::FRecord& StructuredArchiveRo
 		}
 	}
 
+	if (SaveContext.GetCurrentHarvestingRealm() != ESaveRealm::Optional)
+	{
+		SaveContext.PackageSavedHash = SavedHash;
+	}
+
 	return ESavePackageResult::Success;
 }
 
@@ -3225,6 +3342,12 @@ void PostSavePackage(FSaveContext& SaveContext)
 		return;
 	}
 
+	// Clear dirty flag if desired
+	if (!SaveContext.IsKeepDirty())
+	{
+		SaveContext.GetPackage()->SetDirtyFlag(false);
+	}
+
 	// then adjust flags that should be modified at the outcome of a save
 	if (!SaveContext.IsFromAutoSave() && !SaveContext.IsProceduralSave())
 	{
@@ -3269,6 +3392,21 @@ void PostSavePackage(FSaveContext& SaveContext)
 	}
 #endif
 
+	// update data on the Package about its disk data (path, savedhash, etc) if we're saving to a valid mounted path
+	// and we aren't currently cooking
+	const FPackagePath& PackagePath = SaveContext.GetTargetPackagePath();
+	if (SaveContext.IsUpdatingLoadedPath())
+	{
+#if WITH_EDITOR
+		Package->SetLoadedPath(PackagePath);
+#endif
+		// Calling a function on SaveContext is required because SetFileSize is private and SaveContext is a friend.
+		SaveContext.UpdatePackageFileSize(SaveContext.PackageHeaderAndExportSize);
+#if WITH_EDITORONLY_DATA
+		Package->SetSavedHash(SaveContext.PackageSavedHash);
+#endif
+	}
+
 	// Send a message that the package was saved
 	PRAGMA_DISABLE_DEPRECATION_WARNINGS;
 	UPackage::PackageSavedEvent.Broadcast(SaveContext.GetFilename(), Package);
@@ -3276,15 +3414,6 @@ void PostSavePackage(FSaveContext& SaveContext)
 	FObjectSaveContextData& ObjectSaveContext = SaveContext.GetObjectSaveContext();
 	ObjectSaveContext.Object = nullptr;
 	UPackage::PackageSavedWithContextEvent.Broadcast(SaveContext.GetFilename(), Package, FObjectPostSaveContext(ObjectSaveContext));
-
-	// update the internal package filename path if we're saving to a valid mounted path and we aren't currently cooking
-#if WITH_EDITOR
-	const FPackagePath& PackagePath = SaveContext.GetTargetPackagePath();
-	if (SaveContext.IsUpdatingLoadedPath())
-	{
-		Package->SetLoadedPath(PackagePath);
-	}
-#endif
 }
 
 ESavePackageResult SaveHarvestedRealms(FSaveContext& SaveContext, ESaveRealm HarvestingContextToSave)
@@ -3707,23 +3836,6 @@ FSavePackageResultStruct UPackage::Save2(UPackage* InPackage, UObject* InAsset,
 		// in case of failure or cancellation, do not exit here, still run cleanup (e.g. PostSaveRoot)
 	}
 
-	SlowTask.EnterProgressFrame();
-	// If the save was successful, update dirty flag and file size if desired
-	if (SaveContext.Result == ESavePackageResult::Success)
-	{
-		// Clear dirty flag if desired
-		if (!SaveContext.IsKeepDirty())
-		{
-			SaveContext.GetPackage()->SetDirtyFlag(false);
-		}
-
-		// Update package FileSize value
-		if (SaveContext.IsUpdatingLoadedPath())
-		{
-			SaveContext.UpdatePackageFileSize(SaveContext.PackageHeaderAndExportSize);
-		}
-	}
-
 	// PostSave Asset
 	SlowTask.EnterProgressFrame();
 	if (SaveContext.GetPostSaveRootRequired() && SaveContext.GetAsset())
@@ -3812,22 +3924,6 @@ ESavePackageResult UPackage::SaveConcurrent(TArrayView<FPackageSaveInfo> InPacka
 		SCOPED_SAVETIMER(UPackage_SaveConcurrent_PostSave);
 		for (FSaveContext& SaveContext : PackageSaveContexts)
 		{
-			// If the save was successful, update dirty flag and file size if desired
-			if (SaveContext.Result == ESavePackageResult::Success)
-			{
-				// Clear dirty flag if desired
-				if (!SaveContext.IsKeepDirty())
-				{
-					SaveContext.GetPackage()->SetDirtyFlag(false);
-				}
-
-				// Update package FileSize value
-				if (SaveContext.IsUpdatingLoadedPath())
-				{
-					SaveContext.UpdatePackageFileSize(SaveContext.PackageHeaderAndExportSize);
-				}
-			}
-
 			// PostSave Asset
 			if (SaveContext.GetAsset())
 			{
@@ -3836,10 +3932,7 @@ ESavePackageResult UPackage::SaveConcurrent(TArrayView<FPackageSaveInfo> InPacka
 			}
 
 			// PostSave Package - edit in memory package and send events
-			if (SaveContext.Result == ESavePackageResult::Success)
-			{
-				PostSavePackage(SaveContext);
-			}
+			PostSavePackage(SaveContext);
 			OutResults.Add(SaveContext.GetFinalResult());
 		}
 	}
diff --git a/Engine/Source/Runtime/CoreUObject/Public/UObject/ObjectMacros.h b/Engine/Source/Runtime/CoreUObject/Public/UObject/ObjectMacros.h
--- a/Engine/Source/Runtime/CoreUObject/Public/UObject/ObjectMacros.h
+++ b/Engine/Source/Runtime/CoreUObject/Public/UObject/ObjectMacros.h
@@ -88,15 +88,15 @@ enum ESaveFlags
 	SAVE_NoError					= 0x00000001,	///< Don't generate errors on save
 	SAVE_FromAutosave				= 0x00000002,   ///< Used to indicate this save was initiated automatically
 	SAVE_KeepDirty					= 0x00000004,	///< Do not clear the dirty flag when saving
-	SAVE_KeepGUID					= 0x00000008,	///< Keep the same guid, used to save cooked packages
+	SAVE_KeepGUID UE_DEPRECATED(5.6, "KeepPackageGUIDOnSave is no longer supported, because package save no longer records arbitrary guids but instead records a hash of the package's bytes - PackageSavedHash.") = 0x00000008,
 	SAVE_Async						= 0x00000010,	///< Save to a memory writer, then actually write to disk async
 	SAVE_Unversioned_Native			= 0x00000020,	///< Save all versions as zero. Upon load this is changed to the current version. This is only reasonable to use with full cooked builds for distribution.
 	SAVE_CutdownPackage				= 0x00000040,	///< Saving cutdown packages in a temp location WITHOUT renaming the package.
 	SAVE_KeepEditorOnlyCookedPackages = 0x00000080,  ///< Keep packages which are marked as editor only even though we are cooking
 	SAVE_Concurrent					= 0x00000100,	///< We are save packages in multiple threads at once and should not call non-threadsafe functions or rely on globals. GIsSavingPackage should be set and PreSave/Postsave functions should be called before/after the entire concurrent save.
-	SAVE_DiffOnly UE_DEPRECATED(5.0, "Diffing is now done using FDiffPackageWriter.") = 0x00000200,	
-	SAVE_DiffCallstack UE_DEPRECATED(5.0, "Diffing is now done using FDiffPackageWriter.") = 0x00000400,
-	SAVE_ComputeHash UE_DEPRECATED(5.1, "Add EWriteOptions::ComputeHash directly to CommitPackage options") = 0x00000800,	///< Compute the MD5 hash of the cooked data, as well as the per iostore chunk FIoHash.
+//	SAVE_Unused						= 0x00000200,
+//	SAVE_Unused						= 0x00000400,
+//	SAVE_Unused						= 0x00000800,
 	SAVE_CompareLinker				= 0x00001000,	///< Return the linker save to compare against another
 	SAVE_BulkDataByReference		= 0x00002000,	///< When saving to a different file than the package's LoadedPath, point bulkdata in the new file to be loaded from the original file.
 	SAVE_Unversioned_Properties		= 0x00004000,	///< Properties are saved without property name information, and are saved/loaded in the order of the current binary.
diff --git a/Engine/Source/Runtime/CoreUObject/Public/UObject/Package.h b/Engine/Source/Runtime/CoreUObject/Public/UObject/Package.h
--- a/Engine/Source/Runtime/CoreUObject/Public/UObject/Package.h
+++ b/Engine/Source/Runtime/CoreUObject/Public/UObject/Package.h
@@ -9,6 +9,7 @@
 #include "CoreMinimal.h"
 #include "Delegates/Delegate.h"
 #include "HAL/PlatformMath.h"
+#include "IO/IoHash.h"
 #include "IO/PackageId.h"
 #include "Misc/AssertionMacros.h"
 #include "Misc/DateTime.h"
@@ -280,11 +281,7 @@ private:
 	/** Chunk IDs for the streaming install chunks this package will be placed in.  Empty for no chunk. Used during cooking. */
 	TArray<int32> ChunkIDs;
 
-#if !UE_STRIP_DEPRECATED_PROPERTIES
-	/** GUID of package if it was loaded from disk. Changes at every save. */
-	UE_DEPRECATED(5.4, "Use GetSavedHash/SetSavedHash instead.")
-	FGuid Guid;
-#endif
+	FIoHash SavedHash;
 #endif // WITH_EDITORONLY_DATA
 
 	/** Package Flags */
@@ -900,38 +897,6 @@ public:
 		}
 	}
 
-#if WITH_EDITORONLY_DATA
-	UE_DEPRECATED(5.4, "Use GetSavedHash/SetSavedHash instead.")
-	FORCEINLINE FGuid GetGuid() const
-	{
-#if !UE_STRIP_DEPRECATED_PROPERTIES
-		return Guid;
-#else
-		return FGuid();
-#endif
-	
-	}
-	/** makes our a new fresh Guid */
-	UE_DEPRECATED(5.4, "Use GetSavedHash/SetSavedHash instead.")
-	FORCEINLINE FGuid MakeNewGuid()
-	{
-#if !UE_STRIP_DEPRECATED_PROPERTIES
-		Guid = FGuid::NewGuid();
-		return Guid;
-#else
-		return FGuid();
-#endif
-	}
-	/** sets a specific Guid */
-	UE_DEPRECATED(5.4, "Use GetSavedHash/SetSavedHash instead.")
-	FORCEINLINE void SetGuid(FGuid NewGuid)
-	{
-#if !UE_STRIP_DEPRECATED_PROPERTIES
-		Guid = NewGuid;
-#endif
-	}
-#endif // WITH_EDITORONLY_DATA
-
 	/** returns our FileSize */
 	FORCEINLINE int64 GetFileSize() const
 	{
@@ -1036,13 +1001,6 @@ public:
 	COREUOBJECT_API static bool SavePackage(UPackage* InOuter, UObject* InAsset, const TCHAR* Filename,
 		const FSavePackageArgs& SaveArgs);
 
-	UE_DEPRECATED(5.0, "Pack the arguments into FSavePackageArgs and call the function overload that takes FSavePackageArgs. Note that Conform is no longer implemented.")
-	COREUOBJECT_API static bool SavePackage(UPackage* InOuter, UObject* Base, EObjectFlags TopLevelFlags,
-		const TCHAR* Filename, FOutputDevice* Error = GError, FLinkerNull* Conform = nullptr,
-		bool bForceByteSwapping = false, bool bWarnOfLongFilename = true, uint32 SaveFlags = SAVE_None,
-		const ITargetPlatform* TargetPlatform = nullptr, const FDateTime& FinalTimeStamp = FDateTime::MinValue(),
-		bool bSlowTask = true);
-
 	/** Wait for any SAVE_Async file writes to complete **/
 	COREUOBJECT_API static void WaitForAsyncFileWrites();
 
diff --git a/Engine/Source/Runtime/CoreUObject/Public/UObject/PackageFileSummary.h b/Engine/Source/Runtime/CoreUObject/Public/UObject/PackageFileSummary.h
--- a/Engine/Source/Runtime/CoreUObject/Public/UObject/PackageFileSummary.h
+++ b/Engine/Source/Runtime/CoreUObject/Public/UObject/PackageFileSummary.h
@@ -6,6 +6,7 @@
 #include "Containers/UnrealString.h"
 #include "CoreMinimal.h"
 #include "CoreTypes.h"
+#include "IO/IoHash.h"
 #include "Misc/EngineVersion.h"
 #include "Misc/Guid.h"
 #include "Serialization/CustomVersion.h"
@@ -191,9 +192,13 @@ public:
 	*/
 	int32		ThumbnailTableOffset;
 
-	UE_DEPRECATED(5.4, "Use GetSavedHash/SetSavedHash instead.")
-	FGuid	Guid;
+private:
+	/**
+	 * Hash of the Package's bytes when it was saved to disk. Accessed through GetSavedHash/SetSavedHash.
+	 */
+	FIoHash		SavedHash;
 
+public:
 #if WITH_EDITORONLY_DATA
 	/**
 	* Current persistent id for this package
@@ -361,6 +366,7 @@ public:
 	/** Hash of the package's .uasset/.umap file when it was last saved by the editor. */
 	COREUOBJECT_API FIoHash GetSavedHash() const;
 	COREUOBJECT_API void SetSavedHash(const FIoHash& InSavedHash);
+	COREUOBJECT_API int64 GetSavedHashRelativeOffset() const;
 #endif
 
 	/** I/O functions */
diff --git a/Engine/Source/Runtime/CoreUObject/Public/UObject/SavePackage.h b/Engine/Source/Runtime/CoreUObject/Public/UObject/SavePackage.h
--- a/Engine/Source/Runtime/CoreUObject/Public/UObject/SavePackage.h
+++ b/Engine/Source/Runtime/CoreUObject/Public/UObject/SavePackage.h
@@ -61,9 +61,6 @@ struct FPackageSaveInfo
  */
 struct FSavePackageArgs
 {
-	UE_DEPRECATED(5.1, "TargetPlatform is now on the ArchiveCookData, use GetTargetPlatform() or IsCooking()")
-	const ITargetPlatform* TargetPlatform = nullptr;
-
 	/* nullptr if not cooking, passed to the FArchive */
 	FArchiveCookData* ArchiveCookData = nullptr;
 
@@ -93,8 +90,6 @@ struct FSavePackageArgs
 	FOutputDevice* Error = GError;
 	/** Structure to hold longer-lifetime parameters that apply to multiple saves */
 	FSavePackageContext* SavePackageContext = nullptr;
-	UE_DEPRECATED(4.27, "UPackage::Guid has not been used by the engine for a long time and it will be removed.")
-	TOptional<FGuid> OutputPackageGuid;
 
 	/**
 	 * In/Out list of property overrides per object to apply to during save. This list can be extended by PreSave functions
@@ -102,13 +97,13 @@ struct FSavePackageArgs
 	 */
 	TMap<UObject*, FObjectSaveOverride>* InOutSaveOverrides = nullptr;
 
-	PRAGMA_DISABLE_DEPRECATION_WARNINGS;
 	FSavePackageArgs() = default;
 	FSavePackageArgs(const FSavePackageArgs&) = default;
 	FSavePackageArgs(FSavePackageArgs&&) = default;
 	FSavePackageArgs& operator=(const FSavePackageArgs&) = default;
 	FSavePackageArgs& operator=(FSavePackageArgs&&) = default;
 
+	UE_DEPRECATED(5.6, "Use the default constructor and assign elements individually.")
 	FSavePackageArgs(
 		const ITargetPlatform* InTargetPlatform,
 		FArchiveCookData* InArchiveCookData,
@@ -121,8 +116,7 @@ struct FSavePackageArgs
 		FOutputDevice* InError,
 		FSavePackageContext* InSavePackageContext = nullptr
 	)
-	: TargetPlatform(InTargetPlatform)
-	, ArchiveCookData(InArchiveCookData)
+	: ArchiveCookData(InArchiveCookData)
 	, TopLevelFlags(InTopLevelFlags)
 	, SaveFlags(InSaveFlags)
 	, bForceByteSwapping(bInForceByteSwapping)
@@ -133,7 +127,6 @@ struct FSavePackageArgs
 	, SavePackageContext(InSavePackageContext)
 	{
 	}
-	PRAGMA_ENABLE_DEPRECATION_WARNINGS;
 };
 
 /** Interface for SavePackage to test for caller-specific errors. */
