Added a new define UE_UOBJECT_HASH_USES_ARRAYS that replaces FHashBucket in UObjectHash  implementation with a one that uses arrays instead of sets.
This reduced memory consuption by ~20%

diff --git a/Engine/Source/Runtime/CoreUObject/Private/UObject/UObjectHash.cpp b/Engine/Source/Runtime/CoreUObject/Private/UObject/UObjectHash.cpp
--- a/Engine/Source/Runtime/CoreUObject/Private/UObject/UObjectHash.cpp
+++ b/Engine/Source/Runtime/CoreUObject/Private/UObject/UObjectHash.cpp
@@ -25,10 +25,14 @@
 
 #include <atomic>
 
+#ifndef UE_UOBJECT_HASH_USES_ARRAYS
+#	define UE_UOBJECT_HASH_USES_ARRAYS 0
+#endif
+
+
 #if UE_STORE_OBJECT_LIST_INTERNAL_INDEX
-static_assert(sizeof(FName) == 4, "Internal object index optimization depends exploits 4 bytes padding after the FName");
-static_assert(sizeof(UObjectBase) == 40,
-	"UObjectBase size has changed! This optimization relies on a 4 bytes padding after NamePrivate.Please make sure that the new binary layout it optimal");
+	static_assert(sizeof(FName) == 4, "Internal object index optimization depends exploits 4 bytes padding after the FName");
+	static_assert(sizeof(UObjectBase) == 40, "UObjectBase size has changed! This optimization relies on a 4 bytes padding after NamePrivate.Please make sure that the new binary layout it optimal");
 #endif
 
 DEFINE_LOG_CATEGORY_STATIC(LogUObjectHash, Log, All);
@@ -36,9 +40,9 @@ DEFINE_LOG_CATEGORY_STATIC(LogUObjectHash, Log, All);
 DECLARE_CYCLE_STAT( TEXT( "GetObjectsOfClass" ), STAT_Hash_GetObjectsOfClass, STATGROUP_UObjectHash );
 
 #if !UE_BUILD_TEST && !UE_BUILD_SHIPPING
-DECLARE_CYCLE_STAT( TEXT( "HashObject" ), STAT_Hash_HashObject, STATGROUP_UObjectHash );
-DECLARE_CYCLE_STAT( TEXT( "UnhashObject" ), STAT_Hash_UnhashObject, STATGROUP_UObjectHash );
-DEFINE_STAT( STAT_Hash_NumObjects );
+	DECLARE_CYCLE_STAT( TEXT( "HashObject" ), STAT_Hash_HashObject, STATGROUP_UObjectHash );
+	DECLARE_CYCLE_STAT( TEXT( "UnhashObject" ), STAT_Hash_UnhashObject, STATGROUP_UObjectHash );
+	DEFINE_STAT( STAT_Hash_NumObjects );
 #endif // !UE_BUILD_TEST && !UE_BUILD_SHIPPING
 
 LLM_DEFINE_TAG(UObjectHash, TEXT("UObject hashtables"));
@@ -68,228 +72,560 @@ static FORCENOINLINE void OnHashFailure(UObjectBaseUtility* Object, const TCHAR*
  * The hash is rarely used at runtime. A more space-efficient implementation is possible.
  */
 
+template <typename T>
+struct THashBucketIterator;
 
 /*
  * Special hash bucket to conserve memory.
  * Contains a pointer to head element and an optional list of items if more than one element exists in the bucket.
  * The item list is only allocated if needed.
  */
-struct FHashBucket
+struct FSetHashBucket
 {
-	friend struct FHashBucketIterator;
+	typedef TSet<UObjectBase*>::TIterator TIterator;
 
 	/** This always empty set is used to get an iterator if the bucket doesn't use a TSet (has only 1 element) */
 	static TSet<UObjectBase*> EmptyBucket;
 
 	/*
+	* Acts like an alias for either 1-2 elements or a pointer to a TSet
 	* If these are both null, this bucket is empty
 	* If the first one is null, but the second one is non-null, then the second one is a TSet pointer
 	* If the first one is not null, then it is a uobject ptr, and the second ptr is either null or a second element
 	*/
-	void *ElementsOrSetPtr[2];
+	void* Elements[2];
 
-	FORCEINLINE TSet<UObjectBase*>* GetSet()
+	FORCEINLINE TSet<UObjectBase*>* GetItems()
 	{
-		if (ElementsOrSetPtr[1] && !ElementsOrSetPtr[0])
+		if (Elements[1] && !Elements[0])
 		{
-			return (TSet<UObjectBase*>*)ElementsOrSetPtr[1];
+			return (TSet<UObjectBase*>*)Elements[1];
 		}
 		return nullptr;
 	}
 
-	FORCEINLINE const TSet<UObjectBase*>* GetSet() const
+	FORCEINLINE const TSet<UObjectBase*>* GetItems() const
 	{
-		if (ElementsOrSetPtr[1] && !ElementsOrSetPtr[0])
+		if (Elements[1] && !Elements[0])
 		{
-			return (TSet<UObjectBase*>*)ElementsOrSetPtr[1];
+			return (TSet<UObjectBase*>*)Elements[1];
 		}
 		return nullptr;
 	}
 
 	/** Constructor */
-	FORCEINLINE FHashBucket()
+	FORCEINLINE FSetHashBucket()
 	{
-		ElementsOrSetPtr[0] = nullptr;
-		ElementsOrSetPtr[1] = nullptr;
+		Elements[0] = nullptr;
+		Elements[1] = nullptr;
 	}
-	FORCEINLINE ~FHashBucket()
+
+	FORCEINLINE ~FSetHashBucket()
 	{
-		delete GetSet();
+		delete GetItems();
 	}
+
 	/** Adds an Object to the bucket */
 	FORCEINLINE void Add(UObjectBase* Object)
 	{
-		TSet<UObjectBase*>* Items = GetSet();
+		TSet<UObjectBase*>* Items = GetItems();
 		if (Items)
 		{
 			Items->Add(Object);
 		}
-		else if (ElementsOrSetPtr[0] && ElementsOrSetPtr[1])
+		else if (Elements[0] && Elements[1])
 		{
 			Items = new TSet<UObjectBase*>();
-			Items->Add((UObjectBase*)ElementsOrSetPtr[0]);
-			Items->Add((UObjectBase*)ElementsOrSetPtr[1]);
+			Items->Add((UObjectBase*)Elements[0]);
+			Items->Add((UObjectBase*)Elements[1]);
 			Items->Add(Object);
-			ElementsOrSetPtr[0] = nullptr;
-			ElementsOrSetPtr[1] = Items;
+			Elements[0] = nullptr;
+			Elements[1] = Items;
 		}
-		else if (ElementsOrSetPtr[0])
+		else if (Elements[0])
 		{
-			ElementsOrSetPtr[1] = Object;
+			Elements[1] = Object;
 		}
 		else
 		{
-			ElementsOrSetPtr[0] = Object;
-			checkSlow(!ElementsOrSetPtr[1]);
+			Elements[0] = Object;
+			checkSlow(!Elements[1]);
 		}
 	}
+
 	/** Removes an Object from the bucket */
 	FORCEINLINE int32 Remove(UObjectBase* Object)
 	{
 		int32 Result = 0;
-		TSet<UObjectBase*>* Items = GetSet();
+		TSet<UObjectBase*>* Items = GetItems();
 		if (Items)
 		{
 			Result = Items->Remove(Object);
 			if (Items->Num() <= 2)
 			{
 				auto It = TSet<UObjectBase*>::TIterator(*Items);
-				ElementsOrSetPtr[0] = *It;
+				Elements[0] = *It;
 				checkSlow((bool)It);
 				++It;
-				ElementsOrSetPtr[1] = *It;
+				Elements[1] = *It;
 				delete Items;
 			}
 		}
-		else if (Object == ElementsOrSetPtr[1])
+		else if (Object == Elements[1])
 		{
 			Result = 1;
-			ElementsOrSetPtr[1] = nullptr;
+			Elements[1] = nullptr;
 		}
-		else if (Object == ElementsOrSetPtr[0])
+		else if (Object == Elements[0])
 		{
 			Result = 1;
-			ElementsOrSetPtr[0] = ElementsOrSetPtr[1];
-			ElementsOrSetPtr[1] = nullptr;
+			Elements[0] = Elements[1];
+			Elements[1] = nullptr;
 		}
 		return Result;
 	}
+
 	/** Checks if an Object exists in this bucket */
 	FORCEINLINE bool Contains(UObjectBase* Object) const
 	{
-		const TSet<UObjectBase*>* Items = GetSet();
+		const TSet<UObjectBase*>* Items = GetItems();
 		if (Items)
 		{
 			return Items->Contains(Object);
 		}
-		return Object == ElementsOrSetPtr[0] || Object == ElementsOrSetPtr[1];
+		return Object == Elements[0] || Object == Elements[1];
 	}
+
 	/** Returns the number of Objects in this bucket */
 	FORCEINLINE int32 Num() const
 	{
-		const TSet<UObjectBase*>* Items = GetSet();
+		const TSet<UObjectBase*>* Items = GetItems();
 		if (Items)
 		{
 			return Items->Num();
 		}
-		return !!ElementsOrSetPtr[0] + !!ElementsOrSetPtr[1];
+		return !!Elements[0] + !!Elements[1];
 	}
+
 	/** Returns the amount of memory allocated for and by Items TSet */
 	FORCEINLINE SIZE_T GetAllocatedSize() const
 	{
-		const TSet<UObjectBase*>* Items = GetSet();
+		const TSet<UObjectBase*>* Items = GetItems();
 		if (Items)
 		{
 			return sizeof(*Items) + Items->GetAllocatedSize();
 		}
 		return 0;
 	}
+
 	void Shrink()
 	{
-		TSet<UObjectBase*>* Items = GetSet();
+		TSet<UObjectBase*>* Items = GetItems();
 		if (Items)
 		{
 			Items->Compact();
 		}
 	}
-	FORCEINLINE struct FHashBucketIterator CreateIterator();
 
-private:
+	FORCEINLINE struct THashBucketIterator<FSetHashBucket> CreateIterator();
+
 	/** Gets an iterator for the TSet in this bucket or for the EmptyBucker if Items is null */
-	FORCEINLINE TSet<UObjectBase*>::TIterator GetIteratorForSet()
+	FORCEINLINE TSet<UObjectBase*>::TIterator GetIterator()
 	{
-		TSet<UObjectBase*>* Items = GetSet();
+		TSet<UObjectBase*>* Items = GetItems();
 		return Items ? Items->CreateIterator() : EmptyBucket.CreateIterator();
 	}
 };
-TSet<UObjectBase*> FHashBucket::EmptyBucket;
+
+TSet<UObjectBase*> FSetHashBucket::EmptyBucket;
+
+
+typedef TArray<UObjectBase*> ObjectsArray;
+
+struct FSortedArray : public ObjectsArray
+{
+	FORCEINLINE int32 Insert(UObjectBase* Object)
+	{
+		return ObjectsArray::Insert(Object, Algo::UpperBound(*(ObjectsArray*)(this), Object));
+	}
+
+	FORCEINLINE bool Remove(UObjectBase* Object)
+	{
+		const uint32 Index = FindIndex(Object);
+		if (Index != INDEX_NONE)
+		{
+			const int FillRatio = Num() * 100 / Max();
+			RemoveAt(Index, 1, FillRatio <= 75 ? EAllowShrinking::Yes : EAllowShrinking::No);
+
+			return true;
+		}
+
+		return false;
+	}
+
+	FORCEINLINE bool Contains(const UObjectBase* Object) const
+	{
+		return FindIndex(Object) != INDEX_NONE;
+	}
+
+	FORCEINLINE uint32 FindIndex(const UObjectBase* Object) const
+	{
+		return Algo::BinarySearch(*(ObjectsArray*)(this), Object);
+	}
+
+	FSortedArray()
+	{
+		static_assert(offsetof(FSortedArray, ArrayNum) == sizeof(UObjectBase*));	// make sure that offset of ArrayNum and ArrayMax for TArray is at 8 bytes, to be aliased with Elements[1]
+	}
+};
+
+struct FUnsortedArray : public ObjectsArray
+{
+	FORCEINLINE int32 Insert(UObjectBase* Object)
+	{
+		return Add(Object);
+	}
+
+	FORCEINLINE bool Remove(UObjectBase* Object)
+	{
+		const uint32 Index = FindIndex(Object);
+		if (Index != INDEX_NONE)
+		{
+			const int FillRatio = Num() * 100 / Max();
+			RemoveAt(Index, 1, FillRatio <= 75 ? EAllowShrinking::Yes : EAllowShrinking::No);
+
+			return true;
+		}
+
+		return false;
+	}
+
+	FORCEINLINE uint32 FindIndex(UObjectBase* Object) const
+	{
+		return ObjectsArray::Find(Object);
+	}
+
+	FUnsortedArray()
+	{
+		static_assert(offsetof(FUnsortedArray, ArrayNum) == sizeof(UObjectBase*));	// make sure that offset of ArrayNum and ArrayMax for TArray is at 8 bytes, to be aliased with Elements[1]
+	}
+};
+
+/*
+ * Special hash bucket to conserve memory.
+ * Contains a pointer to head element or an optional list of items if more than one element exists in the bucket.
+ */
+template <typename T>
+struct TArrayHashBucket
+{
+	typedef T FItems;
+	typedef T::TIterator TIterator;
+
+	/** This always empty array is used to get an iterator if the bucket doesn't use a TArray (has only 1 element) */
+	static FItems EmptyBucket;
+
+	/*
+	 * A union that contains one element or acts as a binary alias for a TArray
+	 * TArray's binary representation is { UObjectBase* ptr, uint32 ArrayNum, uint32 ArrayMax }
+	 * So the first Element is always a pointer to either UObjectBase itself or to an array of UObjectBase
+	 * If the second Element is non-null, then ArrayNum is initialized and the whole union is a TArray now
+	 */
+	union
+	{
+		FItems ObjectsList;
+		struct
+		{
+			UObjectBase* Elements[2];
+		};
+	};
+
+	static_assert(sizeof(ObjectsList) == sizeof(Elements));
+
+	FORCEINLINE FItems* GetItems()
+	{
+		if (ObjectsList.Num())
+		{
+			return &ObjectsList;
+		}
+		return nullptr;
+	}
+
+	FORCEINLINE const FItems* GetItems() const
+	{
+		if (ObjectsList.Num())
+		{
+			return &ObjectsList;
+		}
+		return nullptr;
+	}
+
+	FORCEINLINE TArrayHashBucket()
+		: Elements{ nullptr, nullptr }
+	{
+	}
+
+	FORCEINLINE TArrayHashBucket(const TArrayHashBucket& Other)
+		: Elements{ Other.Elements[0], Other.Elements[1] }
+	{
+	}
+
+	FORCEINLINE ~TArrayHashBucket()
+	{
+		FItems* List = GetItems();
+		if (List)
+		{
+			DestroyList(List);
+		}
+	}
+
+	FORCEINLINE void DestroyList(FItems* List)
+	{
+		List->~FItems();
+	}
+
+	/** Adds an Object to the bucket */
+	FORCEINLINE int32 Add(UObjectBase* Object)
+	{
+		FItems* Items = GetItems();
+		if (Items)
+		{
+			return Items->Insert(Object);
+		}
+		else if (Elements[0])
+		{
+			UObjectBase* Tmp = Elements[0];
+			new (&ObjectsList) FItems();
+			ObjectsList.Add(Tmp);
+			return ObjectsList.Insert(Object);
+		}
+		else
+		{
+			Elements[0] = Object;
+			checkSlow(!Elements[1]);
+			return 0;
+		}
+	}
+
+	/** Removes an Object from the bucket */
+	FORCEINLINE int32 Remove(UObjectBase* Object)
+	{
+		FItems* Items = GetItems();
+		if (Items)
+		{
+			if (Items->Remove(Object))
+			{
+				if (Items->Num() == 1)
+				{
+					UObjectBase* Tmp = (*Items)[0];
+					DestroyList(Items);
+					Elements[0] = Tmp;
+					Elements[1] = nullptr;
+				}
+
+				return 1;
+			}
+		}
+		else if (Object == Elements[0])
+		{
+			Elements[0] = nullptr;
+			return 1;
+		}
+		return 0;
+	}
+
+	/** Checks if an Object exists in this bucket */
+	FORCEINLINE bool Contains(UObjectBase* Object) const
+	{
+		const FItems* Items = GetItems();
+		if (Items)
+		{
+			return Items->Contains(Object);
+		}
+		return Object == Elements[0];
+	}
+
+	/** Returns the number of Objects in this bucket */
+	FORCEINLINE int32 Num() const
+	{
+		const FItems* Items = GetItems();
+		if (Items)
+		{
+			return Items->Num();
+		}
+		return !!Elements[0];
+	}
+
+	/** Returns the amount of memory allocated for and by Items TArray */
+	FORCEINLINE SIZE_T GetAllocatedSize() const
+	{
+		const FItems* Items = GetItems();
+		if (Items)
+		{
+			return Items->GetAllocatedSize();
+		}
+		return 0;
+	}
+
+	void Shrink()
+	{
+		FItems* Items = GetItems();
+		if (Items)
+		{
+			Items->Shrink();
+		}
+	}
+
+	FORCEINLINE struct THashBucketIterator<TArrayHashBucket<T>> CreateIterator();
+
+#if UE_STORE_OBJECT_LIST_INTERNAL_INDEX
+	// special set of functions to support this code path
+	FORCEINLINE UObjectBase*& Last()
+	{
+		FItems* Items = GetItems();
+		if (Items)
+		{
+			return Items->Last();
+		}
+
+		return Elements[0];
+	}
+
+	FORCEINLINE void Pop(EAllowShrinking AllowShrinking = EAllowShrinking::Default)
+	{
+		FItems* Items = GetItems();
+		if (Items)
+		{
+			Items->Pop(AllowShrinking);
+
+			if (Items->Num() == 0)
+			{
+				DestroyList(Items);
+				Elements[0] = nullptr;
+				Elements[1] = nullptr;
+			}
+		}
+		else
+		{
+			Elements[0] = nullptr;
+		}
+	}
+
+	FORCEINLINE UObjectBase*& operator[](int32 Index)
+	{
+		FItems* Items = GetItems();
+		if (Items)
+		{
+			return (*Items)[Index];
+		}
+
+		checkSlow(Index == 0);
+		return Elements[0];
+	}
+
+#endif
+
+	/** Gets an iterator for the TArray in this bucket or for the EmptyBucker if Items is null */
+	FORCEINLINE TIterator GetIterator()
+	{
+		FItems* Items = GetItems();
+		return Items ? Items->CreateIterator() : EmptyBucket.CreateIterator();
+	}
+};
+
+template <typename T>
+TArrayHashBucket<T>::FItems TArrayHashBucket<T>::EmptyBucket;
 
 /** Hash Bucket Iterator. Iterates over all Objects in the bucket */
-struct FHashBucketIterator
+template <typename T>
+struct THashBucketIterator
 {
-	FHashBucket& Bucket;
-	TSet<UObjectBase*>::TIterator SetIterator;
+	T& Bucket;
+	T::TIterator Iterator;
 	bool bItems;
 	bool bReachedEndNoItems;
 	bool bSecondItem;
 
-	FORCEINLINE FHashBucketIterator(FHashBucket& InBucket)
+	FORCEINLINE THashBucketIterator(T& InBucket)
 		: Bucket(InBucket)
-		, SetIterator(InBucket.GetIteratorForSet())
-		, bItems(!!InBucket.GetSet())
-		, bReachedEndNoItems(!InBucket.ElementsOrSetPtr[0] && !InBucket.ElementsOrSetPtr[1])
+		, Iterator(InBucket.GetIterator())
+		, bItems(!!InBucket.GetItems())
+		, bReachedEndNoItems(!InBucket.Elements[0] && !InBucket.Elements[1])
 		, bSecondItem(false)
 	{
 	}
+
 	/** Advances the iterator to the next element. */
-	FORCEINLINE FHashBucketIterator& operator++()
+	FORCEINLINE THashBucketIterator& operator++()
 	{
 		if (bItems)
 		{
-			++SetIterator;
+			++Iterator;
 		}
 		else
 		{
-			bReachedEndNoItems = bSecondItem || !Bucket.ElementsOrSetPtr[1];
+			bReachedEndNoItems = bSecondItem || !Bucket.Elements[1];
 			bSecondItem = true;
 		}
 		return *this;
 	}
+
 	/** conversion to "bool" returning true if the iterator is valid. */
 	FORCEINLINE explicit operator bool() const
 	{
 		if (bItems)
 		{
-			return (bool)SetIterator;
+			return (bool)Iterator;
 		}
 		else
 		{
 			return !bReachedEndNoItems;
 		}
 	}
+
 	/** inverse of the "bool" operator */
 	FORCEINLINE bool operator !() const
 	{
 		return !(bool)*this;
 	}
+
 	FORCEINLINE UObjectBase*& operator*()
 	{
 		if (bItems)
 		{
-			return *SetIterator;
+			return *Iterator;
 		}
 		else
 		{
-			return (UObjectBase*&)Bucket.ElementsOrSetPtr[!!bSecondItem];
+			return (UObjectBase*&)Bucket.Elements[!!bSecondItem];
 		}
 	}
 };
 
-FORCEINLINE FHashBucketIterator FHashBucket::CreateIterator()
+template <typename T>
+FORCEINLINE THashBucketIterator<TArrayHashBucket<T>> TArrayHashBucket<T>::CreateIterator()
+{
+	return THashBucketIterator(*this);
+}
+
+FORCEINLINE THashBucketIterator<FSetHashBucket> FSetHashBucket::CreateIterator()
 {
-	return FHashBucketIterator(*this);
+	return THashBucketIterator(*this);
 }
 
+#if UE_UOBJECT_HASH_USES_ARRAYS
+	template <typename T>
+	struct FHashBucketIterator : public THashBucketIterator<T>
+	{
+		FHashBucketIterator(T& Bucket) : THashBucketIterator<T>(Bucket) {}
+	};
+	using FHashBucket = TArrayHashBucket<FSortedArray>;
+#else
+	// For example Editor builds need to use FSetHashBucket as a bucket type since we have a significantly higher number of UObjects than on a client, so perf with TArrayHashBucket will degrade
+	using FHashBucketIterator = THashBucketIterator<FSetHashBucket>;
+	using FHashBucket = FSetHashBucket;
+#endif
+
+
 /**
 * Wrapper around a TMap with FHashBucket values that supports read only locks
 */
@@ -396,7 +732,7 @@ public:
 	/** Map of object to their outers, used to avoid an object iterator to find such things. **/
 	TBucketMap<UObjectBase*> ObjectOuterMap;
 #if UE_STORE_OBJECT_LIST_INTERNAL_INDEX
-	TBucketMap<UClass*, TArray<UObjectBase*> > ClassToObjectListMap;
+	TBucketMap<UClass*, TArrayHashBucket<FUnsortedArray> > ClassToObjectListMap;
 #else
 	TBucketMap<UClass*> ClassToObjectListMap;
 #endif
