Added missing typename keyword to fix compilation errors in UObjectHash.cpp

diff --git a/Engine/Source/Runtime/CoreUObject/Private/UObject/UObjectHash.cpp b/Engine/Source/Runtime/CoreUObject/Private/UObject/UObjectHash.cpp
--- a/Engine/Source/Runtime/CoreUObject/Private/UObject/UObjectHash.cpp
+++ b/Engine/Source/Runtime/CoreUObject/Private/UObject/UObjectHash.cpp
@@ -319,7 +319,7 @@ template <typename T>
 struct TArrayHashBucket
 {
 	typedef T FItems;
-	typedef T::TIterator TIterator;
+	typedef typename T::TIterator TIterator;
 
 	/** This always empty array is used to get an iterator if the bucket doesn't use a TArray (has only 1 element) */
 	static FItems EmptyBucket;
@@ -533,14 +533,14 @@ struct TArrayHashBucket
 };
 
 template <typename T>
-TArrayHashBucket<T>::FItems TArrayHashBucket<T>::EmptyBucket;
+typename TArrayHashBucket<T>::FItems TArrayHashBucket<T>::EmptyBucket;
 
 /** Hash Bucket Iterator. Iterates over all Objects in the bucket */
 template <typename T>
 struct THashBucketIterator
 {
 	T& Bucket;
-	T::TIterator Iterator;
+	typename T::TIterator Iterator;
 	bool bItems;
 	bool bReachedEndNoItems;
 	bool bSecondItem;
