Constructing Instance Data Objects (IDOs) using an archetype chain:

Issue:
- Subobject archetypes under a Class Default Object (CDO) might have initialized properties that their own CDO does not initialize (e.g., happens with Verse overrides).
- This is an issue for IDOs because they do not use templates when constructed (they only use their CDO).
- This change addresses this problem by using an archetype chain when constructing IDOs.

List of changes:
- Constructing sub-object archetypes recursively and on-demand which are then used as templates during IDO construction.
- Bypassing the type match check for IDO templates (types could be different when IDO classes have unknown properties).
- Excluding loose properties from subobject instancing.
- Forcing an archetype match during SerializeScriptProperties during IDO serialization. This is necessary as normal archetype searching does not work on IDOs as they do not use an 'outer' hierarchy.

This change can be disabled by using this CVar:
IDO.EnableIDOArchetypeChain 0

[RN] Constructing Instance Data Objects (IDOs) using an archetype chain. IDO. This change constructs sub-object archetypes recursively and on-demand which are then used as templates during IDO construction.

diff --git a/Engine/Source/Runtime/CoreUObject/Internal/UObject/InstanceDataObjectUtils.h b/Engine/Source/Runtime/CoreUObject/Internal/UObject/InstanceDataObjectUtils.h
--- a/Engine/Source/Runtime/CoreUObject/Internal/UObject/InstanceDataObjectUtils.h
+++ b/Engine/Source/Runtime/CoreUObject/Internal/UObject/InstanceDataObjectUtils.h
@@ -33,6 +33,8 @@ UE_API bool IsInstanceDataObjectSupportEnabledForClass(const UClass* Class);
 UE_API bool IsUninitializedAlertUIEnabled();
 /** Query if InstanceDataObjects should be saved instead of instances. */
 UE_API bool IsInstanceDataObjectImpersonationEnabledOnSave();
+/** Query if InstanceDataObjects should be constructed using an archetype chain. */
+UE_API bool IsInstanceDataObjectArchetypeChainEnabled();
 
 bool StructContainsLooseProperties(const UStruct* Struct);
 UE_API bool StructIsInstanceDataObjectStruct(const UStruct* Struct);
diff --git a/Engine/Source/Runtime/CoreUObject/Internal/UObject/PropertyBagRepository.h b/Engine/Source/Runtime/CoreUObject/Internal/UObject/PropertyBagRepository.h
--- a/Engine/Source/Runtime/CoreUObject/Internal/UObject/PropertyBagRepository.h
+++ b/Engine/Source/Runtime/CoreUObject/Internal/UObject/PropertyBagRepository.h
@@ -154,6 +154,9 @@ private:
 
 	// Delete owner reference and disassociate all data. Returns success.
 	bool RemoveAssociationUnsafe(const UObject* Owner);
+
+	// Internal function for creating IDOs. Before calling this function the repository's lock should be acquired.
+	UObject* CreateInstanceDataObjectInternal(UObject* Owner, FArchive& Archive, int64 StartOffset, int64 EndOffset, bool bIsArchetype = false);
 };
 
 #endif // WITH_EDITORONLY_DATA
diff --git a/Engine/Source/Runtime/CoreUObject/Private/UObject/CoreNative.cpp b/Engine/Source/Runtime/CoreUObject/Private/UObject/CoreNative.cpp
--- a/Engine/Source/Runtime/CoreUObject/Private/UObject/CoreNative.cpp
+++ b/Engine/Source/Runtime/CoreUObject/Private/UObject/CoreNative.cpp
@@ -12,6 +12,7 @@
 #include "UObject/OverridableManager.h"
 
 #if WITH_EDITORONLY_DATA
+#include "UObject/InstanceDataObjectUtils.h"
 #include "Serialization/PropertyLocalizationDataGathering.h"
 #include "StructUtils/InstancedStruct.h"
 #endif	// WITH_EDITORONLY_DATA
@@ -169,6 +170,21 @@ void FObjectInstancingGraph::SetDestinationRoot(UObject* DestinationSubobjectRoo
 	{
 		bCanUseDynamicInstancing = true;
 	}
+
+	#if WITH_EDITORONLY_DATA
+	if (UE::IsInstanceDataObject(DestinationRoot))
+	{
+		const UClass* DestinationRootClass = DestinationRoot->GetClass();
+
+		for (const FProperty* Property = DestinationRootClass->RefLink; Property; Property = Property->NextRef)
+		{
+			if (UE::IsPropertyLoose(Property))
+			{
+				AddPropertyToSubobjectExclusionList(Property);
+			}
+		}
+	}
+	#endif
 }
 
 UObject* FObjectInstancingGraph::GetDestinationObject(UObject* SourceObject)
diff --git a/Engine/Source/Runtime/CoreUObject/Private/UObject/InstanceDataObjectUtils.cpp b/Engine/Source/Runtime/CoreUObject/Private/UObject/InstanceDataObjectUtils.cpp
--- a/Engine/Source/Runtime/CoreUObject/Private/UObject/InstanceDataObjectUtils.cpp
+++ b/Engine/Source/Runtime/CoreUObject/Private/UObject/InstanceDataObjectUtils.cpp
@@ -313,6 +313,13 @@ namespace UE
 		TEXT("When enabled, blueprint instances (and prefab instances) can have IDOs generated for them")
 		);
 
+	bool bEnableIDOArchetypeChain = true;
+	FAutoConsoleVariableRef EnableIDOArchetypeChainCVar(
+		TEXT("IDO.EnableArchetypeChain"),
+		bEnableIDOArchetypeChain,
+		TEXT("When enabled, IDOs will be constructed using an archetype chain")
+	);
+
 	static void SetEnableAllIDOFeatures(bool bEnabled)
 	{
 		EnableIDOSupportCVar->Set(bEnabled);
@@ -320,6 +327,7 @@ namespace UE
 		EnableIDOImpersonationOnSaveCVar->Set(bEnabled);
 		EnableIDOsForBlueprintArchetypesCVar->Set(bEnabled);
 		EnableIDOsForBlueprintInstancesCVar->Set(bEnabled);
+		EnableIDOArchetypeChainCVar->Set(bEnabled);
 	}
 
 	static bool AreAllIDOFeaturesEnabled()
@@ -329,7 +337,8 @@ namespace UE
 			bEnableUninitializedUI && 
 			bEnableIDOImpersonationOnSave && 
 			bEnableIDOsForBlueprintArchetypes && 
-			bEnableIDOsForBlueprintInstances;
+			bEnableIDOsForBlueprintInstances &&
+			bEnableIDOArchetypeChain;
 	}
 
 	FAutoConsoleCommand EnableAllIDOFeaturesCommand(
@@ -370,6 +379,11 @@ namespace UE
 		return bEnableIDOImpersonationOnSave;
 	}
 
+	bool IsInstanceDataObjectArchetypeChainEnabled()
+	{
+		return bEnableIDOArchetypeChain;
+	}
+
 	static const UObject* GetBlueprintGeneratedObject(const UObject* InObject)
 	{
 		static const FName NAME_BlueprintGeneratedClass(TEXT("BlueprintGeneratedClass"));
diff --git a/Engine/Source/Runtime/CoreUObject/Private/UObject/PropertyBagRepository.cpp b/Engine/Source/Runtime/CoreUObject/Private/UObject/PropertyBagRepository.cpp
--- a/Engine/Source/Runtime/CoreUObject/Private/UObject/PropertyBagRepository.cpp
+++ b/Engine/Source/Runtime/CoreUObject/Private/UObject/PropertyBagRepository.cpp
@@ -9,6 +9,7 @@
 #include "Containers/Queue.h"
 #include "Serialization/ObjectReader.h"
 #include "Serialization/ObjectWriter.h"
+#include "Serialization/ArchiveProxy.h"
 #include "UObject/GarbageCollection.h"
 #include "UObject/InstanceDataObjectUtils.h"
 #include "UObject/LinkerLoad.h"
@@ -129,6 +130,35 @@ public:
 	}
 };
 
+class FArchetypeMatchingArchive : public FArchiveProxy
+{
+public:
+
+	FArchetypeMatchingArchive(FArchive& InInnerArchive, UObject* InObject, UObject* InArchetype)
+		: FArchiveProxy(InInnerArchive)
+		, Object(InObject)
+		, Archetype(InArchetype)
+	{		
+	}
+
+	UObject* GetArchetypeFromLoader(const UObject* Obj) override
+	{
+		if (Archetype && (Obj == Object))
+		{
+			return Archetype;
+		}
+		else
+		{
+			return InnerArchive.GetArchetypeFromLoader(Obj);
+		}
+	}
+
+private:
+
+	UObject* Object = nullptr;
+	UObject* Archetype = nullptr;
+};
+
 void FPropertyBagRepository::FPropertyBagAssociationData::Destroy()
 {
 	if (InstanceDataObject && InstanceDataObject->IsValidLowLevel())
@@ -788,15 +818,9 @@ FString FPropertyBagRepository::GetReferencerName() const
 	return TEXT("FPropertyBagRepository");
 }
 
-UObject* FPropertyBagRepository::CreateInstanceDataObject(UObject* Owner, FArchive& Archive, int64 StartOffset, int64 EndOffset)
+// Before calling this function the repository's lock should be acquired.
+UObject* FPropertyBagRepository::CreateInstanceDataObjectInternal(UObject* Owner, FArchive& Archive, int64 StartOffset, int64 EndOffset, bool bIsArchetype)
 {
-	FPropertyBagRepositoryLock LockRepo(this);
-	FPropertyBagAssociationData& BagData = AssociatedData.FindOrAdd(Owner);
-	if (BagData.InstanceDataObject)
-	{
-		return BagData.InstanceDataObject;
-	}
-
 	// TODO: Should the IDO or its class be in another package?
 	UObject* Outer = GetTransientPackage();
 
@@ -806,35 +830,117 @@ UObject* FPropertyBagRepository::CreateInstanceDataObject(UObject* Owner, FArchi
 	const FUnknownEnumNames* EnumNamesPtr = !EnumNames.IsEmpty() ? &EnumNames : nullptr;
 	const UClass* InstanceDataObjectClass = CreateInstanceDataObjectClass(PropertyTree.Get(), EnumNamesPtr, Owner->GetClass(), Outer);
 
+	UObject* OwnerArchetype = nullptr;
+	UObject* Template = nullptr;
+
+	bool bIsIDOArchetypeChainEnabled = UE::IsInstanceDataObjectArchetypeChainEnabled();
+
+	if (bIsIDOArchetypeChainEnabled)
+	{
+		OwnerArchetype = Owner->GetArchetype();
+
+		// If an archetype is not the CDO then it's a sub-object archetype.
+		bool bIsSubobjectArchetype = OwnerArchetype && !OwnerArchetype->HasAllFlags(RF_ClassDefaultObject);
+
+		// Subobject archetypes (under a CDO) might have initialized properties that their own CDO does
+		// not initialize (e.g., Verse overrides). For this reason we create those IDO archetypes on demand.
+		if (bIsSubobjectArchetype)
+		{
+			// We cannot do a FindAndAdd here because the map could be modified while creating more
+			// IDO archetypes recursively.
+			FPropertyBagAssociationData* TemplateBagData = AssociatedData.Find(OwnerArchetype);
+			if (TemplateBagData)
+			{
+				check(TemplateBagData->InstanceDataObject);
+				Template = TemplateBagData->InstanceDataObject;
+			}
+			else
+			{
+				// Serialize the owner's archetype into an arhive we can use to construct its IDO.
+				TArray<uint8> OwnerArchetypeData;
+				OwnerArchetypeData.Reserve(OwnerArchetype->GetClass()->GetStructureSize());
+
+				FObjectWriter Writer(OwnerArchetypeData);
+				OwnerArchetype->GetClass()->SerializeTaggedProperties(Writer, (uint8*)OwnerArchetype, OwnerArchetype->GetClass(), nullptr);
+				FObjectReader Reader(OwnerArchetypeData);
+				Reader.ArMergeOverrides = true;
+				Reader.ArPreserveArrayElements = true;
+
+				int64 TemplateEndOffset = Reader.TotalSize();
+				int64 TemplateStartOffset = (TemplateEndOffset == INDEX_NONE) ? INDEX_NONE : 0;
+
+				Template = CreateInstanceDataObjectInternal(OwnerArchetype, Reader, TemplateStartOffset, TemplateEndOffset, /*bIsArchetype*/true);
+				check(Template != nullptr);
+			}
+		}
+	}
+
 	// Construct the IDO.
 	FStaticConstructObjectParameters Params(InstanceDataObjectClass);
 	Params.SetFlags |= EObjectFlags::RF_Transactional;
+	// Set the RF_ArchetypeObject flag on all IDO archetypes so that they can be identified as such.
+	if (bIsArchetype || Owner->HasAllFlags(EObjectFlags::RF_ArchetypeObject))
+	{
+		Params.SetFlags |= EObjectFlags::RF_ArchetypeObject;
+	}
+
 	Params.Name = MakeUniqueObjectName(Outer, Params.Class, Owner->GetFName(), EUniqueObjectNameOptions::UniversallyUnique);
 	Params.Outer = Outer;
+
 	UObject* InstanceDataObject = StaticConstructObject_Internal(Params);
 
+	// We are not passing the template in FStaticConstructObjectParameters because of potential type mismatches
+	// between the IDO and its template. This can happen when both objects have the same 'owner' type but a 
+	// different number of unknown properties. For this reason we do a copy instead.
+	if (Template)
+	{
+		CopyTaggedProperties(Template, InstanceDataObject);
+	}
+
 	InstanceDataObjectToOwner.Add(InstanceDataObject, Owner);
+
+	FPropertyBagAssociationData& BagData = AssociatedData.Add(Owner);
 	BagData.InstanceDataObject = InstanceDataObject;
 	BagData.bNeedsFixup = StructContainsLooseProperties(InstanceDataObjectClass);
 
 	// Load the IDO.
 	if (StartOffset != EndOffset)
 	{
-		const int64 OffsetToRestore = Archive.Tell();
-		FScopedObjectSerializeContext ObjectSerializeContext(InstanceDataObject, Archive);
-		Archive.Seek(StartOffset);
+		// We want to force the template we used as the archetype during SerializeScriptProperties. For this reason 
+		// we wrap the archive in a proxy that will return the correct archetype.
+		FArchetypeMatchingArchive ArchiveWrapper(Archive, InstanceDataObject, Template);
+
+		const int64 OffsetToRestore = ArchiveWrapper.Tell();
+		FScopedObjectSerializeContext ObjectSerializeContext(InstanceDataObject, ArchiveWrapper);
+		ArchiveWrapper.Seek(StartOffset);
 		{
-			FGuardValue_Bitfield(Archive.ArMergeOverrides, true);
-			InstanceDataObject->SerializeScriptProperties(Archive);
+			FGuardValue_Bitfield(ArchiveWrapper.ArMergeOverrides, true);
+			InstanceDataObject->SerializeScriptProperties(ArchiveWrapper);
 		}
-		ensureMsgf(Archive.Tell() == EndOffset,
+		ensureMsgf(ArchiveWrapper.Tell() == EndOffset,
 			TEXT("Serializing %s into its IDO consumed %" INT64_FMT " bytes when %" INT64_FMT " bytes were expected."),
-			*Owner->GetPathName(), Archive.Tell() - StartOffset, EndOffset - StartOffset);
-		Archive.Seek(OffsetToRestore);
+			*Owner->GetPathName(), ArchiveWrapper.Tell() - StartOffset, EndOffset - StartOffset);
+		ArchiveWrapper.Seek(OffsetToRestore);
 	}
 	return BagData.InstanceDataObject;
 }
 
+UObject* FPropertyBagRepository::CreateInstanceDataObject(UObject* Owner, FArchive& Archive, int64 StartOffset, int64 EndOffset)
+{
+	FPropertyBagRepositoryLock LockRepo(this);
+
+	FPropertyBagAssociationData* BagData = AssociatedData.Find(Owner);
+	if (BagData)
+	{
+		check(BagData->InstanceDataObject);
+		return BagData->InstanceDataObject;
+	}
+	else
+	{
+		return CreateInstanceDataObjectInternal(Owner, Archive, StartOffset, EndOffset);
+	}
+}
+
 // Not sure this is necessary.
 void FPropertyBagRepository::ShrinkMaps()
 {
