Removed 13 dll exports from FBulkDataRequest.

diff --git a/Engine/Source/Runtime/CoreUObject/Private/Serialization/BulkDataStreaming.cpp b/Engine/Source/Runtime/CoreUObject/Private/Serialization/BulkDataStreaming.cpp
--- a/Engine/Source/Runtime/CoreUObject/Private/Serialization/BulkDataStreaming.cpp
+++ b/Engine/Source/Runtime/CoreUObject/Private/Serialization/BulkDataStreaming.cpp
@@ -469,20 +469,9 @@ bool TryMemoryMapBulkData(
 
 //////////////////////////////////////////////////////////////////////////////
 
-class FBulkDataRequest::IHandle
+FBulkDataRequest::IHandle::~IHandle()
 {
-public:
-	virtual ~IHandle() = default;
-
-	virtual void AddRef() = 0;
-	virtual void Release() = 0;
-	virtual uint32 GetRefCount() const = 0;
-
-	virtual FBulkDataRequest::EStatus GetStatus() const = 0;
-	virtual bool Cancel() = 0;
-	virtual bool Wait(uint32 Milliseconds) = 0;
-	virtual void UpdatePriority(const EAsyncIOPriorityAndFlags Priority) = 0;
-};
+}
 
 //////////////////////////////////////////////////////////////////////////////
 
@@ -493,8 +482,6 @@ public:
 	FHandleBase(FBulkDataRequest::EStatus InStatus)
 		: Status(uint32(InStatus))
 	{ }
-	
-	virtual ~FHandleBase() = default;
 
 	virtual void AddRef() override final
 	{
@@ -636,7 +623,7 @@ public:
 		TRACE_COUNTER_INCREMENT(BulkDataBatchRequest_Count);
 	}
 	
-	~FBatchHandle()
+	virtual ~FBatchHandle() override
 	{
 		Cancel();
 		Wait(MAX_uint32);
@@ -752,83 +739,6 @@ private:
 
 //////////////////////////////////////////////////////////////////////////////
 
-FBulkDataRequest::FBulkDataRequest()
-{
-}
-
-FBulkDataRequest::FBulkDataRequest(FBulkDataRequest::IHandle* InHandle)
-	: Handle(InHandle)
-{
-	check(InHandle != nullptr);
-}
-
-FBulkDataRequest::~FBulkDataRequest()
-{
-}
-
-FBulkDataRequest::FBulkDataRequest(FBulkDataRequest&& Other)
-	: Handle(MoveTemp(Other.Handle))
-{
-}
-
-FBulkDataRequest& FBulkDataRequest::operator=(FBulkDataRequest&& Other)
-{
-	Handle = MoveTemp(Other.Handle);
-
-	return *this;
-}
-
-FBulkDataRequest::EStatus FBulkDataRequest::GetStatus() const
-{
-	return Handle.IsValid() ? Handle->GetStatus() : EStatus::None;
-}
-
-bool FBulkDataRequest::IsNone() const
-{
-	return FBulkDataRequest::EStatus::None == GetStatus();
-}
-
-bool FBulkDataRequest::IsPending() const
-{
-	return FBulkDataRequest::EStatus::Pending == GetStatus();
-}
-
-bool FBulkDataRequest::IsOk() const
-{
-	return FBulkDataRequest::EStatus::Ok == GetStatus();
-}
-
-bool FBulkDataRequest::IsCompleted() const
-{
-	const uint32 Status = static_cast<uint32>(GetStatus());
-	return Status > static_cast<uint32>(EStatus::Pending);
-}
-
-void FBulkDataRequest::Cancel()
-{
-	if (Handle.IsValid())
-	{
-		Handle->Cancel();
-	}
-}
-
-void FBulkDataRequest::Reset()
-{
-	Handle.SafeRelease();
-}
-
-int32 FBulkDataRequest::GetRefCount() const
-{
-	if (Handle)
-	{
-		return Handle->GetRefCount();
-	}
-
-	return -1;
-}
-
-//////////////////////////////////////////////////////////////////////////////
-
 void FBulkDataBatchRequest::Wait()
 {
 	if (Handle.IsValid())
diff --git a/Engine/Source/Runtime/CoreUObject/Public/Serialization/BulkData.h b/Engine/Source/Runtime/CoreUObject/Public/Serialization/BulkData.h
--- a/Engine/Source/Runtime/CoreUObject/Public/Serialization/BulkData.h
+++ b/Engine/Source/Runtime/CoreUObject/Public/Serialization/BulkData.h
@@ -1192,8 +1192,6 @@ public:
 class FBulkDataRequest
 {
 public:
-	class IHandle;
-
 	/** Bulk data request status. */
 	enum class EStatus : uint32
 	{
@@ -1209,6 +1207,23 @@ public:
 		Error
 	};
 
+protected:
+	class IHandle
+	{
+	public:
+		virtual ~IHandle();
+
+		virtual void AddRef() = 0;
+		virtual void Release() = 0;
+		virtual uint32 GetRefCount() const = 0;
+
+		virtual FBulkDataRequest::EStatus GetStatus() const = 0;
+		virtual bool Cancel() = 0;
+		virtual bool Wait(uint32 Milliseconds) = 0;
+		virtual void UpdatePriority(const EAsyncIOPriorityAndFlags Priority) = 0;
+	};
+
+public:
 	/** Completion callback. */
 	using FCompletionCallback = TFunction<void(EStatus)>;
 
@@ -1216,49 +1231,50 @@ public:
 	static constexpr EAsyncIOPriorityAndFlags DefaultPriority = AIOP_BelowNormal;
 
 	/** Constructs a new handle to bulk data request. */
-	COREUOBJECT_API FBulkDataRequest();
+	FBulkDataRequest() = default;
 
 	/** Destructor, cancels and waits for any pending requests. */
-	COREUOBJECT_API ~FBulkDataRequest();
-
-	/** Moves ownership from an invalid or pending request. */
-	COREUOBJECT_API FBulkDataRequest(FBulkDataRequest&&);
+	~FBulkDataRequest() = default;
 
 	/** Moves ownership from an invalid or pending request. */
-	COREUOBJECT_API FBulkDataRequest& operator=(FBulkDataRequest&&);
+	FBulkDataRequest(FBulkDataRequest&&) = default;
+	FBulkDataRequest& operator=(FBulkDataRequest&&) = default;
 	
-	/** Not copy constructable. */
+	/** Not copy constructable or copy assignable. */
 	FBulkDataRequest(const FBulkDataRequest&) = delete;
-
-	/** Not copy assignable. */
 	FBulkDataRequest& operator=(const FBulkDataRequest&) = delete;
 
 	/** Returns current status of the request. */
-	COREUOBJECT_API EStatus GetStatus() const;
+	EStatus GetStatus() const { return Handle.IsValid() ? Handle->GetStatus() : EStatus::None; }
 
 	/** Returns whether the request is associated with a pending or completed request. */
-	COREUOBJECT_API bool IsNone() const;
+	bool IsNone() const { return FBulkDataRequest::EStatus::None == GetStatus(); }
 
 	/** Returns whether the request is pending. */
-	COREUOBJECT_API bool IsPending() const;
+	bool IsPending() const { return FBulkDataRequest::EStatus::Pending == GetStatus(); }
 
 	/** Returns whether the request completed successfully. */
-	COREUOBJECT_API bool IsOk() const;
+	bool IsOk() const { return FBulkDataRequest::EStatus::Ok == GetStatus(); }
 
 	/** Returns whether the request has been completed. */
-	COREUOBJECT_API bool IsCompleted() const;
+	bool IsCompleted() const { return GetStatus() > EStatus::Pending; }
 
 	/** Cancel the pending request. Cancelling an invalid request will immediately return. */
-	COREUOBJECT_API void Cancel();
+	void Cancel() { if (Handle.IsValid()) { Handle->Cancel(); } }
 
 	/** Reset the request handle to an invalid state. Will cancel and wait if the request is not completed. */
-	COREUOBJECT_API void Reset();
+	void Reset() { Handle.SafeRelease(); }
 
 protected:
+	friend class FHandleBase;
+	friend class FChunkBatchReadRequest;
 	friend class FBulkDataBatchRequest;
 	
-	COREUOBJECT_API FBulkDataRequest(IHandle* InHandle);
-	COREUOBJECT_API int32 GetRefCount() const;
+	FBulkDataRequest(FBulkDataRequest::IHandle* InHandle)
+		: Handle(InHandle)
+	{
+		check(InHandle != nullptr);
+	}
 
 	TRefCountPtr<IHandle> Handle;
 };
