SavePackage: PackageFileSummary: Fix serialization error in unversioned package load due to mismatched custom versions in the PackageFileSummary. This was an error when moving around order of serialization in the previous SavePackage change. We are supposed to overwrite empty custom version container with latest custom versions.
#rnx

diff --git a/Engine/Source/Runtime/CoreUObject/Private/UObject/PackageFileSummary.cpp b/Engine/Source/Runtime/CoreUObject/Private/UObject/PackageFileSummary.cpp
--- a/Engine/Source/Runtime/CoreUObject/Private/UObject/PackageFileSummary.cpp
+++ b/Engine/Source/Runtime/CoreUObject/Private/UObject/PackageFileSummary.cpp
@@ -143,11 +143,10 @@ void operator<<(FStructuredArchive::FSlot Slot, FPackageFileSummary& Sum)
 			
 				Record << SA_VALUE(TEXT("FileVersionLicenseeUE4"), Sum.FileVersionLicenseeUE);
 
-				if (!Sum.FileVersionUE.FileVersionUE4 && !Sum.FileVersionUE.FileVersionUE5 && !Sum.FileVersionLicenseeUE)
+				// Record whether the summary was unversioned when it was loaded off disk
+				Sum.bUnversioned = !Sum.FileVersionUE.FileVersionUE4 && !Sum.FileVersionUE.FileVersionUE5 && !Sum.FileVersionLicenseeUE;
+				if (Sum.bUnversioned)
 				{
-					// Record that the summary was unversioned when it was loaded off disk
-					Sum.bUnversioned = true; 
-
 #if WITH_EDITOR
 					if (!GAllowUnversionedContentInEditor)
 					{
@@ -161,7 +160,6 @@ void operator<<(FStructuredArchive::FSlot Slot, FPackageFileSummary& Sum)
 					// Use the latest supported versions
 					Sum.FileVersionUE = GPackageFileUEVersion;
 					Sum.FileVersionLicenseeUE = GPackageFileLicenseeUEVersion;
-					Sum.CustomVersionContainer = FCurrentCustomVersions::GetAll();
 				}
 
 				if (Sum.FileVersionUE < VER_UE4_OLDEST_LOADABLE_PACKAGE || Sum.IsFileVersionTooNew())
@@ -185,6 +183,13 @@ void operator<<(FStructuredArchive::FSlot Slot, FPackageFileSummary& Sum)
 				{
 					Sum.CustomVersionContainer.Serialize(Record.EnterField(TEXT("CustomVersions")), GetCustomVersionFormatForArchive(LegacyFileVersion));
 				}
+
+				if (Sum.bUnversioned)
+				{
+					// Overwrite the CustomVersionContainer that we deserialized; it was an empty container written
+					// during unversioned savepackage. Overwrite with the latest version of all custom versions.
+					Sum.CustomVersionContainer = FCurrentCustomVersions::GetAll();
+				}
 			}
 			else
 			{
