Physics ID Computation: Prototype of migrating physics objects to other servers based on a naive strategy:
1. Before physics, each objects on the servers looks for other objects in proximity in a custom tick.
2. If there is neighboring object, the object would find the server with minimal id to migrate both of the objects
3. In a custom hard sync before physics, all objects are migrated into the physics server for simulation.
4. All objects are brought back to original server after physics simulation to preserve server boundary
Physics write locks are disabled at the moment because they cause live locks without adding borrowing flags into the serialization code.

diff --git a/Engine/Source/Runtime/CoreUObject/Private/UObject/RemoteObject.cpp b/Engine/Source/Runtime/CoreUObject/Private/UObject/RemoteObject.cpp
--- a/Engine/Source/Runtime/CoreUObject/Private/UObject/RemoteObject.cpp
+++ b/Engine/Source/Runtime/CoreUObject/Private/UObject/RemoteObject.cpp
@@ -89,6 +89,19 @@ public:
 
 		return Stub;
 	}
+
+	void ClearAllPhysicsIds()
+	{
+#if UE_WITH_REMOTE_OBJECT_HANDLE
+		for (TPair<FRemoteObjectId, UE::RemoteObject::Handle::FRemoteObjectStub*>& Elem : RemoteObjects)
+		{
+			if (Elem.Value)
+			{
+				Elem.Value->PhysicsOverrideServerId = FRemoteServerId(0);
+			}
+		}
+#endif
+	}
 };
 FRemoteObjectMaps* ObjectMaps = nullptr;
 
@@ -307,6 +320,52 @@ void ChangeOwnerServerId(const UObject* Object, FRemoteServerId NewOwnerServerId
 #endif
 }
 
+FRemoteServerId GetPhysicsServerId(const UObject* Object)
+{
+	using namespace UE::CoreUObject::Private;
+	using namespace UE::RemoteObject::Private;
+
+	FRemoteServerId Result;
+#if UE_WITH_REMOTE_OBJECT_HANDLE
+	UE::RemoteObject::Handle::FRemoteObjectStub* RemoteStub = FindRemoteObjectStub(FObjectHandleUtils::GetRemoteId(Object));
+	if (RemoteStub)
+	{
+		Result = RemoteStub->PhysicsOverrideServerId;
+	}
+	else
+	{
+		//If the object doesn't have stub, then use the owner id as physics id. 
+		Result = GetOwnerServerId(Object);
+	}
+#endif
+	return Result;
+}
+
+void ChangePhysicsServerId(const UObject* Object, FRemoteServerId NewPhysicsServerId)
+{
+	using namespace UE::CoreUObject::Private;
+	using namespace UE::RemoteObject::Private;
+
+#if UE_WITH_REMOTE_OBJECT_HANDLE
+	FRemoteObjectId ThisObject = FObjectHandleUtils::GetRemoteId(Object);
+	UE::RemoteObject::Handle::FRemoteObjectStub* RemoteStub = ObjectMaps->FindOrAddRemoteObjectStub(FObjectHandleUtils::GetRemoteId(Object));
+	if (RemoteStub)
+	{
+		RemoteStub->PhysicsOverrideServerId = NewPhysicsServerId;
+	}
+#endif
+}
+
+void ClearAllPhysicsServerId()
+{
+#if UE_WITH_REMOTE_OBJECT_HANDLE
+	if (Private::ObjectMaps)
+	{
+		Private::ObjectMaps->ClearAllPhysicsIds();
+	}
+#endif
+}
+
 UObject* ResolveObject(const FRemoteObjectStub* Stub, ERemoteReferenceType RefType /*= ERemoteReferenceType::Strong*/)
 {
 	using namespace UE::RemoteObject::Transfer;
diff --git a/Engine/Source/Runtime/CoreUObject/Private/UObject/RemoteObjectSerialization.cpp b/Engine/Source/Runtime/CoreUObject/Private/UObject/RemoteObjectSerialization.cpp
--- a/Engine/Source/Runtime/CoreUObject/Private/UObject/RemoteObjectSerialization.cpp
+++ b/Engine/Source/Runtime/CoreUObject/Private/UObject/RemoteObjectSerialization.cpp
@@ -91,10 +91,10 @@ void LoadObjectFromDisk(FRemoteTransactionId RequestId, FRemoteWorkPriority Requ
 	FileReader->Close();
 	IFileManager::Get().Delete(*Filename, false, true, true);
 
-	UE::RemoteObject::Transfer::OnObjectDataReceived(RequestId, FRemoteServerId(), ObjectId, CurrentOwnerServerId, ObjectData);
+	UE::RemoteObject::Transfer::OnObjectDataReceived(RequestId, FRemoteServerId(), FRemoteServerId(), ObjectId, CurrentOwnerServerId, ObjectData);
 }
 
-void SaveObjectToDisk(FRemoteTransactionId RequestId, FRemoteObjectId ObjectId, FRemoteServerId DestinationServerId, FRemoteServerId OwnerServerId, FRemoteObjectData& ObjectData)
+void SaveObjectToDisk(FRemoteTransactionId RequestId, FRemoteObjectId ObjectId, FRemoteServerId DestinationServerId, FRemoteServerId OwnerServerId, FRemoteServerId PhysicsId, FRemoteObjectData& ObjectData)
 {
 	TUniquePtr<FArchive> FileWriter{ IFileManager::Get().CreateFileWriter(*GenerateRemoteObjectFilename(ObjectId, DestinationServerId)) };
 	checkf(FileWriter.IsValid(), TEXT("Unable to create file writer for remote object %s"), *ObjectId.ToString());
diff --git a/Engine/Source/Runtime/CoreUObject/Private/UObject/RemoteObjectSerialization.h b/Engine/Source/Runtime/CoreUObject/Private/UObject/RemoteObjectSerialization.h
--- a/Engine/Source/Runtime/CoreUObject/Private/UObject/RemoteObjectSerialization.h
+++ b/Engine/Source/Runtime/CoreUObject/Private/UObject/RemoteObjectSerialization.h
@@ -99,5 +99,5 @@ namespace UE::RemoteObject::Serialization
 namespace UE::RemoteObject::Serialization::Disk
 {
 	void LoadObjectFromDisk(FRemoteTransactionId RequestId, FRemoteWorkPriority RequestPriority, FRemoteObjectId ObjectId, FRemoteServerId CurrentOwnerServerId, FRemoteServerId DestinationServerId);
-	void SaveObjectToDisk(FRemoteTransactionId RequestId, FRemoteObjectId ObjectId, FRemoteServerId ServerId, FRemoteServerId OwnerServerId, FRemoteObjectData& ObjectData);
+	void SaveObjectToDisk(FRemoteTransactionId RequestId, FRemoteObjectId ObjectId, FRemoteServerId ServerId, FRemoteServerId OwnerServerId, FRemoteServerId PhysicsId, FRemoteObjectData& ObjectData);
 }
\ No newline at end of file
diff --git a/Engine/Source/Runtime/CoreUObject/Private/UObject/RemoteObjectTransfer.cpp b/Engine/Source/Runtime/CoreUObject/Private/UObject/RemoteObjectTransfer.cpp
--- a/Engine/Source/Runtime/CoreUObject/Private/UObject/RemoteObjectTransfer.cpp
+++ b/Engine/Source/Runtime/CoreUObject/Private/UObject/RemoteObjectTransfer.cpp
@@ -22,7 +22,7 @@ namespace UE::RemoteObject::Transfer
 const FRemoteServerId DatabaseId(0x3ff);
 
 TDelegate<void(FRemoteTransactionId /*RequestId*/, FRemoteWorkPriority /*RequestPriority*/, FRemoteObjectId /*ObjectId*/, FRemoteServerId /*CurrentOwnerServerId*/, FRemoteServerId /*DestinationServerId*/)> RequestRemoteObjectDelegate;
-TDelegate<void(FRemoteTransactionId /*RequestId*/, FRemoteObjectId /*ObjectId*/, FRemoteServerId /*DestinationServerId*/, FRemoteServerId /*OwnerServerId*/, FRemoteObjectData& /*ObjectData*/)> RemoteObjectTransferDelegate;
+TDelegate<void(FRemoteTransactionId /*RequestId*/, FRemoteObjectId /*ObjectId*/, FRemoteServerId /*DestinationServerId*/, FRemoteServerId /*OwnerServerId*/, FRemoteServerId /*PhysicsId*/, FRemoteObjectData& /*ObjectData*/)> RemoteObjectTransferDelegate;
 TDelegate<void(FRemoteTransactionId /*RequestId*/, FRemoteObjectId /*ObjectId*/, FRemoteServerId /*DestinationServerId*/)> RemoteObjectDeniedTransferDelegate;
 
 TMulticastDelegate<void(const FRemoteObjectData& /*ObjectData*/, FRemoteServerId /*RemoteServerId*/, EObjectMigrationFlags /*MigrationFlags*/)> OnObjectDataSentDelegate;
@@ -31,10 +31,10 @@ TMulticastDelegate<void(const FRemoteObjectData& /*ObjectData*/, FRemoteServerId
 TMulticastDelegate<void(const TArray<UObject*>& /*MigratedObjects*/, FRemoteServerId /*RemoteServerId*/, EObjectMigrationFlags /*MigrationFlags*/)> OnObjectsReceivedDelegate;
 TMulticastDelegate<void(FRemoteTransactionId /*RequestId*/, FRemoteObjectId /*ObjectId*/)> OnObjectTouchedDelegate;
 
-TDelegate<void(FRemoteTransactionId /*RequestId*/, FRemoteObjectId /*ObjectId*/, FRemoteServerId /*DatabaseId*/, FRemoteServerId /*OwnerServerId*/, FRemoteObjectData& /*ObjectData*/)> StoreRemoteObjectDataDelegate;
+TDelegate<void(FRemoteTransactionId /*RequestId*/, FRemoteObjectId /*ObjectId*/, FRemoteServerId /*DatabaseId*/, FRemoteServerId /*OwnerServerId*/, FRemoteServerId /*PhysicsId*/, FRemoteObjectData& /*ObjectData*/)> StoreRemoteObjectDataDelegate;
 TDelegate<void(FRemoteTransactionId /*RequestId*/, FRemoteWorkPriority /*RequestPriority*/, FRemoteObjectId /*ObjectId*/, FRemoteServerId /*DatabaseId*/, FRemoteServerId /*DestinationServerId*/)> RestoreRemoteObjectDataDelegate;
 
-UObject* MigrateObjectFromObjectData(FRemoteServerId OwnerServerId, FRemoteObjectData& ObjectData, FRemoteServerId RemoteServerId, EObjectMigrationFlags ObjectFlags);
+UObject* MigrateObjectFromObjectData(FRemoteServerId OwnerServerId, FRemoteServerId PhysicsId, FRemoteObjectData& ObjectData, FRemoteServerId RemoteServerId, EObjectMigrationFlags ObjectFlags);
 
 /**
  * Remote object transfer queue. 
@@ -190,10 +190,10 @@ public:
 	* @param ObjectData Object data to send. Send request takes ownership of the data
 	* @return Future that's going to receive the object data
 	*/
-	void SendRemoteObject(FRemoteTransactionId RequestId, FRemoteObjectId ObjectId, FRemoteServerId DestinationServerId, FRemoteServerId OwnerServerId, FRemoteObjectData* ObjectData)
+	void SendRemoteObject(FRemoteTransactionId RequestId, FRemoteObjectId ObjectId, FRemoteServerId DestinationServerId, FRemoteServerId OwnerServerId, FRemoteServerId PhysicsId, FRemoteObjectData* ObjectData)
 	{
 		check(!AutoRTFM::IsClosed());
-		RemoteObjectTransferDelegate.Execute(RequestId, ObjectId, DestinationServerId, OwnerServerId, *ObjectData);
+		RemoteObjectTransferDelegate.Execute(RequestId, ObjectId, DestinationServerId, OwnerServerId, PhysicsId, *ObjectData);
 	}
 
 	void SendDeniedRemoteObject(FRemoteTransactionId RequestId, FRemoteObjectId ObjectId, FRemoteServerId DestinationServerId)
@@ -208,7 +208,7 @@ public:
 	* @param Data Object data. Data ownership is transferred to the receive request (if it exists)
 	* @return True if a receive request for the specified object exists (and object data ownership was transferred), false otherwise
 	*/
-	void FulfillReceiveRequest(FRemoteTransactionId RequestId, FRemoteServerId OwnerServerId, FRemoteObjectId ObjectId, FRemoteServerId RemoteServerId, FRemoteObjectData* ObjectData)
+	void FulfillReceiveRequest(FRemoteTransactionId RequestId, FRemoteServerId OwnerServerId, FRemoteServerId PhysicsId, FRemoteObjectId ObjectId, FRemoteServerId RemoteServerId, FRemoteObjectData* ObjectData)
 	{
 		using namespace UE::RemoteObject::Serialization;
 
@@ -221,7 +221,7 @@ public:
 
 		const bool bAssignedOwnership = (RequestId == PushMigrationRequestId);
 		EObjectMigrationFlags MigrationFlags = bAssignedOwnership ? EObjectMigrationFlags::Acquire : EObjectMigrationFlags::Receive;
-		UObject* Result = MigrateObjectFromObjectData(OwnerServerId, *ObjectData, RemoteServerId, MigrationFlags);
+		UObject* Result = MigrateObjectFromObjectData(OwnerServerId, PhysicsId, *ObjectData, RemoteServerId, MigrationFlags);
 		checkf(UE::CoreUObject::Private::FObjectHandleUtils::GetRemoteId(Result) == ObjectId, TEXT("PeerId:%s requested object %s migration from PeerId:%s but received %s"),
 			*GetGlobalServerId().ToString(), *ObjectId.ToString(), *RemoteServerId.ToString(), *UE::CoreUObject::Private::FObjectHandleUtils::GetRemoteId(Result).ToString());
 		checkf(OwnerServerId == Handle::GetOwnerServerId(Result), TEXT("Expected owner id %s doesn't match post-deserialize owner: %s"), *OwnerServerId.ToString(), *Handle::GetOwnerServerId(Result).ToString());
@@ -250,7 +250,7 @@ public:
 
 FRemoteObjectTransferQueue GTransferQueue;
 
-UObject* MigrateObjectFromObjectData(FRemoteServerId OwnerServerId, FRemoteObjectData& ObjectData, FRemoteServerId RemoteServerId, EObjectMigrationFlags MigrationFlags)
+UObject* MigrateObjectFromObjectData(FRemoteServerId OwnerServerId, FRemoteServerId PhysicsId, FRemoteObjectData& ObjectData, FRemoteServerId RemoteServerId, EObjectMigrationFlags MigrationFlags)
 {
 	using namespace UE::RemoteObject::Serialization;
 	check(!AutoRTFM::IsClosed());
@@ -309,6 +309,7 @@ UObject* MigrateObjectFromObjectData(FRemoteServerId OwnerServerId, FRemoteObjec
 	for (UObject* ReceivedObject : ReceivedObjects)
 	{
 		UE::RemoteObject::Handle::ChangeOwnerServerId(ReceivedObject, OwnerServerId);
+		UE::RemoteObject::Handle::ChangePhysicsServerId(ReceivedObject, PhysicsId);
 	}
 
 	// Notify that objects have been migrated from a remote server
@@ -455,6 +456,7 @@ void MigrateObjectToRemoteServerWithExplicitPriority(FRemoteTransactionId Reques
 
 	const bool bOwnershipChange = (RequestId == PushMigrationRequestId);
 	FRemoteServerId OwnerServerId = bOwnershipChange ? DestinationServerId : Handle::GetOwnerServerId(RootObject);
+	FRemoteServerId PhysicsId = Handle::GetPhysicsServerId(RootObject);
 
 	UE_LOG(LogRemoteObject, Verbose, 
 		TEXT("     : MigrateObjectToRemoteServer(%s) sending obj %s [%s] (root %s, isremote %d, changing owner %d, owner %s)"),
@@ -555,7 +557,7 @@ void MigrateObjectToRemoteServerWithExplicitPriority(FRemoteTransactionId Reques
 		// Notify that object data has been migrated (actually will be but SendObjectData does a MoveTemp on the data)
 		OnObjectDataSentDelegate.Broadcast(*ObjectData, DestinationServerId, MigrationFlags);
 
-		GTransferQueue.SendRemoteObject(RequestId, ObjectId, DestinationServerId, OwnerServerId, ObjectData);
+		GTransferQueue.SendRemoteObject(RequestId, ObjectId, DestinationServerId, OwnerServerId, PhysicsId, ObjectData);
 
 		checkf(SentObjects.Num() > 0, TEXT("PeerId:%s unable to migrate object %s to PeerId:%s"), *GetGlobalServerId().ToString(), *UE::CoreUObject::Private::FObjectHandleUtils::GetRemoteId(Object).ToString(), *DestinationServerId.ToString());
 
@@ -646,10 +648,10 @@ void MigrateObjectToRemoteServerWithExplicitPriority(FRemoteTransactionId Reques
 	}
 }
 
-void OnObjectDataReceived(FRemoteTransactionId RequestId, FRemoteServerId OwnerServerId, FRemoteObjectId ObjectId, FRemoteServerId RemoteServerId, FRemoteObjectData* Data)
+void OnObjectDataReceived(FRemoteTransactionId RequestId, FRemoteServerId OwnerServerId, FRemoteServerId PhysicsId, FRemoteObjectId ObjectId, FRemoteServerId RemoteServerId, FRemoteObjectData* Data)
 {
 	check(!AutoRTFM::IsClosed());
-	GTransferQueue.FulfillReceiveRequest(RequestId, OwnerServerId, ObjectId, RemoteServerId, Data);
+	GTransferQueue.FulfillReceiveRequest(RequestId, OwnerServerId, PhysicsId, ObjectId, RemoteServerId, Data);
 }
 
 void OnObjectDataDenied(FRemoteTransactionId RequestId, FRemoteObjectId ObjectId, FRemoteServerId RemoteServerId)
@@ -731,7 +733,7 @@ void StoreObjectToDatabase(UObject* Object)
 		TSet<UObject*> SentObjects;
 		{
 			TUniquePtr<FRemoteObjectData> ObjectData(SerializeObjectData(OutermostUnreachableRemoteReference, SentObjects));
-			StoreRemoteObjectDataDelegate.Execute(PushMigrationRequestId, FRemoteObjectId(OutermostUnreachableRemoteReference), DatabaseId, FRemoteServerId(), *ObjectData);
+			StoreRemoteObjectDataDelegate.Execute(PushMigrationRequestId, FRemoteObjectId(OutermostUnreachableRemoteReference), DatabaseId, FRemoteServerId(), FRemoteServerId(), *ObjectData);
 		}
 
 		checkf(SentObjects.Num() > 0, TEXT("PeerId:%s unable to store object %s to DatabaseId:%s"), *GetGlobalServerId().ToString(), *FRemoteObjectId(Object).ToString(), *DatabaseId.ToString());
diff --git a/Engine/Source/Runtime/CoreUObject/Public/UObject/RemoteObject.h b/Engine/Source/Runtime/CoreUObject/Public/UObject/RemoteObject.h
--- a/Engine/Source/Runtime/CoreUObject/Public/UObject/RemoteObject.h
+++ b/Engine/Source/Runtime/CoreUObject/Public/UObject/RemoteObject.h
@@ -42,9 +42,13 @@ namespace UE::RemoteObject::Handle
 		FRemoteServerId ResidentServerId;
 
 		/** Server id of the server that has ownership of the object
-		    (note: only valid if the object is Local) */
+			(note: only valid if the object is Local) */
 		FRemoteServerId OwningServerId;
 
+		/** Server id of the which server the object should be on for physics
+			Physics ID = 0 (default) means physics ID should not be considered when migrating the object.   */
+		FRemoteServerId PhysicsOverrideServerId = FRemoteServerId(0);
+
 		FRemoteObjectStub() = default;
 		COREUOBJECT_API explicit FRemoteObjectStub(UObject* Object);
 	};
@@ -117,4 +121,16 @@ namespace UE::RemoteObject::Handle
 	* Sets the owner server id for a locally resident object
 	*/
 	COREUOBJECT_API void ChangeOwnerServerId(const UObject* Object, FRemoteServerId NewOwnerServerId);
+
+	/**
+	* Get the physics server id where the object should be simulated on
+	*/
+	COREUOBJECT_API FRemoteServerId GetPhysicsServerId(const UObject* Object);
+
+	/**
+	* Sets the physics id where the object should go to
+	*/
+	COREUOBJECT_API void ChangePhysicsServerId(const UObject* Object, FRemoteServerId NewPhysicsServerId);
+
+	COREUOBJECT_API void ClearAllPhysicsServerId();
 }
diff --git a/Engine/Source/Runtime/CoreUObject/Public/UObject/RemoteObjectTransfer.h b/Engine/Source/Runtime/CoreUObject/Public/UObject/RemoteObjectTransfer.h
--- a/Engine/Source/Runtime/CoreUObject/Public/UObject/RemoteObjectTransfer.h
+++ b/Engine/Source/Runtime/CoreUObject/Public/UObject/RemoteObjectTransfer.h
@@ -58,7 +58,7 @@ namespace UE::RemoteObject::Transfer
 	* @param RemoteServerId Server that sent object data
 	* @param Data Remote object data. Data ownership is transferred to the remote object queue
 	*/
-	COREUOBJECT_API void OnObjectDataReceived(FRemoteTransactionId RequestId, FRemoteServerId OwnerServerId, FRemoteObjectId ObjectId, FRemoteServerId RemoteServerId, FRemoteObjectData* Data);
+	COREUOBJECT_API void OnObjectDataReceived(FRemoteTransactionId RequestId, FRemoteServerId OwnerServerId, FRemoteServerId PhysicsId, FRemoteObjectId ObjectId, FRemoteServerId RemoteServerId, FRemoteObjectData* Data);
 
 	/**
 	* Called when a remote object request was denied by a remote server
@@ -128,7 +128,7 @@ namespace UE::RemoteObject::Transfer
 	extern COREUOBJECT_API void InitRemoteObjectTransfer();
 
 	/** Delegate that transfers object data to another server */
-	extern COREUOBJECT_API TDelegate<void(FRemoteTransactionId /*RequestId*/, FRemoteObjectId /*ObjectId*/, FRemoteServerId /*DestinationServerId*/, FRemoteServerId /*OwnerServerId*/, FRemoteObjectData& /*ObjectData*/)> RemoteObjectTransferDelegate;
+	extern COREUOBJECT_API TDelegate<void(FRemoteTransactionId /*RequestId*/, FRemoteObjectId /*ObjectId*/, FRemoteServerId /*DestinationServerId*/, FRemoteServerId /*OwnerServerId*/, FRemoteServerId /*PhysicsId*/, FRemoteObjectData& /*ObjectData*/)> RemoteObjectTransferDelegate;
 
 	/** Delegate that handles an object request being denied */
 	extern COREUOBJECT_API TDelegate<void(FRemoteTransactionId /*RequestId*/, FRemoteObjectId /*ObjectId*/, FRemoteServerId /*DestinationServerId*/)> RemoteObjectDeniedTransferDelegate;
@@ -159,7 +159,7 @@ namespace UE::RemoteObject::Transfer
 	extern COREUOBJECT_API TMulticastDelegate<void(FRemoteTransactionId /*RequestId*/, FRemoteObjectId /*ObjectId*/)> OnObjectTouchedDelegate;
 
 	/** Delegate that stores locally unrachable object data into a database. */
-	extern COREUOBJECT_API TDelegate<void(FRemoteTransactionId /*RequestId*/, FRemoteObjectId /*ObjectId*/, FRemoteServerId /*DatabaseId*/, FRemoteServerId /*OwnerServerId*/, FRemoteObjectData& /*ObjectData*/)> StoreRemoteObjectDataDelegate;
+	extern COREUOBJECT_API TDelegate<void(FRemoteTransactionId /*RequestId*/, FRemoteObjectId /*ObjectId*/, FRemoteServerId /*DatabaseId*/, FRemoteServerId /*OwnerServerId*/, FRemoteServerId /*PhysicsId*/, FRemoteObjectData& /*ObjectData*/)> StoreRemoteObjectDataDelegate;
 
 	/** Delegate that restores object data from a database. */
 	extern COREUOBJECT_API TDelegate<void(FRemoteTransactionId /*RequestId*/, FRemoteWorkPriority /*RequestPriority*/, FRemoteObjectId /*ObjectId*/, FRemoteServerId /*DatabaseId*/, FRemoteServerId /*DestinationServerId*/)> RestoreRemoteObjectDataDelegate;
