[StructUtils] Add operation to rename properties within a PropertyBag

diff --git a/Engine/Source/Runtime/CoreUObject/Private/StructUtils/PropertyBag.cpp b/Engine/Source/Runtime/CoreUObject/Private/StructUtils/PropertyBag.cpp
--- a/Engine/Source/Runtime/CoreUObject/Private/StructUtils/PropertyBag.cpp
+++ b/Engine/Source/Runtime/CoreUObject/Private/StructUtils/PropertyBag.cpp
@@ -1,6 +1,8 @@
 // Copyright Epic Games, Inc. All Rights Reserved.
 
 #include "StructUtils/PropertyBag.h"
+
+#include "HAL/IConsoleManager.h"
 #include "UObject/EnumProperty.h"
 #include "UObject/Package.h"
 #include "UObject/TextProperty.h"
@@ -12,6 +14,11 @@
 #include "StructUtils/UserDefinedStruct.h"
 #endif // WITH_EDITOR
 
+static TAutoConsoleVariable<bool> CVarEnforceValidName(
+	TEXT("StructUtils.EnforceValidPropertyName"),
+	true,
+	TEXT("Enforce valid naming on new and renamed property bag properties."));
+
 #include UE_INLINE_GENERATED_CPP_BY_NAME(PropertyBag)
 
 struct FPropertyBagCustomVersion
@@ -1429,7 +1436,56 @@ namespace UE::StructUtils::Private
 		}
 	}
 
-}; // UE::StructUtils::Private
+	EPropertyBagAlterationResult RenameProperty(TArray<FPropertyBagPropertyDesc>& Descs, const FName SourcePropertyName, const FName TargetPropertyName)
+	{
+		FName NewTargetPropertyName = TargetPropertyName;
+		if (CVarEnforceValidName.GetValueOnAnyThread())
+		{
+			if (SourcePropertyName.IsNone() || TargetPropertyName.IsNone())
+			{
+				return EPropertyBagAlterationResult::PropertyNameEmpty;
+			}
+			else if (!FInstancedPropertyBag::IsPropertyNameValid(TargetPropertyName))
+			{
+				return EPropertyBagAlterationResult::PropertyNameInvalidCharacters;
+			}
+		}
+		else
+		{
+			NewTargetPropertyName = FInstancedPropertyBag::SanitizePropertyName(NewTargetPropertyName);
+		}
+
+		int32 RenamedPropertyIndex = INDEX_NONE;
+		// Rename the first instance of PropertyName to NewName
+		for (int32 Index = 0; Index < Descs.Num(); ++Index)
+		{
+			if (Descs[Index].Name == NewTargetPropertyName)
+			{
+				return EPropertyBagAlterationResult::TargetPropertyAlreadyExists;
+			}
+			else if (Descs[Index].Name == SourcePropertyName)
+			{
+				if (!ensure(RenamedPropertyIndex == INDEX_NONE))
+				{
+					// There are duplicate source properties
+					return EPropertyBagAlterationResult::InternalError;
+				}
+
+				RenamedPropertyIndex = Index;
+			}
+		}
+
+		if (RenamedPropertyIndex != INDEX_NONE)
+		{
+			Descs[RenamedPropertyIndex].Name = NewTargetPropertyName;
+			return EPropertyBagAlterationResult::Success;
+		}
+		else
+		{
+			return EPropertyBagAlterationResult::SourcePropertyNotFound;
+		}
+	}
+} // UE::StructUtils::Private
 
 
 //----------------------------------------------------------------//
@@ -1691,7 +1747,7 @@ int32 FInstancedPropertyBag::GetNumPropertiesInBag() const
 	return 0;
 }
 
-void FInstancedPropertyBag::AddProperties(const TConstArrayView<FPropertyBagPropertyDesc> NewDescs)
+EPropertyBagAlterationResult FInstancedPropertyBag::AddProperties(const TConstArrayView<FPropertyBagPropertyDesc> NewDescs, const bool bOverwrite)
 {
 	TArray<FPropertyBagPropertyDesc> Descs;
 	if (const UPropertyBag* CurrentBagStruct = GetPropertyBagStruct())
@@ -1701,11 +1757,30 @@ void FInstancedPropertyBag::AddProperties(const TConstArrayView<FPropertyBagProp
 
 	for (const FPropertyBagPropertyDesc& NewDesc : NewDescs)
 	{
+		if (CVarEnforceValidName.GetValueOnAnyThread())
+		{
+			if (NewDesc.Name.IsNone())
+			{
+				return EPropertyBagAlterationResult::PropertyNameEmpty;
+			}
+			else if (!IsPropertyNameValid(NewDesc.Name))
+			{
+				return EPropertyBagAlterationResult::PropertyNameInvalidCharacters;
+			}
+		}
+
 		FPropertyBagPropertyDesc* ExistingProperty = Descs.FindByPredicate([&NewDesc](const FPropertyBagPropertyDesc& Desc) { return Desc.Name == NewDesc.Name; });
 		if (ExistingProperty != nullptr)
 		{
-			ExistingProperty->ValueType = NewDesc.ValueType;
-			ExistingProperty->ValueTypeObject = NewDesc.ValueTypeObject;
+			if (bOverwrite)
+			{
+				ExistingProperty->ValueType = NewDesc.ValueType;
+				ExistingProperty->ValueTypeObject = NewDesc.ValueTypeObject;
+			}
+			else
+			{
+				return EPropertyBagAlterationResult::TargetPropertyAlreadyExists;
+			}
 		}
 		else if (NewDesc.ValueType != EPropertyBagPropertyType::None)
 		{
@@ -1715,26 +1790,28 @@ void FInstancedPropertyBag::AddProperties(const TConstArrayView<FPropertyBagProp
 
 	const UPropertyBag* NewBagStruct = UPropertyBag::GetOrCreateFromDescs(Descs);
 	MigrateToNewBagStruct(NewBagStruct);
+
+	return EPropertyBagAlterationResult::Success;
 }
-	
-void FInstancedPropertyBag::AddProperty(const FName InName, const EPropertyBagPropertyType InValueType, const UObject* InValueTypeObject)
+
+EPropertyBagAlterationResult FInstancedPropertyBag::AddProperty(const FName InName, const EPropertyBagPropertyType InValueType, const UObject* InValueTypeObject, const bool bOverwrite)
 {
-	AddProperties({ FPropertyBagPropertyDesc(InName, InValueType, InValueTypeObject) });
+	return AddProperties({ FPropertyBagPropertyDesc(InName, InValueType, InValueTypeObject) }, bOverwrite);
 }
 
-void FInstancedPropertyBag::AddContainerProperty(const FName InName, const EPropertyBagContainerType InContainerType, const EPropertyBagPropertyType InValueType, const UObject* InValueTypeObject)
+EPropertyBagAlterationResult FInstancedPropertyBag::AddContainerProperty(const FName InName, const EPropertyBagContainerType InContainerType, const EPropertyBagPropertyType InValueType, const UObject* InValueTypeObject, const bool bOverwrite)
 {
-	AddProperties({ FPropertyBagPropertyDesc(InName, InContainerType, InValueType, InValueTypeObject) });
+	return AddProperties({ FPropertyBagPropertyDesc(InName, InContainerType, InValueType, InValueTypeObject) }, bOverwrite);
 }
 
-void FInstancedPropertyBag::AddContainerProperty(const FName InName, const FPropertyBagContainerTypes InContainerTypes, const EPropertyBagPropertyType InValueType, UObject* InValueTypeObject)
+EPropertyBagAlterationResult FInstancedPropertyBag::AddContainerProperty(const FName InName, const FPropertyBagContainerTypes InContainerTypes, const EPropertyBagPropertyType InValueType, UObject* InValueTypeObject, const bool bOverwrite)
 {
-	AddProperties({ FPropertyBagPropertyDesc(InName, InContainerTypes, InValueType, InValueTypeObject) });
+	return AddProperties({ FPropertyBagPropertyDesc(InName, InContainerTypes, InValueType, InValueTypeObject) }, bOverwrite);
 }
 
-void FInstancedPropertyBag::AddProperty(const FName InName, const FProperty* InSourceProperty)
+EPropertyBagAlterationResult FInstancedPropertyBag::AddProperty(const FName InName, const FProperty* InSourceProperty, const bool bOverwrite)
 {
-	AddProperties({ FPropertyBagPropertyDesc(InName, InSourceProperty ) });
+	return AddProperties({ FPropertyBagPropertyDesc(InName, InSourceProperty ) }, bOverwrite);
 }
 
 EPropertyBagResult FInstancedPropertyBag::ReplaceAllPropertiesAndValues(const TConstArrayView<FPropertyBagPropertyDesc> InDescs, const TConstArrayView<TConstArrayView<uint8>> InValues)
@@ -1769,7 +1846,7 @@ EPropertyBagResult FInstancedPropertyBag::ReplaceAllPropertiesAndValues(const TC
 	return EPropertyBagResult::Success;
 }
 
-void FInstancedPropertyBag::RemovePropertiesByName(const TConstArrayView<FName> PropertiesToRemove)
+EPropertyBagAlterationResult FInstancedPropertyBag::RemovePropertiesByName(const TConstArrayView<FName> PropertiesToRemove)
 {
 	TArray<FPropertyBagPropertyDesc> Descs;
 	if (const UPropertyBag* CurrentBagStruct = GetPropertyBagStruct())
@@ -1784,11 +1861,32 @@ void FInstancedPropertyBag::RemovePropertiesByName(const TConstArrayView<FName>
 
 	const UPropertyBag* NewBagStruct = UPropertyBag::GetOrCreateFromDescs(Descs);
 	MigrateToNewBagStruct(NewBagStruct);
+
+	return EPropertyBagAlterationResult::Success;
 }
-	
-void FInstancedPropertyBag::RemovePropertyByName(const FName PropertyToRemove)
+
+EPropertyBagAlterationResult FInstancedPropertyBag::RemovePropertyByName(const FName PropertyToRemove)
 {
-	RemovePropertiesByName({ PropertyToRemove });
+	return RemovePropertiesByName({PropertyToRemove});
+}
+
+EPropertyBagAlterationResult FInstancedPropertyBag::RenameProperty(const FName PropertyToRename, const FName NewName)
+{
+	TArray<FPropertyBagPropertyDesc> Descs;
+	if (const UPropertyBag* CurrentBagStruct = GetPropertyBagStruct())
+	{
+		Descs = CurrentBagStruct->GetPropertyDescs();
+	}
+
+	const EPropertyBagAlterationResult Result = UE::StructUtils::Private::RenameProperty(Descs, PropertyToRename, NewName);
+
+	if (Result == EPropertyBagAlterationResult::Success)
+	{
+		const UPropertyBag* NewBagStruct = UPropertyBag::GetOrCreateFromDescs(Descs);
+		MigrateToNewBagStruct(NewBagStruct);
+	}
+
+	return Result;
 }
 
 void FInstancedPropertyBag::MigrateToNewBagStruct(const UPropertyBag* NewBagStruct)
@@ -3053,6 +3151,33 @@ void FInstancedPropertyBag::GetPreloadDependencies(TArray<UObject*>& OutDeps)
 	}
 }
 
+inline bool FInstancedPropertyBag::IsPropertyNameValid(const FString& Name)
+{
+	return !Name.IsEmpty() && FName::IsValidXName(Name, INVALID_OBJECTNAME_CHARACTERS);
+}
+
+inline bool FInstancedPropertyBag::IsPropertyNameValid(const FName Name)
+{
+	return !Name.IsNone() && FName::IsValidXName(Name, INVALID_OBJECTNAME_CHARACTERS);
+}
+
+FName FInstancedPropertyBag::SanitizePropertyName(const FString& Name)
+{
+	if (!IsPropertyNameValid(Name))
+	{
+		return !Name.IsEmpty() ? MakeObjectNameFromDisplayLabel(Name, NAME_None) : FName(TEXT("Property"));
+	}
+	else
+	{
+		return FName(Name);
+	}
+}
+
+FName FInstancedPropertyBag::SanitizePropertyName(const FName Name)
+{
+	return SanitizePropertyName(Name.ToString());
+}
+
 //----------------------------------------------------------------//
 //  FPropertyBagArrayRef
 //----------------------------------------------------------------//
diff --git a/Engine/Source/Runtime/CoreUObject/Public/StructUtils/PropertyBag.h b/Engine/Source/Runtime/CoreUObject/Public/StructUtils/PropertyBag.h
--- a/Engine/Source/Runtime/CoreUObject/Public/StructUtils/PropertyBag.h
+++ b/Engine/Source/Runtime/CoreUObject/Public/StructUtils/PropertyBag.h
@@ -168,6 +168,18 @@ enum class EPropertyBagResult : uint8
 	DuplicatedValue,	// Tried to set an already existing Set Entry
 };
 
+/** Property bag alterations result code. */
+UENUM()
+enum class EPropertyBagAlterationResult : uint8
+{
+	Success,                       // Operation succeeded.
+	InternalError,                 // The operation could not be completed, due to an internal property bag error.
+	PropertyNameEmpty,             // The property name is empty.
+	PropertyNameInvalidCharacters, // The property name contains illegal characters. Consider using FInstancedPropertyBag::SanitizePropertyName.
+	SourcePropertyNotFound,        // The source property or property name was not found.
+	TargetPropertyAlreadyExists    // The target property or property name already exists.
+};
+
 USTRUCT()
 struct COREUOBJECT_API FPropertyBagPropertyDescMetaData
 {
@@ -392,8 +404,10 @@ struct COREUOBJECT_API FInstancedPropertyBag
 	 * Adds properties to the bag. If property of same name already exists, it will be replaced with the new type.
 	 * Numeric property values will be converted if possible, when a property's type changes.
 	 * @param Descs Descriptors of new properties to add.
+	 * @param bOverwrite Overwrite the property if it already exists.
+	 * @return The result of the alteration.
 	 */
-	void AddProperties(const TConstArrayView<FPropertyBagPropertyDesc> Descs);
+	EPropertyBagAlterationResult AddProperties(const TConstArrayView<FPropertyBagPropertyDesc> Descs, bool bOverwrite = true);
 	
 	/**
 	 * Adds a new property to the bag. If property of same name already exists, it will be replaced with the new type.
@@ -401,8 +415,10 @@ struct COREUOBJECT_API FInstancedPropertyBag
 	 * @param InName Name of the new property
 	 * @param InValueType Type of the new property
 	 * @param InValueTypeObject Type object (for struct, class, enum) of the new property
+	 * @param bOverwrite Overwrite the property if it already exists.
+	 * @return The result of the alteration.
 	 */
-	void AddProperty(const FName InName, const EPropertyBagPropertyType InValueType, const UObject* InValueTypeObject = nullptr);
+	EPropertyBagAlterationResult AddProperty(const FName InName, const EPropertyBagPropertyType InValueType, const UObject* InValueTypeObject = nullptr, bool bOverwrite = true);
 
 	/**
 	 * Adds a new container property to the bag. If property of same name already exists, it will be replaced with the new type.
@@ -410,8 +426,10 @@ struct COREUOBJECT_API FInstancedPropertyBag
 	 * @param InContainerType Type of the new container
 	 * @param InValueType Type of the new property
 	 * @param InValueTypeObject Type object (for struct, class, enum) of the new property
+	 * @param bOverwrite Overwrite the property if it already exists.
+	 * @return The result of the alteration.
 	 */
-	void AddContainerProperty(const FName InName, const EPropertyBagContainerType InContainerType, const EPropertyBagPropertyType InValueType, const UObject* InValueTypeObject = nullptr);
+	EPropertyBagAlterationResult AddContainerProperty(const FName InName, const EPropertyBagContainerType InContainerType, const EPropertyBagPropertyType InValueType, const UObject* InValueTypeObject = nullptr, bool bOverwrite = true);
 
 	/**
 	 * Adds a new container property to the bag. If property of same name already exists, it will be replaced with the new type.
@@ -419,15 +437,20 @@ struct COREUOBJECT_API FInstancedPropertyBag
 	 * @param InContainerTypes List of (optionally nested) containers to create
 	 * @param InValueType Type of the new property
 	 * @param InValueTypeObject Type object (for struct, class, enum) of the new property
+	 * @param bOverwrite Overwrite the property if it already exists.
+	 * @return The result of the alteration.
 	 */
-	void AddContainerProperty(const FName InName, const FPropertyBagContainerTypes InContainerTypes, const EPropertyBagPropertyType InValueType, UObject* InValueTypeObject);
+	EPropertyBagAlterationResult AddContainerProperty(const FName InName, const FPropertyBagContainerTypes InContainerTypes, const EPropertyBagPropertyType InValueType, UObject* InValueTypeObject, bool bOverwrite = true);
 
 	/**
 	 * Adds a new property to the bag. Property type duplicated from source property to. If property of same name already exists, it will be replaced with the new type.
 	 * @param InName Name of the property to add
 	 * @param InSourceProperty The property to add
+	 * @param bOverwrite Overwrite the property if it already exists.
+	 * @return The result of the alteration.
+
 	 */
-	void AddProperty(const FName InName, const FProperty* InSourceProperty);
+	EPropertyBagAlterationResult AddProperty(const FName InName, const FProperty* InSourceProperty, bool bOverwrite = true);
 
 	/**
 	 * Clears all properties, then adds the supplied properties to the bag and sets their respective values.
@@ -439,14 +462,26 @@ struct COREUOBJECT_API FInstancedPropertyBag
 
 	/**
 	 * Removes properties from the bag by name if they exist.
+	 * @param PropertiesToRemove The names of the properties to remove.
+	 * @return The result of the alteration.
 	 */
-	void RemovePropertiesByName(const TConstArrayView<FName> PropertiesToRemove);
+	EPropertyBagAlterationResult RemovePropertiesByName(const TConstArrayView<FName> PropertiesToRemove);
 	
 	/**
 	 * Removes a property from the bag by name if it exists.
+	 * @param PropertyToRemove The name of the property to remove.
+	 * @return The result of the alteration.
 	 */
-	void RemovePropertyByName(const FName PropertyToRemove);
-	
+	EPropertyBagAlterationResult RemovePropertyByName(const FName PropertyToRemove);
+
+	/**
+	 * Renames a property in the bag if it exists.
+	 * @param PropertyToRename The name of the property to rename
+	 * @param NewName The new name of the property
+	 * @return The result of the alteration
+	 */
+	EPropertyBagAlterationResult RenameProperty(FName PropertyToRename, FName NewName);
+
 	/**
 	 * Changes the type of this bag and migrates existing values.
 	 * The properties are matched between the bags based on the property ID.
@@ -773,6 +808,34 @@ struct COREUOBJECT_API FInstancedPropertyBag
 	void AddStructReferencedObjects(FReferenceCollector& Collector);
 	void GetPreloadDependencies(TArray<UObject*>& OutDeps);
 
+	/**
+	 * Checks whether a provided name is a valid property bag property name.
+	 * @param Name The name to check for invalid characters.
+	 * @return Sanitized property name based on the input string.
+	 */
+	static bool IsPropertyNameValid(const FString& Name);
+
+	/**
+	 * Checks whether a provided name is a valid property bag property name.
+	 * @param Name The name to check for invalid characters.
+	 * @return Sanitized property name based on the input name.
+	 */
+	static bool IsPropertyNameValid(FName Name);
+
+	/**
+	 * Returns a sanitized version of the provided name without invalid characters.
+	 * @param Name The name to sanitize.
+	 * @return Sanitized property name based on the input string.
+	 */
+	static FName SanitizePropertyName(const FString& Name);
+
+	/**
+	 * Returns a sanitized version of the provided name without invalid characters.
+	 * @param Name The name to sanitize.
+	 * @return Sanitized property name based on the input name.
+	 */
+	static FName SanitizePropertyName(FName Name);
+
 protected:
 	const void* GetValueAddress(const FPropertyBagPropertyDesc* Desc) const;
 	void* GetMutableValueAddress(const FPropertyBagPropertyDesc* Desc);
