Fixed regression warnings in the scenegraph workflow unittests

#rnx

diff --git a/Engine/Source/Runtime/CoreUObject/Private/UObject/OverriddenPropertySet.cpp b/Engine/Source/Runtime/CoreUObject/Private/UObject/OverriddenPropertySet.cpp
--- a/Engine/Source/Runtime/CoreUObject/Private/UObject/OverriddenPropertySet.cpp
+++ b/Engine/Source/Runtime/CoreUObject/Private/UObject/OverriddenPropertySet.cpp
@@ -396,12 +396,12 @@ FOverriddenPropertyNode& FOverriddenPropertySet::FindOrAddNode(FOverriddenProper
 	return OverriddenPropertyNodes.Get(NewID);
 }
 
-EOverriddenPropertyOperation FOverriddenPropertySet::GetOverriddenPropertyOperation(const FOverriddenPropertyNode& ParentPropertyNode, FPropertyVisitorPath::Iterator PropertyIterator, bool* bOutInheritedOperation, const void* Data) const
+EOverriddenPropertyOperation FOverriddenPropertySet::GetOverriddenPropertyOperation(const FOverriddenPropertyNode* ParentPropertyNode, FPropertyVisitorPath::Iterator PropertyIterator, bool* bOutInheritedOperation, const void* Data) const
 {
 	FOverridableManager& OverridableManager = FOverridableManager::Get();
 
 	const void* SubValuePtr = Data;
-	const FOverriddenPropertyNode* OverriddenPropertyNode = &ParentPropertyNode;
+	const FOverriddenPropertyNode* OverriddenPropertyNode = ParentPropertyNode;
 	int32 ArrayIndex = INDEX_NONE;
 	while (PropertyIterator && (!OverriddenPropertyNode || OverriddenPropertyNode->Operation != EOverriddenPropertyOperation::Replace))
 	{
@@ -517,6 +517,17 @@ EOverriddenPropertyOperation FOverriddenPropertySet::GetOverriddenPropertyOperat
 		}
 
 		OverriddenPropertyNode = CurrentOverriddenPropertyNode;
+		// While digging down the path, if there is one property that is always overridden
+		// stop there and return replace
+		if (CurrentProperty->HasAnyPropertyFlags(CPF_ExperimentalAlwaysOverriden))
+		{
+			if (bOutInheritedOperation)
+			{
+				*bOutInheritedOperation = NextPropertyIterator ? true : false;
+			}
+			return EOverriddenPropertyOperation::Replace;
+		}
+
 		++PropertyIterator;
 	}
 
@@ -1310,11 +1321,7 @@ void FOverriddenPropertySet::RemoveOverriddenSubProperties(FOverriddenPropertyNo
 
 EOverriddenPropertyOperation FOverriddenPropertySet::GetOverriddenPropertyOperation(FPropertyVisitorPath::Iterator PropertyIterator, bool* bOutInheritedOperation /*= nullptr*/) const
 {
-	if (const FOverriddenPropertyNode* RootNode = OverriddenPropertyNodes.Find(RootNodeID))
-	{
-		return GetOverriddenPropertyOperation(*RootNode, PropertyIterator, bOutInheritedOperation, Owner);
-	}
-	return EOverriddenPropertyOperation::None;
+	return GetOverriddenPropertyOperation(OverriddenPropertyNodes.Find(RootNodeID), PropertyIterator, bOutInheritedOperation, Owner);
 }
 
 bool FOverriddenPropertySet::ClearOverriddenProperty(FPropertyVisitorPath::Iterator PropertyIterator)
@@ -1342,11 +1349,7 @@ void FOverriddenPropertySet::NotifyPropertyChange(const EPropertyNotificationTyp
 
 EOverriddenPropertyOperation FOverriddenPropertySet::GetOverriddenPropertyOperation(const FArchiveSerializedPropertyChain* CurrentPropertyChain, FProperty* Property) const
 {
-	if (const FOverriddenPropertyNode* RootNode = OverriddenPropertyNodes.Find(RootNodeID))
-	{
-		return GetOverriddenPropertyOperation(*RootNode, CurrentPropertyChain, Property);
-	}
-	return EOverriddenPropertyOperation::None;
+	return GetOverriddenPropertyOperation(OverriddenPropertyNodes.Find(RootNodeID), CurrentPropertyChain, Property);
 }
 
 FOverriddenPropertyNode* FOverriddenPropertySet::SetOverriddenPropertyOperation(EOverriddenPropertyOperation Operation, const FArchiveSerializedPropertyChain* CurrentPropertyChain, FProperty* Property)
@@ -1380,12 +1383,12 @@ const FOverriddenPropertyNode* FOverriddenPropertySet::GetOverriddenPropertyNode
 	return nullptr;
 }
 
-EOverriddenPropertyOperation FOverriddenPropertySet::GetOverriddenPropertyOperation(const FOverriddenPropertyNode& ParentPropertyNode, const FArchiveSerializedPropertyChain* CurrentPropertyChain, FProperty* Property) const
+EOverriddenPropertyOperation FOverriddenPropertySet::GetOverriddenPropertyOperation(const FOverriddenPropertyNode* ParentPropertyNode, const FArchiveSerializedPropertyChain* CurrentPropertyChain, FProperty* Property) const
 {
 	// No need to look further
 	// if it is the entire property is replaced or
 	// if it is the FOverriddenPropertySet struct which is always Overridden
-	if (ParentPropertyNode.Operation == EOverriddenPropertyOperation::Replace)
+	if (ParentPropertyNode && ParentPropertyNode->Operation == EOverriddenPropertyOperation::Replace)
 	{
 		return EOverriddenPropertyOperation::Replace;
 	}
@@ -1398,29 +1401,31 @@ EOverriddenPropertyOperation FOverriddenPropertySet::GetOverriddenPropertyOperat
 	}
 
 	TArray<class FProperty*, TInlineAllocator<8>>::TConstIterator PropertyIterator = PropertyChain.GetRootIterator();
-	const FOverriddenPropertyNode* OverriddenPropertyNode = &ParentPropertyNode;
-	while (PropertyIterator && OverriddenPropertyNode && 
-		((OverriddenPropertyNode->Operation != EOverriddenPropertyOperation::Replace) || ((*PropertyIterator)->HasAnyPropertyFlags(CPF_ExperimentalAlwaysOverriden))))
+	const FOverriddenPropertyNode* OverriddenPropertyNode = ParentPropertyNode;
+	while (PropertyIterator && (!OverriddenPropertyNode || (OverriddenPropertyNode->Operation != EOverriddenPropertyOperation::Replace)))
 	{
 		const FProperty* CurrentProperty = (*PropertyIterator);
-		if (const FOverriddenPropertyNodeID* CurrentPropKey = OverriddenPropertyNode->SubPropertyNodeKeys.Find(CurrentProperty))
+		if (OverriddenPropertyNode)
 		{
-			OverriddenPropertyNode = OverriddenPropertyNodes.Find(*CurrentPropKey);
-			checkf(OverriddenPropertyNode, TEXT("Expecting a node"));
+			if (const FOverriddenPropertyNodeID* CurrentPropKey = OverriddenPropertyNode->SubPropertyNodeKeys.Find(CurrentProperty))
+			{
+				OverriddenPropertyNode = OverriddenPropertyNodes.Find(*CurrentPropKey);
+				checkf(OverriddenPropertyNode, TEXT("Expecting a node"));
+			}
+			else
+			{
+				OverriddenPropertyNode = nullptr;
+			}
 		}
-		else
+		// While digging down the path, if there is one property that is always overridden
+		// stop there and return replace
+		if (CurrentProperty->HasAnyPropertyFlags(CPF_ExperimentalAlwaysOverriden))
 		{
-			OverriddenPropertyNode = nullptr;
-			break;
+			return EOverriddenPropertyOperation::Replace;
 		}
 		++PropertyIterator;
 	}
 
-	if (PropertyIterator && (*PropertyIterator)->HasAnyPropertyFlags(CPF_ExperimentalAlwaysOverriden))
-	{
-		return EOverriddenPropertyOperation::Replace;
-	}
-
 	return OverriddenPropertyNode ? OverriddenPropertyNode->Operation : EOverriddenPropertyOperation::None;
 }
 
@@ -1446,6 +1451,12 @@ FOverriddenPropertyNode* FOverriddenPropertySet::SetOverriddenPropertyOperation(
 	while (PropertyIterator && OverriddenPropertyNode->Operation != EOverriddenPropertyOperation::Replace)
 	{
 		const FProperty* CurrentProperty = (*PropertyIterator);
+		// While digging down the path, if the operation is replace and one of the property is always overridden
+		// then there isn't anything to do
+		if (Operation == EOverriddenPropertyOperation::Replace && CurrentProperty->HasAnyPropertyFlags(CPF_ExperimentalAlwaysOverriden))
+		{
+			return nullptr;
+		}
 		OverriddenPropertyNode = &FindOrAddNode(*OverriddenPropertyNode, CurrentProperty);
 		++PropertyIterator;
 	}
diff --git a/Engine/Source/Runtime/CoreUObject/Public/UObject/OverriddenPropertySet.h b/Engine/Source/Runtime/CoreUObject/Public/UObject/OverriddenPropertySet.h
--- a/Engine/Source/Runtime/CoreUObject/Public/UObject/OverriddenPropertySet.h
+++ b/Engine/Source/Runtime/CoreUObject/Public/UObject/OverriddenPropertySet.h
@@ -437,11 +437,11 @@ protected:
 
 	FOverriddenPropertyNode& FindOrAddNode(FOverriddenPropertyNode& ParentPropertyNode, FOverriddenPropertyNodeID NodeID);
 
-	EOverriddenPropertyOperation GetOverriddenPropertyOperation(const FOverriddenPropertyNode& ParentPropertyNode, FPropertyVisitorPath::Iterator PropertyIterator, bool* bOutInheritedOperation, const void* Data) const;
+	EOverriddenPropertyOperation GetOverriddenPropertyOperation(const FOverriddenPropertyNode* ParentPropertyNode, FPropertyVisitorPath::Iterator PropertyIterator, bool* bOutInheritedOperation, const void* Data) const;
 	bool ClearOverriddenProperty(FOverriddenPropertyNode& ParentPropertyNode, FPropertyVisitorPath::Iterator PropertyIterator, const void* Data);
 	void NotifyPropertyChange(FOverriddenPropertyNode* ParentPropertyNode, const EPropertyNotificationType Notification, FPropertyVisitorPath::Iterator PropertyIterator, const EPropertyChangeType::Type ChangeType, const void* Data, bool& bNeedsCleanup);
 
-	EOverriddenPropertyOperation GetOverriddenPropertyOperation(const FOverriddenPropertyNode& ParentPropertyNode, const FArchiveSerializedPropertyChain* CurrentPropertyChain, FProperty* Property) const;
+	EOverriddenPropertyOperation GetOverriddenPropertyOperation(const FOverriddenPropertyNode* ParentPropertyNode, const FArchiveSerializedPropertyChain* CurrentPropertyChain, FProperty* Property) const;
 	FOverriddenPropertyNode* SetOverriddenPropertyOperation(EOverriddenPropertyOperation Operation, FOverriddenPropertyNode& ParentPropertyNode, const FArchiveSerializedPropertyChain* CurrentPropertyChain, FProperty* Property);
 	const FOverriddenPropertyNode* GetOverriddenPropertyNode(const FOverriddenPropertyNode& ParentPropertyNode, const FArchiveSerializedPropertyChain* CurrentPropertyChain) const;
 
