[AutoRTFM] Do not undo UObjectBase::ObjectListInternalIndex

ObjectListInternalIndex is an index into an array of UObjectHash's ClassToObjectListMap.
This index must be preserved, even if an reconstructed object is reverted as the table is not instrumented by AutoRTFM.

diff --git a/Engine/Source/Programs/AutoRTFMEngineTests/AutoRTFMEngineTests.Target.cs b/Engine/Source/Programs/AutoRTFMEngineTests/AutoRTFMEngineTests.Target.cs
--- a/Engine/Source/Programs/AutoRTFMEngineTests/AutoRTFMEngineTests.Target.cs
+++ b/Engine/Source/Programs/AutoRTFMEngineTests/AutoRTFMEngineTests.Target.cs
@@ -50,6 +50,12 @@ public class AutoRTFMEngineTestsTarget : TargetRules
 		// Disable crashreporter to improve startup time.
 		GlobalDefinitions.Add("NOINITCRASHREPORTER=1");
 
+		// Enable the UE_STORE_OBJECT_LIST_INTERNAL_INDEX optimization, as this
+		// has caused trouble with UObject reconstruction in transactions (SOL-7852)
+		bFNameOutlineNumber = true;
+		GlobalDefinitions.Add("WITH_CASE_PRESERVING_NAME=0");
+		GlobalDefinitions.Add("UE_STORE_OBJECT_LIST_INTERNAL_INDEX=1");
+
 		if (!bGenerateProjectFiles)
 		{
 			bUseAutoRTFMCompiler = true;
diff --git a/Engine/Source/Runtime/CoreUObject/Private/UObject/UObjectGlobals.cpp b/Engine/Source/Runtime/CoreUObject/Private/UObject/UObjectGlobals.cpp
--- a/Engine/Source/Runtime/CoreUObject/Private/UObject/UObjectGlobals.cpp
+++ b/Engine/Source/Runtime/CoreUObject/Private/UObject/UObjectGlobals.cpp
@@ -3753,7 +3753,19 @@ UObject* StaticAllocateObject
 					// including fields that may be accessed only in the open (e.g. object flags which
 					// are set atomically), and mixing closed writes with open writes will trigger the
 					// memory validator.
-					AutoRTFM::RecordOpenWriteNoMemoryValidation((void *)Obj, TotalSize);
+#if UE_STORE_OBJECT_LIST_INTERNAL_INDEX
+					// ObjectListInternalIndex is an index into an array of UObjectHash's ClassToObjectListMap.
+					// This index must be preserved, even if a reconstructed object is reverted as the
+					// table is not instrumented by AutoRTFM.
+					char* ObjectStart = reinterpret_cast<char*>(Obj);
+					char* InternalIndexStart = reinterpret_cast<char*>(&Obj->ObjectListInternalIndex);
+					char* InternalIndexEnd = InternalIndexStart + sizeof(Obj->ObjectListInternalIndex);
+					char* ObjectEnd = ObjectStart + TotalSize;
+					AutoRTFM::RecordOpenWriteNoMemoryValidation(ObjectStart, InternalIndexStart - ObjectStart);
+					AutoRTFM::RecordOpenWriteNoMemoryValidation(InternalIndexEnd, ObjectEnd - InternalIndexEnd);
+#else
+					AutoRTFM::RecordOpenWriteNoMemoryValidation(Obj, TotalSize);
+#endif
 				};
 
 				// Get the name before we start the destroy, as destroy renames it
diff --git a/Engine/Source/Runtime/CoreUObject/Private/UObject/UObjectHash.cpp b/Engine/Source/Runtime/CoreUObject/Private/UObject/UObjectHash.cpp
--- a/Engine/Source/Runtime/CoreUObject/Private/UObject/UObjectHash.cpp
+++ b/Engine/Source/Runtime/CoreUObject/Private/UObject/UObjectHash.cpp
@@ -1150,6 +1150,8 @@ FORCEINLINE void RemoveFromClassMap(FUObjectHashTables& ThreadHash, UObjectBase*
 	{
 		auto& ObjectList = ThreadHash.ClassToObjectListMap.FindOrAdd(Object->GetClass());
 #if UE_STORE_OBJECT_LIST_INTERNAL_INDEX
+		ensureMsgf(ObjectList[Object->ObjectListInternalIndex] == Object,
+			TEXT("Object doesn't match the one stored in the ObjectList"));
 		ObjectList.Last()->ObjectListInternalIndex = Object->ObjectListInternalIndex;
 		ObjectList[Object->ObjectListInternalIndex] = ObjectList.Last();
 		ObjectList.Pop(EAllowShrinking::No);
diff --git a/Engine/Source/Runtime/CoreUObject/Public/UObject/UObjectBase.h b/Engine/Source/Runtime/CoreUObject/Public/UObject/UObjectBase.h
--- a/Engine/Source/Runtime/CoreUObject/Public/UObject/UObjectBase.h
+++ b/Engine/Source/Runtime/CoreUObject/Public/UObject/UObjectBase.h
@@ -60,6 +60,18 @@ class UObjectBase
 		const TCHAR* PackageName,
 		const TCHAR* Name
 		);
+	friend COREUOBJECT_API UObject* StaticAllocateObject( // for access to ObjectListInternalIndex
+		const UClass*,
+		UObject*,
+		FName,
+		EObjectFlags,
+		EInternalObjectFlags,
+		bool,
+		bool*,
+		UPackage*,
+		int32,
+		FRemoteObjectId);
+
 protected:
 	UObjectBase() :
 		ClassPrivate(NoInit),
diff --git a/Engine/Source/Runtime/Engine/Private/Tests/AutoRTFM/AutoRTFMActorComponentTests.cpp b/Engine/Source/Runtime/Engine/Private/Tests/AutoRTFM/AutoRTFMActorComponentTests.cpp
--- a/Engine/Source/Runtime/Engine/Private/Tests/AutoRTFM/AutoRTFMActorComponentTests.cpp
+++ b/Engine/Source/Runtime/Engine/Private/Tests/AutoRTFM/AutoRTFMActorComponentTests.cpp
@@ -518,6 +518,12 @@ AUTORTFM_ACTOR_COMPONENT_TEST(ReconstructActor)
 	TestTrueExpr(OldWeak == Old);
 	
 	const uint32 OldID = Old->GetUniqueID();
+
+	// Create some additional actors to append more objects to the UObjectHash tables
+	for (int I = 0; I < 5; I++)
+	{
+		NewObject<AAutoRTFMTestActor>();
+	}
 	
 	AutoRTFM::Testing::Abort([&]
 		{
@@ -533,6 +539,10 @@ AUTORTFM_ACTOR_COMPONENT_TEST(ReconstructActor)
 	TestTrueExpr(OldWeak == Old);
 	TestTrueExpr(OldID == Old->GetUniqueID());
 	TestTrueExpr(123 == Old->ActorCategory);
+
+	// Finally reconstruct the object one more time. This exercises code that
+	// can detect corrupt UObjectHash tables
+	NewObject<AAutoRTFMTestActor>(Outer, Name);
 }
 
 // Reproduction for mixed open / closed writes on FUObjectItem::SerialNumber.
