Added a function to create an IDO without an archive to load from and fixed archetype serialization

#rnx

diff --git a/Engine/Source/Runtime/CoreUObject/Internal/UObject/InstanceDataObjectUtils.h b/Engine/Source/Runtime/CoreUObject/Internal/UObject/InstanceDataObjectUtils.h
--- a/Engine/Source/Runtime/CoreUObject/Internal/UObject/InstanceDataObjectUtils.h
+++ b/Engine/Source/Runtime/CoreUObject/Internal/UObject/InstanceDataObjectUtils.h
@@ -56,6 +56,18 @@ void CopyPropertyValueSerializedData(const FFieldVariant& OldField, void* OldDat
 /** Returns whether the object is an IDO. */
 UE_API bool IsInstanceDataObject(const UObject* Object);
 
+/**
+ * Creates an IDO for Owner and serializes it from Archive.
+ *
+ * USE THE OVERLOAD TAKING AN ARCHIVE IF UNKNOWN PROPERTY VALUES ARE REQUIRED!
+ *
+ * Creates an IDO class containing the union of the fields in Owner and any associated unknown property tree.
+ *
+ * @param Owner   Object from which to construct an InstanceDataObject.
+ * @return The InstanceDataObject for Owner.
+ */
+UE_API UObject* CreateInstanceDataObject(UObject* Owner);
+
 /**
  * Creates an IDO for Owner and serializes it from Archive.
  *
diff --git a/Engine/Source/Runtime/CoreUObject/Internal/UObject/PropertyBagRepository.h b/Engine/Source/Runtime/CoreUObject/Internal/UObject/PropertyBagRepository.h
--- a/Engine/Source/Runtime/CoreUObject/Internal/UObject/PropertyBagRepository.h
+++ b/Engine/Source/Runtime/CoreUObject/Internal/UObject/PropertyBagRepository.h
@@ -68,9 +68,10 @@ public:
 	 * @param Archive		Archive from which to serialize script properties for the object.
 	 * @param StartOffset	Offset to seek Archive to before serializing script properties for the object.
 	 * @param EndOffset		Offset to expect the Archive to be at after serializing script properties for the object.
+	 * @param bIsArchetype	True if the IDO is to be treated as an archetype even when RF_ArchetypeObject is not set.
 	 * @return				The InstanceDataObject object, UClass derived from associated unknown property tree.
 	 */
-	UObject* CreateInstanceDataObject(UObject* Owner, FArchive& Archive, int64 StartOffset, int64 EndOffset);
+	UObject* CreateInstanceDataObject(UObject* Owner, FArchive& Archive, int64 StartOffset, int64 EndOffset, bool bIsArchetype = false);
 	UObject* DuplicateInstanceDataObject(UObject* SourceOwner, UObject* DestOwner);
 
 	// called at the end of postload to copy data from Owner to its IDO
@@ -154,9 +155,6 @@ private:
 
 	// Delete owner reference and disassociate all data. Returns success.
 	bool RemoveAssociationUnsafe(const UObject* Owner);
-
-	// Internal function for creating IDOs. Before calling this function the repository's lock should be acquired.
-	UObject* CreateInstanceDataObjectInternal(UObject* Owner, FArchive& Archive, int64 StartOffset, int64 EndOffset, bool bIsArchetype = false);
 };
 
 #endif // WITH_EDITORONLY_DATA
diff --git a/Engine/Source/Runtime/CoreUObject/Private/UObject/InstanceDataObjectUtils.cpp b/Engine/Source/Runtime/CoreUObject/Private/UObject/InstanceDataObjectUtils.cpp
--- a/Engine/Source/Runtime/CoreUObject/Private/UObject/InstanceDataObjectUtils.cpp
+++ b/Engine/Source/Runtime/CoreUObject/Private/UObject/InstanceDataObjectUtils.cpp
@@ -1299,6 +1299,20 @@ bool IsInstanceDataObject(const UObject* Object)
 	return Object && Object->GetClass()->UObject::IsA(UInstanceDataObjectClass::StaticClass());
 }
 
+UObject* CreateInstanceDataObject(UObject* Owner)
+{
+	TArray<uint8> OwnerData;
+
+	FObjectWriter Writer(OwnerData);
+	Writer.ArNoDelta = true;
+	Owner->SerializeScriptProperties(Writer);
+
+	FObjectReader Reader(OwnerData);
+	Reader.ArMergeOverrides = true;
+	Reader.ArPreserveArrayElements = true;
+	return CreateInstanceDataObject(Owner, Reader, 0, Reader.TotalSize());
+}
+
 UObject* CreateInstanceDataObject(UObject* Owner, FArchive& Ar, int64 StartOffset, int64 EndOffset)
 {
 	return FPropertyBagRepository::Get().CreateInstanceDataObject(Owner, Ar, StartOffset, EndOffset);
diff --git a/Engine/Source/Runtime/CoreUObject/Private/UObject/PropertyBagRepository.cpp b/Engine/Source/Runtime/CoreUObject/Private/UObject/PropertyBagRepository.cpp
--- a/Engine/Source/Runtime/CoreUObject/Private/UObject/PropertyBagRepository.cpp
+++ b/Engine/Source/Runtime/CoreUObject/Private/UObject/PropertyBagRepository.cpp
@@ -818,9 +818,15 @@ FString FPropertyBagRepository::GetReferencerName() const
 	return TEXT("FPropertyBagRepository");
 }
 
-// Before calling this function the repository's lock should be acquired.
-UObject* FPropertyBagRepository::CreateInstanceDataObjectInternal(UObject* Owner, FArchive& Archive, int64 StartOffset, int64 EndOffset, bool bIsArchetype)
+UObject* FPropertyBagRepository::CreateInstanceDataObject(UObject* Owner, FArchive& Archive, int64 StartOffset, int64 EndOffset, bool bIsArchetype)
 {
+	// Limit the scope of the lock during to the find because this calls itself recursively to handle archetypes.
+	if (FPropertyBagRepositoryLock LockRepo(this); FPropertyBagAssociationData* BagData = AssociatedData.Find(Owner))
+	{
+		check(BagData->InstanceDataObject);
+		return BagData->InstanceDataObject;
+	}
+
 	// TODO: Should the IDO or its class be in another package?
 	UObject* Outer = GetTransientPackage();
 
@@ -830,14 +836,11 @@ UObject* FPropertyBagRepository::CreateInstanceDataObjectInternal(UObject* Owner
 	const FUnknownEnumNames* EnumNamesPtr = !EnumNames.IsEmpty() ? &EnumNames : nullptr;
 	const UClass* InstanceDataObjectClass = CreateInstanceDataObjectClass(PropertyTree.Get(), EnumNamesPtr, Owner->GetClass(), Outer);
 
-	UObject* OwnerArchetype = nullptr;
+	// Find the template for the IDO.
 	UObject* Template = nullptr;
-
-	bool bIsIDOArchetypeChainEnabled = UE::IsInstanceDataObjectArchetypeChainEnabled();
-
-	if (bIsIDOArchetypeChainEnabled)
+	if (IsInstanceDataObjectArchetypeChainEnabled())
 	{
-		OwnerArchetype = Owner->GetArchetype();
+		UObject* OwnerArchetype = Owner->GetArchetype();
 
 		// If an archetype is not the CDO then it's a sub-object archetype.
 		bool bIsSubobjectArchetype = OwnerArchetype && !OwnerArchetype->HasAllFlags(RF_ClassDefaultObject);
@@ -846,32 +849,8 @@ UObject* FPropertyBagRepository::CreateInstanceDataObjectInternal(UObject* Owner
 		// not initialize (e.g., Verse overrides). For this reason we create those IDO archetypes on demand.
 		if (bIsSubobjectArchetype)
 		{
-			// We cannot do a FindAndAdd here because the map could be modified while creating more
-			// IDO archetypes recursively.
-			FPropertyBagAssociationData* TemplateBagData = AssociatedData.Find(OwnerArchetype);
-			if (TemplateBagData)
-			{
-				check(TemplateBagData->InstanceDataObject);
-				Template = TemplateBagData->InstanceDataObject;
-			}
-			else
-			{
-				// Serialize the owner's archetype into an arhive we can use to construct its IDO.
-				TArray<uint8> OwnerArchetypeData;
-				OwnerArchetypeData.Reserve(OwnerArchetype->GetClass()->GetStructureSize());
-
-				FObjectWriter Writer(OwnerArchetypeData);
-				OwnerArchetype->GetClass()->SerializeTaggedProperties(Writer, (uint8*)OwnerArchetype, OwnerArchetype->GetClass(), nullptr);
-				FObjectReader Reader(OwnerArchetypeData);
-				Reader.ArMergeOverrides = true;
-				Reader.ArPreserveArrayElements = true;
-
-				int64 TemplateEndOffset = Reader.TotalSize();
-				int64 TemplateStartOffset = (TemplateEndOffset == INDEX_NONE) ? INDEX_NONE : 0;
-
-				Template = CreateInstanceDataObjectInternal(OwnerArchetype, Reader, TemplateStartOffset, TemplateEndOffset, /*bIsArchetype*/true);
-				check(Template != nullptr);
-			}
+			Template = UE::CreateInstanceDataObject(OwnerArchetype);
+			check(Template);
 		}
 	}
 
@@ -890,13 +869,14 @@ UObject* FPropertyBagRepository::CreateInstanceDataObjectInternal(UObject* Owner
 	UObject* InstanceDataObject = StaticConstructObject_Internal(Params);
 
 	// We are not passing the template in FStaticConstructObjectParameters because of potential type mismatches
-	// between the IDO and its template. This can happen when both objects have the same 'owner' type but a 
+	// between the IDO and its template. This can happen when both objects have the same 'owner' type but a
 	// different number of unknown properties. For this reason we do a copy instead.
 	if (Template)
 	{
 		CopyTaggedProperties(Template, InstanceDataObject);
 	}
 
+	FPropertyBagRepositoryLock LockRepo(this);
 	InstanceDataObjectToOwner.Add(InstanceDataObject, Owner);
 
 	FPropertyBagAssociationData& BagData = AssociatedData.Add(Owner);
@@ -925,22 +905,6 @@ UObject* FPropertyBagRepository::CreateInstanceDataObjectInternal(UObject* Owner
 	return BagData.InstanceDataObject;
 }
 
-UObject* FPropertyBagRepository::CreateInstanceDataObject(UObject* Owner, FArchive& Archive, int64 StartOffset, int64 EndOffset)
-{
-	FPropertyBagRepositoryLock LockRepo(this);
-
-	FPropertyBagAssociationData* BagData = AssociatedData.Find(Owner);
-	if (BagData)
-	{
-		check(BagData->InstanceDataObject);
-		return BagData->InstanceDataObject;
-	}
-	else
-	{
-		return CreateInstanceDataObjectInternal(Owner, Archive, StartOffset, EndOffset);
-	}
-}
-
 // Not sure this is necessary.
 void FPropertyBagRepository::ShrinkMaps()
 {
