Remote Objects:  Attempting to codify some rules around borrowing / assigning ownership / returning borrowed objects.    Created a new FUObjectMigrationContext to keep track of the current migration parameters.  Passing that context to the Remote Object FArchives so that we can get this context when Serializing or Deserializing.  That allows us to choose what fidelity of data we're serializing (we anticipate using this for Physics).

#rnx

diff --git a/Engine/Source/Runtime/Core/Public/Serialization/Archive.h b/Engine/Source/Runtime/Core/Public/Serialization/Archive.h
--- a/Engine/Source/Runtime/Core/Public/Serialization/Archive.h
+++ b/Engine/Source/Runtime/Core/Public/Serialization/Archive.h
@@ -42,6 +42,7 @@ template <typename T> struct TObjectPtr;
 
 typedef TFunction<bool (double RemainingTime)> FExternalReadCallback;
 class FField;
+struct FUObjectMigrationContext;
 struct FUObjectSerializeContext;
 
 enum class EFileRegionType : uint8;
@@ -336,6 +337,12 @@ public:
 	 */
 	CORE_API void SetMigratingRemoteObjects(bool bMigrating);
 
+	/** When migrating objects, we may need some context as to the intent of the migration during serialization */
+	virtual const FUObjectMigrationContext* GetMigrationContext() const
+	{
+		return nullptr;
+	}
+
 	/** Returns true if this Archive is migrating remote object */
 #if UE_WITH_REMOTE_OBJECT_HANDLE
 	FORCEINLINE bool IsMigratingRemoteObjects() const
@@ -2183,6 +2190,8 @@ public:
 	using FArchiveState::ShouldSkipUpdateCustomVersion;
 	using FArchiveState::SetMigratingRemoteObjects;
 	using FArchiveState::IsMigratingRemoteObjects;
+	using FArchiveState::GetMigrationContext;
+
 	
 private:
 	using FArchiveState::ArUEVer;
diff --git a/Engine/Source/Runtime/CoreUObject/Internal/UObject/UObjectMigrationContext.h b/Engine/Source/Runtime/CoreUObject/Internal/UObject/UObjectMigrationContext.h
new file mode 100644
--- /dev/null
+++ b/Engine/Source/Runtime/CoreUObject/Internal/UObject/UObjectMigrationContext.h
@@ -0,0 +1,89 @@
+// Copyright Epic Games, Inc. All Rights Reserved.
+
+#pragma once
+
+#include "Templates/UniquePtr.h"
+#include "UObject/RemoteExecutor.h"
+#include "UObject/RemoteObjectTypes.h"
+
+enum class EObjectMigrationSide : uint8
+{
+	/** We are receiving (migrating-in) an object */
+	Receive,
+
+	/** We are sending (migrating-out) an object */
+	Send,
+};
+
+/**
+* Possible types of migration that have happened when you have received an Object
+*/
+enum class EObjectMigrationRecvType : uint8
+{
+	Invalid,
+
+	/** Received an object without changing ownership (implies it will be sent back to the server that owns it eventually) */
+	Borrowed,
+
+	/** We are receiving an Object we have previously loaned-out (we already own it) */
+	ReturnedLoan,
+
+	/** We received objects and must take ownership of the them (we have no choice) */
+	AssignedOwnership
+};
+COREUOBJECT_API const TCHAR* ToString(EObjectMigrationRecvType Value);
+
+/**
+	* Possible types of migration that are happening when you are sending an Object
+	*/
+enum class EObjectMigrationSendType : uint8
+{
+	Invalid,
+
+	/** Loaned-out an object while retaining ownership (implies we want this back, we still own it) */
+	Loan,
+
+	/** This borrowed object is being returned (sent back) to the server that loaned it out */
+	ReturnBorrowed,
+
+	/** Reassign objects to relinquish ownership and assign them to the destination server */
+	ReassignOwnership
+};
+COREUOBJECT_API const TCHAR* ToString(EObjectMigrationSendType Value);
+
+/** Structure that holds context for what we're intending to do during migration */
+struct COREUOBJECT_API FUObjectMigrationContext
+{
+	/** The Object that is being migrated, it may not be the top-level UObject */
+	FRemoteObjectId ObjectId;
+
+	/** The ServerId on the other side of the migration */
+	FRemoteServerId RemoteServerId;
+
+	/** The (New) Owner of the ObjectId as a result of this migration */
+	FRemoteServerId OwnerServerId;
+
+	/** The Server that this Object's Physics reside on */
+	FRemoteServerId PhysicsServerId;
+
+	/** The RequestId that corresponds to this migration request */
+	FRemoteTransactionId RequestId;
+
+	/** Which side of the migration does this context apply to? Sending or Receiving? */
+	EObjectMigrationSide MigrationSide;
+
+	/** Are the values of this context valid? */
+	bool IsValid() const;
+
+	/**
+	* During migration, figure out the role of the object being received
+	*/
+	EObjectMigrationRecvType GetObjectMigrationRecvType(const UObjectBase* Object) const;
+	EObjectMigrationRecvType GetObjectMigrationRecvType(FRemoteObjectId ObjectId) const;
+
+	/**
+	* During migration, figure out the role of the object being sent
+	*/
+	EObjectMigrationSendType GetObjectMigrationSendType(const UObjectBase* Object) const;
+	EObjectMigrationSendType GetObjectMigrationSendType(FRemoteObjectId ObjectId) const;
+};
diff --git a/Engine/Source/Runtime/CoreUObject/Private/UObject/RemoteObject.cpp b/Engine/Source/Runtime/CoreUObject/Private/UObject/RemoteObject.cpp
--- a/Engine/Source/Runtime/CoreUObject/Private/UObject/RemoteObject.cpp
+++ b/Engine/Source/Runtime/CoreUObject/Private/UObject/RemoteObject.cpp
@@ -14,6 +14,7 @@
 #include "UObject/RemoteObjectPrivate.h"
 #include "UObject/ObjectHandlePrivate.h"
 #include "UObject/UObjectHash.h"
+#include "UObject/UObjectMigrationContext.h"
 #include "Misc/CommandLine.h"
 #include "Misc/CoreDelegates.h"
 #include <atomic>
@@ -129,7 +130,17 @@ void InitRemoteObjects()
 	}
 	if (!UE::RemoteObject::Transfer::RequestRemoteObjectDelegate.IsBound())
 	{
-		UE::RemoteObject::Transfer::RequestRemoteObjectDelegate.BindStatic(&UE::RemoteObject::Serialization::Disk::LoadObjectFromDisk);
+		UE::RemoteObject::Transfer::RequestRemoteObjectDelegate.BindLambda(
+			[](FRemoteTransactionId RequestId, FRemoteWorkPriority RequestPriority, FRemoteObjectId ObjectId, FRemoteServerId LastKnownResidentServerId, FRemoteServerId DestinationServerId)
+			{
+				// Turns a request into an immediate load
+				FUObjectMigrationContext MigrationContext {
+					.ObjectId = ObjectId, .RemoteServerId = DestinationServerId, .OwnerServerId = LastKnownResidentServerId,
+					.PhysicsServerId = LastKnownResidentServerId, .RequestId = RequestId, .MigrationSide = EObjectMigrationSide::Receive
+				};
+				UE::RemoteObject::Serialization::Disk::LoadObjectFromDisk(MigrationContext);
+			}
+		);
 	}
 	if (!UE::RemoteObject::Transfer::StoreRemoteObjectDataDelegate.IsBound())
 	{
@@ -153,6 +164,7 @@ void RegisterRemoteObjectId(FRemoteObjectId ObjectId, FRemoteServerId ResidentSe
 
 void RegisterSharedObject(UObject* Object)
 {
+	ensureMsgf(!Object->HasAnyFlags(EObjectFlags::RF_ClassDefaultObject | EObjectFlags::RF_ArchetypeObject), TEXT("We're about to set an ArchetypeObject %s as remote"), *GetNameSafe(Object));
 	FUObjectItem* ObjectItem = GUObjectArray.ObjectToObjectItem(Object);
 	ObjectItem->SetFlags(EInternalObjectFlags::RemoteReference);
 }
@@ -160,6 +172,8 @@ void RegisterSharedObject(UObject* Object)
 void MarkAsRemote(UObject* Object, FRemoteServerId DestinationServerId)
 {
 	static_assert(sizeof(FObjectHandle) == sizeof(UObject*));
+	checkf(!Object->IsA<UClass>(), TEXT("Attempted to Migrate UClass %s which is considered an asset and never allowed to migrate"), *GetNameSafe(Object));
+
 	FUObjectItem* ObjectItem = GUObjectArray.ObjectToObjectItem(Object);
 	ObjectItem->SetFlags(EInternalObjectFlags::Remote);
 	ObjectItem->ClearFlags(EInternalObjectFlags_RootFlags | EInternalObjectFlags::RemoteReference);
@@ -171,6 +185,7 @@ void MarkAsLocal(UObject* Object)
 	FUObjectItem* ObjectItem = GUObjectArray.ObjectToObjectItem(Object);
 	ObjectItem->ClearFlags(EInternalObjectFlags::Remote);
 	ObjectItem->SetFlags(EInternalObjectFlags::RemoteReference);
+	ensureMsgf(!Object->HasAnyFlags(EObjectFlags::RF_ClassDefaultObject | EObjectFlags::RF_ArchetypeObject), TEXT("We're about to set an ArchetypeObject %s as remote reference"), *GetNameSafe(Object));
 }
 
 UE::RemoteObject::Handle::FRemoteObjectStub* FindRemoteObjectStub(FRemoteObjectId ObjectId)
diff --git a/Engine/Source/Runtime/CoreUObject/Private/UObject/RemoteObjectSerialization.cpp b/Engine/Source/Runtime/CoreUObject/Private/UObject/RemoteObjectSerialization.cpp
--- a/Engine/Source/Runtime/CoreUObject/Private/UObject/RemoteObjectSerialization.cpp
+++ b/Engine/Source/Runtime/CoreUObject/Private/UObject/RemoteObjectSerialization.cpp
@@ -8,6 +8,7 @@
 #include "UObject/Class.h"
 #include "UObject/PropertyPortFlags.h"
 #include "UObject/UObjectHash.h"
+#include "UObject/UObjectMigrationContext.h"
 #include "UObject/UnrealType.h"
 #include "Serialization/ArchiveUObject.h"
 #include "Serialization/MemoryWriter.h"
@@ -48,7 +49,7 @@ FArchive& operator<<(FArchive& Ar, FRemoteObjectBytes& Chunk)
 namespace UE::RemoteObject::Serialization::Disk
 {
 
-void SerializeNameTables(FArchive& Ar, FRemoteObjectData& InObjectData)
+void SerializeNameTables(FArchive& Ar, const FRemoteObjectData& InObjectData)
 {
 	FNameIndexType NumNames = IntCastChecked<FNameIndexType>(InObjectData.Names.Num());
 	Ar << NumNames;
@@ -57,7 +58,7 @@ void SerializeNameTables(FArchive& Ar, FRemoteObjectData& InObjectData)
 		FString NameString = Name.ToString();
 		Ar << NameString;
 	}
-	Ar << InObjectData.PathNames;
+	Ar << const_cast<TArray<FRemoteObjectPathName>&>(InObjectData.PathNames);
 }
 
 void DeserializeNameTables(FArchive& Ar, FRemoteObjectData& OutObjectData)
@@ -79,28 +80,31 @@ FString GenerateRemoteObjectFilename(FRemoteObjectId ObjectId, FRemoteServerId D
 	return FPaths::Combine(*FPaths::ProjectSavedDir(), *(FString::Printf(TEXT("%s_%s.remote"), *ObjectId.ToString(), *DestinationServerId.ToString())));
 }
 
-void LoadObjectFromDisk(FRemoteTransactionId RequestId, FRemoteWorkPriority RequestPriority, FRemoteObjectId ObjectId, FRemoteServerId CurrentOwnerServerId, FRemoteServerId DestinationServerId)
+void LoadObjectFromDisk(const FUObjectMigrationContext& MigrationContext)
 {
-	const FString Filename(GenerateRemoteObjectFilename(ObjectId, CurrentOwnerServerId));
+	const FString Filename(GenerateRemoteObjectFilename(MigrationContext.ObjectId, MigrationContext.OwnerServerId));
 	TUniquePtr<FArchive> FileReader{ IFileManager::Get().CreateFileReader(*Filename) };
-	checkf(FileReader.IsValid(), TEXT("Unable to create file reader for remote object %s"), *ObjectId.ToString());
-	FRemoteObjectData* ObjectData = new FRemoteObjectData();
-	
-	DeserializeNameTables(*FileReader, *ObjectData);
-	*FileReader << ObjectData->Bytes;
+	checkf(FileReader.IsValid(), TEXT("Unable to create file reader for remote object %s"), *MigrationContext.ObjectId.ToString());
+
+	FRemoteObjectData ObjectData;	
+	DeserializeNameTables(*FileReader, ObjectData);
+	*FileReader << ObjectData.Bytes;
 	FileReader->Close();
 	IFileManager::Get().Delete(*Filename, false, true, true);
 
-	UE::RemoteObject::Transfer::OnObjectDataReceived(RequestId, FRemoteServerId(), FRemoteServerId(), ObjectId, CurrentOwnerServerId, ObjectData);
+	// We have transferred ownership from the Database to the local server
+	FRemoteServerId LocalServerId = UE::RemoteObject::GetGlobalServerId();
+	const FRemoteServerId& DatabaseId = UE::RemoteObject::Transfer::DatabaseId;
+	UE::RemoteObject::Transfer::OnObjectDataReceived(MigrationContext.RequestId, LocalServerId, LocalServerId, MigrationContext.ObjectId, DatabaseId, ObjectData);
 }
 
-void SaveObjectToDisk(FRemoteTransactionId RequestId, FRemoteObjectId ObjectId, FRemoteServerId DestinationServerId, FRemoteServerId OwnerServerId, FRemoteServerId PhysicsId, FRemoteObjectData& ObjectData)
+void SaveObjectToDisk(const UE::RemoteObject::Transfer::FMigrateSendParams& Params)
 {
-	TUniquePtr<FArchive> FileWriter{ IFileManager::Get().CreateFileWriter(*GenerateRemoteObjectFilename(ObjectId, DestinationServerId)) };
-	checkf(FileWriter.IsValid(), TEXT("Unable to create file writer for remote object %s"), *ObjectId.ToString());
+	TUniquePtr<FArchive> FileWriter{ IFileManager::Get().CreateFileWriter(*GenerateRemoteObjectFilename(Params.MigrationContext.ObjectId, Params.MigrationContext.OwnerServerId)) };
+	checkf(FileWriter.IsValid(), TEXT("Unable to create file writer for remote object %s"), *Params.MigrationContext.ObjectId.ToString());
 	
-	SerializeNameTables(*FileWriter, ObjectData);
-	*FileWriter << ObjectData.Bytes;
+	SerializeNameTables(*FileWriter, Params.ObjectData);
+	*FileWriter << const_cast<TArray<FRemoteObjectBytes>&>(Params.ObjectData.Bytes);
 	FileWriter->Close();
 }
 
@@ -353,6 +357,7 @@ protected:
 	TMap<FName, FNameIndexType> NameMap;
 	TMap<UObject*, FNameIndexType> PathNameMap;
 	TMap<UObject*, UObject*>* ObjectToArchetypeMap = nullptr;
+	const FUObjectMigrationContext& MigrationContext;
 
 public:
 
@@ -422,7 +427,12 @@ private:
 				Info.Object = *ReplacementObject;
 				Info.Id = FRemoteObjectId(Info.Object);
 			}
-			Info.Object->SetInternalFlags(EInternalObjectFlags::RemoteReference);
+
+			const bool bGloballyKnownAsset = Info.Object->HasAnyFlags(EObjectFlags::RF_ClassDefaultObject | EObjectFlags::RF_ArchetypeObject);
+			if (!bGloballyKnownAsset)
+			{
+				Info.Object->SetInternalFlags(EInternalObjectFlags::RemoteReference);
+			}
 		}
 
 		*this << Info.Type;
@@ -449,8 +459,9 @@ private:
 
 public:
 
-	FArchiveRemoteObjectWriter(UObject* InRootObject, FRemoteObjectData& OutObjectData, const TCHAR* ArchiveName = nullptr)
-		: TArchiveRemoteObjectBase<FMemoryWriter>(OutObjectData, ArchiveName ? ArchiveName : TEXT("RemoteObjectWriter"))
+	FArchiveRemoteObjectWriter(UObject* InRootObject, FRemoteObjectData& OutObjectData, const FUObjectMigrationContext& InMigrationContext, const TCHAR* ArchiveName = nullptr)
+		: TArchiveRemoteObjectBase<FMemoryWriter>(OutObjectData, ArchiveName ? ArchiveName : TEXT("RemoteObjectWriter")),
+		  MigrationContext(InMigrationContext)
 	{
 		static_assert(sizeof(FNameIndexType) == sizeof(FRemoteObjectPathName::FNameIndexType), "Name Index type must match the type used for storing remote object path name indices");
 
@@ -458,6 +469,11 @@ public:
 		ObjectsToSerialize.Add(InRootObject);
 	}
 
+	virtual const FUObjectMigrationContext* GetMigrationContext() const override
+	{
+		return &MigrationContext;
+	}
+
 	virtual UObject* GetArchetypeFromLoader(const UObject* Obj) override
 	{
 		if (ObjectToArchetypeMap)
@@ -605,8 +621,8 @@ class FArchetypeDeltaWriter : public FArchiveRemoteObjectWriter
 
 public:
 
-	FArchetypeDeltaWriter(UObject* InRootObject, FRemoteObjectData& OutObjectData, TMap<UObject*, UObject*>& InSubstituteArchetypeMap)
-		: FArchiveRemoteObjectWriter(InRootObject, OutObjectData, TEXT("RemoteArchetypeDeltaWriter"))
+	FArchetypeDeltaWriter(UObject* InRootObject, FRemoteObjectData& OutObjectData, const FUObjectMigrationContext& InMigrationContext, TMap<UObject*, UObject*>& InSubstituteArchetypeMap)
+		: FArchiveRemoteObjectWriter(InRootObject, OutObjectData, InMigrationContext, TEXT("RemoteArchetypeDeltaWriter"))
 		, SubstituteArchetypeMap(InSubstituteArchetypeMap)
 	{
 	}
@@ -672,6 +688,7 @@ class FArchiveRemoteObjectReader : public TArchiveRemoteObjectBase<FMemoryReader
 	const TArray<FName>& Names;
 	const TArray<UObject*>& ResolvedPathNameObjects;
 	const ERemoteObjectSerializationFlags DeserializeFlags = ERemoteObjectSerializationFlags::None;
+	const FUObjectMigrationContext& MigrationContext;
 
 	ERemoteReferenceType ReadObjectReference(FObjectPtr& Value)
 	{
@@ -727,12 +744,24 @@ class FArchiveRemoteObjectReader : public TArchiveRemoteObjectBase<FMemoryReader
 
 public:
 
-	FArchiveRemoteObjectReader(FRemoteObjectData& InObjectData, const TArray<UObject*>& InResolvedPathNames, ERemoteObjectSerializationFlags InDeserializeFlags)
+	/**
+	 * @param InObjectData:  The serialized object data we are should deserialize
+	 * @param InResolvedPathNames:  The existing Resolved Objects that corresponds to InObjectData.PathNames
+	 * @param bInAssignedOwnership:  We must take ownership of the objects we deserialize. We could already have ownership of those objects.
+	 * @param InDeserializeFlags:  The flags for how we should treat references during deserialization.
+	 */
+	FArchiveRemoteObjectReader(FRemoteObjectData& InObjectData, const TArray<UObject*>& InResolvedPathNames, const FUObjectMigrationContext& InMigrationContext, ERemoteObjectSerializationFlags InDeserializeFlags)
 		: TArchiveRemoteObjectBase<FMemoryReader>(InObjectData, TEXT("RemoteObjectReader"))
 		, Names(InObjectData.Names)
 		, ResolvedPathNameObjects(InResolvedPathNames)
 		, DeserializeFlags(InDeserializeFlags)
+		, MigrationContext(InMigrationContext)
+	{
+	}
+
+	virtual const FUObjectMigrationContext* GetMigrationContext() const override
 	{
+		return &MigrationContext;
 	}
 
 	virtual FArchive& operator<<(UObject*& Obj) override
@@ -787,6 +816,7 @@ public:
 class FArchiveRemoteReferencePatcher : public FArchiveUObject
 {
 	TMap<FRemoteObjectId, UObject*> IdToObjectMap;
+	const FUObjectMigrationContext& MigrationContext;
 
 	void InitObjectMap(const TArray<UObject*>& DeserializedObjects)
 	{
@@ -810,7 +840,8 @@ class FArchiveRemoteReferencePatcher : public FArchiveUObject
 
 public:
 
-	FArchiveRemoteReferencePatcher(const TArray<UObject*>& DeserializedObjects)
+	FArchiveRemoteReferencePatcher(const TArray<UObject*>& DeserializedObjects, const FUObjectMigrationContext& InMigrationContext)
+		: MigrationContext(InMigrationContext)
 	{
 		ArIsObjectReferenceCollector = true;
 		SetIsPersistent(false);
@@ -819,6 +850,11 @@ public:
 		SetPortFlags(PPF_MigratingRemoteObjects);
 	}
 
+	virtual const FUObjectMigrationContext* GetMigrationContext() const override
+	{
+		return &MigrationContext;
+	}
+
 	virtual FArchive& operator<<(FObjectPtr& Value) override
 	{
 		if (!Value.IsResolved())
@@ -833,13 +869,15 @@ public:
 class FArchiveSubObjectGatherer : public FArchiveUObject
 {
 	UObject* RootObject = nullptr;
+	const FUObjectMigrationContext& MigrationContext;
 	TMap<UObject*, UObject*>& ArchetypeToObjectMap;
 	TArray<UObject*>& ObjectsToSerialize;
 
 public:
 
-	FArchiveSubObjectGatherer(UObject* InRootObject, TMap<UObject*, UObject*>& InArchetypeToObjectMap, TArray<UObject*>& InObjectsToSerialize)
+	FArchiveSubObjectGatherer(UObject* InRootObject, const FUObjectMigrationContext& InMigrationContext, TMap<UObject*, UObject*>& InArchetypeToObjectMap, TArray<UObject*>& InObjectsToSerialize)
 		: RootObject(InRootObject)
+		, MigrationContext(InMigrationContext)
 		, ArchetypeToObjectMap(InArchetypeToObjectMap)
 		, ObjectsToSerialize(InObjectsToSerialize)
 	{
@@ -850,6 +888,11 @@ public:
 		SetPortFlags(PPF_MigratingRemoteObjects);
 	}
 
+	virtual const FUObjectMigrationContext* GetMigrationContext() const override
+	{
+		return &MigrationContext;
+	}
+
 	virtual FArchive& operator<<(FObjectPtr& Value) override
 	{
 		FRemoteObjectReferenceInfo Info = FArchiveRemoteObjectWriter::GetReferenceInfo(Value, RootObject);
@@ -1052,28 +1095,26 @@ void SerializeObjectDataInternal(FArchiveRemoteObjectWriter& Ar, UObject* Reques
 	}
 }
 
-FRemoteObjectData* SerializeObjectData(UObject* InObject, TSet<UObject*>& OutObjects)
+FRemoteObjectData SerializeObjectData(UObject* InObject, TSet<UObject*>& OutObjects, const FUObjectMigrationContext& MigrationContext)
 {
-	using namespace UE::RemoteObject::Handle;
-
 	UObject* Object = InObject;
 	if (GRemoteObjectsMigrateFullHierarchy)
 	{
 		Object = FindCanonicalRootObjectForSerialization(Object);
 	}
 
-	FRemoteObjectData* ObjectData = new FRemoteObjectData();	
+	FRemoteObjectData ObjectData;
 	FRemoteObjectId RequestedObjectId(InObject);
 
 	{
-		FArchiveRemoteObjectWriter Ar(Object, *ObjectData);
+		FArchiveRemoteObjectWriter Ar(Object, ObjectData, MigrationContext);
 		SerializeObjectDataInternal(Ar, InObject, RequestedObjectId, OutObjects);
 	}
 
-	return ObjectData;
+	return MoveTemp(ObjectData);
 }
 
-void ResetRemoteObject(UObject* InObject)
+void ResetRemoteObject(UObject* InObject, const FUObjectMigrationContext& MigrationContext)
 {
 	UObject* Object = InObject;
 	if (GRemoteObjectsMigrateFullHierarchy)
@@ -1088,7 +1129,7 @@ void ResetRemoteObject(UObject* InObject)
 	{
 		TArray<UObject*> ObjectsToSerialize;
 		TSet<UObject*> SerializedObjects;
-		FArchiveSubObjectGatherer SubobjectGather(Object, ReverseArchetypeToObjectMap, ObjectsToSerialize);
+		FArchiveSubObjectGatherer SubobjectGather(Object, MigrationContext, ReverseArchetypeToObjectMap, ObjectsToSerialize);
 		ObjectsToSerialize.Add(Object);
 
 		for (int32 Index = 0; Index < ObjectsToSerialize.Num(); ++Index)
@@ -1108,7 +1149,7 @@ void ResetRemoteObject(UObject* InObject)
 		// This will produce a delta between the archetypes and their instances which will then be used to restore the instances' state to the archetypes'
 		FRemoteObjectId RequestedObjectId(Object);
 		UObject* Archetype = Object->GetArchetype();
-		FArchetypeDeltaWriter Ar(Archetype, *ArchetypeDelta, ReverseArchetypeToObjectMap);
+		FArchetypeDeltaWriter Ar(Archetype, *ArchetypeDelta, MigrationContext, ReverseArchetypeToObjectMap);
 		TSet<UObject*> ArchetypeSerializedObjects;
 		SerializeObjectDataInternal(Ar, Archetype, RequestedObjectId, ArchetypeSerializedObjects);
 	}
@@ -1122,7 +1163,7 @@ void ResetRemoteObject(UObject* InObject)
 		// Use any existing objects and deserialize archetype delta on top of them avoiding re-construction which may lead to undiserable side-effects
 		// Additionally we don't want to recursively re-enter this function so let the deserialization process know we're already resetting migrated object(s)
 		ERemoteObjectSerializationFlags DeserializationFlags = ERemoteObjectSerializationFlags::PreserveRemoteReferences | ERemoteObjectSerializationFlags::UseExistingObjects | ERemoteObjectSerializationFlags::Resetting;
-		DeserializeObjectData(*ArchetypeDelta, DeserializedIds, DeserializedObjects, DeserializationFlags);
+		DeserializeObjectData(*ArchetypeDelta, MigrationContext, DeserializedIds, DeserializedObjects, DeserializationFlags);
 
 		const FRemoteServerId LocalServerId = UE::RemoteObject::GetGlobalServerId();
 		for (UObject* DeserializedObject : DeserializedObjects)
@@ -1133,7 +1174,11 @@ void ResetRemoteObject(UObject* InObject)
 	delete ArchetypeDelta;
 }
 
-int32 DeserializeObjectData(FRemoteObjectData& ObjectData, TArray<FRemoteObjectId>& OutObjectRemoteIds, TArray<UObject*>& OutObjects, ERemoteObjectSerializationFlags DeserializeFlags)
+int32 DeserializeObjectData(FRemoteObjectData& ObjectData,
+							const FUObjectMigrationContext& MigrationContext,
+							TArray<FRemoteObjectId>& OutObjectRemoteIds,
+							TArray<UObject*>& OutObjects,
+							ERemoteObjectSerializationFlags DeserializeFlags)
 {
 	using namespace UE::RemoteObject::Handle;
 	using namespace UE::RemoteObject::Private;
@@ -1147,7 +1192,8 @@ int32 DeserializeObjectData(FRemoteObjectData& ObjectData, TArray<FRemoteObjectI
 
 	ResolvePathNames(ObjectData, ResolvedPathNameObjects);
 
-	FArchiveRemoteObjectReader Ar(ObjectData, ResolvedPathNameObjects, DeserializeFlags);
+	// If we are being assigned ownership, we *must* take ownership of the objects.  Note: We also may already have ownership of those objects. 
+	FArchiveRemoteObjectReader Ar(ObjectData, ResolvedPathNameObjects, MigrationContext, DeserializeFlags);
 
 	{
 		FRemoteObjectReaderScope Scope(Ar, TEXT("Header"));
@@ -1201,6 +1247,12 @@ int32 DeserializeObjectData(FRemoteObjectData& ObjectData, TArray<FRemoteObjectI
 			OutObjects.Add(ConstructRemoteObject(ObjectHeaders[ObjectIndex], &ConstructionOverrides, DeserializeFlags));
 		}
 
+		UObject* RootObject = OutObjects.Num() ? OutObjects[0] : nullptr;
+		if (!ensureMsgf(RootObject, TEXT("%hs had objects to construct but could not reconstruct them"), __func__))
+		{
+			return RequestedObjectIndex;
+		}
+
 		// If we're already resetting a borrowed object we don't want to change the ownership until the object is deserialized using remote server data
 		// and we also don't want to recursively reset the object
 		const bool bResetting = EnumHasAnyFlags(DeserializeFlags, ERemoteObjectSerializationFlags::Resetting);
@@ -1217,7 +1269,7 @@ int32 DeserializeObjectData(FRemoteObjectData& ObjectData, TArray<FRemoteObjectI
 
 			if (bReturningBorrowedObject)
 			{
-				ResetRemoteObject(OutObjects[0]);
+				ResetRemoteObject(OutObjects[0], MigrationContext);
 			}
 		}
 
@@ -1243,7 +1295,7 @@ int32 DeserializeObjectData(FRemoteObjectData& ObjectData, TArray<FRemoteObjectI
 		// Patch any unresolved remote references (skip if we're resetting an object to its CDO state)
 		if (!bResetting)
 		{
-			FArchiveRemoteReferencePatcher PatchAr(OutObjects);
+			FArchiveRemoteReferencePatcher PatchAr(OutObjects, MigrationContext);
 			for (UObject* Object : OutObjects)
 			{
 				if (Object)
diff --git a/Engine/Source/Runtime/CoreUObject/Private/UObject/RemoteObjectSerialization.h b/Engine/Source/Runtime/CoreUObject/Private/UObject/RemoteObjectSerialization.h
--- a/Engine/Source/Runtime/CoreUObject/Private/UObject/RemoteObjectSerialization.h
+++ b/Engine/Source/Runtime/CoreUObject/Private/UObject/RemoteObjectSerialization.h
@@ -72,20 +72,22 @@ namespace UE::RemoteObject::Serialization
 	* Serializes an object and its subobject (or if the object is a default subobject, its parent and the parent's subobjects)
 	* @param InObject Object to be serialized
 	* @param OutObjects All objects that have been serialized (Object and its subobjects and/or parent)
+	* @param MigrationContext Contains the meta data of the current migration request
 	* @return Remote object data representing the serialized objects
 	*/
-	FRemoteObjectData* SerializeObjectData(UObject* InObject, TSet<UObject*>& OutObjects);
+	FRemoteObjectData SerializeObjectData(UObject* InObject, TSet<UObject*>& OutObjects, const FUObjectMigrationContext& MigrationContext);
 
 	/**
 	* Deserializes remote object data
-	* @param ObjectData to be deserialized
+	* @param ObjectData the data to deserailize 
+	* @param MigrationContext the Context (meta data) of the current migration that's causing the deserialization
 	* @param OutObjectRemoteIds Remote IDs of the deserialized objects
 	* @param OutReceivedObjects All deserialized objects
 	* @param DeserializeFlags Flags modifying the behavior of the deserialization process
 	* @return Index of an object in OutReceivedObjects that was the main object the migration request was triggered for 
 	* (usually 0 but if a migration requests a default subobject then its parent is also migrated and the return value will be > 0)
 	*/
-	int32 DeserializeObjectData(FRemoteObjectData& ObjectData, TArray<FRemoteObjectId>& OutObjectRemoteIds, TArray<UObject*>& OutReceivedObjects, ERemoteObjectSerializationFlags DeserializeFlags = ERemoteObjectSerializationFlags::None);
+	int32 DeserializeObjectData(FRemoteObjectData& ObjectData, const FUObjectMigrationContext& MigrationContext, TArray<FRemoteObjectId>& OutObjectRemoteIds, TArray<UObject*>& OutReceivedObjects, ERemoteObjectSerializationFlags DeserializeFlags = ERemoteObjectSerializationFlags::None);
 
 	/**
 	* Finds the canonical 'root' object that is used for remote object serialization
@@ -98,6 +100,6 @@ namespace UE::RemoteObject::Serialization
 
 namespace UE::RemoteObject::Serialization::Disk
 {
-	void LoadObjectFromDisk(FRemoteTransactionId RequestId, FRemoteWorkPriority RequestPriority, FRemoteObjectId ObjectId, FRemoteServerId CurrentOwnerServerId, FRemoteServerId DestinationServerId);
-	void SaveObjectToDisk(FRemoteTransactionId RequestId, FRemoteObjectId ObjectId, FRemoteServerId ServerId, FRemoteServerId OwnerServerId, FRemoteServerId PhysicsId, FRemoteObjectData& ObjectData);
+	void LoadObjectFromDisk(const FUObjectMigrationContext& MigrationContext);
+	void SaveObjectToDisk(const UE::RemoteObject::Transfer::FMigrateSendParams& Params);
 }
\ No newline at end of file
diff --git a/Engine/Source/Runtime/CoreUObject/Private/UObject/RemoteObjectTransfer.cpp b/Engine/Source/Runtime/CoreUObject/Private/UObject/RemoteObjectTransfer.cpp
--- a/Engine/Source/Runtime/CoreUObject/Private/UObject/RemoteObjectTransfer.cpp
+++ b/Engine/Source/Runtime/CoreUObject/Private/UObject/RemoteObjectTransfer.cpp
@@ -13,28 +13,32 @@
 #include "UObject/RemoteObjectPrivate.h"
 #include "UObject/ObjectHandlePrivate.h"
 #include "UObject/UObjectHash.h"
+#include "UObject/UObjectMigrationContext.h"
 #include "HAL/CriticalSection.h"
 #include "AutoRTFM/AutoRTFM.h"
 #include "Templates/Atomic.h"
 
+
 namespace UE::RemoteObject::Transfer
 {
 const FRemoteServerId DatabaseId(0x3ff);
-
-TDelegate<void(FRemoteTransactionId /*RequestId*/, FRemoteWorkPriority /*RequestPriority*/, FRemoteObjectId /*ObjectId*/, FRemoteServerId /*CurrentOwnerServerId*/, FRemoteServerId /*DestinationServerId*/)> RequestRemoteObjectDelegate;
-TDelegate<void(FRemoteTransactionId /*RequestId*/, FRemoteObjectId /*ObjectId*/, FRemoteServerId /*DestinationServerId*/, FRemoteServerId /*OwnerServerId*/, FRemoteServerId /*PhysicsId*/, FRemoteObjectData& /*ObjectData*/)> RemoteObjectTransferDelegate;
-TDelegate<void(FRemoteTransactionId /*RequestId*/, FRemoteObjectId /*ObjectId*/, FRemoteServerId /*DestinationServerId*/)> RemoteObjectDeniedTransferDelegate;
-
-TMulticastDelegate<void(const FRemoteObjectData& /*ObjectData*/, FRemoteServerId /*RemoteServerId*/, EObjectMigrationFlags /*MigrationFlags*/)> OnObjectDataSentDelegate;
-TMulticastDelegate<void(const TSet<UObject*>& /*MigratedObjects*/, FRemoteServerId /*RemoteServerId*/, EObjectMigrationFlags /*MigrationFlags*/)> OnObjectsSentDelegate;
-TMulticastDelegate<void(const FRemoteObjectData& /*ObjectData*/, FRemoteServerId /*RemoteServerId*/, EObjectMigrationFlags /*MigrationFlags*/)> OnObjectDataReceivedDelegate;
-TMulticastDelegate<void(const TArray<UObject*>& /*MigratedObjects*/, FRemoteServerId /*RemoteServerId*/, EObjectMigrationFlags /*MigrationFlags*/)> OnObjectsReceivedDelegate;
-TMulticastDelegate<void(FRemoteTransactionId /*RequestId*/, FRemoteObjectId /*ObjectId*/)> OnObjectTouchedDelegate;
-
-TDelegate<void(FRemoteTransactionId /*RequestId*/, FRemoteObjectId /*ObjectId*/, FRemoteServerId /*DatabaseId*/, FRemoteServerId /*OwnerServerId*/, FRemoteServerId /*PhysicsId*/, FRemoteObjectData& /*ObjectData*/)> StoreRemoteObjectDataDelegate;
-TDelegate<void(FRemoteTransactionId /*RequestId*/, FRemoteWorkPriority /*RequestPriority*/, FRemoteObjectId /*ObjectId*/, FRemoteServerId /*DatabaseId*/, FRemoteServerId /*DestinationServerId*/)> RestoreRemoteObjectDataDelegate;
-
-UObject* MigrateObjectFromObjectData(FRemoteServerId OwnerServerId, FRemoteServerId PhysicsId, FRemoteObjectData& ObjectData, FRemoteServerId RemoteServerId, EObjectMigrationFlags ObjectFlags);
+static TAutoConsoleVariable<bool> CVarRemoveMissingMigratedComponents(TEXT("DSTM.RemoveMissingMigratedComponents"), false, TEXT("When we receive less SubObjects (Components) than we expect, remove the extras that we have locally (assume they were removed on purpose)"));
+
+UObject* MigrateObjectFromObjectData(FRemoteObjectData& ObjectData, const FUObjectMigrationContext& MigrationContext);
+
+// Prevent us from having to match this definition to the declartion in the header just to initialize the delegate
+#define DEFINE_MIGRATION_DELEGATE(x) decltype(x) x
+	DEFINE_MIGRATION_DELEGATE(RemoteObjectTransferDelegate);
+	DEFINE_MIGRATION_DELEGATE(RemoteObjectDeniedTransferDelegate);
+	DEFINE_MIGRATION_DELEGATE(RequestRemoteObjectDelegate);
+	DEFINE_MIGRATION_DELEGATE(OnObjectDataReceivedDelegate);
+	DEFINE_MIGRATION_DELEGATE(OnObjectsReceivedDelegate);
+	DEFINE_MIGRATION_DELEGATE(OnObjectDataSentDelegate);
+	DEFINE_MIGRATION_DELEGATE(OnObjectsSentDelegate);
+	DEFINE_MIGRATION_DELEGATE(OnObjectTouchedDelegate);
+	DEFINE_MIGRATION_DELEGATE(StoreRemoteObjectDataDelegate);
+	DEFINE_MIGRATION_DELEGATE(RestoreRemoteObjectDataDelegate);
+#undef DEFINE_MIGRATION_DELEGATE
 
 /**
  * Remote object transfer queue. 
@@ -168,32 +172,31 @@ public:
 	{
 	}
 
-	void RequestRemoteObject(FRemoteTransactionId RequestId, FRemoteWorkPriority RequestPriority, FRemoteObjectId ObjectId, FRemoteServerId CurrentOwnerServerId, FRemoteServerId DestinationServerId)
+	void RequestRemoteObject(FRemoteTransactionId RequestId, FRemoteWorkPriority RequestPriority, FRemoteObjectId ObjectId, FRemoteServerId LastKnownResidentServerId, FRemoteServerId DestinationServerId)
 	{
 		check(!AutoRTFM::IsClosed());
 
 		// Request an object from remote server
-		if (CurrentOwnerServerId != DatabaseId)
+		if (LastKnownResidentServerId != DatabaseId)
 		{
-			RequestRemoteObjectDelegate.Execute(RequestId, RequestPriority, ObjectId, CurrentOwnerServerId, DestinationServerId);
+			// We are forwarding a request for pull-migrate from the resident server id
+			// This isn't technically a migration, it is actually a request for migration...
+			RequestRemoteObjectDelegate.Execute(RequestId, RequestPriority, ObjectId, LastKnownResidentServerId, DestinationServerId);
 		}
 		else
 		{
-			RestoreRemoteObjectDataDelegate.Execute(RequestId, RequestPriority, ObjectId, DatabaseId, GetGlobalServerId());
+			// We are intercepting a request to pull-migrate since we believe the object resides in our database, migrate it "from" the database
+			FUObjectMigrationContext MigrationContext { .ObjectId = ObjectId, .RemoteServerId = DatabaseId, .OwnerServerId = DatabaseId, .PhysicsServerId = DatabaseId, .RequestId = RequestId, .MigrationSide = EObjectMigrationSide::Receive };
+			RestoreRemoteObjectDataDelegate.Execute(MigrationContext);
 		}
 	}
 
 	/**
-	* Sends an object to a remote server (adds a send request)
-	* @param ObjectId Id of the object being requested
-	* @param DestinationServerId Remote server id
-	* @param ObjectData Object data to send. Send request takes ownership of the data
-	* @return Future that's going to receive the object data
 	*/
-	void SendRemoteObject(FRemoteTransactionId RequestId, FRemoteObjectId ObjectId, FRemoteServerId DestinationServerId, FRemoteServerId OwnerServerId, FRemoteServerId PhysicsId, FRemoteObjectData* ObjectData)
+	void SendRemoteObject(const FMigrateSendParams& Params)
 	{
 		check(!AutoRTFM::IsClosed());
-		RemoteObjectTransferDelegate.Execute(RequestId, ObjectId, DestinationServerId, OwnerServerId, PhysicsId, *ObjectData);
+		RemoteObjectTransferDelegate.Execute(Params);
 	}
 
 	void SendDeniedRemoteObject(FRemoteTransactionId RequestId, FRemoteObjectId ObjectId, FRemoteServerId DestinationServerId)
@@ -206,25 +209,22 @@ public:
 	* Fullfils receive request
 	* @param ObjectId Id of the object that was requested
 	* @param Data Object data. Data ownership is transferred to the receive request (if it exists)
-	* @return True if a receive request for the specified object exists (and object data ownership was transferred), false otherwise
 	*/
-	void FulfillReceiveRequest(FRemoteTransactionId RequestId, FRemoteServerId OwnerServerId, FRemoteServerId PhysicsId, FRemoteObjectId ObjectId, FRemoteServerId RemoteServerId, FRemoteObjectData* ObjectData)
+	void FulfillReceiveRequest(FRemoteObjectData& ObjectData, const FUObjectMigrationContext& MigrationContext)
 	{
 		using namespace UE::RemoteObject::Serialization;
 
 		check(!AutoRTFM::IsClosed());
 
-		const int32 PendingRequestIndex = PendingObjectRequests.Find(ObjectId);
+		const int32 PendingRequestIndex = PendingObjectRequests.Find(MigrationContext.ObjectId);
 		const bool bFoundRequest = (PendingRequestIndex != INDEX_NONE);
 
-		UE_LOG(LogRemoteObject, Verbose, TEXT("     : FulfillReceiveRequest(%s) %s (was requested: %d, owner: %s)"), *RequestId.ToString(), *ObjectId.ToString(), bFoundRequest, *OwnerServerId.ToString());
+		UE_LOG(LogRemoteObject, Verbose, TEXT("     : FulfillReceiveRequest(%s) %s (was requested: %d, owner: %s)"), *MigrationContext.RequestId.ToString(), *MigrationContext.ObjectId.ToString(), bFoundRequest, *MigrationContext.OwnerServerId.ToString());
 
-		const bool bAssignedOwnership = (RequestId == PushMigrationRequestId);
-		EObjectMigrationFlags MigrationFlags = bAssignedOwnership ? EObjectMigrationFlags::Acquire : EObjectMigrationFlags::Receive;
-		UObject* Result = MigrateObjectFromObjectData(OwnerServerId, PhysicsId, *ObjectData, RemoteServerId, MigrationFlags);
-		checkf(UE::CoreUObject::Private::FObjectHandleUtils::GetRemoteId(Result) == ObjectId, TEXT("PeerId:%s requested object %s migration from PeerId:%s but received %s"),
-			*GetGlobalServerId().ToString(), *ObjectId.ToString(), *RemoteServerId.ToString(), *UE::CoreUObject::Private::FObjectHandleUtils::GetRemoteId(Result).ToString());
-		checkf(OwnerServerId == Handle::GetOwnerServerId(Result), TEXT("Expected owner id %s doesn't match post-deserialize owner: %s"), *OwnerServerId.ToString(), *Handle::GetOwnerServerId(Result).ToString());
+		UObject* Result = MigrateObjectFromObjectData(ObjectData, MigrationContext);
+		checkf(UE::CoreUObject::Private::FObjectHandleUtils::GetRemoteId(Result) == MigrationContext.ObjectId, TEXT("PeerId:%s requested object %s migration from PeerId:%s but received %s"),
+			   *GetGlobalServerId().ToString(), *MigrationContext.ObjectId.ToString(), *MigrationContext.RemoteServerId.ToString(), *UE::CoreUObject::Private::FObjectHandleUtils::GetRemoteId(Result).ToString());
+		checkf(MigrationContext.OwnerServerId == Handle::GetOwnerServerId(Result), TEXT("Expected owner id %s doesn't match post-deserialize owner: %s"), *MigrationContext.OwnerServerId.ToString(), *Handle::GetOwnerServerId(Result).ToString());
 
 		if (bFoundRequest)
 		{
@@ -250,20 +250,20 @@ public:
 
 FRemoteObjectTransferQueue GTransferQueue;
 
-UObject* MigrateObjectFromObjectData(FRemoteServerId OwnerServerId, FRemoteServerId PhysicsId, FRemoteObjectData& ObjectData, FRemoteServerId RemoteServerId, EObjectMigrationFlags MigrationFlags)
+UObject* MigrateObjectFromObjectData(FRemoteObjectData& ObjectData, const FUObjectMigrationContext& MigrationContext)
 {
 	using namespace UE::RemoteObject::Serialization;
 	check(!AutoRTFM::IsClosed());
 
 	// Notify that object data has been migrated from a remote server
-	OnObjectDataReceivedDelegate.Broadcast(ObjectData, RemoteServerId, MigrationFlags);
+	OnObjectDataReceivedDelegate.Broadcast(ObjectData, MigrationContext);
 
 	TArray<FRemoteObjectId> ReceivedObjectRemoteIds;
 	TArray<UObject*> ReceivedObjects;
-	int32 RequestedObjectIndex = DeserializeObjectData(ObjectData, ReceivedObjectRemoteIds, ReceivedObjects);
+	int32 RequestedObjectIndex = DeserializeObjectData(ObjectData, MigrationContext, ReceivedObjectRemoteIds, ReceivedObjects);
 
+	// Debug section
 	checkf(ReceivedObjects.Num() > 0, TEXT("PeerId:%s unable to deserialize object data (%d bytes)"), *GetGlobalServerId().ToString(), ObjectData.GetNumBytes());
-
 	if (ReceivedObjects.Num() > 0)
 	{
 		if (UE_LOG_ACTIVE(LogRemoteObject, VeryVerbose))
@@ -282,38 +282,59 @@ UObject* MigrateObjectFromObjectData(FRemoteServerId OwnerServerId, FRemoteServe
 
 			UE_LOG(LogRemoteObject, VeryVerbose, TEXT("         <done listing subobjects>"));
 		}
-	}
 
-#if UE_WITH_REMOTE_OBJECT_HANDLE
-	const bool bAssignedOwnership = EnumHasAnyFlags(MigrationFlags, EObjectMigrationFlags::ChangeOwnership);
+		// Let's make sure we received data for all of the objects
+		const UObject* RootObject = ReceivedObjects[0];
+		TArray<UObject*> AllSubObjects;
+		GetObjectsWithOuter(RootObject, AllSubObjects);
 
-	FPostMigrateParams PostMigrateParams;
-	PostMigrateParams.RequestingServerId = RemoteServerId;
-	PostMigrateParams.DestinationServerId = GetGlobalServerId();
-	for (UObject* ReceivedObject : ReceivedObjects)
-	{
-		PostMigrateParams.MigrationFlags = MigrationFlags;
-
-		// Recompute the flags as having only changed ownership if we're being assigned ownership we don't already have
-		const bool bAlreadyOwned = UE::RemoteObject::Handle::IsOwned(ReceivedObject);
-		if (bAssignedOwnership && bAlreadyOwned)
+		const bool bSameNumberOfObjects = (1 + AllSubObjects.Num() == ReceivedObjects.Num());
+		if (!bSameNumberOfObjects)
 		{
-			EnumRemoveFlags(PostMigrateParams.MigrationFlags, EObjectMigrationFlags::ChangeOwnership);
+			UE_LOG(LogRemoteObject, Warning, TEXT("We received migrated '%s' but we received %d Objects whereas %d Objects now live locally.  Differences:"), *GetNameSafe(RootObject), ReceivedObjects.Num(), 1 + AllSubObjects.Num());
+			TArray<UObject*> ReceivedButNotInstanced;
+			for (UObject* Obj : ReceivedObjects)
+			{
+				if (AllSubObjects.Remove(Obj) < 1)
+				{
+					ReceivedButNotInstanced.Add(Obj);
+				}
+			}
+			// We'll always have the root object in this list due to AllSubObjects not containing the root object
+			ReceivedButNotInstanced.RemoveAtSwap(0);
+
+			auto WrappedGetNameSafe = [](const UObject* O) { return GetNameSafe(O); };
+			UE_CLOG(!ReceivedButNotInstanced.IsEmpty(), LogRemoteObject, Warning, TEXT(" Received but not instanced: %s"), *FString::JoinBy(ReceivedButNotInstanced, TEXT(", "), WrappedGetNameSafe));
+			UE_CLOG(!AllSubObjects.IsEmpty(), LogRemoteObject, Warning, TEXT(" Existing instance but not received: %s"), *FString::JoinBy(AllSubObjects, TEXT(", "), WrappedGetNameSafe));
+
+			if (CVarRemoveMissingMigratedComponents.GetValueOnGameThread())
+			{
+				// Let's remove all of those instances that we shouldn't have
+				for (UObject* Obj : AllSubObjects)
+				{
+					UE_LOG(LogRemoteObject, Warning, TEXT(" Removing: %s"), *GetNameSafe(Obj));
+					Obj->MarkAsGarbage();
+				}
+			}
 		}
+	}
 
-		ReceivedObject->PostMigrate(PostMigrateParams);
+#if UE_WITH_REMOTE_OBJECT_HANDLE
+	for (UObject* ReceivedObject : ReceivedObjects)
+	{
+		ReceivedObject->PostMigrate(MigrationContext);
 	}
 #endif
 
 	// Set the ownership to the correct server id
 	for (UObject* ReceivedObject : ReceivedObjects)
 	{
-		UE::RemoteObject::Handle::ChangeOwnerServerId(ReceivedObject, OwnerServerId);
-		UE::RemoteObject::Handle::ChangePhysicsServerId(ReceivedObject, PhysicsId);
+		UE::RemoteObject::Handle::ChangeOwnerServerId(ReceivedObject, MigrationContext.OwnerServerId);
+		UE::RemoteObject::Handle::ChangePhysicsServerId(ReceivedObject, MigrationContext.PhysicsServerId);
 	}
 
 	// Notify that objects have been migrated from a remote server
-	OnObjectsReceivedDelegate.Broadcast(ReceivedObjects, RemoteServerId, MigrationFlags);
+	OnObjectsReceivedDelegate.Broadcast(ReceivedObjects, MigrationContext);
 
 	// do any existing requests require any of the objects in this hierarchy?
 	// if so, we can now update the canonical root id
@@ -449,24 +470,27 @@ void MigrateObjectToRemoteServerWithExplicitPriority(FRemoteTransactionId Reques
 	UObject* const RootObject = FindCanonicalRootObjectForSerialization(Object);
 	const FRemoteObjectId CanonicalRootObjectId = UE::CoreUObject::Private::FObjectHandleUtils::GetRemoteId(RootObject);
 
+	// Figure out the ownership semantics
+	const bool bPushOwnershipChange = (RequestId == PushMigrationRequestId);
+	FRemoteServerId OwnerServerId = bPushOwnershipChange ? DestinationServerId : Handle::GetOwnerServerId(RootObject);
+	FRemoteServerId PhysicsServerId = Handle::GetPhysicsServerId(RootObject);
+	FUObjectMigrationContext MigrationContext{ .ObjectId = ObjectId, .RemoteServerId = DestinationServerId, .OwnerServerId = OwnerServerId, .PhysicsServerId = PhysicsServerId, .RequestId = RequestId, .MigrationSide = EObjectMigrationSide::Send };
+
 	// is the object locked?
 	const FRemoteObjectRequest* HigherPriorityRequest = nullptr;
 	const FRemoteObjectRequest* EqualPriorityRequest = nullptr;
 	const FRemoteObjectRequest* LowerPriorityRequest = nullptr;
 
-	const bool bOwnershipChange = (RequestId == PushMigrationRequestId);
-	FRemoteServerId OwnerServerId = bOwnershipChange ? DestinationServerId : Handle::GetOwnerServerId(RootObject);
-	FRemoteServerId PhysicsId = Handle::GetPhysicsServerId(RootObject);
-
 	UE_LOG(LogRemoteObject, Verbose, 
-		TEXT("     : MigrateObjectToRemoteServer(%s) sending obj %s [%s] (root %s, isremote %d, changing owner %d, owner %s)"),
+		TEXT("     : MigrateObjectToRemoteServer(%s) sending obj %s [%s] (root %s, isremote %d, owner %s, dest %s)"),
 		*RequestId.ToString(),
 		*ObjectId.ToString(),
 		*Object->GetName(),
 		*CanonicalRootObjectId.ToString(),
 		Handle::IsRemote(Object),
-		bOwnershipChange,
-		*OwnerServerId.ToString());
+		*OwnerServerId.ToString(),
+		*DestinationServerId.ToString()
+	);
 
 	for (int32 RequestIndex = 0; RequestIndex < GTransferQueue.GetRequestCount(); RequestIndex++)
 	{
@@ -537,27 +561,25 @@ void MigrateObjectToRemoteServerWithExplicitPriority(FRemoteTransactionId Reques
 	}
 	else
 	{
-		UE_LOG(LogRemoteObject, VeryVerbose, TEXT("     : MigrateObjectToRemoteServer(%s) serializing obj %s (isremote %d owner changing %d owner %s)"), *RequestId.ToString(), *ObjectId.ToString(), Handle::IsRemote(Object), bOwnershipChange, *OwnerServerId.ToString());
+		UE_LOG(LogRemoteObject, VeryVerbose, TEXT("     : MigrateObjectToRemoteServer(%s) serializing obj %s"), *RequestId.ToString(), *ObjectId.ToString(), Handle::IsRemote(Object), *OwnerServerId.ToString(), *DestinationServerId.ToString());
 
 		// send the object
 		TSet<UObject*> SentObjects;
-		FRemoteObjectData* ObjectData = SerializeObjectData(Object, SentObjects);
+		FRemoteObjectData ObjectData = SerializeObjectData(Object, SentObjects, MigrationContext);
 
-		UE_LOG(LogRemoteObject, VeryVerbose, TEXT("     : MigrateObjectToRemoteServer(%s) serialization complete obj %s (isremote %d, %d bytes)"), *RequestId.ToString(), *ObjectId.ToString(), Handle::IsRemote(Object), ObjectData->GetNumBytes());
+		UE_LOG(LogRemoteObject, VeryVerbose, TEXT("     : MigrateObjectToRemoteServer(%s) serialization complete obj %s (%d bytes)"), *RequestId.ToString(), *ObjectId.ToString(), ObjectData.GetNumBytes());
 
 		for(UObject* SO : SentObjects)
 		{
 			FRemoteObjectId SubObjectId = UE::CoreUObject::Private::FObjectHandleUtils::GetRemoteId(SO);
 			UE_LOG(LogRemoteObject, VeryVerbose, TEXT("         sent subobject : %s [%s]"), *SubObjectId.ToString(), *GetNameSafe(SO));
 		}
-
-		// Are we loaning it out or changing ownership?
-		const EObjectMigrationFlags MigrationFlags = bOwnershipChange ? EObjectMigrationFlags::Reassign : EObjectMigrationFlags::Loan;
 			
 		// Notify that object data has been migrated (actually will be but SendObjectData does a MoveTemp on the data)
-		OnObjectDataSentDelegate.Broadcast(*ObjectData, DestinationServerId, MigrationFlags);
+		OnObjectDataSentDelegate.Broadcast(ObjectData, MigrationContext);
 
-		GTransferQueue.SendRemoteObject(RequestId, ObjectId, DestinationServerId, OwnerServerId, PhysicsId, ObjectData);
+		FMigrateSendParams SendRemoteObjectParams = { .MigrationContext = MigrationContext, .ObjectData = ObjectData };
+		GTransferQueue.SendRemoteObject(SendRemoteObjectParams);
 
 		checkf(SentObjects.Num() > 0, TEXT("PeerId:%s unable to migrate object %s to PeerId:%s"), *GetGlobalServerId().ToString(), *UE::CoreUObject::Private::FObjectHandleUtils::GetRemoteId(Object).ToString(), *DestinationServerId.ToString());
 
@@ -567,19 +589,16 @@ void MigrateObjectToRemoteServerWithExplicitPriority(FRemoteTransactionId Reques
 			FUnsafeToMigrateScope UnsafeToMigrateScope;
 
 			// Notify that objects have been migrated to a remote server
-			OnObjectsSentDelegate.Broadcast(SentObjects, DestinationServerId, MigrationFlags);
+			OnObjectsSentDelegate.Broadcast(SentObjects, MigrationContext);
 
-			FPostMigrateParams PostMigrateParams;
-			PostMigrateParams.RequestingServerId = RequestPriority.GetRootServerId();
-			PostMigrateParams.DestinationServerId = DestinationServerId;
-			PostMigrateParams.MigrationFlags = MigrationFlags;
+			const bool bSentToOwnerServer = (OwnerServerId == DestinationServerId);
 			for (UObject* SentObject : SentObjects)
 			{
 #if UE_WITH_REMOTE_OBJECT_HANDLE
-				SentObject->PostMigrate(PostMigrateParams);
+				SentObject->PostMigrate(MigrationContext);
 #endif
 				MarkAsRemote(SentObject, DestinationServerId);
-				if (bOwnershipChange)
+				if (bSentToOwnerServer)
 				{
 					UE::RemoteObject::Handle::ChangeOwnerServerId(SentObject, OwnerServerId);
 				}
@@ -648,10 +667,12 @@ void MigrateObjectToRemoteServerWithExplicitPriority(FRemoteTransactionId Reques
 	}
 }
 
-void OnObjectDataReceived(FRemoteTransactionId RequestId, FRemoteServerId OwnerServerId, FRemoteServerId PhysicsId, FRemoteObjectId ObjectId, FRemoteServerId RemoteServerId, FRemoteObjectData* Data)
+void OnObjectDataReceived(FRemoteTransactionId RequestId, FRemoteServerId OwnerServerId, FRemoteServerId PhysicsId, FRemoteObjectId ObjectId, FRemoteServerId RemoteServerId, FRemoteObjectData& Data)
 {
 	check(!AutoRTFM::IsClosed());
-	GTransferQueue.FulfillReceiveRequest(RequestId, OwnerServerId, PhysicsId, ObjectId, RemoteServerId, Data);
+
+	FUObjectMigrationContext MigrationContext { .ObjectId = ObjectId, .RemoteServerId = RemoteServerId, .OwnerServerId = OwnerServerId, .PhysicsServerId = PhysicsId, .RequestId = RequestId, .MigrationSide = EObjectMigrationSide::Receive };
+	GTransferQueue.FulfillReceiveRequest(Data, MigrationContext);
 }
 
 void OnObjectDataDenied(FRemoteTransactionId RequestId, FRemoteObjectId ObjectId, FRemoteServerId RemoteServerId)
@@ -732,8 +753,25 @@ void StoreObjectToDatabase(UObject* Object)
 
 		TSet<UObject*> SentObjects;
 		{
-			TUniquePtr<FRemoteObjectData> ObjectData(SerializeObjectData(OutermostUnreachableRemoteReference, SentObjects));
-			StoreRemoteObjectDataDelegate.Execute(PushMigrationRequestId, FRemoteObjectId(OutermostUnreachableRemoteReference), DatabaseId, FRemoteServerId(), FRemoteServerId(), *ObjectData);
+			// Let's setup a fake migration context for storing this
+			FRemoteObjectId RootObjectId = FRemoteObjectId { OutermostUnreachableRemoteReference };
+			FRemoteServerId LocalServerId = GetGlobalServerId();
+			FRemoteServerId OwnerServerId = UE::RemoteObject::Handle::GetOwnerServerId(Object);
+			if (!ensureMsgf(LocalServerId == OwnerServerId, TEXT("%hs attempted to store non-locally-owned Object %s to Server %s's Database"), __func__, *GetNameSafe(Object), *LocalServerId.ToString()))
+			{
+				// This shouldn't happen obviously, but we're currently hitting this case.  So if it happens,
+				// don't store in the database, instead fetch it back from the Owning Server if needed
+				MarkAsRemote(Object, OwnerServerId); 
+				return;
+			}
+			// Setup a migration context for "push migrating to the new owner: the Database Server"
+			FUObjectMigrationContext StoreToDatabaseContext{ .ObjectId = RootObjectId, .RemoteServerId = DatabaseId, .OwnerServerId = DatabaseId, .PhysicsServerId = DatabaseId, .RequestId = PushMigrationRequestId, .MigrationSide = EObjectMigrationSide::Send };
+
+			// Do the actual serialization
+			FRemoteObjectData ObjectData = SerializeObjectData(OutermostUnreachableRemoteReference, SentObjects, StoreToDatabaseContext);
+
+			const FMigrateSendParams Params { .MigrationContext = StoreToDatabaseContext, .ObjectData = ObjectData };
+			StoreRemoteObjectDataDelegate.Execute(Params);
 		}
 
 		checkf(SentObjects.Num() > 0, TEXT("PeerId:%s unable to store object %s to DatabaseId:%s"), *GetGlobalServerId().ToString(), *FRemoteObjectId(Object).ToString(), *DatabaseId.ToString());
diff --git a/Engine/Source/Runtime/CoreUObject/Private/UObject/UObjectMigrationContext.cpp b/Engine/Source/Runtime/CoreUObject/Private/UObject/UObjectMigrationContext.cpp
new file mode 100644
--- /dev/null
+++ b/Engine/Source/Runtime/CoreUObject/Private/UObject/UObjectMigrationContext.cpp
@@ -0,0 +1,150 @@
+// Copyright Epic Games, Inc. All Rights Reserved.
+
+#include "UObject/UObjectMigrationContext.h"
+
+#include "Serialization/Archive.h"
+#include "UObject/RemoteObjectTransfer.h"
+#include "UObject/UObjectThreadContext.h"
+
+bool FUObjectMigrationContext::IsValid() const
+{
+	return ObjectId.IsValid() && RemoteServerId.IsValid() && OwnerServerId.IsValid() && RequestId.IsValid();
+}
+
+EObjectMigrationRecvType FUObjectMigrationContext::GetObjectMigrationRecvType(const UObjectBase* Object) const
+{
+	return GetObjectMigrationRecvType(FRemoteObjectId { Object });
+}
+
+EObjectMigrationRecvType FUObjectMigrationContext::GetObjectMigrationRecvType(FRemoteObjectId InObjectId) const
+{
+#if UE_WITH_REMOTE_OBJECT_HANDLE
+	if (!ensureMsgf(IsValid(), TEXT("%hs called on an invalid MigrationContext"), __func__))
+	{
+		return EObjectMigrationRecvType::Invalid;
+	}
+
+	const FRemoteServerId LocalServerId = UE::RemoteObject::GetGlobalServerId();
+	const bool bAlreadyOwns = UE::RemoteObject::Handle::IsOwned(InObjectId);
+
+	const bool bPushAssignOwnership = (RequestId == UE::RemoteObject::Transfer::PushMigrationRequestId);
+	if (bPushAssignOwnership)
+	{
+		ensureMsgf(OwnerServerId == LocalServerId,
+				   TEXT("%s attempted to push-assign ownership of %s to %s but the receiving server of that request is actually %s"),
+				   *RemoteServerId.ToString(), *ObjectId.ToString(), *OwnerServerId.ToString(), *LocalServerId.ToString());
+
+		if (bAlreadyOwns)
+		{
+			// We were pushed an object that we already owned, so the loan has been returned to us
+			return EObjectMigrationRecvType::ReturnedLoan;
+		}
+
+		// We were pushed an object that we do not already own, so we must assume ownership of it
+		return EObjectMigrationRecvType::AssignedOwnership;
+	}
+
+	const bool bReturnLoan = (OwnerServerId == LocalServerId);
+	if (bReturnLoan)
+	{
+		ensureMsgf(bAlreadyOwns,
+				   TEXT("Remote Server %s sent obj %s and assumed %s owns it (but it doesn't)"),
+				   *RemoteServerId.ToString(), *InObjectId.ToString(), *LocalServerId.ToString());
+
+		// A request we must have made (a non-push) has been fulfilled that returned a loaned object back to us
+		return EObjectMigrationRecvType::ReturnedLoan;
+	}
+
+	// We received an object we requested and ownership is not involved, so we are just borrowing it
+	return EObjectMigrationRecvType::Borrowed;
+#else
+	return EObjectMigrationRecvType::Invalid;
+#endif
+}
+
+EObjectMigrationSendType FUObjectMigrationContext::GetObjectMigrationSendType(const UObjectBase* Object) const
+{
+	return GetObjectMigrationSendType(FRemoteObjectId { Object });
+}
+
+EObjectMigrationSendType FUObjectMigrationContext::GetObjectMigrationSendType(FRemoteObjectId InObjectId) const
+{
+#if UE_WITH_REMOTE_OBJECT_HANDLE
+	if (!ensureMsgf(IsValid(), TEXT("%hs passed an invalid MigrationContext"), __func__))
+	{
+		return EObjectMigrationSendType::Invalid;
+	}
+
+	const FRemoteServerId LocalServerId = UE::RemoteObject::GetGlobalServerId();
+	const bool bCurrentlyOwns = UE::RemoteObject::Handle::IsOwned(InObjectId);
+
+	const bool bPushAssignOwnership = (RequestId == UE::RemoteObject::Transfer::PushMigrationRequestId);
+	if (bPushAssignOwnership)
+	{
+		ensureMsgf(OwnerServerId == RemoteServerId,
+				   TEXT("Server %s attempted to push Object %s to Server %s which is not the Owner (%s is the owner)"),
+				   *LocalServerId.ToString(), *ObjectId.ToString(), *RemoteServerId.ToString(), *OwnerServerId.ToString());
+
+		if (bCurrentlyOwns)
+		{
+			// We have ownership of this object and are pushing it to the destination (reassigning the ownership)
+			return EObjectMigrationSendType::ReassignOwnership;
+		}
+
+		// Send an object we don't own back to the owning server.  The fact that it was a push is irrelevant in the outcome.
+		return EObjectMigrationSendType::ReturnBorrowed;
+	}
+
+	const bool bReturnBorrowed = (OwnerServerId == RemoteServerId);
+	if (bReturnBorrowed)
+	{
+		ensureMsgf(!bCurrentlyOwns,
+				   TEXT("Server %s is sending a non-push migration request for Object %s to Server %s but it actually owns the Object (this should have been a push)"),
+				   *LocalServerId.ToString(), *InObjectId.ToString(), *RemoteServerId.ToString());
+
+		// Send an object we don't own back to the owning server (it requested the object back).
+		return EObjectMigrationSendType::ReturnBorrowed;
+	}
+
+	// By default we loan the object to the destination server but ownership does not change
+	return EObjectMigrationSendType::Loan;
+#else
+	return EObjectMigrationSendType::Invalid;
+#endif
+}
+
+const TCHAR* ToString(EObjectMigrationRecvType Value)
+{
+	switch (Value)
+	{
+		case EObjectMigrationRecvType::Borrowed:
+			return TEXT("Borrowed");
+
+		case EObjectMigrationRecvType::ReturnedLoan:
+			return TEXT("ReturnedLoan");
+
+		case EObjectMigrationRecvType::AssignedOwnership:
+			return TEXT("AssignedOwnership");
+
+		default:
+			return TEXT("Invalid");
+	}
+}
+
+const TCHAR* ToString(EObjectMigrationSendType Value)
+{
+	switch (Value)
+	{
+		case EObjectMigrationSendType::Loan:
+			return TEXT("Loan");
+
+		case EObjectMigrationSendType::ReturnBorrowed:
+			return TEXT("ReturnBorrowed");
+
+		case EObjectMigrationSendType::ReassignOwnership:
+			return TEXT("ReassignOwnership");
+
+		default:
+			return TEXT("Invalid");
+	}
+}
diff --git a/Engine/Source/Runtime/CoreUObject/Public/UObject/Object.h b/Engine/Source/Runtime/CoreUObject/Public/UObject/Object.h
--- a/Engine/Source/Runtime/CoreUObject/Public/UObject/Object.h
+++ b/Engine/Source/Runtime/CoreUObject/Public/UObject/Object.h
@@ -547,7 +547,7 @@ public:
 	/**
 	 * Called when an object is about to be migrated and after it's been migrated to another process.
 	 */
-	virtual void PostMigrate(const struct FPostMigrateParams& PostMigrateParams) {}
+	virtual void PostMigrate(const struct FUObjectMigrationContext& MigrationContext) {}
 #endif
 
 	/**
diff --git a/Engine/Source/Runtime/CoreUObject/Public/UObject/RemoteObjectTransfer.h b/Engine/Source/Runtime/CoreUObject/Public/UObject/RemoteObjectTransfer.h
--- a/Engine/Source/Runtime/CoreUObject/Public/UObject/RemoteObjectTransfer.h
+++ b/Engine/Source/Runtime/CoreUObject/Public/UObject/RemoteObjectTransfer.h
@@ -9,41 +9,20 @@
 #include "Async/Async.h"
 #include "UObject/RemoteExecutor.h"
 
-enum class EObjectMigrationFlags : uint8
-{
-	Invalid 		= 0,
-	Receive			= (1 << 0),
-	Send 			= (1 << 1),
-	ChangeOwnership = (1 << 2),
-
-	/** Receive an object without changing ownership (implies it will be sent back to the server that loaned it) */
-	Borrow = Receive,
-
-	/** Loan an object while retaining ownership (implies the server loaning wants it back) */
-	Loan = Send,
+struct FUObjectMigrationContext;
 
-	/** Acquired objects mean the receiving side owns the object */
-	Acquire = Receive | ChangeOwnership,
-
-	/** Reassign objects to relinquish ownership and assign them to the destination server */
-	Reassign = Send | ChangeOwnership,
-};
-ENUM_CLASS_FLAGS(EObjectMigrationFlags);
-
-struct FPostMigrateParams
+namespace UE::RemoteObject::Transfer
 {
-	/** The server ID that initiated the migration request */
-	FRemoteServerId RequestingServerId;
-
-	/** The server ID that the request is destined for */
-	FRemoteServerId DestinationServerId;
+	/** Information for performing a migration (send) an object to a remote server */
+	struct FMigrateSendParams
+	{
+		/** The migration context (meta data) of the send request */
+		FUObjectMigrationContext& MigrationContext;
 
-	/** The flags that determine the semantics of the object migration request */
-	EObjectMigrationFlags MigrationFlags;
-};
+		/** The serialized data of the object being sent */
+		FRemoteObjectData ObjectData;
+	};
 
-namespace UE::RemoteObject::Transfer
-{
 	constexpr FRemoteTransactionId PushMigrationRequestId = FRemoteTransactionId(0xFFFFFFFFu);
 
 	/**
@@ -58,7 +37,7 @@ namespace UE::RemoteObject::Transfer
 	* @param RemoteServerId Server that sent object data
 	* @param Data Remote object data. Data ownership is transferred to the remote object queue
 	*/
-	COREUOBJECT_API void OnObjectDataReceived(FRemoteTransactionId RequestId, FRemoteServerId OwnerServerId, FRemoteServerId PhysicsId, FRemoteObjectId ObjectId, FRemoteServerId RemoteServerId, FRemoteObjectData* Data);
+	COREUOBJECT_API void OnObjectDataReceived(FRemoteTransactionId RequestId, FRemoteServerId OwnerServerId, FRemoteServerId PhysicsId, FRemoteObjectId ObjectId, FRemoteServerId RemoteServerId, FRemoteObjectData& Data);
 
 	/**
 	* Called when a remote object request was denied by a remote server
@@ -125,44 +104,40 @@ namespace UE::RemoteObject::Transfer
 
 	extern COREUOBJECT_API const FRemoteServerId DatabaseId;
 
-	extern COREUOBJECT_API void InitRemoteObjectTransfer();
+	COREUOBJECT_API void InitRemoteObjectTransfer();
 
 	/** Delegate that transfers object data to another server */
-	extern COREUOBJECT_API TDelegate<void(FRemoteTransactionId /*RequestId*/, FRemoteObjectId /*ObjectId*/, FRemoteServerId /*DestinationServerId*/, FRemoteServerId /*OwnerServerId*/, FRemoteServerId /*PhysicsId*/, FRemoteObjectData& /*ObjectData*/)> RemoteObjectTransferDelegate;
+	extern COREUOBJECT_API TDelegate<void(const FMigrateSendParams& Params)> RemoteObjectTransferDelegate;
 
 	/** Delegate that handles an object request being denied */
 	extern COREUOBJECT_API TDelegate<void(FRemoteTransactionId /*RequestId*/, FRemoteObjectId /*ObjectId*/, FRemoteServerId /*DestinationServerId*/)> RemoteObjectDeniedTransferDelegate;
 
-	/** Delegate that requests remote object data from CurrentOwnerServerId to be transferred to DestinationServerId */
-	extern COREUOBJECT_API TDelegate<void(FRemoteTransactionId /*RequestId*/, FRemoteWorkPriority /*RequestPriority*/, FRemoteObjectId /*ObjectId*/, FRemoteServerId /*CurrentOwnerServerId*/, FRemoteServerId /*DestinationServerId*/)> RequestRemoteObjectDelegate;
+	/** Delegate that requests remote object data from LastKnownResidentServerId to be transferred to DestinationServerId. Allows requests to be forwarded if ObjectId does not reside on LastKnownResidentServerId. */
+	extern COREUOBJECT_API TDelegate<void(FRemoteTransactionId /*RequestId*/, FRemoteWorkPriority /*RequestPriority*/, FRemoteObjectId /*ObjectId*/, FRemoteServerId /*LastKnownResidentServerId*/, FRemoteServerId /*DestinationServerId*/)> RequestRemoteObjectDelegate;
 
-	/** 
-	 * Delegate executed when objects have been migrated from another server. 
-	 */
-	extern COREUOBJECT_API TMulticastDelegate<void(const FRemoteObjectData& /*ObjectData*/, FRemoteServerId /*RemoteServerId*/, EObjectMigrationFlags /*MigrationFlags*/)> OnObjectDataReceivedDelegate;
+	/** Delegate executed when objects have been migrated from another server. */
+	extern COREUOBJECT_API TMulticastDelegate<void(const FRemoteObjectData& /*ObjectData*/, const FUObjectMigrationContext& /*MigrationContext*/)> OnObjectDataReceivedDelegate;
 
-	/**
-	 * Delegate executed when objects have been migrated from another server.
-	 */
-	extern COREUOBJECT_API TMulticastDelegate<void(const TArray<UObject*>& /*MigratedObjects*/, FRemoteServerId /*RemoteServerId*/, EObjectMigrationFlags /*MigrationFlags*/)> OnObjectsReceivedDelegate;
+	/** Delegate executed when objects have been migrated from another server. */
+	extern COREUOBJECT_API TMulticastDelegate<void(const TArray<UObject*>& /*MigratedObjects*/, const FUObjectMigrationContext& /*MigrationContext*/)> OnObjectsReceivedDelegate;
 
 	/** Delegate executed when object data has been migrated to another server */
-	extern COREUOBJECT_API TMulticastDelegate<void(const FRemoteObjectData& /*ObjectData*/, FRemoteServerId /*RemoteServerId*/, EObjectMigrationFlags /*MigrationFlags*/)> OnObjectDataSentDelegate;
+	extern COREUOBJECT_API TMulticastDelegate<void(const FRemoteObjectData& /*ObjectData*/, const FUObjectMigrationContext& /*MigrationContext*/)> OnObjectDataSentDelegate;
 
 	/** 
 	 * Delegate executed when objects have been migrated to another server. Migrated objects can still be resolved with weak pointers but at this point 
 	 * any changes to their internal state won't be migrated across to the remote server 
 	 */
-	extern COREUOBJECT_API TMulticastDelegate<void(const TSet<UObject*>& /*MigratedObjects*/, FRemoteServerId /*RemoteServerId*/, EObjectMigrationFlags /*MigrationFlags*/)> OnObjectsSentDelegate;
+	extern COREUOBJECT_API TMulticastDelegate<void(const TSet<UObject*>& /*MigratedObjects*/, const FUObjectMigrationContext& /*MigrationContext*/)> OnObjectsSentDelegate;
 
 	/** Delegate executed when an object has been accessed by a transaction */
 	extern COREUOBJECT_API TMulticastDelegate<void(FRemoteTransactionId /*RequestId*/, FRemoteObjectId /*ObjectId*/)> OnObjectTouchedDelegate;
 
 	/** Delegate that stores locally unrachable object data into a database. */
-	extern COREUOBJECT_API TDelegate<void(FRemoteTransactionId /*RequestId*/, FRemoteObjectId /*ObjectId*/, FRemoteServerId /*DatabaseId*/, FRemoteServerId /*OwnerServerId*/, FRemoteServerId /*PhysicsId*/, FRemoteObjectData& /*ObjectData*/)> StoreRemoteObjectDataDelegate;
+	extern COREUOBJECT_API TDelegate<void(const FMigrateSendParams&)> StoreRemoteObjectDataDelegate;
 
 	/** Delegate that restores object data from a database. */
-	extern COREUOBJECT_API TDelegate<void(FRemoteTransactionId /*RequestId*/, FRemoteWorkPriority /*RequestPriority*/, FRemoteObjectId /*ObjectId*/, FRemoteServerId /*DatabaseId*/, FRemoteServerId /*DestinationServerId*/)> RestoreRemoteObjectDataDelegate;
+	extern COREUOBJECT_API TDelegate<void(const FUObjectMigrationContext&)> RestoreRemoteObjectDataDelegate;
 
 } // namespace UE::RemoveObject::Transfer
 
diff --git a/Engine/Source/Runtime/Engine/Classes/GameFramework/Actor.h b/Engine/Source/Runtime/Engine/Classes/GameFramework/Actor.h
--- a/Engine/Source/Runtime/Engine/Classes/GameFramework/Actor.h
+++ b/Engine/Source/Runtime/Engine/Classes/GameFramework/Actor.h
@@ -2335,7 +2335,7 @@ public:
 	ENGINE_API virtual void PostLoad() override;
 	ENGINE_API virtual void PostLoadSubobjects( FObjectInstancingGraph* OuterInstanceGraph ) override;
 #if UE_WITH_REMOTE_OBJECT_HANDLE
-	ENGINE_API virtual void PostMigrate(const struct FPostMigrateParams& PostMigrateParams) override;
+	ENGINE_API virtual void PostMigrate(const struct FUObjectMigrationContext& MigrationContext) override;
 #endif
 	ENGINE_API virtual void BeginDestroy() override;
 #if !UE_STRIP_DEPRECATED_PROPERTIES
diff --git a/Engine/Source/Runtime/Engine/Private/Actor.cpp b/Engine/Source/Runtime/Engine/Private/Actor.cpp
--- a/Engine/Source/Runtime/Engine/Private/Actor.cpp
+++ b/Engine/Source/Runtime/Engine/Private/Actor.cpp
@@ -67,6 +67,7 @@
 #if UE_WITH_REMOTE_OBJECT_HANDLE
 #include "UObject/RemoteExecutor.h"
 #include "UObject/RemoteObjectTransfer.h"
+#include "UObject/UObjectMigrationContext.h"
 #endif
 
 #if WITH_EDITOR
@@ -1232,27 +1233,26 @@ void AActor::PostLoadSubobjects(FObjectInstancingGraph* OuterInstanceGraph)
 
 static USceneComponent* FixupNativeActorComponents(AActor* Actor);
 
-void AActor::PostMigrate(const FPostMigrateParams& PostMigrateParams)
+void AActor::PostMigrate(const FUObjectMigrationContext& MigrationContext)
 {
-	Super::PostMigrate(PostMigrateParams);
+	Super::PostMigrate(MigrationContext);
 
 	// We are using this as a proxy for having ownership until we sort out those mechanics
 	const FRemoteServerId LocalServerId = UE::RemoteObject::GetGlobalServerId();
 	const bool bHasOwnership = UE::RemoteObject::Handle::IsOwned(this);
-	const bool bRequestOwnershipChange = EnumHasAnyFlags(PostMigrateParams.MigrationFlags, EObjectMigrationFlags::ChangeOwnership);
-
-	UE_LOG(LogActor, Verbose, TEXT("%hs: %s %d (LocalSrvId: %s. ReqSrvId: %s. DstSrvId: %s. HasOwnership: %d. ChangeOwnership: %d)"), __func__, *GetName(), static_cast<int>(PostMigrateParams.MigrationFlags), *LocalServerId.ToString(), *PostMigrateParams.RequestingServerId.ToString(), *PostMigrateParams.DestinationServerId.ToString(), bHasOwnership, bRequestOwnershipChange);
 
 	UWorld* const World = GetWorld();
 	checkf(World, TEXT("Migrated Actor %s (%s) doesn't have a World"), *GetPathName(), *FRemoteObjectId(this).ToString());
 
 	TGuardValue<ITransaction*> SuppressTransaction(GUndo, nullptr);
-	if (EnumHasAnyFlags(PostMigrateParams.MigrationFlags, EObjectMigrationFlags::Send))
+	if (MigrationContext.MigrationSide == EObjectMigrationSide::Send)
 	{
-		ensureMsgf(PostMigrateParams.DestinationServerId != LocalServerId, TEXT("%hs: Sending Object %s from %s to %s but destination is the same as the Local Server Id."), __func__, *GetName(), *PostMigrateParams.RequestingServerId.ToString(), *PostMigrateParams.DestinationServerId.ToString());
+		EObjectMigrationSendType MigrationSendType = MigrationContext.GetObjectMigrationSendType(this);
+		UE_LOG(LogActor, Verbose, TEXT("%hs: Send %s (LocalSrvId: %s. DstSrvId: %s. HasOwnership: %d. MigrationType: %s)"), __func__, *GetName(), *LocalServerId.ToString(), *MigrationContext.RemoteServerId.ToString(), bHasOwnership, ToString(MigrationSendType));
+		ensureMsgf(MigrationContext.RemoteServerId != LocalServerId, TEXT("%hs: Sending Object %s from %s to %s but destination is the same as the Local Server Id."), __func__, *GetName(), *LocalServerId.ToString(), *MigrationContext.RemoteServerId.ToString());
 
-		const bool bRelinquishOwnership = !bHasOwnership || bRequestOwnershipChange;
-		if (bRelinquishOwnership)
+		const bool bKeepOwnership = (MigrationSendType == EObjectMigrationSendType::Loan);
+		if (!bKeepOwnership)
 		{
 			// Remove the actor from the world on this server as it is about to be migrated to another server.
 			// Migrated actors are unique to the server instance that currently owns their memory and are not replicated to other servers.
@@ -1290,9 +1290,10 @@ void AActor::PostMigrate(const FPostMigrateParams& PostMigrateParams)
 			RegisterAllActorTickFunctions(/*bRegisterTickFunctions*/ false, /*bIncludeComponents*/ true);
 		}
 	}
-	else if (EnumHasAnyFlags(PostMigrateParams.MigrationFlags, EObjectMigrationFlags::Receive))
+	else if (MigrationContext.MigrationSide == EObjectMigrationSide::Receive)
 	{
-		ensureMsgf(PostMigrateParams.DestinationServerId == LocalServerId, TEXT("Unexpected DestinationServerId (%s was not the local ServerId: %s)"), *PostMigrateParams.DestinationServerId.ToString(), *LocalServerId.ToString());
+		EObjectMigrationRecvType MigrationRecvType = MigrationContext.GetObjectMigrationRecvType(this);
+		UE_LOG(LogActor, Verbose, TEXT("%hs: Recv %s (LocalSrvId: %s. DstSrvId: %s. HasOwnership: %d. MigrationType: %s)"), __func__, *GetName(), *LocalServerId.ToString(), *MigrationContext.RemoteServerId.ToString(), bHasOwnership, ToString(MigrationRecvType));
 
 		// Add the actor that was migrated from another server to this server's world
 		constexpr bool bOnlyAddIfUnique = true;
@@ -1316,8 +1317,7 @@ void AActor::PostMigrate(const FPostMigrateParams& PostMigrateParams)
 		bHasFinishedSpawning = true;
 
 		// If we're being passed ownership, register us with the global systems
-		const bool bAcquiredOwnership = !bHasOwnership && bRequestOwnershipChange;
-		if (bAcquiredOwnership)
+		if (MigrationRecvType == EObjectMigrationRecvType::AssignedOwnership)
 		{
 			BuildReplicatedComponentsInfo();
 
diff --git a/Engine/Source/Runtime/Engine/Private/Components/PrimitiveComponent.cpp b/Engine/Source/Runtime/Engine/Private/Components/PrimitiveComponent.cpp
--- a/Engine/Source/Runtime/Engine/Private/Components/PrimitiveComponent.cpp
+++ b/Engine/Source/Runtime/Engine/Private/Components/PrimitiveComponent.cpp
@@ -73,6 +73,10 @@
 #include "Engine/StaticMesh.h"
 #endif
 
+#if UE_WITH_REMOTE_OBJECT_HANDLE
+#include "UObject/UObjectMigrationContext.h"
+#endif
+
 #define LOCTEXT_NAMESPACE "PrimitiveComponent"
 
 
@@ -1254,6 +1258,20 @@ void UPrimitiveComponent::Serialize(FArchive& Ar)
 	// For remote objects we need non-UProperties to be transferred
 	if (Ar.IsMigratingRemoteObjects())
 	{
+		const FUObjectMigrationContext* MigrationContext = Ar.GetMigrationContext();
+		check(MigrationContext);
+
+		if (Ar.IsLoading())
+		{
+			EObjectMigrationRecvType ObjMigrationRecvType = MigrationContext->GetObjectMigrationRecvType(this);
+			UE_LOG(LogTemp, Log, TEXT("%hs: %s ObjMigrationRecvType: %s"), __func__, *GetName(), ToString(ObjMigrationRecvType));
+		}
+		else if (Ar.IsSaving())
+		{
+			EObjectMigrationSendType ObjMigrationSendType = MigrationContext->GetObjectMigrationSendType(this);
+			UE_LOG(LogTemp, Log, TEXT("%hs: %s ObjMigrationSendType: %s"), __func__, *GetName(), ToString(ObjMigrationSendType));
+		}
+
 		{
 	    	int32 NumOverlappingComponents = OverlappingComponents.Num();
     		Ar << NumOverlappingComponents;
