[Build Fix] Move inline from FInstancedPropertyBag implementation file to header to fix Mac build error

#rnx

diff --git a/Engine/Source/Runtime/CoreUObject/Private/StructUtils/PropertyBag.cpp b/Engine/Source/Runtime/CoreUObject/Private/StructUtils/PropertyBag.cpp
--- a/Engine/Source/Runtime/CoreUObject/Private/StructUtils/PropertyBag.cpp
+++ b/Engine/Source/Runtime/CoreUObject/Private/StructUtils/PropertyBag.cpp
@@ -3151,16 +3151,6 @@ void FInstancedPropertyBag::GetPreloadDependencies(TArray<UObject*>& OutDeps)
 	}
 }
 
-inline bool FInstancedPropertyBag::IsPropertyNameValid(const FString& Name)
-{
-	return !Name.IsEmpty() && FName::IsValidXName(Name, INVALID_OBJECTNAME_CHARACTERS);
-}
-
-inline bool FInstancedPropertyBag::IsPropertyNameValid(const FName Name)
-{
-	return !Name.IsNone() && FName::IsValidXName(Name, INVALID_OBJECTNAME_CHARACTERS);
-}
-
 FName FInstancedPropertyBag::SanitizePropertyName(const FString& Name)
 {
 	if (!IsPropertyNameValid(Name))
diff --git a/Engine/Source/Runtime/CoreUObject/Public/StructUtils/PropertyBag.h b/Engine/Source/Runtime/CoreUObject/Public/StructUtils/PropertyBag.h
--- a/Engine/Source/Runtime/CoreUObject/Public/StructUtils/PropertyBag.h
+++ b/Engine/Source/Runtime/CoreUObject/Public/StructUtils/PropertyBag.h
@@ -813,14 +813,20 @@ struct COREUOBJECT_API FInstancedPropertyBag
 	 * @param Name The name to check for invalid characters.
 	 * @return Sanitized property name based on the input string.
 	 */
-	static bool IsPropertyNameValid(const FString& Name);
+	static bool IsPropertyNameValid(const FString& Name)
+	{
+		return !Name.IsEmpty() && FName::IsValidXName(Name, INVALID_OBJECTNAME_CHARACTERS);
+	}
 
 	/**
 	 * Checks whether a provided name is a valid property bag property name.
 	 * @param Name The name to check for invalid characters.
 	 * @return Sanitized property name based on the input name.
 	 */
-	static bool IsPropertyNameValid(FName Name);
+	static bool IsPropertyNameValid(FName Name)
+	{
+		return !Name.IsNone() && FName::IsValidXName(Name, INVALID_OBJECTNAME_CHARACTERS);
+	}
 
 	/**
 	 * Returns a sanitized version of the provided name without invalid characters.
