[ZenLoader]
- Add support for linker load external read dependencies

diff --git a/Engine/Source/Runtime/CoreUObject/Private/Serialization/AsyncLoading2.cpp b/Engine/Source/Runtime/CoreUObject/Private/Serialization/AsyncLoading2.cpp
--- a/Engine/Source/Runtime/CoreUObject/Private/Serialization/AsyncLoading2.cpp
+++ b/Engine/Source/Runtime/CoreUObject/Private/Serialization/AsyncLoading2.cpp
@@ -7121,6 +7121,11 @@ EEventLoadNodeExecutionResult FAsyncPackage2::Event_PreloadLinkerLoadExports(FAs
 		{
 			return EEventLoadNodeExecutionResult::Timeout;
 		}
+
+		if (Package->LinkerLoadState->Linker->ExternalReadDependencies.Num())
+		{
+			Package->ExternalReadDependencies.Append(MoveTemp(Package->LinkerLoadState->Linker->ExternalReadDependencies));
+		}
 	}
 
 	if (Package->ExternalReadDependencies.Num() == 0)
