SoftObjectPtr class now has ctors accepting TNotNull pointers.

diff --git a/Engine/Source/Runtime/CoreUObject/Public/UObject/SoftObjectPtr.h b/Engine/Source/Runtime/CoreUObject/Public/UObject/SoftObjectPtr.h
--- a/Engine/Source/Runtime/CoreUObject/Public/UObject/SoftObjectPtr.h
+++ b/Engine/Source/Runtime/CoreUObject/Public/UObject/SoftObjectPtr.h
@@ -242,6 +242,27 @@ public:
 	{
 	}
 
+#if UE_ENABLE_NOTNULL_WRAPPER
+	/** Construct from an object already in memory. Convenience function for implicit initialization from TNotNull pointers. */
+	template <
+		typename U
+		UE_REQUIRES(std::is_convertible_v<U*, T*>)
+	>
+	[[nodiscard]] FORCEINLINE TSoftObjectPtr(TNotNull<U*> Object)
+		: SoftObjectPtr(Object)
+	{
+	}
+	/** Construct from an object already in memory. Convenience function for implicit initialization from TNotNull pointers. */
+	template <
+		typename U
+		UE_REQUIRES(std::is_convertible_v<U*, T*>)
+	>
+	[[nodiscard]] FORCEINLINE TSoftObjectPtr(TNotNull<TObjectPtr<U>> Object)
+		: SoftObjectPtr(static_cast<TObjectPtr<U>>(Object).Get())
+	{
+	}
+#endif
+
 	/** Construct from a TObjectPtr<U> which may or may not be in memory. */
 	template <
 		typename U
@@ -321,6 +342,29 @@ public:
 		return *this;
 	}
 
+#if UE_ENABLE_NOTNULL_WRAPPER
+	/** Copy from an object already in memory. Convenience function for implicit assignment from TNotNull pointers. */
+	template <
+		typename U
+		UE_REQUIRES(std::is_convertible_v<U*, T*>)
+	>
+	FORCEINLINE TSoftObjectPtr& operator=(TNotNull<U*> Object)
+	{
+		SoftObjectPtr = Object;
+		return *this;
+	}
+	/** Copy from an object already in memory. Convenience function for implicit assignment from TNotNull pointers. */
+	template <
+		typename U
+		UE_REQUIRES(std::is_convertible_v<U*, T*>)
+	>
+	FORCEINLINE TSoftObjectPtr& operator=(TNotNull<TObjectPtr<U>> Object)
+	{
+		SoftObjectPtr = static_cast<TObjectPtr<U>>(Object).Get();
+		return *this;
+	}
+#endif
+
 	/** Copy from a TObjectPtr<U> which may or may not be in memory. */
 	template <
 		typename U
