Added a new define UE_PACK_FUOBJECT_ITEM that packs FUObjectItem to save memory in a Shipping build, reducing struct size from 24 bytes down to 20

diff --git a/Engine/Source/Programs/Shared/EpicGames.Core/UnrealEngineTypes.cs b/Engine/Source/Programs/Shared/EpicGames.Core/UnrealEngineTypes.cs
--- a/Engine/Source/Programs/Shared/EpicGames.Core/UnrealEngineTypes.cs
+++ b/Engine/Source/Programs/Shared/EpicGames.Core/UnrealEngineTypes.cs
@@ -672,97 +672,94 @@ namespace EpicGames.Core
 		/// <summary>
 		/// One of the flags used by Garbage Collector to determine UObject's reachability state
 		/// </summary>
-		ReachabilityFlag0 = 1 << 0,
+		ReachabilityFlag0 = 1 << 14,
 
 		/// <summary>
 		/// One of the flags used by Garbage Collector to determine UObject's reachability state
 		/// </summary>
-		ReachabilityFlag1 = 1 << 1,
+		ReachabilityFlag1 = 1 << 15,
 
 		/// <summary>
 		/// One of the flags used by Garbage Collector to determine UObject's reachability state
 		/// </summary>
-		ReachabilityFlag2 = 1 << 2,
-
-		/// <summary>
-		/// Object was constructed in an AutoRTFM transaction that has been aborted
-		/// </summary>
-		AutoRTFMConstructionAborted = 1 << 17,
+		ReachabilityFlag2 = 1 << 16,
 
 		/// <summary>
 		/// Object is no longer owned by this process
-		/// </summary>
-		Remote = 1 << 18,
+		/// <summary>
+		Remote = 1 << 17,
 
 		/// <summary>
 		/// Object referenced by remote process
 		/// </summary>
-		RemoteReference = 1 << 19,
-		
+		RemoteReference = 1 << 18,
+
 		/// <summary>
 		/// Object is ready to be imported by another package during loading
 		/// </summary>
-		LoaderImport = 1 << 20,
+		LoaderImport = 1 << 19,
 
 		/// <summary>
 		/// Garbage from logical point of view and should not be referenced. This flag is mirrored in EObjectFlags as RF_Garbage for performance
 		/// </summary>
-		Garbage = 1 << 21,
+		Garbage = 1 << 20,
 
 		/// <summary>
-		/// Object is being asynchronously loaded
+		/// Object is being asynchronously loaded.
 		/// </summary>
-		AsyncLoadingPhase1 = 1 << 22,
+		AsyncLoadingPhase1 = 1 << 21,
 
 		/// <summary>
 		/// External reference to object in cluster exists
 		/// </summary>
-		ReachableInCluster = 1 << 23,
+		ReachableInCluster = 1 << 22,
 
 		/// <summary>
 		/// Root of a cluster
 		/// </summary>
-		ClusterRoot = 1 << 24,
+		ClusterRoot = 1 << 23,
 
 		/// <summary>
 		/// Native (UClass only).
 		/// </summary>
-		Native = 1 << 25,
+		Native = 1 << 24,
 
 		/// <summary>
 		/// Object exists only on a different thread than the game thread.
 		/// </summary>
-		Async = 1 << 26,
+		Async = 1 << 25,
 
 		/// <summary>
 		/// Object is being asynchronously loaded.
 		/// </summary>
-		AsyncLoading = 1 << 27,
+		AsyncLoadingPhase2 = 1 << 26,
 
 		/// <summary>
 		/// Object is not reachable on the object graph.
 		/// </summary>
-		Unreachable = 1 << 28,
+		Unreachable = 1 << 27,
 
 		/// <summary>
 		/// Object currently has ref-counts associated with it.
 		/// </summary>
-		RefCounted = 1 << 29,
+		RefCounted = 1 << 28,
 
 		/// <summary>
 		/// Object will not be garbage collected, even if unreferenced.
 		/// </summary>
-		RootSet = 1 << 30,
+		RootSet = 1 << 29,
 
 		/// <summary>
 		/// Object didn't have its class constructor called yet (only the UObjectBase one to initialize its most basic members)
 		/// </summary>
-		PendingConstruction = 1 << 31,
+		PendingConstruction = 1 << 30,
 
-		GarbageCollectionKeepFlags = Native | Async | AsyncLoading | LoaderImport | RefCounted,
+		/// <summary>
+		/// Object was constructed in an AutoRTFM transaction that has been aborted
+		/// </summary>
+		AutoRTFMConstructionAborted = 1 << 31,
 
-		//~ Make sure this is up to date!
-		AllFlags = ReachabilityFlag0 | ReachabilityFlag1 | ReachabilityFlag2 | LoaderImport | Garbage | ReachableInCluster | ClusterRoot | Native | Async | AsyncLoading | Unreachable | RootSet | PendingConstruction | RefCounted
+		AsyncLoading = AsyncLoadingPhase1 | AsyncLoadingPhase2,
 	};
 
 	/// <summary>
diff --git a/Engine/Source/Runtime/Core/Public/Clang/ClangPlatformMath.h b/Engine/Source/Runtime/Core/Public/Clang/ClangPlatformMath.h
--- a/Engine/Source/Runtime/Core/Public/Clang/ClangPlatformMath.h
+++ b/Engine/Source/Runtime/Core/Public/Clang/ClangPlatformMath.h
@@ -21,7 +21,7 @@ struct FClangPlatformMath : public FGenericPlatformMath
 	 *
 	 * @return the number of zeros before the first "on" bit
 	 */
-	static FORCEINLINE uint8 CountLeadingZeros8(uint8 Value)
+	static FORCEINLINE constexpr uint8 CountLeadingZeros8(uint8 Value)
 	{
 		return uint8(__builtin_clz((uint32(Value) << 1) | 1) - 23);
 	}
@@ -33,7 +33,7 @@ struct FClangPlatformMath : public FGenericPlatformMath
 	 *
 	 * @return the number of zeros before the first "on" bit
 	 */
-	static FORCEINLINE uint32 CountLeadingZeros(uint32 Value)
+	static FORCEINLINE constexpr uint32 CountLeadingZeros(uint32 Value)
 	{
 		return __builtin_clzll((uint64(Value) << 1) | 1) - 31;
 	}
@@ -45,7 +45,7 @@ struct FClangPlatformMath : public FGenericPlatformMath
 	 *
 	 * @return the number of zeros before the first "on" bit
 	 */
-	static FORCEINLINE uint64 CountLeadingZeros64(uint64 Value)
+	static FORCEINLINE constexpr uint64 CountLeadingZeros64(uint64 Value)
 	{
 		if (Value == 0)
 		{
@@ -62,7 +62,7 @@ struct FClangPlatformMath : public FGenericPlatformMath
 	 *
 	 * @return the number of zeros after the last "on" bit
 	 */
-	static FORCEINLINE uint32 CountTrailingZeros(uint32 Value)
+	static FORCEINLINE constexpr uint32 CountTrailingZeros(uint32 Value)
 	{
 		if (Value == 0)
 		{
@@ -71,6 +71,11 @@ struct FClangPlatformMath : public FGenericPlatformMath
 	
 		return (uint32)__builtin_ctz(Value);
 	}
+
+	static FORCEINLINE constexpr uint32 CountTrailingZerosConstExpr(uint32 Value)
+	{
+		return CountTrailingZeros(Value);
+	}
 	
 	/**
 	 * Counts the number of trailing zeros in the bit representation of the value
@@ -79,7 +84,7 @@ struct FClangPlatformMath : public FGenericPlatformMath
 	 *
 	 * @return the number of zeros after the last "on" bit
 	 */
-	static FORCEINLINE uint64 CountTrailingZeros64(uint64 Value)
+	static FORCEINLINE constexpr uint64 CountTrailingZeros64(uint64 Value)
 	{
 		if (Value == 0)
 		{
@@ -89,22 +94,28 @@ struct FClangPlatformMath : public FGenericPlatformMath
 		return (uint64)__builtin_ctzll(Value);
 	}
 
-	static FORCEINLINE uint32 FloorLog2(uint32 Value)
+	static FORCEINLINE constexpr uint64 CountTrailingZeros64ConstExpr(uint64 Value)
+	{
+		return CountTrailingZeros64(Value);
+	}
+	
+
+	static FORCEINLINE constexpr uint32 FloorLog2(uint32 Value)
 	{
 		return 31 - __builtin_clz(Value | 1);
 	}
 
-	static FORCEINLINE uint32 FloorLog2NonZero(uint32 Value)
+	static FORCEINLINE constexpr uint32 FloorLog2NonZero(uint32 Value)
 	{
 		return 31 - __builtin_clz(Value);
 	}
 
-	static FORCEINLINE uint64 FloorLog2_64(uint64 Value)
+	static FORCEINLINE constexpr uint64 FloorLog2_64(uint64 Value)
 	{
 		return 63 - __builtin_clzll(Value | 1);
 	}
 
-	static FORCEINLINE uint64 FloorLog2NonZero_64(uint64 Value)
+	static FORCEINLINE constexpr uint64 FloorLog2NonZero_64(uint64 Value)
 	{
 		return 63 - __builtin_clzll(Value);
 	}
diff --git a/Engine/Source/Runtime/Core/Public/GenericPlatform/GenericPlatformMath.h b/Engine/Source/Runtime/Core/Public/GenericPlatform/GenericPlatformMath.h
--- a/Engine/Source/Runtime/Core/Public/GenericPlatform/GenericPlatformMath.h
+++ b/Engine/Source/Runtime/Core/Public/GenericPlatform/GenericPlatformMath.h
@@ -739,7 +739,7 @@ struct FGenericPlatformMath
 	 *
 	 * @return the number of zeros after the last "on" bit
 	 */
-	[[nodiscard]] static constexpr FORCEINLINE uint32 CountTrailingZeros(uint32 Value)
+	[[nodiscard]] static constexpr FORCEINLINE uint32 CountTrailingZerosConstExpr(uint32 Value)
 	{
 		if (Value == 0)
 		{
@@ -754,6 +754,11 @@ struct FGenericPlatformMath
 		return Result;
 	}
 
+	[[nodiscard]] static constexpr FORCEINLINE uint32 CountTrailingZeros(uint32 Value)
+	{
+		return CountTrailingZerosConstExpr(Value);
+	}
+
 	/**
 	 * Counts the number of trailing zeros in the bit representation of the value
 	 *
@@ -761,7 +766,7 @@ struct FGenericPlatformMath
 	 *
 	 * @return the number of zeros after the last "on" bit
 	 */
-	[[nodiscard]] static constexpr FORCEINLINE uint64 CountTrailingZeros64(uint64 Value)
+	[[nodiscard]] static constexpr FORCEINLINE uint64 CountTrailingZeros64ConstExpr(uint64 Value)
 	{
 		if (Value == 0)
 		{
@@ -776,6 +781,11 @@ struct FGenericPlatformMath
 		return Result;
 	}
 
+	[[nodiscard]] static constexpr FORCEINLINE uint64 CountTrailingZeros64(uint64 Value)
+	{
+		return CountTrailingZeros64ConstExpr(Value);
+	}
+
 	/**
 	 * Returns smallest N such that (1<<N)>=Arg.
 	 * Note: CeilLogTwo(0)=0 
diff --git a/Engine/Source/Runtime/CoreUObject/Private/UObject/GarbageCollection.cpp b/Engine/Source/Runtime/CoreUObject/Private/UObject/GarbageCollection.cpp
--- a/Engine/Source/Runtime/CoreUObject/Private/UObject/GarbageCollection.cpp
+++ b/Engine/Source/Runtime/CoreUObject/Private/UObject/GarbageCollection.cpp
@@ -652,7 +652,7 @@ bool FUObjectItem::SetRootFlags(EInternalObjectFlags FlagsToSet)
 		{
 			if (!(GUObjectArray.IsOpenForDisregardForGC() & GUObjectArray.DisregardForGCEnabled())) //-V792
 			{
-				GRoots.Add(GUObjectArray.ObjectToIndex(Object));
+				GRoots.Add(GUObjectArray.ObjectToIndex(GetObject()));
 			}
 		}
 		bIChangedIt = AtomicallySetFlag_ForGC(FlagsToSet);
@@ -661,8 +661,8 @@ bool FUObjectItem::SetRootFlags(EInternalObjectFlags FlagsToSet)
 	{
 		// Setting any of the root flags on an object during incremental reachability requires a GC barrier
 		// to make sure an object with root flags does not get Garbage Collected
-		checkf(Object, TEXT("Setting an internal object flag on a null object entry"));
-		Object->MarkAsReachable();
+		checkf(GetObject(), TEXT("Setting an internal object flag on a null object entry"));
+		GetObject()->MarkAsReachable();
 	}
 	return bIChangedIt;
 }
@@ -684,7 +684,7 @@ bool FUObjectItem::ClearRootFlags(EInternalObjectFlags FlagsToClear)
 
 	if ((OldFlags & RootFlags) != 0 && ((OldFlags & ~(int32)FlagsToClear) & RootFlags) == 0)
 	{
-		GRoots.Remove(GUObjectArray.ObjectToIndex(Object));
+		GRoots.Remove(GUObjectArray.ObjectToIndex(GetObject()));
 	}
 	return AtomicallyClearFlag_ForGC((EInternalObjectFlags)FlagsToClear);
 }
diff --git a/Engine/Source/Runtime/CoreUObject/Private/UObject/UObjectArray.cpp b/Engine/Source/Runtime/CoreUObject/Private/UObject/UObjectArray.cpp
--- a/Engine/Source/Runtime/CoreUObject/Private/UObject/UObjectArray.cpp
+++ b/Engine/Source/Runtime/CoreUObject/Private/UObject/UObjectArray.cpp
@@ -35,7 +35,7 @@ void FUObjectItem::CreateStatID() const
 	TArray<UObjectBase const*, TInlineAllocator<24>> ClassChain;
 
 	// Build class hierarchy
-	UObjectBase const* Target = Object;
+	UObjectBase const* Target = GetObject();
 	do
 	{
 		ClassChain.Add(Target);
@@ -43,9 +43,9 @@ void FUObjectItem::CreateStatID() const
 	} while (Target);
 
 	// Start with class name
-	if (Object->GetClass())
+	if (GetObject()->GetClass())
 	{
-		Object->GetClass()->GetFName().GetDisplayNameEntry()->AppendNameToString(LongName);
+		GetObject()->GetClass()->GetFName().GetDisplayNameEntry()->AppendNameToString(LongName);
 	}
 
 	// Now process from parent -> child so we can append strings more efficiently.
@@ -284,7 +284,6 @@ void FUObjectArray::AllocateUObjectIndex(UObjectBase* Object, EInternalObjectFla
 	// Add to global table.
 	FUObjectItem* ObjectItem = IndexToObject(Index);
 	UE_CLOG(ObjectItem->GetObject() != nullptr, LogUObjectArray, Fatal, TEXT("Attempting to add %s at index %d but another object (0x%016llx) exists at that index!"), *Object->GetFName().ToString(), Index, (int64)(PTRINT)ObjectItem->GetObject());
-	ObjectItem->SetObject(Object);
 	// At this point all not-compiled-in objects are not fully constructed yet and this is the earliest we can mark them as such
 	ObjectItem->Flags = (int32)EInternalObjectFlags::PendingConstruction;
 	if (!(IsOpenForDisregardForGC() & GUObjectArray.DisregardForGCEnabled())) //-V792
@@ -293,6 +292,7 @@ void FUObjectArray::AllocateUObjectIndex(UObjectBase* Object, EInternalObjectFla
 		// under the same UObjectArray lock as swapping reachability flags inside of GC, see FGCFlags::SwapReachableAndMaybeUnreachable()
 		ObjectItem->Flags |= (int32)UE::GC::Private::FGCFlags::GetReachableFlagValue_ForGC();
 	}
+	ObjectItem->SetObject(Object);
 	ObjectItem->RefCount = 0;
 	ObjectItem->ClusterRootIndex = 0;
 	ObjectItem->SerialNumber = SerialNumber;
@@ -479,12 +479,12 @@ bool FUObjectArray::IsValid(const UObjectBase* Object) const
 		return false;
 	}
 	const FUObjectItem& Slot = ObjObjects[Index];
-	if( Slot.Object == NULL )
+	if( Slot.GetObject() == NULL )
 	{
 		UE_LOG(LogUObjectArray, Warning, TEXT("Empty slot") );
 		return false;
 	}
-	if( Slot.Object != Object )
+	if( Slot.GetObject() != Object )
 	{
 		UE_LOG(LogUObjectArray, Warning, TEXT("Other object in slot") );
 		return false;
diff --git a/Engine/Source/Runtime/CoreUObject/Public/UObject/ObjectMacros.h b/Engine/Source/Runtime/CoreUObject/Public/UObject/ObjectMacros.h
--- a/Engine/Source/Runtime/CoreUObject/Public/UObject/ObjectMacros.h
+++ b/Engine/Source/Runtime/CoreUObject/Public/UObject/ObjectMacros.h
@@ -619,9 +619,11 @@ enum class EInternalObjectFlags : int32
 {
 	None = 0,
 
-	ReachabilityFlag0 = 1 << 0, ///< One of the flags used by Garbage Collector to determine UObject's reachability state
-	ReachabilityFlag1 = 1 << 1, ///< One of the flags used by Garbage Collector to determine UObject's reachability state
-	ReachabilityFlag2 = 1 << 2, ///< One of the flags used by Garbage Collector to determine UObject's reachability state
+	MinFlagBitIndex = 14,
+
+	ReachabilityFlag0 = 1 << 14, ///< One of the flags used by Garbage Collector to determine UObject's reachability state
+	ReachabilityFlag1 = 1 << 15, ///< One of the flags used by Garbage Collector to determine UObject's reachability state
+	ReachabilityFlag2 = 1 << 16, ///< One of the flags used by Garbage Collector to determine UObject's reachability state
 
 	AutoRTFMConstructionAborted = 1 << 17, //< Object was constructed in an AutoRTFM transaction that has been aborted
 	Remote = 1 << 18, //< Object is no longer owned by this process
diff --git a/Engine/Source/Runtime/CoreUObject/Public/UObject/UObjectArray.h b/Engine/Source/Runtime/CoreUObject/Public/UObject/UObjectArray.h
--- a/Engine/Source/Runtime/CoreUObject/Public/UObject/UObjectArray.h
+++ b/Engine/Source/Runtime/CoreUObject/Public/UObject/UObjectArray.h
@@ -21,6 +21,10 @@ namespace Verse
 }
 #endif
 
+#ifndef UE_PACK_FUOBJECT_ITEM
+#	define UE_PACK_FUOBJECT_ITEM 0
+#endif
+
 /**
 * Controls whether the number of available elements is being tracked in the ObjObjects array.
 * By default it is only tracked in WITH_EDITOR builds as it adds a small amount of tracking overhead
@@ -37,21 +41,37 @@ namespace UE::GC::Private
 /**
 * Single item in the UObject array.
 */
-struct
-#if !STATS && !ENABLE_STATNAMEDEVENTS_UOBJECT
-	// Packing avoids 20% mem waste and improves perf
-	GCC_PACK(4)
-#endif
-	FUObjectItem
+struct FUObjectItem
 {
 	friend class FUObjectArray;
 	friend class UE::GC::Private::FGCFlags;
 
+#if !UE_PACK_FUOBJECT_ITEM
 	// Pointer to the allocated object
+	UE_DEPRECATED(5.5, "Use GetObject() and SetObject() to access Object.")
 	class UObjectBase* Object;
+
+#else
 private:
-	// Internal flags. These can only be changed via Set* and Clear* functions
+	// Currently we assume UObjects are aligned by 8 bytes, that gives us 3 lower bits as zeros that we can discard.
+	// This will give us total 45 bits in a pointer that we pack into a int32 and the remaining 13 bits we pack with Flags
+	// EInternalObjectFlags::MinFlagBitIndex at the time of writing this is 14 and we have only 1 bit left in the EInternalObjectFlags for future use
+	// We can increase UObject alignment to 16 bytes to get one more bit and reduce the overall addressable virtual memory range to get more bits if necessary
+	constexpr static int32 UObjectAlignment = 8;
+	constexpr static int32 UObjectPtrTrailingZeroes = FMath::CountTrailingZerosConstExpr(UObjectAlignment);
+	static_assert(int(EInternalObjectFlags::MinFlagBitIndex) >= 48 - 32 - UObjectPtrTrailingZeroes, "We need at least 13 bits to pack higher bits of a UObject pointer into Flags");
+	constexpr static int32 FlagsMask = 0xFFFFFFFF << int(EInternalObjectFlags::MinFlagBitIndex);
+	constexpr static int32 PtrMask = ~FlagsMask;
+
+	// Stores lower 32 bits of UObject pointer shifted by 3 to the left as all our allocations are at leas 8 bytes aligned and lower 3 bits will always be 0
+	uint32 ObjectPtrLow;
+#endif
+
+private:
+	// Stores EInternalObjectFlags (and higher 13 bits of the UObject pointer packed together if UE_PACK_FUOBJECT_ITEM is set to 1)
+	// These can only be changed via Set* and Clear* functions
 	int32 Flags;
+
 public:
 	// UObject Owner Cluster Index
 	int32 ClusterRootIndex;	
@@ -75,7 +95,11 @@ public:
 #endif // STATS || ENABLE_STATNAMEDEVENTS
 
 	FUObjectItem()
+#if !UE_PACK_FUOBJECT_ITEM
 		: Object(nullptr)
+#else
+		: ObjectPtrLow(0)
+#endif
 		, Flags(0)
 		, ClusterRootIndex(0)
 		, SerialNumber(0)
@@ -104,12 +128,27 @@ public:
 
 	FORCEINLINE class UObjectBase* GetObject() const
 	{
+#if UE_PACK_FUOBJECT_ITEM
+		const uintptr_t Object = (uintptr_t(Flags & PtrMask) << (32 + UObjectPtrTrailingZeroes)) | (uintptr_t(ObjectPtrLow) << UObjectPtrTrailingZeroes);
+		return (class UObjectBase*)Object;
+#else
+		PRAGMA_DISABLE_DEPRECATION_WARNINGS
 		return Object;
+		PRAGMA_ENABLE_DEPRECATION_WARNINGS
+#endif
 	}
 
 	FORCEINLINE void SetObject(class UObjectBase* InObject)
 	{
+#if UE_PACK_FUOBJECT_ITEM
+		const uint32 ObjectPtrHi = ((uintptr_t)InObject & 0xFFFF00000000ULL) >> (32 + UObjectPtrTrailingZeroes);
+		Flags |= ObjectPtrHi;
+		ObjectPtrLow = ((uintptr_t)InObject >> UObjectPtrTrailingZeroes) & 0xFFFFFFFF;
+#else
+		PRAGMA_DISABLE_DEPRECATION_WARNINGS
 		Object = InObject;
+		PRAGMA_ENABLE_DEPRECATION_WARNINGS
+#endif
 	}
 
 	FORCEINLINE void SetOwnerIndex(int32 OwnerIndex)
@@ -154,7 +193,11 @@ public:
 
 	FORCEINLINE EInternalObjectFlags GetFlags() const
 	{
+#if UE_PACK_FUOBJECT_ITEM
+		return EInternalObjectFlags(GetFlagsInternal() & FlagsMask);
+#else
 		return EInternalObjectFlags(GetFlagsInternal());
+#endif
 	}
 
 	FORCEINLINE void ClearFlags(EInternalObjectFlags FlagsToClear)
