Fixed IDO GC to continue to consider IDOs after turning off support for IDOs

diff --git a/Engine/Source/Runtime/CoreUObject/Internal/UObject/InstanceDataObjectUtils.h b/Engine/Source/Runtime/CoreUObject/Internal/UObject/InstanceDataObjectUtils.h
--- a/Engine/Source/Runtime/CoreUObject/Internal/UObject/InstanceDataObjectUtils.h
+++ b/Engine/Source/Runtime/CoreUObject/Internal/UObject/InstanceDataObjectUtils.h
@@ -27,8 +27,8 @@ namespace UE
 UE_API bool IsInstanceDataObjectSupportEnabled();
 /** Query if InstanceDataObject support is enabled for a specific object. */
 UE_API bool IsInstanceDataObjectSupportEnabled(const UObject* Object);
-/** Query if InstanceDataObject support is enabled for a specific class. */
-UE_API bool IsInstanceDataObjectSupportEnabledForClass(const UClass* Class);
+/** Query if InstanceDataObject support is enabled for garbage collection for a specific class. */
+UE_API bool IsInstanceDataObjectSupportEnabledForGC(const UClass* Class);
 /** Query if uninitialized alert UI is enabled **/
 UE_API bool IsUninitializedAlertUIEnabled();
 /** Query if InstanceDataObjects should be saved instead of instances. */
diff --git a/Engine/Source/Runtime/CoreUObject/Private/UObject/GarbageCollection.cpp b/Engine/Source/Runtime/CoreUObject/Private/UObject/GarbageCollection.cpp
--- a/Engine/Source/Runtime/CoreUObject/Private/UObject/GarbageCollection.cpp
+++ b/Engine/Source/Runtime/CoreUObject/Private/UObject/GarbageCollection.cpp
@@ -6764,7 +6764,7 @@ static UE::GC::ObjectAROFn GetARO(UClass* Class)
 	UE::GC::ObjectAROFn ARO = Class->CppClassStaticFunctions.GetAddReferencedObjects();
 	check(ARO != nullptr);
 #if WITH_EDITORONLY_DATA
-	return ARO != &UObject::AddReferencedObjects || UE::IsInstanceDataObjectSupportEnabledForClass(Class) ? ARO : nullptr;
+	return ARO != &UObject::AddReferencedObjects || UE::IsInstanceDataObjectSupportEnabledForGC(Class) ? ARO : nullptr;
 #else
 	return ARO != &UObject::AddReferencedObjects ? ARO : nullptr;
 #endif
diff --git a/Engine/Source/Runtime/CoreUObject/Private/UObject/InstanceDataObjectUtils.cpp b/Engine/Source/Runtime/CoreUObject/Private/UObject/InstanceDataObjectUtils.cpp
--- a/Engine/Source/Runtime/CoreUObject/Private/UObject/InstanceDataObjectUtils.cpp
+++ b/Engine/Source/Runtime/CoreUObject/Private/UObject/InstanceDataObjectUtils.cpp
@@ -261,7 +261,10 @@ namespace UE
 
 	static void OnInstanceDataObjectSupportChanged(IConsoleVariable*);
 
+	bool bEverCreatedIDO = false;
+
 	bool bEnableIDOSupport = true;
+	bool bEverEnabledIDOSupport = true;
 	FAutoConsoleVariableRef EnableIDOSupportCVar(
 		TEXT("IDO.Enable"),
 		bEnableIDOSupport,
@@ -270,6 +273,7 @@ namespace UE
 	);
 
 	bool bEnableIDOSupportOnEveryObject = false;
+	bool bEverEnabledIDOSupportOnEveryObject = false;
 	FAutoConsoleVariableRef EnableIDOSupportOnEveryObjectCVar(
 		TEXT("IDO.EnableOnEveryObject"),
 		bEnableIDOSupportOnEveryObject,
@@ -399,9 +403,20 @@ namespace UE
 		return nullptr;
 	}
 
+	bool IsInstanceDataObjectSupportEnabledForClass(const UClass* Class)
+	{
+		return bEnableIDOSupport && (bEnableIDOSupportOnEveryObject || Class->CanCreateInstanceDataObject());
+	}
+
+	bool IsInstanceDataObjectSupportEnabledForGC(const UClass* Class)
+	{
+		// Garbage Collection must always consider IDOs once an IDO has been created in the relevant category.
+		return bEverEnabledIDOSupport && (bEverEnabledIDOSupportOnEveryObject || Class->CanCreateInstanceDataObject());
+	}
+
 	bool IsInstanceDataObjectSupportEnabled(const UObject* InObject)
 	{
-		if (!IsInstanceDataObjectSupportEnabled() || !InObject)
+		if (!InObject || !bEverEnabledIDOSupport)
 		{
 			return false;
 		}
@@ -424,6 +439,11 @@ namespace UE
 			return true;
 		}
 
+		if (!IsInstanceDataObjectSupportEnabled())
+		{
+			return false;
+		}
+
 		//@todo FH: change to check trait when available or use config object
 		const UClass* ObjClass = InObject->GetClass();
 		if (!IsInstanceDataObjectSupportEnabledForClass(ObjClass))
@@ -451,13 +471,11 @@ namespace UE
 		return true;
 	}
 
-	bool IsInstanceDataObjectSupportEnabledForClass(const UClass* Class)
-	{
-		return bEnableIDOSupport && (bEnableIDOSupportOnEveryObject || Class->CanCreateInstanceDataObject());
-	}
-
 	static void OnInstanceDataObjectSupportChanged(IConsoleVariable*)
 	{
+		bEverEnabledIDOSupport = bEnableIDOSupport || (bEverEnabledIDOSupport && bEverCreatedIDO);
+		bEverEnabledIDOSupportOnEveryObject = bEnableIDOSupportOnEveryObject || (bEverEnabledIDOSupportOnEveryObject && bEverCreatedIDO);
+
 		// The reference token stream is dependent on the return value of IsInstanceDataObjectSupportEnabledForClass.
 		TArray<UClass*> AllClasses;
 		AllClasses.Add(UObject::StaticClass());
@@ -1315,6 +1333,7 @@ UObject* CreateInstanceDataObject(UObject* Owner)
 
 UObject* CreateInstanceDataObject(UObject* Owner, FArchive& Ar, int64 StartOffset, int64 EndOffset)
 {
+	bEverCreatedIDO = true;
 	return FPropertyBagRepository::Get().CreateInstanceDataObject(Owner, Ar, StartOffset, EndOffset);
 }
 
