Force disable UE_PACK_FUOBJECT_ITEM on ARM64 when FORCE_ANSI_ALLOCATOR is defined as it might be using MTE or HWASan and this feature will interfere with it

diff --git a/Engine/Source/Programs/Shared/EpicGames.Core/UnrealEngineTypes.cs b/Engine/Source/Programs/Shared/EpicGames.Core/UnrealEngineTypes.cs
--- a/Engine/Source/Programs/Shared/EpicGames.Core/UnrealEngineTypes.cs
+++ b/Engine/Source/Programs/Shared/EpicGames.Core/UnrealEngineTypes.cs
@@ -686,7 +686,7 @@ namespace EpicGames.Core
 
 		/// <summary>
 		/// Object is no longer owned by this process
-		/// <summary>
+		/// </summary>
 		Remote = 1 << 17,
 
 		/// <summary>
diff --git a/Engine/Source/Runtime/CoreUObject/Public/UObject/UObjectArray.h b/Engine/Source/Runtime/CoreUObject/Public/UObject/UObjectArray.h
--- a/Engine/Source/Runtime/CoreUObject/Public/UObject/UObjectArray.h
+++ b/Engine/Source/Runtime/CoreUObject/Public/UObject/UObjectArray.h
@@ -23,6 +23,10 @@ namespace Verse
 
 #ifndef UE_PACK_FUOBJECT_ITEM
 #	define UE_PACK_FUOBJECT_ITEM 0
+#elif defined(UE_PACK_FUOBJECT_ITEM) && defined(PLATFORM_CPU_ARM_FAMILY) && defined(FORCE_ANSI_ALLOCATOR)
+//	disable packing on ARM64 is we use Ansi allocator as it might use MTE or HWASan which uses top byte of a pointer and this feature will discard that
+#	undef UE_PACK_FUOBJECT_ITEM
+#	define UE_PACK_FUOBJECT_ITEM 0
 #endif
 
 /**
