[Backout] - CL40686690
#rnx
Original CL Desc
-----------------------------------------------------------------
Force disable UE_PACK_FUOBJECT_ITEM on ARM64 when FORCE_ANSI_ALLOCATOR is defined as it might be using MTE or HWASan and this feature will interfere with it

diff --git a/Engine/Source/Runtime/CoreUObject/Public/UObject/UObjectArray.h b/Engine/Source/Runtime/CoreUObject/Public/UObject/UObjectArray.h
--- a/Engine/Source/Runtime/CoreUObject/Public/UObject/UObjectArray.h
+++ b/Engine/Source/Runtime/CoreUObject/Public/UObject/UObjectArray.h
@@ -23,10 +23,6 @@ namespace Verse
 
 #ifndef UE_PACK_FUOBJECT_ITEM
 #	define UE_PACK_FUOBJECT_ITEM 0
-#elif defined(UE_PACK_FUOBJECT_ITEM) && defined(PLATFORM_CPU_ARM_FAMILY) && defined(FORCE_ANSI_ALLOCATOR)
-//	disable packing on ARM64 is we use Ansi allocator as it might use MTE or HWASan which uses top byte of a pointer and this feature will discard that
-#	undef UE_PACK_FUOBJECT_ITEM
-#	define UE_PACK_FUOBJECT_ITEM 0
 #endif
 
 /**
