Added IDO.Unknowns CVars to toggle inclusion of unknown properties, enums, and structs in IDOs

diff --git a/Engine/Source/Runtime/CoreUObject/Private/UObject/InstanceDataObjectUtils.cpp b/Engine/Source/Runtime/CoreUObject/Private/UObject/InstanceDataObjectUtils.cpp
--- a/Engine/Source/Runtime/CoreUObject/Private/UObject/InstanceDataObjectUtils.cpp
+++ b/Engine/Source/Runtime/CoreUObject/Private/UObject/InstanceDataObjectUtils.cpp
@@ -281,6 +281,27 @@ namespace UE
 		FConsoleVariableDelegate::CreateStatic(OnInstanceDataObjectSupportChanged)
 	);
 
+	bool bEnableIDOUnknownProperties = true;
+	FAutoConsoleVariableRef EnableIDOUnknownProperties(
+		TEXT("IDO.Unknowns.EnableProperties"),
+		bEnableIDOUnknownProperties,
+		TEXT("When enabled, IDOs will include unknown properties.")
+	);
+
+	bool bEnableIDOUnknownEnums = true;
+	FAutoConsoleVariableRef EnableIDOUnknownEnums(
+		TEXT("IDO.Unknowns.EnableEnums"),
+		bEnableIDOUnknownEnums,
+		TEXT("When enabled, IDOs will include unknown enum names.")
+	);
+
+	bool bEnableIDOUnknownStructs = true;
+	FAutoConsoleVariableRef EnableIDOUnknownStructs(
+		TEXT("IDO.Unknowns.EnableStructs"),
+		bEnableIDOUnknownStructs,
+		TEXT("When enabled, IDOs will include unknown structs and the properties within them.")
+	);
+
 	FString ExcludedUnknownPropertyTypesVar = TEXT("VerseFunctionProperty");
 	FAutoConsoleVariableRef ExcludedUnknownPropertyTypesCVar(
 		TEXT("IDO.Unknowns.ExcludedTypes"),
@@ -315,7 +336,7 @@ namespace UE
 		TEXT("IDO.EnableBlueprintInstances"),
 		bEnableIDOsForBlueprintInstances,
 		TEXT("When enabled, blueprint instances (and prefab instances) can have IDOs generated for them")
-		);
+	);
 
 	bool bEnableIDOArchetypeChain = true;
 	FAutoConsoleVariableRef EnableIDOArchetypeChainCVar(
@@ -324,9 +345,30 @@ namespace UE
 		TEXT("When enabled, IDOs will be constructed using an archetype chain")
 	);
 
+	FAutoConsoleCommand EnableIDOUnknownsCommand(
+		TEXT("IDO.Unknowns.Enable"),
+		TEXT("Use this command to toggle IDO.Unknowns.* on or off, or to report their current state."),
+		FConsoleCommandWithArgsAndOutputDeviceDelegate::CreateLambda([](const TArray<FString>& Args, FOutputDevice& OutputDevice)
+		{
+			if (Args.Num() == 1)
+			{
+				bool bEnabled = Args[0] == TEXT("True") || Args[0] == TEXT("1");
+				EnableIDOUnknownProperties->Set(bEnabled);
+				EnableIDOUnknownEnums->Set(bEnabled);
+				EnableIDOUnknownStructs->Set(bEnabled);
+			}
+
+			bool bEnabled = bEnableIDOUnknownProperties && bEnableIDOUnknownEnums && bEnableIDOUnknownStructs;
+			OutputDevice.Logf(TEXT("IDO.Unknowns.Enable = \"%s\""), bEnabled ? TEXT("True") : TEXT("False"));
+		})
+	);
+
 	static void SetEnableAllIDOFeatures(bool bEnabled)
 	{
 		EnableIDOSupportCVar->Set(bEnabled);
+		EnableIDOUnknownProperties->Set(bEnabled);
+		EnableIDOUnknownEnums->Set(bEnabled);
+		EnableIDOUnknownStructs->Set(bEnabled);
 		EnableUninitializedUICVar->Set(bEnabled);
 		EnableIDOImpersonationOnSaveCVar->Set(bEnabled);
 		EnableIDOsForBlueprintArchetypesCVar->Set(bEnabled);
@@ -337,10 +379,13 @@ namespace UE
 	static bool AreAllIDOFeaturesEnabled()
 	{
 		return
-			bEnableIDOSupport && 
-			bEnableUninitializedUI && 
-			bEnableIDOImpersonationOnSave && 
-			bEnableIDOsForBlueprintArchetypes && 
+			bEnableIDOSupport &&
+			bEnableIDOUnknownProperties &&
+			bEnableIDOUnknownEnums &&
+			bEnableIDOUnknownStructs &&
+			bEnableUninitializedUI &&
+			bEnableIDOImpersonationOnSave &&
+			bEnableIDOsForBlueprintArchetypes &&
 			bEnableIDOsForBlueprintInstances &&
 			bEnableIDOArchetypeChain;
 	}
@@ -348,16 +393,15 @@ namespace UE
 	FAutoConsoleCommand EnableAllIDOFeaturesCommand(
 		TEXT("IDO.EnableAllFeatures"),
 		TEXT("Call this method to toggle all IDO related features on"),
-		FConsoleCommandWithArgsAndOutputDeviceDelegate::CreateLambda([](const TArray< FString >& Args, FOutputDevice& OutputDevice)
+		FConsoleCommandWithArgsAndOutputDeviceDelegate::CreateLambda([](const TArray<FString>& Args, FOutputDevice& OutputDevice)
 		{
 			if (Args.Num() == 1)
 			{
 				bool Val = Args[0] == TEXT("True") || Args[0] == TEXT("1");
 				SetEnableAllIDOFeatures(Val);
 			}
-			
+
 			OutputDevice.Logf(TEXT("IDO.EnableAllFeatures = \"%s\""), AreAllIDOFeaturesEnabled() ? TEXT("True") : TEXT("False"));
-			
 		})
 	);
 
@@ -501,7 +545,6 @@ namespace UE
 		return Class->IsA(UInstanceDataObjectClass::StaticClass()) || Class->IsA(UInstanceDataObjectStruct::StaticClass());
 	}
 
-
 	bool StructContainsLooseProperties(const UStruct* Struct)
 	{
 		return Struct->GetBoolMetaData(NAME_ContainsLoosePropertiesMetadata);
@@ -539,7 +582,7 @@ namespace UE
 
 	UEnum* FindOrCreateInstanceDataObjectEnum(UEnum* TemplateEnum, UObject* Outer, const FProperty* Property, const FUnknownEnumNames* EnumNames)
 	{
-		if (!TemplateEnum || !EnumNames)
+		if (!bEnableIDOUnknownEnums || !TemplateEnum || !EnumNames)
 		{
 			return TemplateEnum;
 		}
@@ -686,7 +729,7 @@ namespace UE
 		if (FStructProperty* AsStructProperty = CastField<FStructProperty>(Property))
 		{
 			// Structs that use native or binary serialization cannot safely generate an IDO.
-			if (!AsStructProperty->Struct->UseNativeSerialization() && !(AsStructProperty->Struct->StructFlags & STRUCT_Immutable))
+			if (bEnableIDOUnknownStructs && !AsStructProperty->Struct->UseNativeSerialization() && !(AsStructProperty->Struct->StructFlags & STRUCT_Immutable))
 			{
 				//@note: Transfer existing metadata over as we build the InstanceDataObject from the struct or it owners, if any, this is useful for testing purposes
 				TStringBuilder<256> OriginalName;
@@ -1089,6 +1132,8 @@ namespace UE
 
 	UClass* CreateInstanceDataObjectClass(const FPropertyPathNameTree* PropertyTree, const FUnknownEnumNames* EnumNames, UClass* OwnerClass, UObject* Outer)
 	{
+		PropertyTree = bEnableIDOUnknownProperties ? PropertyTree : nullptr;
+
 		FBlake3Hash Key;
 		{
 			FBlake3 KeyBuilder;
