Fix issue in physics ID computation: If an object does not have a remote stub, its physics ID should return 0 to indicate that it does not participate in push migration. Originally owner ID is used instead, which will make an object stay on the owning server forever instead of migrating to the correct server.

diff --git a/Engine/Source/Runtime/CoreUObject/Private/UObject/RemoteObject.cpp b/Engine/Source/Runtime/CoreUObject/Private/UObject/RemoteObject.cpp
--- a/Engine/Source/Runtime/CoreUObject/Private/UObject/RemoteObject.cpp
+++ b/Engine/Source/Runtime/CoreUObject/Private/UObject/RemoteObject.cpp
@@ -349,8 +349,8 @@ FRemoteServerId GetPhysicsServerId(const UObject* Object)
 	}
 	else
 	{
-		//If the object doesn't have stub, then use the owner id as physics id. 
-		Result = GetOwnerServerId(Object);
+		//If the object doesn't have stub, then the physics ID is invalid. It should not participate in push migration.
+		Result = FRemoteServerId(0);
 	}
 #endif
 	return Result;
