Update FPackageName::GetPackageMountPoint to take FStringView instead of const FString&
Save performance by removing dynamic allocation/deallocation of FString

#rnx

diff --git a/Engine/Source/Runtime/CoreUObject/Private/Misc/PackageName.cpp b/Engine/Source/Runtime/CoreUObject/Private/Misc/PackageName.cpp
--- a/Engine/Source/Runtime/CoreUObject/Private/Misc/PackageName.cpp
+++ b/Engine/Source/Runtime/CoreUObject/Private/Misc/PackageName.cpp
@@ -1419,13 +1419,13 @@ bool FPackageName::MountPointExists(const TCHAR* RootPath)
 	return FLongPackagePathsSingleton::Get().MountPointExists(FStringView(RootPath));
 }
 
-FName FPackageName::GetPackageMountPoint(const FString& InPackagePath, bool InWithoutSlashes)
+FName FPackageName::GetPackageMountPoint(FStringView InPackagePath, bool InWithoutSlashes)
 {
 	FLongPackagePathsSingleton& Paths = FLongPackagePathsSingleton::Get();
 	
 	UE::TReadScopeLock ScopeLock(Paths.MountLock);
 	int32 WithoutSlashes = InWithoutSlashes ? 1 : 0;
-	auto TryRootPath = [&InPackagePath, WithoutSlashes](FStringView RootPath)
+	auto TryRootPath = [InPackagePath, WithoutSlashes](FStringView RootPath)
 	{
 		if (FPathViews::IsParentPathOf(RootPath, InPackagePath))
 		{
@@ -1454,6 +1454,16 @@ FName FPackageName::GetPackageMountPoint(const FString& InPackagePath, bool InWi
 	return FName();
 }
 
+FName FPackageName::GetPackageMountPoint(const FString& InPackagePath, bool InWithoutSlashes)
+{
+	return GetPackageMountPoint(FStringView(InPackagePath), InWithoutSlashes);
+}
+
+FName FPackageName::GetPackageMountPoint(const TCHAR* InPackagePath, bool InWithoutSlashes)
+{
+	return GetPackageMountPoint(FStringView(InPackagePath), InWithoutSlashes);
+}
+
 FString FPackageName::GetContentPathForPackageRoot(FStringView InMountPoint)
 {
 	FLongPackagePathsSingleton& Paths = FLongPackagePathsSingleton::Get();
diff --git a/Engine/Source/Runtime/CoreUObject/Private/UObject/LinkerLoad.cpp b/Engine/Source/Runtime/CoreUObject/Private/UObject/LinkerLoad.cpp
--- a/Engine/Source/Runtime/CoreUObject/Private/UObject/LinkerLoad.cpp
+++ b/Engine/Source/Runtime/CoreUObject/Private/UObject/LinkerLoad.cpp
@@ -5114,8 +5114,8 @@ bool FLinkerLoad::IsPackageReferenceAllowed(UPackage* InPackage)
 		if ((LoadFlags & LOAD_ForDiff) == 0)
 #endif //if WITH_EDITOR
 		{
-			FName MountPointName = FPackageName::GetPackageMountPoint(LinkerRoot->GetName());
-			FName ImportMountPointName = FPackageName::GetPackageMountPoint(InPackage->GetName());
+			FName MountPointName = FPackageName::GetPackageMountPoint(FNameBuilder(LinkerRoot->GetFName()).ToView());
+			FName ImportMountPointName = FPackageName::GetPackageMountPoint(FNameBuilder(InPackage->GetFName()).ToView());
 			if (MountPointName != ImportMountPointName)
 			{
 				return false;
diff --git a/Engine/Source/Runtime/CoreUObject/Public/Misc/PackageName.h b/Engine/Source/Runtime/CoreUObject/Public/Misc/PackageName.h
--- a/Engine/Source/Runtime/CoreUObject/Public/Misc/PackageName.h
+++ b/Engine/Source/Runtime/CoreUObject/Public/Misc/PackageName.h
@@ -344,7 +344,9 @@ public:
 	 * @param InWithoutSlashes Optional parameters that keeps the slashes around the mount point if false
 	 * @return FName corresponding to the mount point, or Empty if invalid
 	 */
+	static COREUOBJECT_API FName GetPackageMountPoint(FStringView InPackagePath, bool InWithoutSlashes = true);
 	static COREUOBJECT_API FName GetPackageMountPoint(const FString& InPackagePath, bool InWithoutSlashes = true);
+	static COREUOBJECT_API FName GetPackageMountPoint(const TCHAR* InPackagePath, bool InWithoutSlashes = true);
 
 	/**
 	 * Get the path associated with the given package root
