[UE] Add container property override support for instanced subobjects when experimental override serialization logic is not in use and explicit instanced property flags are not present.

Notes:
- Maintains consistent override tracking for instanced subobjects when heterogenous reference values are in use.

diff --git a/Engine/Source/Runtime/CoreUObject/Private/UObject/OverriddenPropertySet.cpp b/Engine/Source/Runtime/CoreUObject/Private/UObject/OverriddenPropertySet.cpp
--- a/Engine/Source/Runtime/CoreUObject/Private/UObject/OverriddenPropertySet.cpp
+++ b/Engine/Source/Runtime/CoreUObject/Private/UObject/OverriddenPropertySet.cpp
@@ -442,12 +442,12 @@ EOverriddenPropertyOperation FOverriddenPropertySet::GetOverriddenPropertyOperat
 
 			// Only special case is instanced subobjects, otherwise we fallback to full array override
 			checkf(ArrayProperty->Inner, TEXT("Expecting an inner type for Arrays"));
-			if (const FObjectPropertyBase* InnerObjectProperty = ArrayProperty->Inner->HasAnyPropertyFlags(CPF_PersistentInstance) ? CastField<FObjectPropertyBase>(ArrayProperty->Inner) : nullptr)
+			if (const FObjectPropertyBase* InnerObjectProperty = CastField<FObjectPropertyBase>(ArrayProperty->Inner))
 			{
 				FScriptArrayHelper ArrayHelper(ArrayProperty, SubValuePtr);
 				if(ArrayHelper.IsValidIndex(ArrayIndex))
 				{
-					if (UObject* SubObject = InnerObjectProperty->GetObjectPropertyValue(ArrayHelper.GetElementPtr(ArrayIndex)))
+					if (UObject* SubObject = TryGetInstancedSubObjectValue(InnerObjectProperty, ArrayHelper.GetElementPtr(ArrayIndex)))
 					{
 						if (NextPropertyIterator)
 						{
@@ -488,9 +488,9 @@ EOverriddenPropertyOperation FOverriddenPropertySet::GetOverriddenPropertyOperat
 				if (NextPropertyIterator)
 				{
 					// Forward any sub queries to the subobject
-					if (const FObjectPropertyBase* ValueInstancedObjectProperty = MapProperty->ValueProp->HasAnyPropertyFlags(CPF_PersistentInstance) ? CastField<FObjectPropertyBase>(MapProperty->ValueProp) : nullptr)
+					if (const FObjectPropertyBase* ValueObjectProperty = CastField<FObjectPropertyBase>(MapProperty->ValueProp))
 					{
-						if (UObject* ValueSubObject = ValueInstancedObjectProperty->GetObjectPropertyValue(MapHelper.GetValuePtr(InternalMapIndex)))
+						if (UObject* ValueSubObject = TryGetInstancedSubObjectValue(ValueObjectProperty, MapHelper.GetValuePtr(InternalMapIndex)))
 						{
 							return OverridableManager.GetOverriddenPropertyOperation(ValueSubObject, NextPropertyIterator, bOutInheritedOperation);
 						}
@@ -592,41 +592,38 @@ bool FOverriddenPropertySet::ClearOverriddenProperty(FOverriddenPropertyNode& Pa
 			// Only special case is instanced subobjects, otherwise we fallback to full array override
 			if (FObjectPropertyBase* InnerObjectProperty = CastField<FObjectPropertyBase>(ArrayProperty->Inner))
 			{
-				if (InnerObjectProperty->HasAnyPropertyFlags(CPF_PersistentInstance))
-				{
-					FScriptArrayHelper ArrayHelper(ArrayProperty, SubValuePtr);
+				FScriptArrayHelper ArrayHelper(ArrayProperty, SubValuePtr);
 
-					if(ArrayIndex == INDEX_NONE)
+				if(ArrayIndex == INDEX_NONE)
+				{
+					// This is a case of the entire array needs to be cleared
+					// Need to loop through every sub object and clear them
+					for (int i = 0; i < ArrayHelper.Num(); ++i)
 					{
-						// This is a case of the entire array needs to be cleared
-						// Need to loop through every sub object and clear them
-						for (int i = 0; i < ArrayHelper.Num(); ++i)
+						if (UObject* SubObject = TryGetInstancedSubObjectValue(InnerObjectProperty, ArrayHelper.GetElementPtr(i)))
 						{
-							if (UObject* SubObject = InnerObjectProperty->GetObjectPropertyValue(ArrayHelper.GetElementPtr(i)))
-							{
-								OverridableManager.ClearInstancedSubObjectOverrides(Owner, SubObject);
-							}
+							OverridableManager.ClearInstancedSubObjectOverrides(Owner, SubObject);
 						}
-						bClearedOverrides = true;
 					}
-					else if(ArrayHelper.IsValidIndex(ArrayIndex))
+					bClearedOverrides = true;
+				}
+				else if(ArrayHelper.IsValidIndex(ArrayIndex))
+				{
+					if (UObject* SubObject = TryGetInstancedSubObjectValue(InnerObjectProperty, ArrayHelper.GetElementPtr(ArrayIndex)))
 					{
-						if (UObject* SubObject = InnerObjectProperty->GetObjectPropertyValue(ArrayHelper.GetElementPtr(ArrayIndex)))
+						if (NextPropertyIterator)
 						{
-							if (NextPropertyIterator)
-							{
-								return OverridableManager.ClearOverriddenProperty(SubObject, NextPropertyIterator);
-							}
-							else if(CurrentOverriddenPropertyNode)
+							return OverridableManager.ClearOverriddenProperty(SubObject, NextPropertyIterator);
+						}
+						else if (CurrentOverriddenPropertyNode)
+						{
+							const FOverriddenPropertyNodeID  SubObjectID(SubObject);
+							FOverriddenPropertyNodeID CurrentPropKey;
+							if (CurrentOverriddenPropertyNode->SubPropertyNodeKeys.RemoveAndCopyValue(SubObjectID, CurrentPropKey))
 							{
-								const FOverriddenPropertyNodeID  SubObjectID(SubObject);
-								FOverriddenPropertyNodeID CurrentPropKey;
-								if (CurrentOverriddenPropertyNode->SubPropertyNodeKeys.RemoveAndCopyValue(SubObjectID, CurrentPropKey))
-								{
-									verifyf(OverriddenPropertyNodes.Remove(CurrentPropKey), TEXT("Expecting a node to be removed"));
-									OverridableManager.ClearInstancedSubObjectOverrides(Owner, SubObject);
-									return true;
-								}
+								verifyf(OverriddenPropertyNodes.Remove(CurrentPropKey), TEXT("Expecting a node to be removed"));
+								OverridableManager.ClearInstancedSubObjectOverrides(Owner, SubObject);
+								return true;
 							}
 						}
 					}
@@ -642,7 +639,7 @@ bool FOverriddenPropertySet::ClearOverriddenProperty(FOverriddenPropertyNode& Pa
 			FScriptMapHelper MapHelper(MapProperty, SubValuePtr);
 
 			const int32 InternalMapIndex = ArrayIndex != INDEX_NONE ? MapHelper.FindInternalIndex(ArrayIndex) : INDEX_NONE;
-			const FObjectPropertyBase* ValueInstancedObjectProperty = MapProperty->ValueProp->HasAnyPropertyFlags(CPF_PersistentInstance) ? CastField<FObjectPropertyBase>(MapProperty->ValueProp) : nullptr;
+			const FObjectPropertyBase* ValueObjectProperty = CastField<FObjectPropertyBase>(MapProperty->ValueProp);
 
 			// If there is a next node, it is probably because the map value is holding a instanced subobject and the user is changing value on it.
 			// So forward the call to the instanced subobject
@@ -651,7 +648,7 @@ bool FOverriddenPropertySet::ClearOverriddenProperty(FOverriddenPropertyNode& Pa
 				if(MapHelper.IsValidIndex(InternalMapIndex))
 				{
 					checkf(MapProperty->ValueProp, TEXT("Expecting a value type for Maps"));
-					if (UObject* ValueSubObject = ValueInstancedObjectProperty ? ValueInstancedObjectProperty->GetObjectPropertyValue(MapHelper.GetValuePtr(InternalMapIndex)) : nullptr)
+					if (UObject* ValueSubObject = TryGetInstancedSubObjectValue(ValueObjectProperty, MapHelper.GetValuePtr(InternalMapIndex)))
 					{
 						return OverridableManager.ClearOverriddenProperty(ValueSubObject, NextPropertyIterator);
 					}
@@ -660,13 +657,13 @@ bool FOverriddenPropertySet::ClearOverriddenProperty(FOverriddenPropertyNode& Pa
 			else if(InternalMapIndex == INDEX_NONE)
 			{
 				// Users want to clear all of the overrides on the array, but in the case of instanced subobject, we need to clear the overrides on them as well.
-				if (ValueInstancedObjectProperty)
+				if (ValueObjectProperty)
 				{
 					// This is a case of the entire array needs to be cleared
 					// Need to loop through every sub object and clear them
 					for (FScriptMapHelper::FIterator It(MapHelper); It; ++It)
 					{
-						if (UObject* ValueSubObject = ValueInstancedObjectProperty->GetObjectPropertyValue(MapHelper.GetValuePtr(It.GetInternalIndex())))
+						if (UObject* ValueSubObject = TryGetInstancedSubObjectValue(ValueObjectProperty, MapHelper.GetValuePtr(It.GetInternalIndex())))
 						{
 							OverridableManager.ClearInstancedSubObjectOverrides(Owner, ValueSubObject);
 						}
@@ -684,7 +681,7 @@ bool FOverriddenPropertySet::ClearOverriddenProperty(FOverriddenPropertyNode& Pa
 				{
 					verifyf(OverriddenPropertyNodes.Remove(CurrentPropKey), TEXT("Expecting a node to be removed"));
 
-					if (UObject* ValueSubObject = ValueInstancedObjectProperty ? ValueInstancedObjectProperty->GetObjectPropertyValue(MapHelper.GetValuePtr(InternalMapIndex)) : nullptr)
+					if (UObject* ValueSubObject = TryGetInstancedSubObjectValue(ValueObjectProperty, MapHelper.GetValuePtr(InternalMapIndex)))
 					{
 						// In the case of a instanced subobject, clear all the overrides on the subobject as well
 						OverridableManager.ClearInstancedSubObjectOverrides(Owner, ValueSubObject);
@@ -814,209 +811,209 @@ void FOverriddenPropertySet::NotifyPropertyChange(FOverriddenPropertyNode* Paren
 		// Only special case is instanced subobjects, otherwise we fallback to full array override
 		if (FObjectPropertyBase* InnerObjectProperty = CastField<FObjectPropertyBase>(ArrayProperty->Inner))
 		{
-			if (InnerObjectProperty->HasAnyPropertyFlags(CPF_PersistentInstance))
+			FScriptArrayHelper ArrayHelper(ArrayProperty, SubValuePtr);
+			int32 ArrayIndex = PropertyIterator->Index;
+			checkf(ArrayIndex == INDEX_NONE || PropertyIterator->PropertyInfo == EPropertyVisitorInfoType::ContainerIndex, TEXT("Expecting a container index"));
+
+			if (!NextPropertyIterator)
 			{
-				FScriptArrayHelper ArrayHelper(ArrayProperty, SubValuePtr);
-				int32 ArrayIndex = PropertyIterator->Index;
-				checkf(ArrayIndex == INDEX_NONE || PropertyIterator->PropertyInfo == EPropertyVisitorInfoType::ContainerIndex, TEXT("Expecting a container index"));
+				checkf(ArrayProperty->Inner, TEXT("Expecting an inner type for Arrays"));
 
-				if (!NextPropertyIterator)
+				static TArray<uint8> SavedPreEditSubObjects;
+				FScriptArrayHelper PreEditSubObjectsArrayHelper(ArrayProperty, &SavedPreEditSubObjects);
+				if (Notification == EPropertyNotificationType::PreEdit)
 				{
-					checkf(ArrayProperty->Inner, TEXT("Expecting an inner type for Arrays"));
-
-					static TArray<uint8> SavedPreEditSubObjects;
-					FScriptArrayHelper PreEditSubObjectsArrayHelper(ArrayProperty, &SavedPreEditSubObjects);
-					if (Notification == EPropertyNotificationType::PreEdit)
+					PreEditSubObjectsArrayHelper.EmptyAndAddValues(ArrayHelper.Num());
+					for(int32 i = 0; i < ArrayHelper.Num(); i++)
 					{
-						PreEditSubObjectsArrayHelper.EmptyAndAddValues(ArrayHelper.Num());
-						for(int32 i = 0; i < ArrayHelper.Num(); i++)
-						{
-							InnerObjectProperty->SetObjectPropertyValue(PreEditSubObjectsArrayHelper.GetElementPtr(i), InnerObjectProperty->GetObjectPropertyValue(ArrayHelper.GetElementPtr(i)));
-						}
-						return;
+						InnerObjectProperty->SetObjectPropertyValue(PreEditSubObjectsArrayHelper.GetElementPtr(i), InnerObjectProperty->GetObjectPropertyValue(ArrayHelper.GetElementPtr(i)));
 					}
+					return;
+				}
 
-					auto ArrayReplace = [&]
+				auto ArrayReplace = [&]
+				{
+					if (SubPropertyNode)
 					{
-						if (SubPropertyNode)
-						{
-							// Overriding all entry in the array
-							SubPropertyNode->Operation = EOverriddenPropertyOperation::Replace;
-						}
+						// Overriding all entry in the array
+						SubPropertyNode->Operation = EOverriddenPropertyOperation::Replace;
+					}
 
-						// This is a case of the entire array is overridden
-						// Need to loop through every sub object and setup them up as overridden
-						for (int i = 0; i < ArrayHelper.Num(); ++i)
+					// This is a case of the entire array is overridden
+					// Need to loop through every sub object and setup them up as overridden
+					for (int i = 0; i < ArrayHelper.Num(); ++i)
+					{
+						if (UObject* SubObject = TryGetInstancedSubObjectValue(InnerObjectProperty, ArrayHelper.GetElementPtr(i)))
 						{
-							if (UObject* SubObject = InnerObjectProperty->GetObjectPropertyValue(ArrayHelper.GetElementPtr(i)))
+							if(SubPropertyNode)
 							{
-								if(SubPropertyNode)
-								{
-									const FOverriddenPropertyNodeID SubObjectID(SubObject);
-									FOverriddenPropertyNode& SubObjectNode = FindOrAddNode(*SubPropertyNode, SubObjectID);
-									SubObjectNode.Operation = EOverriddenPropertyOperation::Replace;
-								}
-
-								OverridableManager.OverrideInstancedSubObject(Owner, SubObject);
+								const FOverriddenPropertyNodeID SubObjectID(SubObject);
+								FOverriddenPropertyNode& SubObjectNode = FindOrAddNode(*SubPropertyNode, SubObjectID);
+								SubObjectNode.Operation = EOverriddenPropertyOperation::Replace;
 							}
+
+							OverridableManager.OverrideInstancedSubObject(Owner, SubObject);
 						}
-					};
+					}
+				};
 
-					auto ArrayAddImpl = [&]()
+				auto ArrayAddImpl = [&]()
+				{
+					checkf(ArrayHelper.IsValidIndex(ArrayIndex), TEXT("ArrayAdd change type expected to have an valid index"));
+					if (UObject* AddedSubObject = InnerObjectProperty->GetObjectPropertyValue(ArrayHelper.GetElementPtr(ArrayIndex)))
 					{
-						checkf(ArrayHelper.IsValidIndex(ArrayIndex), TEXT("ArrayAdd change type expected to have an valid index"));
-						if (UObject* AddedSubObject = InnerObjectProperty->GetObjectPropertyValue(ArrayHelper.GetElementPtr(ArrayIndex)))
+						if(SubPropertyNode)
 						{
-							if(SubPropertyNode)
-							{
-								const FOverriddenPropertyNodeID  AddedSubObjectID(AddedSubObject);
-								FOverriddenPropertyNode& AddedSubObjectNode = FindOrAddNode(*SubPropertyNode, AddedSubObjectID);
-								AddedSubObjectNode.Operation = EOverriddenPropertyOperation::Add;
+							const FOverriddenPropertyNodeID  AddedSubObjectID(AddedSubObject);
+							FOverriddenPropertyNode& AddedSubObjectNode = FindOrAddNode(*SubPropertyNode, AddedSubObjectID);
+							AddedSubObjectNode.Operation = EOverriddenPropertyOperation::Add;
 
-								// Notify the subobject that it was added
-								if (FOverriddenPropertySet* AddedSubObjectOverriddenProperties = OverridableManager.GetOverriddenProperties(AddedSubObject))
-								{
-									AddedSubObjectOverriddenProperties->bWasAdded = true;
-								}
+							// Notify the subobject that it was added
+							if (FOverriddenPropertySet* AddedSubObjectOverriddenProperties = OverridableManager.GetOverriddenProperties(AddedSubObject))
+							{
+								AddedSubObjectOverriddenProperties->bWasAdded = true;
 							}
 						}
-					};
+					}
+				};
 
-					auto ArrayRemoveImpl = [&]()
+				auto ArrayRemoveImpl = [&]()
+				{
+					checkf(PreEditSubObjectsArrayHelper.IsValidIndex(ArrayIndex), TEXT("ArrayRemove change type expected to have an valid index"));
+					if (UObject* RemovedSubObject = InnerObjectProperty->GetObjectPropertyValue(PreEditSubObjectsArrayHelper.GetElementPtr(ArrayIndex)))
 					{
-						checkf(PreEditSubObjectsArrayHelper.IsValidIndex(ArrayIndex), TEXT("ArrayRemove change type expected to have an valid index"));
-						if (UObject* RemovedSubObject = InnerObjectProperty->GetObjectPropertyValue(PreEditSubObjectsArrayHelper.GetElementPtr(ArrayIndex)))
+						if(SubPropertyNode)
 						{
-							if(SubPropertyNode)
+							// Check if there is a matching archetype for this object
+							UObject* RemovedSubObjectArchetype = RemovedSubObject->GetArchetype();
+							if (RemovedSubObjectArchetype && !RemovedSubObjectArchetype->HasAnyFlags(RF_ClassDefaultObject))
 							{
-								// Check if there is a matching archetype for this object
-								UObject* RemovedSubObjectArchetype = RemovedSubObject->GetArchetype();
-								if (RemovedSubObjectArchetype && !RemovedSubObjectArchetype->HasAnyFlags(RF_ClassDefaultObject))
+								const FOverriddenPropertyNodeID RemovedSubObjectID (RemovedSubObjectArchetype);
+								FOverriddenPropertyNode& RemovedSubObjectNode = FindOrAddNode(*SubPropertyNode, RemovedSubObjectID);
+								if (RemovedSubObjectNode.Operation == EOverriddenPropertyOperation::Add)
 								{
-									const FOverriddenPropertyNodeID RemovedSubObjectID (RemovedSubObjectArchetype);
-									FOverriddenPropertyNode& RemovedSubObjectNode = FindOrAddNode(*SubPropertyNode, RemovedSubObjectID);
-									if (RemovedSubObjectNode.Operation == EOverriddenPropertyOperation::Add)
+									// An add then a remove becomes no opt
+									FOverriddenPropertyNodeID RemovedNodeID;
+									if (SubPropertyNode->SubPropertyNodeKeys.RemoveAndCopyValue(RemovedSubObjectID, RemovedNodeID))
 									{
-										// An add then a remove becomes no opt
-										FOverriddenPropertyNodeID RemovedNodeID;
-										if (SubPropertyNode->SubPropertyNodeKeys.RemoveAndCopyValue(RemovedSubObjectID, RemovedNodeID))
-										{
-											verifyf(OverriddenPropertyNodes.Remove(RemovedNodeID), TEXT("Expecting the node to be removed"));
-											bNeedsCleanup = true;
-										}
-									}
-									else
-									{
-										RemovedSubObjectNode.Operation = EOverriddenPropertyOperation::Remove;
+										verifyf(OverriddenPropertyNodes.Remove(RemovedNodeID), TEXT("Expecting the node to be removed"));
+										bNeedsCleanup = true;
 									}
 								}
 								else
 								{
-									// Figure out if it is a remove of a previously added element
-									const FOverriddenPropertyNodeID RemovedSubObjectID (RemovedSubObject);
-									if (const FOverriddenPropertyNodeID* AddedSubObjectID = SubPropertyNode->SubPropertyNodeKeys.Find(RemovedSubObjectID))
+									RemovedSubObjectNode.Operation = EOverriddenPropertyOperation::Remove;
+								}
+							}
+							else
+							{
+								// Figure out if it is a remove of a previously added element
+								const FOverriddenPropertyNodeID RemovedSubObjectID (RemovedSubObject);
+								if (const FOverriddenPropertyNodeID* AddedSubObjectID = SubPropertyNode->SubPropertyNodeKeys.Find(RemovedSubObjectID))
+								{
+									FOverriddenPropertyNode* AddedSubObjectNode = OverriddenPropertyNodes.Find(*AddedSubObjectID);
+									checkf(AddedSubObjectNode, TEXT("Expecting a node"));
+									if (AddedSubObjectNode->Operation != EOverriddenPropertyOperation::Add)
 									{
-										FOverriddenPropertyNode* AddedSubObjectNode = OverriddenPropertyNodes.Find(*AddedSubObjectID);
-										checkf(AddedSubObjectNode, TEXT("Expecting a node"));
-										if (AddedSubObjectNode->Operation != EOverriddenPropertyOperation::Add)
-										{
-											UE_LOG(LogOverridableObject, Warning, TEXT("This removed object:%s(0x%p) was not tracked as an add in the overridden properties"), *GetNameSafe(RemovedSubObject), RemovedSubObject);
-										}
-
-										// An add then a remove becomes no opt
-										FOverriddenPropertyNodeID RemovedNodeID;
-										if (SubPropertyNode->SubPropertyNodeKeys.RemoveAndCopyValue(RemovedSubObjectID, RemovedNodeID))
-										{
-											verifyf(OverriddenPropertyNodes.Remove(RemovedNodeID), TEXT("Expecting the node to be removed"));
-											bNeedsCleanup = true;
-										}
+										UE_LOG(LogOverridableObject, Warning, TEXT("This removed object:%s(0x%p) was not tracked as an add in the overridden properties"), *GetNameSafe(RemovedSubObject), RemovedSubObject);
 									}
-									else
+
+									// An add then a remove becomes no opt
+									FOverriddenPropertyNodeID RemovedNodeID;
+									if (SubPropertyNode->SubPropertyNodeKeys.RemoveAndCopyValue(RemovedSubObjectID, RemovedNodeID))
 									{
-										UE_LOG(LogOverridableObject, Log, TEXT("This removed object:%s(0x%p) was not tracked in the overridden properties"), *GetNameSafe(RemovedSubObject), RemovedSubObject);
+										verifyf(OverriddenPropertyNodes.Remove(RemovedNodeID), TEXT("Expecting the node to be removed"));
+										bNeedsCleanup = true;
 									}
 								}
+								else
+								{
+									UE_LOG(LogOverridableObject, Log, TEXT("This removed object:%s(0x%p) was not tracked in the overridden properties"), *GetNameSafe(RemovedSubObject), RemovedSubObject);
+								}
 							}
 						}
-					};
-
-					// Only arrays flagged overridable logic can record deltas, for now just override entire array
-					if (!ArrayProperty->HasAnyPropertyFlags(CPF_ExperimentalOverridableLogic))
-					{
-						if(ChangeType == EPropertyChangeType::Unspecified && ArrayIndex == INDEX_NONE)
-						{
-							// Overriding all entry in the array + override instanced sub obejects
-							ArrayReplace();
-						}
-						else if (SubPropertyNode)
-						{
-							// Overriding all entry in the array
-							SubPropertyNode->Operation = EOverriddenPropertyOperation::Replace;
-						}
-						return;
 					}
+				};
 
-					if (ChangeType & EPropertyChangeType::ValueSet)
+				// Only arrays flagged overridable logic can record deltas, for now just override entire array
+				if (!ArrayProperty->HasAnyPropertyFlags(CPF_ExperimentalOverridableLogic))
+				{
+					if(ChangeType == EPropertyChangeType::Unspecified && ArrayIndex == INDEX_NONE)
 					{
-						checkf(ArrayIndex != INDEX_NONE, TEXT("ValueSet change type should have associated indexes"));
+						// Overriding all entry in the array + override instanced sub objects
+						ArrayReplace();
 					}
-
-					if (ChangeType & (EPropertyChangeType::ValueSet | EPropertyChangeType::Unspecified))
+					else if (SubPropertyNode)
 					{
-						if (ArrayIndex != INDEX_NONE)
-						{
-							// Overriding a single entry in the array
-							ArrayRemoveImpl();
-							ArrayAddImpl();
-						}
-						else
-						{
-							ArrayReplace();
-						}
-						return;
+						// Overriding all entry in the array
+						SubPropertyNode->Operation = EOverriddenPropertyOperation::Replace;
 					}
+					return;
+				}
 
-					if (ChangeType & EPropertyChangeType::ArrayAdd)
-					{
-						ArrayAddImpl();
-						return;
-					}
+				// Note: Currently, if CPF_ExperimentalOverridableLogic is set, we also require the property to be explicitly marked as an instanced subobject.
+				checkf(InnerObjectProperty->HasAllPropertyFlags(CPF_PersistentInstance), TEXT("Only instanced array properties support experimental overridable logic"));
 
-					if (ChangeType & EPropertyChangeType::ArrayRemove)
+				if (ChangeType & EPropertyChangeType::ValueSet)
+				{
+					checkf(ArrayIndex != INDEX_NONE, TEXT("ValueSet change type should have associated indexes"));
+				}
+
+				if (ChangeType & (EPropertyChangeType::ValueSet | EPropertyChangeType::Unspecified))
+				{
+					if (ArrayIndex != INDEX_NONE)
 					{
+						// Overriding a single entry in the array
 						ArrayRemoveImpl();
-						return;
-					}
-			
-					if (ChangeType & EPropertyChangeType::ArrayClear)
-					{
-						checkf(ArrayIndex == INDEX_NONE, TEXT("ArrayClear change type should not have associated indexes"));
-
-						for (int i = 0; i < PreEditSubObjectsArrayHelper.Num(); ++i)
-						{
-							ArrayIndex = i;
-							ArrayRemoveImpl();
-						}
-						return;
+						ArrayAddImpl();
 					}
-			
-					if (ChangeType & EPropertyChangeType::ArrayMove)
+					else
 					{
-						UE_LOG(LogOverridableObject, Log, TEXT("ArrayMove change type is not going to change anything as ordering of object isn't supported yet"));
-						return;
+						ArrayReplace();
 					}
+					return;
+				}
 
-					UE_LOG(LogOverridableObject, Verbose, TEXT("Property change type is not supported will default to full array override"));
+				if (ChangeType & EPropertyChangeType::ArrayAdd)
+				{
+					ArrayAddImpl();
+					return;
 				}
-				// Can only forward to subobject if we have a valid index
-				else if (ArrayHelper.IsValidIndex(ArrayIndex))
+
+				if (ChangeType & EPropertyChangeType::ArrayRemove)
 				{
-					if (UObject* SubObject = InnerObjectProperty->GetObjectPropertyValue(ArrayHelper.GetElementPtr(ArrayIndex)))
+					ArrayRemoveImpl();
+					return;
+				}
+			
+				if (ChangeType & EPropertyChangeType::ArrayClear)
+				{
+					checkf(ArrayIndex == INDEX_NONE, TEXT("ArrayClear change type should not have associated indexes"));
+
+					for (int i = 0; i < PreEditSubObjectsArrayHelper.Num(); ++i)
 					{
-						// This should not be needed in the property grid, as it should already been called on the subobject itself.
-						OverridableManager.NotifyPropertyChange(Notification, SubObject, NextPropertyIterator, ChangeType);
-						return;
+						ArrayIndex = i;
+						ArrayRemoveImpl();
 					}
+					return;
+				}
+			
+				if (ChangeType & EPropertyChangeType::ArrayMove)
+				{
+					UE_LOG(LogOverridableObject, Log, TEXT("ArrayMove change type is not going to change anything as ordering of object isn't supported yet"));
+					return;
+				}
+
+				UE_LOG(LogOverridableObject, Verbose, TEXT("Property change type is not supported will default to full array override"));
+			}
+			// Can only forward to subobject if we have a valid index
+			else if (ArrayHelper.IsValidIndex(ArrayIndex))
+			{
+				if (UObject* SubObject = InnerObjectProperty->GetObjectPropertyValue(ArrayHelper.GetElementPtr(ArrayIndex)))
+				{
+					// This should not be needed in the property grid, as it should already been called on the subobject itself.
+					OverridableManager.NotifyPropertyChange(Notification, SubObject, NextPropertyIterator, ChangeType);
+					return;
 				}
 			}
 		}
@@ -1033,10 +1030,8 @@ void FOverriddenPropertySet::NotifyPropertyChange(FOverriddenPropertyNode* Paren
 		FObjectPropertyBase* KeyObjectProperty = CastField<FObjectPropertyBase>(MapProperty->KeyProp);
 
 		// SubObjects
-		checkf(!KeyObjectProperty || !MapProperty->KeyProp->HasAnyPropertyFlags(CPF_PersistentInstance) || CastField<FClassProperty>(MapProperty->KeyProp), TEXT("Keys as a instanced subobject is not supported yet"));
-
 		checkf(MapProperty->ValueProp, TEXT("Expecting a value type for Maps"));
-		FObjectPropertyBase* ValueInstancedObjectProperty = MapProperty->ValueProp->HasAnyPropertyFlags(CPF_PersistentInstance) ? CastField<FObjectPropertyBase>(MapProperty->ValueProp) : nullptr;
+		FObjectPropertyBase* ValueObjectProperty = CastField<FObjectPropertyBase>(MapProperty->ValueProp);
 
 		FScriptMapHelper MapHelper(MapProperty, SubValuePtr);
 		int32 LogicalMapIndex = PropertyIterator->Index;
@@ -1107,11 +1102,12 @@ void FOverriddenPropertySet::NotifyPropertyChange(FOverriddenPropertyNode* Paren
 					}
 
 					// @todo support instanced object as a key in maps
-					//if (UObject* KeySubObject = KeyInstancedObjectProperty ? KeyInstancedObjectProperty->GetObjectPropertyValue(MapHelper.GetKeyPtr(*It)) : nullptr)
+					//if (UObject* KeySubObject = TryGetInstancedSubObjectValue(KeyObjectProperty, MapHelper.GetKeyPtr(It.GetInternalIndex())))
 					//{
+					//	checkf(false, TEXT("Keys as an instanced subobject is not supported yet"));
 					//	OverridableManager.OverrideInstancedSubObject(Owner, KeySubObject);
 					//}
-					if (UObject* ValueSubObject = ValueInstancedObjectProperty ? ValueInstancedObjectProperty->GetObjectPropertyValue(MapHelper.GetValuePtr(It.GetInternalIndex())) : nullptr)
+					if (UObject* ValueSubObject = TryGetInstancedSubObjectValue(ValueObjectProperty, MapHelper.GetValuePtr(It.GetInternalIndex())))
 					{
 						OverridableManager.OverrideInstancedSubObject(Owner, ValueSubObject);
 					}
@@ -1171,6 +1167,11 @@ void FOverriddenPropertySet::NotifyPropertyChange(FOverriddenPropertyNode* Paren
 				return;
 			}
 
+			// Ensure that an object key type is not explicitly marked as an instanced subobject. This is not supported yet.
+			checkf(!KeyObjectProperty || !KeyObjectProperty->HasAllPropertyFlags(CPF_PersistentInstance), TEXT("Keys as an instanced subobject is not supported yet"));
+			// Note: Currently, if CPF_ExperimentalOverridableLogic is set on the map, we require its value type to be explicitly marked as an instanced subobject.
+			checkf(!ValueObjectProperty || ValueObjectProperty->HasAllPropertyFlags(CPF_PersistentInstance), TEXT("Values must be instanced to support map overrides"));
+
 			if (ChangeType & EPropertyChangeType::ValueSet)
 			{
 				checkf(LogicalMapIndex != INDEX_NONE, TEXT("ValueSet change type should have associated indexes"));
@@ -1233,14 +1234,15 @@ void FOverriddenPropertySet::NotifyPropertyChange(FOverriddenPropertyNode* Paren
 		else if (MapHelper.IsValidIndex(InternalMapIndex))
 		{
 			// @todo support instanced object as a key in maps
-			//if (UObject* SubObject = KeyInstancedObjectProperty ? KeyInstancedObjectProperty->GetObjectPropertyValue(MapHelper.GetValuePtr(InternalMapIndex)) : nullptr)
+			//if (UObject* SubObject = TryGetInstancedSubObjectValue(KeyObjectProperty, MapHelper.GetValuePtr(InternalMapIndex)))
 			//{
+			//	checkf(false, TEXT("Keys as an instanced subobject is not supported yet"));
 			//	// This should not be needed in the property grid, as it should already been called on the subobject.
 			//	OverridableManager.NotifyPropertyChange(Notification, *SubObject, NextPropertyIterator, ChangeType);
 			//	return;
 			//}
 
-			if (UObject* SubObject = ValueInstancedObjectProperty ? ValueInstancedObjectProperty->GetObjectPropertyValue(MapHelper.GetValuePtr(InternalMapIndex)) : nullptr)
+			if (UObject* SubObject = TryGetInstancedSubObjectValue(ValueObjectProperty, MapHelper.GetValuePtr(InternalMapIndex)))
 			{
 				// This should not be needed in the property grid, as it should already been called on the subobject.
 				OverridableManager.NotifyPropertyChange(Notification, SubObject, NextPropertyIterator, ChangeType);
@@ -1319,6 +1321,24 @@ void FOverriddenPropertySet::RemoveOverriddenSubProperties(FOverriddenPropertyNo
 	PropertyNode.SubPropertyNodeKeys.Empty();
 }
 
+UObject* FOverriddenPropertySet::TryGetInstancedSubObjectValue(const FObjectPropertyBase* FromProperty, void* ValuePtr) const
+{
+	// Property can be NULL - in that case there is no value.
+	if (!FromProperty)
+	{
+		return nullptr;
+	}
+
+	UObject* SubObject = FromProperty->GetObjectPropertyValue(ValuePtr);
+	if (FromProperty->HasAnyPropertyFlags(CPF_PersistentInstance)
+		|| (FromProperty->IsA<FObjectProperty>() && SubObject && SubObject->IsIn(Owner)))
+	{
+		return SubObject;
+	}
+
+	return nullptr;
+}
+
 EOverriddenPropertyOperation FOverriddenPropertySet::GetOverriddenPropertyOperation(FPropertyVisitorPath::Iterator PropertyIterator, bool* bOutInheritedOperation /*= nullptr*/) const
 {
 	return GetOverriddenPropertyOperation(OverriddenPropertyNodes.Find(RootNodeID), PropertyIterator, bOutInheritedOperation, Owner);
diff --git a/Engine/Source/Runtime/CoreUObject/Public/UObject/OverriddenPropertySet.h b/Engine/Source/Runtime/CoreUObject/Public/UObject/OverriddenPropertySet.h
--- a/Engine/Source/Runtime/CoreUObject/Public/UObject/OverriddenPropertySet.h
+++ b/Engine/Source/Runtime/CoreUObject/Public/UObject/OverriddenPropertySet.h
@@ -447,6 +447,8 @@ protected:
 
 	void RemoveOverriddenSubProperties(FOverriddenPropertyNode& PropertyIterator);
 
+	UObject* TryGetInstancedSubObjectValue(const FObjectPropertyBase* FromProperty, void* ValuePtr) const;
+
 private:
 	UPROPERTY()
 	TObjectPtr<UObject> Owner = nullptr;
