New disable package localization flag. Default true localization.EnablePackageRemapping, set to false to disable localization.

#rnx

diff --git a/Engine/Source/Runtime/CoreUObject/Private/Internationalization/PackageLocalizationManager.cpp b/Engine/Source/Runtime/CoreUObject/Private/Internationalization/PackageLocalizationManager.cpp
--- a/Engine/Source/Runtime/CoreUObject/Private/Internationalization/PackageLocalizationManager.cpp
+++ b/Engine/Source/Runtime/CoreUObject/Private/Internationalization/PackageLocalizationManager.cpp
@@ -1,6 +1,7 @@
 // Copyright Epic Games, Inc. All Rights Reserved.
 
 #include "Internationalization/PackageLocalizationManager.h"
+#include "HAL/IConsoleManager.h"
 #include "Misc/Paths.h"
 #include "Misc/PackageName.h"
 #include "Interfaces/IPluginManager.h"
@@ -10,6 +11,21 @@
 
 DEFINE_LOG_CATEGORY_STATIC(LogPackageLocalizationManager, Log, All);
 
+namespace PackageLocalizationInternal
+{
+	static bool GEnableLocalizationPackageRemapping = true;
+	static FAutoConsoleVariableRef CVar_EnableLocalizationPackageRemapping(
+		TEXT("localization.EnablePackageRemapping"),
+		GEnableLocalizationPackageRemapping,
+		TEXT("Disables identification of localization packages in order to improve startup time. Always false for the editor, optional (default true) for everything else."),
+		ECVF_Default);
+
+	bool ShouldSkipPackageRemapping()
+	{
+		return GIsEditor || !GEnableLocalizationPackageRemapping;
+	}
+}
+
 class FDefaultPackageLocalizationCache : public FPackageLocalizationCache
 {
 public:
@@ -57,6 +73,12 @@ void FDefaultPackageLocalizationCache::FindAssetGroupPackages(const FName InAsse
 
 void FPackageLocalizationManager::PerformLazyInitialization()
 {
+	if (PackageLocalizationInternal::ShouldSkipPackageRemapping())
+	{
+		UE_LOG(LogPackageLocalizationManager, Display, TEXT("Localization packages disabled."));
+		return;
+	}
+
 	if (!ActiveCache.IsValid() && LazyInitFunc)
 	{
 		LazyInitFunc(*this);
@@ -81,7 +103,7 @@ void FPackageLocalizationManager::InitializeFromCache(const TSharedRef<IPackageL
 	// Only preemptively attempt to conditionally update the cache outside of the editor where such things
 	// will happen almost immediately in a localized game, where as in the editor it's a bunch of work that
 	// likely won't be used until using some localization menus in the editor.
-	if (!GIsEditor)
+	if (PackageLocalizationInternal::ShouldSkipPackageRemapping())
 	{
 		ActiveCache->ConditionalUpdateCache();
 	}
@@ -94,6 +116,11 @@ void FPackageLocalizationManager::InitializeFromDefaultCache()
 
 FName FPackageLocalizationManager::FindLocalizedPackageName(const FName InSourcePackageName)
 {
+	if (PackageLocalizationInternal::ShouldSkipPackageRemapping())
+	{
+		return InSourcePackageName;
+	}
+
 	PerformLazyInitialization();
 
 	FName LocalizedPackageName;
@@ -117,6 +144,11 @@ FName FPackageLocalizationManager::FindLocalizedPackageName(const FName InSource
 
 FName FPackageLocalizationManager::FindLocalizedPackageNameForCulture(const FName InSourcePackageName, const FString& InCultureName)
 {
+	if (PackageLocalizationInternal::ShouldSkipPackageRemapping())
+	{
+		return InSourcePackageName;
+	}
+
 	PerformLazyInitialization();
 
 	FName LocalizedPackageName;
@@ -139,6 +171,11 @@ FName FPackageLocalizationManager::FindLocalizedPackageNameForCulture(const FNam
 
 FName FPackageLocalizationManager::FindLocalizedPackageNameNoCache(const FName InSourcePackageName, const FString& InCultureName) const
 {
+	if (PackageLocalizationInternal::ShouldSkipPackageRemapping())
+	{
+		return InSourcePackageName;
+	}
+
 	// Split the package name into its root and sub-path so that we can convert it into its localized variants for testing
 	FString PackageNameRoot;
 	FString PackageNameSubPath;
@@ -198,7 +235,7 @@ void FPackageLocalizationManager::InvalidateRootSourcePath(const FString& InRoot
 
 void FPackageLocalizationManager::ConditionalUpdateCache()
 {
-	if (ActiveCache.IsValid())
+	if (!PackageLocalizationInternal::ShouldSkipPackageRemapping() && ActiveCache.IsValid())
 	{
 		ActiveCache->ConditionalUpdateCache();
 	}
