Force disable UE_PACK_FUOBJECT_ITEM on ARM64 when FORCE_ANSI_ALLOCATOR is defined as it might be using MTE or HWASan and this feature will interfere with it

diff --git a/Engine/Source/Runtime/CoreUObject/Public/UObject/UObjectArray.h b/Engine/Source/Runtime/CoreUObject/Public/UObject/UObjectArray.h
--- a/Engine/Source/Runtime/CoreUObject/Public/UObject/UObjectArray.h
+++ b/Engine/Source/Runtime/CoreUObject/Public/UObject/UObjectArray.h
@@ -22,7 +22,12 @@ namespace Verse
 #endif
 
 #ifndef UE_PACK_FUOBJECT_ITEM
-#	define UE_PACK_FUOBJECT_ITEM 0
+#	define UE_ENABLE_FUOBJECT_ITEM_PACKING 0
+#elif defined(UE_PACK_FUOBJECT_ITEM) && defined(PLATFORM_CPU_ARM_FAMILY) && defined(FORCE_ANSI_ALLOCATOR)
+//	disable packing on ARM64 is we use Ansi allocator as it might use MTE or HWASan which uses top byte of a pointer and this feature will discard that
+#	define UE_ENABLE_FUOBJECT_ITEM_PACKING 0
+#elif
+#	define UE_ENABLE_FUOBJECT_ITEM_PACKING UE_PACK_FUOBJECT_ITEM
 #endif
 
 /**
@@ -46,7 +51,7 @@ struct FUObjectItem
 	friend class FUObjectArray;
 	friend class UE::GC::Private::FGCFlags;
 
-#if !UE_PACK_FUOBJECT_ITEM
+#if !UE_ENABLE_FUOBJECT_ITEM_PACKING
 	// Pointer to the allocated object
 	UE_DEPRECATED(5.5, "Use GetObject() and SetObject() to access Object.")
 	class UObjectBase* Object;
@@ -95,7 +100,7 @@ public:
 #endif // STATS || ENABLE_STATNAMEDEVENTS
 
 	FUObjectItem()
-#if !UE_PACK_FUOBJECT_ITEM
+#if !UE_ENABLE_FUOBJECT_ITEM_PACKING
 		: Object(nullptr)
 #else
 		: ObjectPtrLow(0)
@@ -128,7 +133,7 @@ public:
 
 	FORCEINLINE class UObjectBase* GetObject() const
 	{
-#if UE_PACK_FUOBJECT_ITEM
+#if UE_ENABLE_FUOBJECT_ITEM_PACKING
 		const uintptr_t Object = (uintptr_t(Flags & PtrMask) << (32 + UObjectPtrTrailingZeroes)) | (uintptr_t(ObjectPtrLow) << UObjectPtrTrailingZeroes);
 		return (class UObjectBase*)Object;
 #else
@@ -140,7 +145,7 @@ public:
 
 	FORCEINLINE void SetObject(class UObjectBase* InObject)
 	{
-#if UE_PACK_FUOBJECT_ITEM
+#if UE_ENABLE_FUOBJECT_ITEM_PACKING
 		const uint32 ObjectPtrHi = ((uintptr_t)InObject & 0xFFFF00000000ULL) >> (32 + UObjectPtrTrailingZeroes);
 		Flags |= ObjectPtrHi;
 		ObjectPtrLow = ((uintptr_t)InObject >> UObjectPtrTrailingZeroes) & 0xFFFFFFFF;
@@ -193,7 +198,7 @@ public:
 
 	FORCEINLINE EInternalObjectFlags GetFlags() const
 	{
-#if UE_PACK_FUOBJECT_ITEM
+#if UE_ENABLE_FUOBJECT_ITEM_PACKING
 		return EInternalObjectFlags(GetFlagsInternal() & FlagsMask);
 #else
 		return EInternalObjectFlags(GetFlagsInternal());
