Skip creating linker load export in ZenLoader for exports which have a null class index. This can happen when a class has been deleted and it is known missing through an InstanceOnly=true, Removed=true ClassRedirect

diff --git a/Engine/Source/Runtime/CoreUObject/Private/Serialization/AsyncLoading2.cpp b/Engine/Source/Runtime/CoreUObject/Private/Serialization/AsyncLoading2.cpp
--- a/Engine/Source/Runtime/CoreUObject/Private/Serialization/AsyncLoading2.cpp
+++ b/Engine/Source/Runtime/CoreUObject/Private/Serialization/AsyncLoading2.cpp
@@ -6861,6 +6861,20 @@ bool FAsyncPackage2::CreateLinkerLoadExports(FAsyncLoadingThreadState2& ThreadSt
 #endif // WITH_METADATA
 		FObjectExport& LinkerExport = LinkerLoadState->Linker->ExportMap[ExportIndex];
 		FExportObject& ExportObject = Data.Exports[ExportIndex];
+
+		if (LinkerExport.ClassIndex.IsNull())
+		{
+			// During FixupExportMap, LinkerLoad may null class indices for exports whose class is FCoreRedirects::IsKnownMissing. We cannot load these types, and since 
+			// it would have been an error to serialize these types in the first place, we can assume if the class is null, it is a missing type we can ignore.
+
+			// Fixup will also null the OuterIndex and ObjectName. Assert these are as expected, and that we aren't just hitting some corrupt state.
+			checkf(LinkerExport.OuterIndex.IsNull() && LinkerExport.ObjectName == NAME_None, TEXT("Trying to create a linker load export that has a null class but non null outer or name"));
+
+			ExportObject.bWasFoundInMemory = false;
+			ExportObject.Object = nullptr;
+			continue;
+		}
+
 		if (UObject* Object = LinkerLoadState->Linker->CreateExport(ExportIndex))
 		{
 			checkf(!Object->IsUnreachable(), TEXT("Trying to store an unreachable object '%s' in the import store"), *Object->GetFullName());
