Minor optimization to save package now that FPackageName::GetPackageMountPoint can take a FStringView

diff --git a/Engine/Source/Runtime/CoreUObject/Private/UObject/SavePackage2.cpp b/Engine/Source/Runtime/CoreUObject/Private/UObject/SavePackage2.cpp
--- a/Engine/Source/Runtime/CoreUObject/Private/UObject/SavePackage2.cpp
+++ b/Engine/Source/Runtime/CoreUObject/Private/UObject/SavePackage2.cpp
@@ -1023,7 +1023,7 @@ ESavePackageResult ValidateImports(FSaveContext& SaveContext)
 		{
 			FName MountPointName = FPackageName::GetPackageMountPoint(PackageName);
 
-			FName ImportMountPointName = FPackageName::GetPackageMountPoint(ImportPackage->GetName());
+			FName ImportMountPointName = FPackageName::GetPackageMountPoint(FNameBuilder(ImportPackage->GetFName()).ToView());
 
 			if (!MountPointName.IsNone() && !ImportMountPointName.IsNone())
 			{
