Unreal.nativs: Make GObjectArrayForDebugVisualizers based visualizers like UObjectBase work on more platforms.
Some platforms were stripping the pointer reference GObjectArrayForDebugVisualizers that was initialized but otherwise unused in native code.
Fix broken visualization of TWeakObjectPtr, and a few other usages of GObjectArrayForDebugVisualizers/_GetObjectItem().

#rnx

diff --git a/Engine/Source/Runtime/CoreUObject/Private/UObject/UObjectArray.cpp b/Engine/Source/Runtime/CoreUObject/Private/UObject/UObjectArray.cpp
--- a/Engine/Source/Runtime/CoreUObject/Private/UObject/UObjectArray.cpp
+++ b/Engine/Source/Runtime/CoreUObject/Private/UObject/UObjectArray.cpp
@@ -108,7 +108,10 @@ FUObjectArray::FUObjectArray()
 , OpenForDisregardForGC(true)
 , PrimarySerialNumber(START_SERIAL_NUMBER)
 {
-	GCoreObjectArrayForDebugVisualizers = &GUObjectArray.ObjObjects;
+	// Set the global state using the (per-module) pointer reference rather than the GCoreObjectArrayForDebugVisualizers pointer directly
+	// in order to make sure the reference does not get stripped in monolithic (non-modular) builds on any platform.
+	extern FChunkedFixedUObjectArray*& GObjectArrayForDebugVisualizers;
+	GObjectArrayForDebugVisualizers = &GUObjectArray.ObjObjects;
 }
 
 void FUObjectArray::AllocateObjectPool(int32 InMaxUObjects, int32 InMaxObjectsNotConsideredByGC, bool bPreAllocateObjectArray)
