Remote Objects:  Changing the strategy for PostMigration of Actors/Components to be as minimalistic as possible.  The Components now register themselves, rather than the Actor and the goal of PostMigrate is to respect the deserialized state, and re-register in-place.  This may require many more PostMigrate functions to be written per ActorComponent to get this right.
Remote Objects:  Prevent sending Attached Children that are going to become remote references.  This case would barf during the PostMigration, requiring more migrations to complete.  I'm going to submit this hack for now with the goal of fixing it properly soon.

#rnx

diff --git a/Engine/Source/Runtime/CoreUObject/Private/UObject/RemoteObjectTransfer.cpp b/Engine/Source/Runtime/CoreUObject/Private/UObject/RemoteObjectTransfer.cpp
--- a/Engine/Source/Runtime/CoreUObject/Private/UObject/RemoteObjectTransfer.cpp
+++ b/Engine/Source/Runtime/CoreUObject/Private/UObject/RemoteObjectTransfer.cpp
@@ -291,7 +291,9 @@ UObject* MigrateObjectFromObjectData(FRemoteObjectData& ObjectData, const FUObje
 		const bool bSameNumberOfObjects = (1 + AllSubObjects.Num() == ReceivedObjects.Num());
 		if (!bSameNumberOfObjects)
 		{
-			UE_LOG(LogRemoteObject, Warning, TEXT("We received migrated '%s' but we received %d Objects whereas %d Objects now live locally.  Differences:"), *GetNameSafe(RootObject), ReceivedObjects.Num(), 1 + AllSubObjects.Num());
+			FRemoteServerId LocalServerId = UE::RemoteObject::GetGlobalServerId();
+
+			UE_LOG(LogRemoteObject, Warning, TEXT("While migrating '%s' [%s -> %s] we received %d Objects whereas %d Objects got instantiated after deserialization.  Differences:"), *GetNameSafe(RootObject), *MigrationContext.RemoteServerId.ToString(), *LocalServerId.ToString(), ReceivedObjects.Num(), 1 + AllSubObjects.Num());
 			TArray<UObject*> ReceivedButNotInstanced;
 			for (UObject* Obj : ReceivedObjects)
 			{
@@ -303,9 +305,9 @@ UObject* MigrateObjectFromObjectData(FRemoteObjectData& ObjectData, const FUObje
 			// We'll always have the root object in this list due to AllSubObjects not containing the root object
 			ReceivedButNotInstanced.RemoveAtSwap(0);
 
-			auto WrappedGetNameSafe = [](const UObject* O) { return GetNameSafe(O); };
-			UE_CLOG(!ReceivedButNotInstanced.IsEmpty(), LogRemoteObject, Warning, TEXT(" Received but not instanced: %s"), *FString::JoinBy(ReceivedButNotInstanced, TEXT(", "), WrappedGetNameSafe));
-			UE_CLOG(!AllSubObjects.IsEmpty(), LogRemoteObject, Warning, TEXT(" Existing instance but not received: %s"), *FString::JoinBy(AllSubObjects, TEXT(", "), WrappedGetNameSafe));
+			auto WrappedGetNameSafe = [](const UObject* O) { return FString::Printf(TEXT("%s (%s)"), *GetNameSafe(O), *(FObjectPtr{ const_cast<UObject*>(O) }.GetRemoteId().ToString())); };
+			UE_CLOG(!ReceivedButNotInstanced.IsEmpty(), LogRemoteObject, Warning, TEXT(" Received objects but they weren't instanced: %s"), *FString::JoinBy(ReceivedButNotInstanced, TEXT(", "), WrappedGetNameSafe));
+			UE_CLOG(!AllSubObjects.IsEmpty(), LogRemoteObject, Warning, TEXT(" New (or Existing) instances not received: %s"), *FString::JoinBy(AllSubObjects, TEXT(", "), WrappedGetNameSafe));
 
 			if (CVarRemoveMissingMigratedComponents.GetValueOnGameThread())
 			{
@@ -561,7 +563,7 @@ void MigrateObjectToRemoteServerWithExplicitPriority(FRemoteTransactionId Reques
 	}
 	else
 	{
-		UE_LOG(LogRemoteObject, VeryVerbose, TEXT("     : MigrateObjectToRemoteServer(%s) serializing obj %s"), *RequestId.ToString(), *ObjectId.ToString(), Handle::IsRemote(Object), *OwnerServerId.ToString(), *DestinationServerId.ToString());
+		UE_LOG(LogRemoteObject, VeryVerbose, TEXT("     : MigrateObjectToRemoteServer(%s) serializing obj %s"), *RequestId.ToString(), *ObjectId.ToString());
 
 		// send the object
 		TSet<UObject*> SentObjects;
@@ -757,13 +759,19 @@ void StoreObjectToDatabase(UObject* Object)
 			FRemoteObjectId RootObjectId = FRemoteObjectId { OutermostUnreachableRemoteReference };
 			FRemoteServerId LocalServerId = GetGlobalServerId();
 			FRemoteServerId OwnerServerId = UE::RemoteObject::Handle::GetOwnerServerId(Object);
-			if (!ensureMsgf(LocalServerId == OwnerServerId, TEXT("%hs attempted to store non-locally-owned Object %s to Server %s's Database"), __func__, *GetNameSafe(Object), *LocalServerId.ToString()))
+			if (LocalServerId != OwnerServerId)
 			{
-				// This shouldn't happen obviously, but we're currently hitting this case.  So if it happens,
-				// don't store in the database, instead fetch it back from the Owning Server if needed
+				// We don't own this, and it's definitely remotely owned, just don't save it to the database
+				if (UE::RemoteObject::Handle::IsRemote(Object))
+				{
+					return;
+				}
+
+				ensureMsgf(false, TEXT("%hs was about to store non-locally-owned Object %s to Server %s's Database; making it remote instead"), __func__, *GetNameSafe(Object), *LocalServerId.ToString());
 				MarkAsRemote(Object, OwnerServerId); 
 				return;
 			}
+
 			// Setup a migration context for "push migrating to the new owner: the Database Server"
 			FUObjectMigrationContext StoreToDatabaseContext{ .ObjectId = RootObjectId, .RemoteServerId = DatabaseId, .OwnerServerId = DatabaseId, .PhysicsServerId = DatabaseId, .RequestId = PushMigrationRequestId, .MigrationSide = EObjectMigrationSide::Send };
 
diff --git a/Engine/Source/Runtime/Engine/Classes/Components/ActorComponent.h b/Engine/Source/Runtime/Engine/Classes/Components/ActorComponent.h
--- a/Engine/Source/Runtime/Engine/Classes/Components/ActorComponent.h
+++ b/Engine/Source/Runtime/Engine/Classes/Components/ActorComponent.h
@@ -1170,6 +1170,9 @@ public:
 	ENGINE_API virtual bool CallRemoteFunction( UFunction* Function, void* Parameters, FOutParmRec* OutParms, FFrame* Stack ) override;
 	ENGINE_API virtual void PostInitProperties() override;
 	ENGINE_API virtual void PostLoad() override;
+#if UE_WITH_REMOTE_OBJECT_HANDLE
+	ENGINE_API virtual void PostMigrate(const struct FUObjectMigrationContext& MigrationContext) override;
+#endif
 	ENGINE_API virtual bool Rename( const TCHAR* NewName=NULL, UObject* NewOuter=NULL, ERenameFlags Flags=REN_None ) override;
 	ENGINE_API virtual void PostRename(UObject* OldOuter, const FName OldName) override;
 	ENGINE_API virtual void Serialize(FArchive& Ar) override;
diff --git a/Engine/Source/Runtime/Engine/Classes/Components/SkeletalMeshComponent.h b/Engine/Source/Runtime/Engine/Classes/Components/SkeletalMeshComponent.h
--- a/Engine/Source/Runtime/Engine/Classes/Components/SkeletalMeshComponent.h
+++ b/Engine/Source/Runtime/Engine/Classes/Components/SkeletalMeshComponent.h
@@ -1766,6 +1766,9 @@ public:
 	ENGINE_API virtual void Serialize(FArchive& Ar) override;
 	ENGINE_API virtual void PostLoad() override;
 	ENGINE_API virtual void PostInitProperties() override;
+#if UE_WITH_REMOTE_OBJECT_HANDLE
+	ENGINE_API virtual void PostMigrate(const struct FUObjectMigrationContext& MigrationContext) override;
+#endif
 #if WITH_EDITOR
 	DECLARE_MULTICAST_DELEGATE(FOnSkeletalMeshPropertyChangedMulticaster)
 	FOnSkeletalMeshPropertyChangedMulticaster OnSkeletalMeshPropertyChanged;
diff --git a/Engine/Source/Runtime/Engine/Private/Actor.cpp b/Engine/Source/Runtime/Engine/Private/Actor.cpp
--- a/Engine/Source/Runtime/Engine/Private/Actor.cpp
+++ b/Engine/Source/Runtime/Engine/Private/Actor.cpp
@@ -1295,23 +1295,20 @@ void AActor::PostMigrate(const FUObjectMigrationContext& MigrationContext)
 		EObjectMigrationRecvType MigrationRecvType = MigrationContext.GetObjectMigrationRecvType(this);
 		UE_LOG(LogActor, Verbose, TEXT("%hs: Recv %s (LocalSrvId: %s. DstSrvId: %s. HasOwnership: %d. MigrationType: %s)"), __func__, *GetName(), *LocalServerId.ToString(), *MigrationContext.RemoteServerId.ToString(), bHasOwnership, ToString(MigrationRecvType));
 
-		// Add the actor that was migrated from another server to this server's world
-		constexpr bool bOnlyAddIfUnique = true;
-		GetLevel()->TryAddActorToList(this, bOnlyAddIfUnique);
-
+		// We can get rid of this check further down the road...
 		// Set the actor's world transform if it has a native rootcomponent.
 		const USceneComponent* SceneRootComponent = FixupNativeActorComponents(this);
 		ensureMsgf(SceneRootComponent == RootComponent, TEXT("%hs: We were expecting RootComponent (%s) to be equal to FixupNativeActorComponents (%s). Transform fix-up code would be required to handle this case."), __func__, *GetNameSafe(RootComponent), *GetNameSafe(SceneRootComponent));
 
+		// Add the actor that was migrated from another server to this server's world
+		constexpr bool bOnlyAddIfUnique = true;
+		GetLevel()->TryAddActorToList(this, bOnlyAddIfUnique);
+
 		// If we don't have Ownership, these components were never registered
 		if (!bHasOwnership)
 		{
+			// This will build up the list of OwnedComponents (even if they are not currently registered)
 			ResetOwnedComponents();
-			RegisterAllComponents();
-
-			PreInitializeComponents();
-			InitializeComponents();
-			PostInitializeComponents();
 		}
 
 		bHasFinishedSpawning = true;
@@ -1325,7 +1322,9 @@ void AActor::PostMigrate(const FUObjectMigrationContext& MigrationContext)
 			BeginReplication();
 			#endif // UE_WITH_IRIS
 
-			RegisterAllActorTickFunctions(true, true);
+			// Each component will be in charge of registering itself in its own PostMigrate
+			constexpr bool bRegisterComponents = false;
+			RegisterAllActorTickFunctions(true, bRegisterComponents);
 
 			// We are going to assume that the Physics queries on the previous owning server had taken multi-server physics into account
 			// thus we do not actually need to check for Initial Overlaps anymore
diff --git a/Engine/Source/Runtime/Engine/Private/Components/ActorComponent.cpp b/Engine/Source/Runtime/Engine/Private/Components/ActorComponent.cpp
--- a/Engine/Source/Runtime/Engine/Private/Components/ActorComponent.cpp
+++ b/Engine/Source/Runtime/Engine/Private/Components/ActorComponent.cpp
@@ -57,6 +57,10 @@
 #include "Net/Iris/ReplicationSystem/ReplicationSystemUtil.h"
 #endif // UE_WITH_IRIS
 
+#if UE_WITH_REMOTE_OBJECT_HANDLE
+#include "UObject/UObjectMigrationContext.h"
+#endif
+
 #include UE_INLINE_GENERATED_CPP_BY_NAME(ActorComponent)
 
 #define LOCTEXT_NAMESPACE "ActorComponent"
@@ -680,6 +684,69 @@ void UActorComponent::PostLoad()
 #endif // WITH_EDITOR
 }
 
+#if UE_WITH_REMOTE_OBJECT_HANDLE
+void UActorComponent::PostMigrate(const struct FUObjectMigrationContext& MigrationContext)
+{
+	if (MigrationContext.MigrationSide == EObjectMigrationSide::Send)
+	{
+		EObjectMigrationSendType SendType = MigrationContext.GetObjectMigrationSendType(this);
+		const bool bKeepOwnership = (SendType == EObjectMigrationSendType::Loan);
+		if (!bKeepOwnership)
+		{
+			// If we're transferring ownership, make sure we don't keep anything registered locally
+			if (IsRegistered())
+			{
+				UnregisterComponent();
+			}
+
+			// Prevent tick functions from happening this frame (unregistering just prevents it from happening *next* frame)
+			SetComponentTickEnabled(false);
+		}
+	}
+	else if (ensureMsgf(MigrationContext.MigrationSide == EObjectMigrationSide::Receive, TEXT("Invalid MigrationContext: Side was not Send or Receive")))
+	{
+		EObjectMigrationRecvType RecvType = MigrationContext.GetObjectMigrationRecvType(this);
+
+		// We are expecting this flow from AActor, but now we must do this per-component in such a way
+		// that we are not mutating our previously deserialized data.
+		// RegisterAllComponents();
+		// PreInitializeComponents();
+		// InitializeComponents();
+		// PostInitializeComponents();
+
+		if (RecvType != EObjectMigrationRecvType::ReturnedLoan)
+		{
+			// If we did not previously have this Component, we need to register it (RegisterAllComponents)
+			// However, since that causes a bunch of side effects which have already happened at some point
+			// let's skip right ahead to OnRegister() in a state that causes the least amount of side-effects:
+			if (IsRegistered())
+			{
+				bool bRestoreAutoActivate = bAutoActivate;
+
+				WorldPrivate = GetOwner()->GetWorld();
+				RegistrationState = EComponentRegistrationState::PreRegistered;
+				bAutoActivate = false;
+				bRegistered = false;
+				{
+					OnRegister();
+				}
+				if (bPhysicsStateCreated)
+				{
+					bPhysicsStateCreated = false;
+					OnCreatePhysicsState();
+				}
+				bAutoActivate = bRestoreAutoActivate;
+			}
+		}
+
+		if (RecvType == EObjectMigrationRecvType::AssignedOwnership)
+		{
+			SetComponentTickEnabled(bIsActive);
+		}
+	}
+}
+#endif
+
 bool UActorComponent::Rename( const TCHAR* InName, UObject* NewOuter, ERenameFlags Flags )
 {
 	bRoutedPostRename = false;
@@ -3031,16 +3098,24 @@ void UActorComponent::Serialize(FArchive& Ar)
 
 		if (Ar.IsSaving())
 		{
-			ComponentFlags |= bHasBeenCreated	<< 0;
-			ComponentFlags |= bHasBegunPlay		<< 1;
+			ComponentFlags |= (bRegistered & 1)					<< 0;
+			ComponentFlags |= (bHasBeenInitialized & 1)			<< 1;
+			ComponentFlags |= (bTickFunctionsRegistered & 1)	<< 2;
+			ComponentFlags |= (bHasBegunPlay & 1)				<< 3;
+			ComponentFlags |= (bIsReadyForReplication & 1) 		<< 4;
+			ComponentFlags |= (bPhysicsStateCreated & 1) 		<< 5;
 		}
 
 		Ar << ComponentFlags;
 
 		if (Ar.IsLoading())
 		{
-			bHasBeenCreated	= (ComponentFlags >> 0) & 1;
-			bHasBegunPlay	= (ComponentFlags >> 1) & 1;
+			bRegistered 				= (ComponentFlags >> 0) & 1;
+			bHasBeenInitialized 		= (ComponentFlags >> 1) & 1;
+			bTickFunctionsRegistered 	= (ComponentFlags >> 2) & 1;
+			bHasBegunPlay				= (ComponentFlags >> 3) & 1;
+			bIsReadyForReplication 		= (ComponentFlags >> 4) & 1;
+			bPhysicsStateCreated 		= (ComponentFlags >> 5) & 1;
 		}
 	}
 #endif // UE_WITH_REMOTE_OBJECT_HANDLE
diff --git a/Engine/Source/Runtime/Engine/Private/Components/SceneComponent.cpp b/Engine/Source/Runtime/Engine/Private/Components/SceneComponent.cpp
--- a/Engine/Source/Runtime/Engine/Private/Components/SceneComponent.cpp
+++ b/Engine/Source/Runtime/Engine/Private/Components/SceneComponent.cpp
@@ -3689,6 +3689,28 @@ void USceneComponent::Serialize(FArchive& Ar)
 {
 	Ar.UsingCustomVersion(FUE5SpecialProjectStreamObjectVersion::GUID);
 
+#if UE_WITH_REMOTE_OBJECT_HANDLE
+	TArray<TObjectPtr<USceneComponent>> RestoreChildren;
+	if (Ar.IsMigratingRemoteObjects() && Ar.IsSaving())
+	{
+		for (int Index = AttachChildren.Num() - 1 ; Index >= 0 ; --Index)
+		{
+			TObjectPtr<USceneComponent>& Child = AttachChildren[Index];
+
+			if (ensureMsgf(!Child.IsRemote(), TEXT("About to save a Remote Attached Child in %s"), *GetFullNameSafe(this)))
+			{
+				if (Child->GetOwner() != GetOwner())
+				{
+					UE_LOG(LogSceneComponent, Warning, TEXT("Saving a Child %s which is not within the same hierarchy as %s (would become a RemoteReference - not sending it)"), *Child->GetName(), *GetName());
+
+					RestoreChildren.Add(Child);
+					AttachChildren.RemoveAtSwap(Index);
+				}
+			}
+		}
+	}
+#endif
+
 	Super::Serialize(Ar);
 
 	if (bComputeBoundsOnceForGame)
@@ -3704,6 +3726,22 @@ void USceneComponent::Serialize(FArchive& Ar)
 			}
 		}
 	}
+
+#if UE_WITH_REMOTE_OBJECT_HANDLE
+	if (Ar.IsMigratingRemoteObjects())
+	{
+		Ar << ComponentToWorld;
+
+		AttachChildren.Append(RestoreChildren);
+		if (Ar.IsLoading())
+		{
+			for (TObjectPtr<USceneComponent>& Child : AttachChildren)
+			{
+				ensureMsgf(!Child.IsRemote(), TEXT("Loaded a Remote Attached Child in %s"), *GetFullNameSafe(this));
+			}
+		}
+	}
+#endif
 }
 
 void USceneComponent::GetLifetimeReplicatedProps(TArray< FLifetimeProperty > & OutLifetimeProps) const
diff --git a/Engine/Source/Runtime/Engine/Private/Components/SkeletalMeshComponent.cpp b/Engine/Source/Runtime/Engine/Private/Components/SkeletalMeshComponent.cpp
--- a/Engine/Source/Runtime/Engine/Private/Components/SkeletalMeshComponent.cpp
+++ b/Engine/Source/Runtime/Engine/Private/Components/SkeletalMeshComponent.cpp
@@ -662,6 +662,19 @@ void USkeletalMeshComponent::PostInitProperties()
 	}
 }
 
+#if UE_WITH_REMOTE_OBJECT_HANDLE
+void USkeletalMeshComponent::PostMigrate(const struct FUObjectMigrationContext& MigrationContext)
+{
+	// When OnRegister gets called from PostMigrate, tell it to not touch the AnimInstance by setting our mode to Custom
+	EAnimationMode::Type RestoreAnimationMode = AnimationMode;
+	AnimationMode = EAnimationMode::AnimationCustomMode;
+
+	Super::PostMigrate(MigrationContext);
+
+	AnimationMode = RestoreAnimationMode;
+}
+#endif
+
 void USkeletalMeshComponent::RegisterComponentTickFunctions(bool bRegister)
 {
 	Super::RegisterComponentTickFunctions(bRegister);
