Fix LoadAssetAsync (and the soft object wrappers) to correctly handle UObjectRedirector, it will now return the redirected asset like the sync load functions

diff --git a/Engine/Source/Runtime/CoreUObject/Private/UObject/UObjectGlobals.cpp b/Engine/Source/Runtime/CoreUObject/Private/UObject/UObjectGlobals.cpp
--- a/Engine/Source/Runtime/CoreUObject/Private/UObject/UObjectGlobals.cpp
+++ b/Engine/Source/Runtime/CoreUObject/Private/UObject/UObjectGlobals.cpp
@@ -1551,15 +1551,16 @@ int32 LoadAssetAsync(FTopLevelAssetPath InAssetPath, FLoadAssetAsyncDelegate InC
 {
 	// Asset paths should always have a valid package
 	FPackagePath PackagePath = FPackagePath::FromPackageNameChecked(InAssetPath.GetPackageName());
+	uint32 LoadFlags = InOptionalParams.LoadFlags;
 	
 	FLoadPackageAsyncOptionalParams PackageParams{
 		.PackagePriority = InOptionalParams.PackagePriority,
 		.InstancingContext = InOptionalParams.InstancingContext,
-		.LoadFlags = InOptionalParams.LoadFlags
+		.LoadFlags = LoadFlags
 	};
 
 	PackageParams.CompletionDelegate = MakeUnique<FLoadPackageAsyncDelegate>(FLoadPackageAsyncDelegate::CreateLambda(
-		[InAssetPath, CompletionDelegate = MoveTemp(InCompletionDelegate)](const FName& LoadedPackageName, UPackage* LoadedPackage, EAsyncLoadingResult::Type Result) mutable
+		[InAssetPath, LoadFlags, CompletionDelegate = MoveTemp(InCompletionDelegate)](const FName& LoadedPackageName, UPackage* LoadedPackage, EAsyncLoadingResult::Type Result) mutable
 		{
 			UObject* LoadedObject = nullptr;
 			if (Result == EAsyncLoadingResult::Succeeded && LoadedPackage)
@@ -1571,6 +1572,19 @@ int32 LoadAssetAsync(FTopLevelAssetPath InAssetPath, FLoadAssetAsyncDelegate InC
 				{
 					Result = EAsyncLoadingResult::Failed;
 				}
+				else if (!(LoadFlags & LOAD_NoRedirects))
+				{
+					// Follow any object redirectors
+					while (UObjectRedirector* Redirector = Cast<UObjectRedirector>(LoadedObject))
+					{
+						LoadedObject = Redirector->DestinationObject;
+						if (UE::GC::GIsIncrementalReachabilityPending && LoadedObject)
+						{
+							// This is handled in StaticFindObjectFast for the redirector itself
+							UE::GC::MarkAsReachable(LoadedObject);
+						}
+					}
+				}
 			}
 
 			CompletionDelegate.ExecuteIfBound(InAssetPath, LoadedObject, Result);
