Update `GetGatherableFieldsForType` to ignore interface `FProperty`s to avoid an assertion.

diff --git a/Engine/Source/Runtime/CoreUObject/Private/Serialization/PropertyLocalizationDataGathering.cpp b/Engine/Source/Runtime/CoreUObject/Private/Serialization/PropertyLocalizationDataGathering.cpp
--- a/Engine/Source/Runtime/CoreUObject/Private/Serialization/PropertyLocalizationDataGathering.cpp
+++ b/Engine/Source/Runtime/CoreUObject/Private/Serialization/PropertyLocalizationDataGathering.cpp
@@ -281,6 +281,10 @@ void FPropertyLocalizationDataGatherer::GatherLocalizationDataFromObjectFields(c
 	// Gather text from the property data
 	for (const FProperty* PropertyField : GatherableFieldsForType.Properties)
 	{
+		if (!Object->IsA(PropertyField->GetOwner<UClass>()))
+		{
+			continue;
+		}
 		const void* ValueAddress = PropertyField->ContainerPtrToValuePtr<void>(Object);
 		const void* DefaultValueAddress = (ArchetypeObject && ArchetypeObject->IsA(PropertyField->GetOwnerClass())) ? PropertyField->ContainerPtrToValuePtr<void>(ArchetypeObject) : nullptr;
 		GatherLocalizationDataFromChildTextProperties(PathToParent, PropertyField, ValueAddress, DefaultValueAddress, GatherTextFlags | (PropertyField->HasAnyPropertyFlags(CPF_EditorOnly) ? EPropertyLocalizationGathererTextFlags::ForceEditorOnly : EPropertyLocalizationGathererTextFlags::None));
