OFPE: Fix missing external entities in PIE.
Don't filter out external objects when serializing array of objects marked as ExperimentalExternalObjects when duplicating for PIE.

diff --git a/Engine/Source/Runtime/CoreUObject/Private/UObject/PropertyArray.cpp b/Engine/Source/Runtime/CoreUObject/Private/UObject/PropertyArray.cpp
--- a/Engine/Source/Runtime/CoreUObject/Private/UObject/PropertyArray.cpp
+++ b/Engine/Source/Runtime/CoreUObject/Private/UObject/PropertyArray.cpp
@@ -425,7 +425,7 @@ void FArrayProperty::SerializeItem(FStructuredArchive::FSlot Slot, void* Value,
 			//	- We have at least one external entity object.
 			// For now, this will be considered as a full array replacement.
 			FObjectProperty* ObjectProperty = CastField<FObjectProperty>(Inner);
-			const bool bFilterExternalObjects = ObjectProperty && HasAnyPropertyFlags(CPF_ExperimentalExternalObjects) && UnderlyingArchive.IsPersistent();
+			const bool bFilterExternalObjects = ObjectProperty && HasAnyPropertyFlags(CPF_ExperimentalExternalObjects) && UnderlyingArchive.IsPersistent() && !UnderlyingArchive.HasAnyPortFlags(PPF_DuplicateForPIE);
 
 			bool bReplaceArray = !Defaults || !UnderlyingArchive.DoDelta() || UnderlyingArchive.IsTransacting();
 
