Added FCoreUObjectDelegates::OnObjectRenamed which will be called when an object is renamed in both the editor and standalone game editor builds.

diff --git a/Engine/Source/Runtime/CoreUObject/Private/UObject/Obj.cpp b/Engine/Source/Runtime/CoreUObject/Private/UObject/Obj.cpp
--- a/Engine/Source/Runtime/CoreUObject/Private/UObject/Obj.cpp
+++ b/Engine/Source/Runtime/CoreUObject/Private/UObject/Obj.cpp
@@ -852,6 +852,12 @@ bool UObject::IsSelectedInEditor() const
 
 #endif // WITH_EDITOR
 
+void UObject::PostRename(UObject* OldOuter, const FName OldName)
+{
+#if WITH_EDITOR
+	FCoreUObjectDelegates::OnObjectRenamed.Broadcast(this, OldOuter, OldName);
+#endif
+}
 
 /**
 	Helper class for tracking the list of classes excluded on a certain target system (client/server)
diff --git a/Engine/Source/Runtime/CoreUObject/Private/UObject/UObjectGlobals.cpp b/Engine/Source/Runtime/CoreUObject/Private/UObject/UObjectGlobals.cpp
--- a/Engine/Source/Runtime/CoreUObject/Private/UObject/UObjectGlobals.cpp
+++ b/Engine/Source/Runtime/CoreUObject/Private/UObject/UObjectGlobals.cpp
@@ -202,6 +202,7 @@ FCoreUObjectDelegates::FOnPreObjectPropertyChanged FCoreUObjectDelegates::OnPreO
 FCoreUObjectDelegates::FOnObjectPropertyChanged FCoreUObjectDelegates::OnObjectPropertyChanged;
 TSet<UObject*> FCoreUObjectDelegates::ObjectsModifiedThisFrame;
 FCoreUObjectDelegates::FOnObjectModified FCoreUObjectDelegates::OnObjectModified;
+FCoreUObjectDelegates::FOnObjectRenamed FCoreUObjectDelegates::OnObjectRenamed;
 FCoreUObjectDelegates::FOnObjectTransacted FCoreUObjectDelegates::OnObjectTransacted;
 FCoreUObjectDelegates::FOnObjectsReplaced FCoreUObjectDelegates::OnObjectsReplaced;
 FCoreUObjectDelegates::FOnObjectsReplaced FCoreUObjectDelegates::OnObjectsReinstanced;
diff --git a/Engine/Source/Runtime/CoreUObject/Public/UObject/Object.h b/Engine/Source/Runtime/CoreUObject/Public/UObject/Object.h
--- a/Engine/Source/Runtime/CoreUObject/Public/UObject/Object.h
+++ b/Engine/Source/Runtime/CoreUObject/Public/UObject/Object.h
@@ -522,7 +522,7 @@ public:
 #endif // WITH_EDITOR
 
 	/** Called at the end of Rename(), but only if the rename was actually carried out */
-	virtual void PostRename(UObject* OldOuter, const FName OldName) {}
+	COREUOBJECT_API virtual void PostRename(UObject* OldOuter, const FName OldName);
 
 	/**
 	 * Called before duplication.
diff --git a/Engine/Source/Runtime/CoreUObject/Public/UObject/UObjectGlobals.h b/Engine/Source/Runtime/CoreUObject/Public/UObject/UObjectGlobals.h
--- a/Engine/Source/Runtime/CoreUObject/Public/UObject/UObjectGlobals.h
+++ b/Engine/Source/Runtime/CoreUObject/Public/UObject/UObjectGlobals.h
@@ -3153,6 +3153,16 @@ struct FCoreUObjectDelegates
 		}
 	}
 
+	/** Called when an object is renamed this gets called in both the editor and standalone game editor builds, for every object renamed */
+	DECLARE_MULTICAST_DELEGATE_ThreeParams(FOnObjectRenamed, UObject*, UObject*, FName);
+	static COREUOBJECT_API FOnObjectRenamed OnObjectRenamed;
+
+	/** Broadcast OnObjectRenamed event */
+	static void BroadcastOnObjectRenamed(UObject* Object, UObject* OldOuter, FName OldName)
+	{
+		OnObjectRenamed.Broadcast(Object, OldOuter, OldName);
+	}
+
 	/** Callback for an object being transacted */
 	DECLARE_MULTICAST_DELEGATE_TwoParams(FOnObjectTransacted, UObject*, const class FTransactionObjectEvent&);
 	static COREUOBJECT_API FOnObjectTransacted OnObjectTransacted;
