Fix serialization of function dependency in the artifact: a function dependency is serialized with the first field being the function name and then all the following fields being the function arguments.
I plan to add the dependency hashed value as the last field but in the case of function dependency, this last field will be interpreted as a function argument.
So to prevent this, I serialize the function arguments in an array. In the deserialization code, I handle both case (the old format where arguments are not in an array and the new format) so old artifacts can still be loaded.

#rnx
#incrementalcook

diff --git a/Engine/Source/Runtime/CoreUObject/Private/Cooker/CookDependency.cpp b/Engine/Source/Runtime/CoreUObject/Private/Cooker/CookDependency.cpp
--- a/Engine/Source/Runtime/CoreUObject/Private/Cooker/CookDependency.cpp
+++ b/Engine/Source/Runtime/CoreUObject/Private/Cooker/CookDependency.cpp
@@ -561,10 +561,12 @@ void FCookDependency::Save(FCbWriter& Writer) const
 		break;
 	case ECookDependency::Function:
 		Writer << FunctionData.Name;
+		Writer.BeginArray();
 		for (FCbFieldViewIterator Iter(FunctionData.Args); Iter; ++Iter)
 		{
 			Writer << Iter;
 		}
+		Writer.EndArray();
 		break;
 	case ECookDependency::TransitiveBuild:
 		Writer << TransitiveBuildData.PackageName;
@@ -635,8 +637,20 @@ bool FCookDependency::Load(FCbFieldView Value)
 		{
 			return false;
 		}
-		*this = FCookDependency::Function(LocalFuncName, FCbFieldIterator::CloneRange(Field));
-		FunctionData.Args.MakeOwned();
+
+		if (Field.IsArray())
+		{
+			FCbFieldViewIterator FunctionArgumentsIterator = Field.CreateViewIterator();
+			*this = FCookDependency::Function(LocalFuncName, FCbFieldIterator::CloneRange(FunctionArgumentsIterator));
+			FunctionData.Args.MakeOwned();
+			++Field;
+		}
+		else // If the arguments are not in an array, then this is the old format and all the fields until the end are function parameters
+		{
+			*this = FCookDependency::Function(LocalFuncName, FCbFieldIterator::CloneRange(Field));
+			FunctionData.Args.MakeOwned();
+		}
+
 		return true;
 	}
 	case ECookDependency::TransitiveBuild:
