Fix some inconsistencies between the PropertyFlags found in cpp and c#

- Add missing C# version of CPF_AllowSelfReference.
- Fixed ParsePropertyFlags now matching EPropertyFlags.
-- Updated some old deprecated properties that have now been taken into use again.
-- Added the missing CPF_TObjectPtr which was causing the array to be out of order.
-- Added a number of recently added properties that were never added to this function.
- Fix some typos in the code documentation for EPropertyFlags.

diff --git a/Engine/Source/Programs/Shared/EpicGames.Core/UnrealEngineTypes.cs b/Engine/Source/Programs/Shared/EpicGames.Core/UnrealEngineTypes.cs
--- a/Engine/Source/Programs/Shared/EpicGames.Core/UnrealEngineTypes.cs
+++ b/Engine/Source/Programs/Shared/EpicGames.Core/UnrealEngineTypes.cs
@@ -1265,6 +1265,11 @@ namespace EpicGames.Core
 		/// </summary>
 		ExperimentalNeverOverriden = 0x0800000000000000,
 
+		/// <summary>
+		/// Enables the instancing graph self referencing logic, delegates and verse function are already using this
+		/// </summary>
+		AllowSelfReference = 0x1000000000000000,
+
 		/// <summary>
 		/// All Native Access Specifier flags
 		/// </summary>
diff --git a/Engine/Source/Runtime/CoreUObject/Private/UObject/Obj.cpp b/Engine/Source/Runtime/CoreUObject/Private/UObject/Obj.cpp
--- a/Engine/Source/Runtime/CoreUObject/Private/UObject/Obj.cpp
+++ b/Engine/Source/Runtime/CoreUObject/Private/UObject/Obj.cpp
@@ -4323,10 +4323,10 @@ COREUOBJECT_API TArray<const TCHAR*> ParsePropertyFlags(EPropertyFlags InFlags)
 		TEXT("CPF_ZeroConstructor"),
 		TEXT("CPF_ReturnParm"),
 		TEXT("CPF_DisableEditOnTemplate"),
-		TEXT("0x0000000000001000"),
+		TEXT("CPF_NonNullable"),
 		TEXT("CPF_Transient"),
 		TEXT("CPF_Config"),
-		TEXT("0x0000000000008000"),
+		TEXT("CPF_RequiredParm"),
 		TEXT("CPF_DisableEditOnInstance"),
 		TEXT("CPF_EditConst"),
 		TEXT("CPF_GlobalConfig"),
@@ -4337,7 +4337,7 @@ COREUOBJECT_API TArray<const TCHAR*> ParsePropertyFlags(EPropertyFlags InFlags)
 		TEXT("0x0000000000800000"),
 		TEXT("CPF_SaveGame"),	
 		TEXT("CPF_NoClear"),
-		TEXT("0x0000000004000000"),
+		TEXT("CPF_Virtual"),
 		TEXT("CPF_ReferenceParm"),
 		TEXT("CPF_BlueprintAssignable"),
 		TEXT("CPF_Deprecated"),
@@ -4367,7 +4367,11 @@ COREUOBJECT_API TArray<const TCHAR*> ParsePropertyFlags(EPropertyFlags InFlags)
 		TEXT("CPF_NativeAccessSpecifierProtected"),
 		TEXT("CPF_NativeAccessSpecifierPrivate"),
 		TEXT("CPF_SkipSerialization"),
+		TEXT("CPF_TObjectPtr")
 		TEXT("CPF_ExperimentalOverridableLogic"),
+		TEXT("CPF_ExperimentalAlwaysOverriden"),
+		TEXT("CPF_ExperimentalNeverOverriden"),
+		TEXT("CPF_AllowSelfReference"),
 	};
 
 	uint64 Flags = InFlags;
diff --git a/Engine/Source/Runtime/CoreUObject/Public/UObject/ObjectMacros.h b/Engine/Source/Runtime/CoreUObject/Public/UObject/ObjectMacros.h
--- a/Engine/Source/Runtime/CoreUObject/Public/UObject/ObjectMacros.h
+++ b/Engine/Source/Runtime/CoreUObject/Public/UObject/ObjectMacros.h
@@ -464,7 +464,7 @@ enum EPropertyFlags : uint64
 	CPF_ExperimentalOverridableLogic	= 0x0200000000000000,	///< ****Experimental*** Property will use different logic to serialize knowing what changes are done against its default use the overridable information provided by the overridable manager on the object
 	CPF_ExperimentalAlwaysOverriden		= 0x0400000000000000,	///< ****Experimental*** Property should never inherit from the parent when using overridable serialization
 	CPF_ExperimentalNeverOverriden		= 0x0800000000000000,	///< ****Experimental*** Property should never be overridden when using overridable serialization
-	CPF_AllowSelfReference				= 0x1000000000000000,	///< Enables the intancing graph self referencing logic, delgates and verse function are already using this
+	CPF_AllowSelfReference				= 0x1000000000000000,	///< Enables the instancing graph self referencing logic, delegates and verse function are already using this
 };
 
 /** All Native Access Specifier flags */
