Renaming PPF_MigratingRemoteObjects to PPF_AvoidRemoteObjectMigration to better reflect its purpose
PPF_AvoidRemoteObjectMigration is set on all archives used by remote object serialization (including helper archives that collect references or store temporary data) but FArchive::IsMigratingRemoteObjects() will only return true if the archive is serializing object data for the purpose of saving or loading migrated object data.

diff --git a/Engine/Source/Runtime/Core/Public/UObject/PropertyPortFlags.h b/Engine/Source/Runtime/Core/Public/UObject/PropertyPortFlags.h
--- a/Engine/Source/Runtime/Core/Public/UObject/PropertyPortFlags.h
+++ b/Engine/Source/Runtime/Core/Public/UObject/PropertyPortFlags.h
@@ -92,8 +92,8 @@ enum EPropertyPortFlags
 	/** indicates that we're exporting properties for display in the property window. - used to hide EditHide items in collapsed structs */
 	PPF_PropertyWindow				= 0x00080000,
 
-	/** indicates that we're serializing remote objects for migration */
-	PPF_MigratingRemoteObjects		= 0x00100000,
+	/** Avoid dereferencing remote object references if possible */
+	PPF_AvoidRemoteObjectMigration	= 0x00100000,
 
 	/** Force fully qualified object names (for debug dumping) */
 	PPF_DebugDump					= 0x00200000,
diff --git a/Engine/Source/Runtime/Core/Public/UObject/ScriptDelegates.h b/Engine/Source/Runtime/Core/Public/UObject/ScriptDelegates.h
--- a/Engine/Source/Runtime/Core/Public/UObject/ScriptDelegates.h
+++ b/Engine/Source/Runtime/Core/Public/UObject/ScriptDelegates.h
@@ -5,7 +5,7 @@
 #include "Containers/Array.h"
 #include "Containers/ContainerAllocationPolicies.h"
 #include "Containers/UnrealString.h"
-#include "CoreTypes.h"
+#include "PropertyPortFlags.h"
 #include "Delegates/DelegateAccessHandler.h"
 #include "Misc/AssertionMacros.h"
 #include "Templates/SharedPointer.h"
@@ -858,7 +858,7 @@ public:
 			// which may be pointing to objects that haven't been migrated yet
 			if (Ar.IsLoading() 
 #if UE_WITH_REMOTE_OBJECT_HANDLE
-				&& !Ar.IsMigratingRemoteObjects()
+				&& !Ar.HasAnyPortFlags(PPF_AvoidRemoteObjectMigration)
 #endif
 				)
 			{
diff --git a/Engine/Source/Runtime/CoreUObject/Private/UObject/Obj.cpp b/Engine/Source/Runtime/CoreUObject/Private/UObject/Obj.cpp
--- a/Engine/Source/Runtime/CoreUObject/Private/UObject/Obj.cpp
+++ b/Engine/Source/Runtime/CoreUObject/Private/UObject/Obj.cpp
@@ -1937,7 +1937,7 @@ void UObject::SerializeScriptProperties( FStructuredArchive::FSlot Slot ) const
 
 		// When migrating remote objects the only instance where we serialize a CDO is when resetting an object to its archetype state in which case
 		// we want to serialize against the actual object class to get the right delta
-		UStruct* DiffClass = (HasAnyFlags(RF_ClassDefaultObject) && !UnderlyingArchive.IsMigratingRemoteObjects()) ? ObjClass->GetSuperClass() : ObjClass;
+		UStruct* DiffClass = (HasAnyFlags(RF_ClassDefaultObject) && !UnderlyingArchive.HasAnyPortFlags(PPF_AvoidRemoteObjectMigration)) ? ObjClass->GetSuperClass() : ObjClass;
 
 		// Query if this object data is being impersonated 
 		const UObject* ThisObject = this;
diff --git a/Engine/Source/Runtime/CoreUObject/Private/UObject/PropertyMulticastDelegate.cpp b/Engine/Source/Runtime/CoreUObject/Private/UObject/PropertyMulticastDelegate.cpp
--- a/Engine/Source/Runtime/CoreUObject/Private/UObject/PropertyMulticastDelegate.cpp
+++ b/Engine/Source/Runtime/CoreUObject/Private/UObject/PropertyMulticastDelegate.cpp
@@ -148,8 +148,9 @@ bool FMulticastDelegateProperty::Identical( const void* A, const void* B, uint32
 		const FMulticastScriptDelegate::UnicastDelegateType& BindingB = DelegateB.InvocationList[CurInvocationIndex];
 
 #if UE_WITH_REMOTE_OBJECT_HANDLE
-		if ((PortFlags & PPF_MigratingRemoteObjects) != 0)
+		if ((PortFlags & PPF_AvoidRemoteObjectMigration) != 0)
 		{
+			// Before comparing the actual object pointers, compare their remote ids to prevent dereferencing objects that may be remote
 			if (BindingA.GetUObjectRef().GetRemoteId() != BindingB.GetUObjectRef().GetRemoteId())
 			{
 				return false;
diff --git a/Engine/Source/Runtime/CoreUObject/Private/UObject/PropertyObject.cpp b/Engine/Source/Runtime/CoreUObject/Private/UObject/PropertyObject.cpp
--- a/Engine/Source/Runtime/CoreUObject/Private/UObject/PropertyObject.cpp
+++ b/Engine/Source/Runtime/CoreUObject/Private/UObject/PropertyObject.cpp
@@ -177,7 +177,7 @@ void FObjectProperty::SerializeItem(FStructuredArchive::FSlot Slot, void* Value,
 			Slot << (*ObjectPtr);
 		}
 #if !(UE_BUILD_TEST || UE_BUILD_SHIPPING) 
-		if (!UnderlyingArchive.IsSaving() && !UnderlyingArchive.IsMigratingRemoteObjects())
+		if (!UnderlyingArchive.IsSaving() && !UnderlyingArchive.HasAnyPortFlags(PPF_AvoidRemoteObjectMigration))
 		{
 			CheckValidObject(Value, CurrentValue, Defaults);
 		}
@@ -186,7 +186,7 @@ void FObjectProperty::SerializeItem(FStructuredArchive::FSlot Slot, void* Value,
 	else
 	{
 #if UE_WITH_OBJECT_HANDLE_LATE_RESOLVE || UE_WITH_REMOTE_OBJECT_HANDLE
-		if (HasAnyPropertyFlags(CPF_TObjectPtr) || UnderlyingArchive.IsMigratingRemoteObjects())
+		if (HasAnyPropertyFlags(CPF_TObjectPtr) || UnderlyingArchive.HasAnyPortFlags(PPF_AvoidRemoteObjectMigration))
 		{
 			FObjectHandle OriginalHandle = ObjectPtr->GetHandle();
 			Slot << *ObjectPtr;
diff --git a/Engine/Source/Runtime/CoreUObject/Private/UObject/PropertyWeakObjectPtr.cpp b/Engine/Source/Runtime/CoreUObject/Private/UObject/PropertyWeakObjectPtr.cpp
--- a/Engine/Source/Runtime/CoreUObject/Private/UObject/PropertyWeakObjectPtr.cpp
+++ b/Engine/Source/Runtime/CoreUObject/Private/UObject/PropertyWeakObjectPtr.cpp
@@ -69,7 +69,7 @@ void FWeakObjectProperty::SerializeItem( FStructuredArchive::FSlot Slot, void* V
 
 	UObject* OldObjectValue = nullptr;
 #if UE_WITH_REMOTE_OBJECT_HANDLE
-	if (!UnderlyingArchive.IsMigratingRemoteObjects())
+	if (!UnderlyingArchive.HasAnyPortFlags(PPF_AvoidRemoteObjectMigration))
 #endif
 	{
 		OldObjectValue = GetObjectPropertyValue(Value);
@@ -79,7 +79,7 @@ void FWeakObjectProperty::SerializeItem( FStructuredArchive::FSlot Slot, void* V
 
 	if (
 #if UE_WITH_REMOTE_OBJECT_HANDLE
-	!UnderlyingArchive.IsMigratingRemoteObjects() &&
+	!UnderlyingArchive.HasAnyPortFlags(PPF_AvoidRemoteObjectMigration) &&
 #endif
 	   (UnderlyingArchive.IsLoading() || UnderlyingArchive.IsModifyingWeakAndStrongReferences()))
 	{
diff --git a/Engine/Source/Runtime/CoreUObject/Private/UObject/RemoteObjectSerialization.cpp b/Engine/Source/Runtime/CoreUObject/Private/UObject/RemoteObjectSerialization.cpp
--- a/Engine/Source/Runtime/CoreUObject/Private/UObject/RemoteObjectSerialization.cpp
+++ b/Engine/Source/Runtime/CoreUObject/Private/UObject/RemoteObjectSerialization.cpp
@@ -165,6 +165,7 @@ protected:
 
 	UObject* RootObject = nullptr;
 	FRemoteObjectData& ObjectData;
+	const FUObjectMigrationContext* MigrationContext = nullptr;
 	TArray<uint8> SerializedBytes;
 	FString ArchiveName;
 
@@ -253,15 +254,15 @@ protected:
 
 public:
 
-	TArchiveRemoteObjectBase(FRemoteObjectData& InObjectData, const TCHAR* InArchiveName)
+	TArchiveRemoteObjectBase(FRemoteObjectData& InObjectData, const FUObjectMigrationContext* InContext, const TCHAR* InArchiveName)
 		: T(SerializedBytes)
 		, ObjectData(InObjectData)
+		, MigrationContext(InContext)
 		, ArchiveName(InArchiveName)
 	{
 		T::SetIsPersistent(false);
 		T::SetUseUnversionedPropertySerialization(true);
-		T::SetMigratingRemoteObjects(true);
-		T::SetPortFlags(PPF_MigratingRemoteObjects);
+		T::SetPortFlags(PPF_AvoidRemoteObjectMigration);
 
 		if (T::IsLoading())
 		{
@@ -308,6 +309,11 @@ public:
 		return ArchiveName;
 	}
 
+	virtual const FUObjectMigrationContext* GetMigrationContext() const override
+	{
+		return MigrationContext;
+	}
+
 	void SetRootObject(UObject* InRoot)
 	{
 		RootObject = InRoot;
@@ -357,7 +363,6 @@ protected:
 	TMap<FName, FNameIndexType> NameMap;
 	TMap<UObject*, FNameIndexType> PathNameMap;
 	TMap<UObject*, UObject*>* ObjectToArchetypeMap = nullptr;
-	const FUObjectMigrationContext& MigrationContext;
 
 public:
 
@@ -459,9 +464,8 @@ private:
 
 public:
 
-	FArchiveRemoteObjectWriter(UObject* InRootObject, FRemoteObjectData& OutObjectData, const FUObjectMigrationContext& InMigrationContext, const TCHAR* ArchiveName = nullptr)
-		: TArchiveRemoteObjectBase<FMemoryWriter>(OutObjectData, ArchiveName ? ArchiveName : TEXT("RemoteObjectWriter")),
-		  MigrationContext(InMigrationContext)
+	FArchiveRemoteObjectWriter(UObject* InRootObject, FRemoteObjectData& OutObjectData, const FUObjectMigrationContext* InMigrationContext, const TCHAR* ArchiveName = nullptr)
+		: TArchiveRemoteObjectBase<FMemoryWriter>(OutObjectData, InMigrationContext, ArchiveName ? ArchiveName : TEXT("RemoteObjectWriter"))
 	{
 		static_assert(sizeof(FNameIndexType) == sizeof(FRemoteObjectPathName::FNameIndexType), "Name Index type must match the type used for storing remote object path name indices");
 
@@ -469,11 +473,6 @@ public:
 		ObjectsToSerialize.Add(InRootObject);
 	}
 
-	virtual const FUObjectMigrationContext* GetMigrationContext() const override
-	{
-		return &MigrationContext;
-	}
-
 	virtual UObject* GetArchetypeFromLoader(const UObject* Obj) override
 	{
 		if (ObjectToArchetypeMap)
@@ -621,8 +620,8 @@ class FArchetypeDeltaWriter : public FArchiveRemoteObjectWriter
 
 public:
 
-	FArchetypeDeltaWriter(UObject* InRootObject, FRemoteObjectData& OutObjectData, const FUObjectMigrationContext& InMigrationContext, TMap<UObject*, UObject*>& InSubstituteArchetypeMap)
-		: FArchiveRemoteObjectWriter(InRootObject, OutObjectData, InMigrationContext, TEXT("RemoteArchetypeDeltaWriter"))
+	FArchetypeDeltaWriter(UObject* InRootObject, FRemoteObjectData& OutObjectData, TMap<UObject*, UObject*>& InSubstituteArchetypeMap)
+		: FArchiveRemoteObjectWriter(InRootObject, OutObjectData, /*MigrationContext*/ nullptr, TEXT("RemoteArchetypeDeltaWriter"))
 		, SubstituteArchetypeMap(InSubstituteArchetypeMap)
 	{
 	}
@@ -688,7 +687,6 @@ class FArchiveRemoteObjectReader : public TArchiveRemoteObjectBase<FMemoryReader
 	const TArray<FName>& Names;
 	const TArray<UObject*>& ResolvedPathNameObjects;
 	const ERemoteObjectSerializationFlags DeserializeFlags = ERemoteObjectSerializationFlags::None;
-	const FUObjectMigrationContext& MigrationContext;
 
 	ERemoteReferenceType ReadObjectReference(FObjectPtr& Value)
 	{
@@ -750,18 +748,12 @@ public:
 	 * @param bInAssignedOwnership:  We must take ownership of the objects we deserialize. We could already have ownership of those objects.
 	 * @param InDeserializeFlags:  The flags for how we should treat references during deserialization.
 	 */
-	FArchiveRemoteObjectReader(FRemoteObjectData& InObjectData, const TArray<UObject*>& InResolvedPathNames, const FUObjectMigrationContext& InMigrationContext, ERemoteObjectSerializationFlags InDeserializeFlags)
-		: TArchiveRemoteObjectBase<FMemoryReader>(InObjectData, TEXT("RemoteObjectReader"))
+	FArchiveRemoteObjectReader(FRemoteObjectData& InObjectData, const TArray<UObject*>& InResolvedPathNames, const FUObjectMigrationContext* InMigrationContext, ERemoteObjectSerializationFlags InDeserializeFlags)
+		: TArchiveRemoteObjectBase<FMemoryReader>(InObjectData, InMigrationContext, TEXT("RemoteObjectReader"))
 		, Names(InObjectData.Names)
 		, ResolvedPathNameObjects(InResolvedPathNames)
 		, DeserializeFlags(InDeserializeFlags)
-		, MigrationContext(InMigrationContext)
-	{
-	}
-
-	virtual const FUObjectMigrationContext* GetMigrationContext() const override
 	{
-		return &MigrationContext;
 	}
 
 	virtual FArchive& operator<<(UObject*& Obj) override
@@ -816,7 +808,6 @@ public:
 class FArchiveRemoteReferencePatcher : public FArchiveUObject
 {
 	TMap<FRemoteObjectId, UObject*> IdToObjectMap;
-	const FUObjectMigrationContext& MigrationContext;
 
 	void InitObjectMap(const TArray<UObject*>& DeserializedObjects)
 	{
@@ -840,19 +831,12 @@ class FArchiveRemoteReferencePatcher : public FArchiveUObject
 
 public:
 
-	FArchiveRemoteReferencePatcher(const TArray<UObject*>& DeserializedObjects, const FUObjectMigrationContext& InMigrationContext)
-		: MigrationContext(InMigrationContext)
+	FArchiveRemoteReferencePatcher(const TArray<UObject*>& DeserializedObjects)
 	{
 		ArIsObjectReferenceCollector = true;
 		SetIsPersistent(false);
 		InitObjectMap(DeserializedObjects);
-		SetMigratingRemoteObjects(true);
-		SetPortFlags(PPF_MigratingRemoteObjects);
-	}
-
-	virtual const FUObjectMigrationContext* GetMigrationContext() const override
-	{
-		return &MigrationContext;
+		SetPortFlags(PPF_AvoidRemoteObjectMigration);
 	}
 
 	virtual FArchive& operator<<(FObjectPtr& Value) override
@@ -869,28 +853,20 @@ public:
 class FArchiveSubObjectGatherer : public FArchiveUObject
 {
 	UObject* RootObject = nullptr;
-	const FUObjectMigrationContext& MigrationContext;
 	TMap<UObject*, UObject*>& ArchetypeToObjectMap;
 	TArray<UObject*>& ObjectsToSerialize;
 
 public:
 
-	FArchiveSubObjectGatherer(UObject* InRootObject, const FUObjectMigrationContext& InMigrationContext, TMap<UObject*, UObject*>& InArchetypeToObjectMap, TArray<UObject*>& InObjectsToSerialize)
+	FArchiveSubObjectGatherer(UObject* InRootObject, TMap<UObject*, UObject*>& InArchetypeToObjectMap, TArray<UObject*>& InObjectsToSerialize)
 		: RootObject(InRootObject)
-		, MigrationContext(InMigrationContext)
 		, ArchetypeToObjectMap(InArchetypeToObjectMap)
 		, ObjectsToSerialize(InObjectsToSerialize)
 	{
 		ArchetypeToObjectMap.Add(RootObject->GetArchetype(), RootObject);
 		ArIsObjectReferenceCollector = true;
 		SetIsPersistent(false);
-		SetMigratingRemoteObjects(true);
-		SetPortFlags(PPF_MigratingRemoteObjects);
-	}
-
-	virtual const FUObjectMigrationContext* GetMigrationContext() const override
-	{
-		return &MigrationContext;
+		SetPortFlags(PPF_AvoidRemoteObjectMigration);
 	}
 
 	virtual FArchive& operator<<(FObjectPtr& Value) override
@@ -1095,7 +1071,7 @@ void SerializeObjectDataInternal(FArchiveRemoteObjectWriter& Ar, UObject* Reques
 	}
 }
 
-FRemoteObjectData SerializeObjectData(UObject* InObject, TSet<UObject*>& OutObjects, const FUObjectMigrationContext& MigrationContext)
+FRemoteObjectData SerializeObjectData(UObject* InObject, TSet<UObject*>& OutObjects, const FUObjectMigrationContext* MigrationContext)
 {
 	UObject* Object = InObject;
 	if (GRemoteObjectsMigrateFullHierarchy)
@@ -1108,13 +1084,14 @@ FRemoteObjectData SerializeObjectData(UObject* InObject, TSet<UObject*>& OutObje
 
 	{
 		FArchiveRemoteObjectWriter Ar(Object, ObjectData, MigrationContext);
+		Ar.SetMigratingRemoteObjects(true);
 		SerializeObjectDataInternal(Ar, InObject, RequestedObjectId, OutObjects);
 	}
 
 	return MoveTemp(ObjectData);
 }
 
-void ResetRemoteObject(UObject* InObject, const FUObjectMigrationContext& MigrationContext)
+void ResetRemoteObject(UObject* InObject)
 {
 	UObject* Object = InObject;
 	if (GRemoteObjectsMigrateFullHierarchy)
@@ -1129,7 +1106,7 @@ void ResetRemoteObject(UObject* InObject, const FUObjectMigrationContext& Migrat
 	{
 		TArray<UObject*> ObjectsToSerialize;
 		TSet<UObject*> SerializedObjects;
-		FArchiveSubObjectGatherer SubobjectGather(Object, MigrationContext, ReverseArchetypeToObjectMap, ObjectsToSerialize);
+		FArchiveSubObjectGatherer SubobjectGather(Object, ReverseArchetypeToObjectMap, ObjectsToSerialize);
 		ObjectsToSerialize.Add(Object);
 
 		for (int32 Index = 0; Index < ObjectsToSerialize.Num(); ++Index)
@@ -1149,7 +1126,7 @@ void ResetRemoteObject(UObject* InObject, const FUObjectMigrationContext& Migrat
 		// This will produce a delta between the archetypes and their instances which will then be used to restore the instances' state to the archetypes'
 		FRemoteObjectId RequestedObjectId(Object);
 		UObject* Archetype = Object->GetArchetype();
-		FArchetypeDeltaWriter Ar(Archetype, *ArchetypeDelta, MigrationContext, ReverseArchetypeToObjectMap);
+		FArchetypeDeltaWriter Ar(Archetype, *ArchetypeDelta, ReverseArchetypeToObjectMap);
 		TSet<UObject*> ArchetypeSerializedObjects;
 		SerializeObjectDataInternal(Ar, Archetype, RequestedObjectId, ArchetypeSerializedObjects);
 	}
@@ -1163,7 +1140,7 @@ void ResetRemoteObject(UObject* InObject, const FUObjectMigrationContext& Migrat
 		// Use any existing objects and deserialize archetype delta on top of them avoiding re-construction which may lead to undiserable side-effects
 		// Additionally we don't want to recursively re-enter this function so let the deserialization process know we're already resetting migrated object(s)
 		ERemoteObjectSerializationFlags DeserializationFlags = ERemoteObjectSerializationFlags::PreserveRemoteReferences | ERemoteObjectSerializationFlags::UseExistingObjects | ERemoteObjectSerializationFlags::Resetting;
-		DeserializeObjectData(*ArchetypeDelta, MigrationContext, DeserializedIds, DeserializedObjects, DeserializationFlags);
+		DeserializeObjectData(*ArchetypeDelta, /*MigrationContext*/ nullptr, DeserializedIds, DeserializedObjects, DeserializationFlags);
 
 		const FRemoteServerId LocalServerId = UE::RemoteObject::GetGlobalServerId();
 		for (UObject* DeserializedObject : DeserializedObjects)
@@ -1175,7 +1152,7 @@ void ResetRemoteObject(UObject* InObject, const FUObjectMigrationContext& Migrat
 }
 
 int32 DeserializeObjectData(FRemoteObjectData& ObjectData,
-							const FUObjectMigrationContext& MigrationContext,
+							const FUObjectMigrationContext* MigrationContext,
 							TArray<FRemoteObjectId>& OutObjectRemoteIds,
 							TArray<UObject*>& OutObjects,
 							ERemoteObjectSerializationFlags DeserializeFlags)
@@ -1189,11 +1166,14 @@ int32 DeserializeObjectData(FRemoteObjectData& ObjectData,
 	TArray<FRemoteObjectHeader> ObjectHeaders;
 	TArray<UObject*> ResolvedPathNameObjects;
 	int32 RequestedObjectIndex = -1;
+	const bool bResetting = EnumHasAnyFlags(DeserializeFlags, ERemoteObjectSerializationFlags::Resetting);
 
 	ResolvePathNames(ObjectData, ResolvedPathNameObjects);
 
 	// If we are being assigned ownership, we *must* take ownership of the objects.  Note: We also may already have ownership of those objects. 
 	FArchiveRemoteObjectReader Ar(ObjectData, ResolvedPathNameObjects, MigrationContext, DeserializeFlags);
+	// If we're calling this function to reset an object to its archetype state then we don't want FArchiveRemoteObjectReader to be marked as migrating remote objects (Which it is by default)
+	Ar.SetMigratingRemoteObjects(!bResetting);
 
 	{
 		FRemoteObjectReaderScope Scope(Ar, TEXT("Header"));
@@ -1255,7 +1235,6 @@ int32 DeserializeObjectData(FRemoteObjectData& ObjectData,
 
 		// If we're already resetting a borrowed object we don't want to change the ownership until the object is deserialized using remote server data
 		// and we also don't want to recursively reset the object
-		const bool bResetting = EnumHasAnyFlags(DeserializeFlags, ERemoteObjectSerializationFlags::Resetting);
 		if (!bResetting)
 		{
 			// ensure remote object stubs are created (ownership will be assigned after PostMigrate)
@@ -1269,7 +1248,7 @@ int32 DeserializeObjectData(FRemoteObjectData& ObjectData,
 
 			if (bReturningBorrowedObject)
 			{
-				ResetRemoteObject(OutObjects[0], MigrationContext);
+				ResetRemoteObject(OutObjects[0]);
 			}
 		}
 
@@ -1295,7 +1274,7 @@ int32 DeserializeObjectData(FRemoteObjectData& ObjectData,
 		// Patch any unresolved remote references (skip if we're resetting an object to its CDO state)
 		if (!bResetting)
 		{
-			FArchiveRemoteReferencePatcher PatchAr(OutObjects, MigrationContext);
+			FArchiveRemoteReferencePatcher PatchAr(OutObjects);
 			for (UObject* Object : OutObjects)
 			{
 				if (Object)
diff --git a/Engine/Source/Runtime/CoreUObject/Private/UObject/RemoteObjectSerialization.h b/Engine/Source/Runtime/CoreUObject/Private/UObject/RemoteObjectSerialization.h
--- a/Engine/Source/Runtime/CoreUObject/Private/UObject/RemoteObjectSerialization.h
+++ b/Engine/Source/Runtime/CoreUObject/Private/UObject/RemoteObjectSerialization.h
@@ -75,7 +75,7 @@ namespace UE::RemoteObject::Serialization
 	* @param MigrationContext Contains the meta data of the current migration request
 	* @return Remote object data representing the serialized objects
 	*/
-	FRemoteObjectData SerializeObjectData(UObject* InObject, TSet<UObject*>& OutObjects, const FUObjectMigrationContext& MigrationContext);
+	FRemoteObjectData SerializeObjectData(UObject* InObject, TSet<UObject*>& OutObjects, const FUObjectMigrationContext* MigrationContext);
 
 	/**
 	* Deserializes remote object data
@@ -87,7 +87,7 @@ namespace UE::RemoteObject::Serialization
 	* @return Index of an object in OutReceivedObjects that was the main object the migration request was triggered for 
 	* (usually 0 but if a migration requests a default subobject then its parent is also migrated and the return value will be > 0)
 	*/
-	int32 DeserializeObjectData(FRemoteObjectData& ObjectData, const FUObjectMigrationContext& MigrationContext, TArray<FRemoteObjectId>& OutObjectRemoteIds, TArray<UObject*>& OutReceivedObjects, ERemoteObjectSerializationFlags DeserializeFlags = ERemoteObjectSerializationFlags::None);
+	int32 DeserializeObjectData(FRemoteObjectData& ObjectData, const FUObjectMigrationContext* MigrationContext, TArray<FRemoteObjectId>& OutObjectRemoteIds, TArray<UObject*>& OutReceivedObjects, ERemoteObjectSerializationFlags DeserializeFlags = ERemoteObjectSerializationFlags::None);
 
 	/**
 	* Finds the canonical 'root' object that is used for remote object serialization
diff --git a/Engine/Source/Runtime/CoreUObject/Private/UObject/RemoteObjectTransfer.cpp b/Engine/Source/Runtime/CoreUObject/Private/UObject/RemoteObjectTransfer.cpp
--- a/Engine/Source/Runtime/CoreUObject/Private/UObject/RemoteObjectTransfer.cpp
+++ b/Engine/Source/Runtime/CoreUObject/Private/UObject/RemoteObjectTransfer.cpp
@@ -260,7 +260,7 @@ UObject* MigrateObjectFromObjectData(FRemoteObjectData& ObjectData, const FUObje
 
 	TArray<FRemoteObjectId> ReceivedObjectRemoteIds;
 	TArray<UObject*> ReceivedObjects;
-	int32 RequestedObjectIndex = DeserializeObjectData(ObjectData, MigrationContext, ReceivedObjectRemoteIds, ReceivedObjects);
+	int32 RequestedObjectIndex = DeserializeObjectData(ObjectData, &MigrationContext, ReceivedObjectRemoteIds, ReceivedObjects);
 
 	// Debug section
 	checkf(ReceivedObjects.Num() > 0, TEXT("PeerId:%s unable to deserialize object data (%d bytes)"), *GetGlobalServerId().ToString(), ObjectData.GetNumBytes());
@@ -567,7 +567,7 @@ void MigrateObjectToRemoteServerWithExplicitPriority(FRemoteTransactionId Reques
 
 		// send the object
 		TSet<UObject*> SentObjects;
-		FRemoteObjectData ObjectData = SerializeObjectData(Object, SentObjects, MigrationContext);
+		FRemoteObjectData ObjectData = SerializeObjectData(Object, SentObjects, &MigrationContext);
 
 		UE_LOG(LogRemoteObject, VeryVerbose, TEXT("     : MigrateObjectToRemoteServer(%s) serialization complete obj %s (%d bytes)"), *RequestId.ToString(), *ObjectId.ToString(), ObjectData.GetNumBytes());
 
@@ -776,7 +776,7 @@ void StoreObjectToDatabase(UObject* Object)
 			FUObjectMigrationContext StoreToDatabaseContext{ .ObjectId = RootObjectId, .RemoteServerId = DatabaseId, .OwnerServerId = DatabaseId, .PhysicsServerId = DatabaseId, .RequestId = PushMigrationRequestId, .MigrationSide = EObjectMigrationSide::Send };
 
 			// Do the actual serialization
-			FRemoteObjectData ObjectData = SerializeObjectData(OutermostUnreachableRemoteReference, SentObjects, StoreToDatabaseContext);
+			FRemoteObjectData ObjectData = SerializeObjectData(OutermostUnreachableRemoteReference, SentObjects, &StoreToDatabaseContext);
 
 			const FMigrateSendParams Params { .MigrationContext = StoreToDatabaseContext, .ObjectData = ObjectData };
 			StoreRemoteObjectDataDelegate.Execute(Params);
diff --git a/Engine/Source/Runtime/Engine/Private/Materials/MaterialInstance.cpp b/Engine/Source/Runtime/Engine/Private/Materials/MaterialInstance.cpp
--- a/Engine/Source/Runtime/Engine/Private/Materials/MaterialInstance.cpp
+++ b/Engine/Source/Runtime/Engine/Private/Materials/MaterialInstance.cpp
@@ -3170,7 +3170,7 @@ void UMaterialInstance::Serialize(FArchive& Ar)
 		const bool bWantToSaveCachedData = Ar.IsCooking();
 #else
 		// We want to copy the cached data when the material is duplicated either directly or during remote object migration
-		const bool bDuplicatingObjectInACookedGame = FPlatformProperties::RequiresCookedData() && (Ar.GetPortFlags() & (PPF_Duplicate | PPF_MigratingRemoteObjects));
+		const bool bDuplicatingObjectInACookedGame = FPlatformProperties::RequiresCookedData() && (Ar.HasAnyPortFlags(PPF_Duplicate | PPF_AvoidRemoteObjectMigration));
 		const bool bWantToSaveCachedData = Ar.IsSaving() && bDuplicatingObjectInACookedGame;
 
 		// Workaround for materials being created by annotation data at runtime not having CachedData
