[AsyncLoading2]
Provide more information about exports that are missing when we try to clear constructed objects.

#rnx

diff --git a/Engine/Source/Runtime/CoreUObject/Private/Serialization/AsyncLoading2.cpp b/Engine/Source/Runtime/CoreUObject/Private/Serialization/AsyncLoading2.cpp
--- a/Engine/Source/Runtime/CoreUObject/Private/Serialization/AsyncLoading2.cpp
+++ b/Engine/Source/Runtime/CoreUObject/Private/Serialization/AsyncLoading2.cpp
@@ -3888,6 +3888,7 @@ private:
 
 	void ProcessExportDependencies(const FAsyncPackageHeaderData& Header, int32 LocalExportIndex, FExportBundleEntry::EExportCommandType CommandType);
 	int32 GetPublicExportIndex(uint64 ExportHash, FAsyncPackageHeaderData*& OutHeader);
+	FName GetFNameFromPackageObjectIndex(FPackageObjectIndex Index);
 	UObject* ConditionalCreateExport(const FAsyncPackageHeaderData& Header, int32 LocalExportIndex);
 	UObject* ConditionalSerializeExport(const FAsyncPackageHeaderData& Header, int32 LocalExportIndex);
 	UObject* ConditionalCreateImport(const FAsyncPackageHeaderData& Header, int32 LocalImportIndex);
@@ -7791,17 +7792,26 @@ Verse::VCell* FAsyncPackage2::ConditionalSerializeCellImport(const FAsyncPackage
 	return SourcePackage->ConditionalSerializeCellExport(*SourcePackageHeader, CellExportIndex);
 }
 
-static FName GetFNameFromPackageObjectIndex(const FAsyncPackageHeaderData& Header, FPackageImportStore& ImportStore, FPackageObjectIndex Index)
+FName FAsyncPackage2::GetFNameFromPackageObjectIndex(FPackageObjectIndex Index)
 {
 	FName Out(NAME_None);
 	if (Index.IsExport())
 	{
-		Header.NameMap.TryGetName(Header.ExportMap[Index.ToExport()].ObjectName, Out);
+#if ALT2_ENABLE_LINKERLOAD_SUPPORT
+		if (LinkerLoadState.IsSet())
+		{
+			Out = LinkerLoadState->Linker->ExportMap[Index.ToExport()].ObjectName;
+		}
+		else
+#endif
+		{
+			HeaderData.NameMap.TryGetName(HeaderData.ExportMap[Index.ToExport()].ObjectName, Out);
+		}
 	}
 	else
 	{
 		// No need to check if the Index is invalid, that is handled internally
-		const UObject* Import = ImportStore.FindOrGetImportObject(Header, Index);
+		const UObject* Import = ImportStore.FindOrGetImportObject(HeaderData, Index);
 		if (Import)
 		{
 			Out = Import->GetFName();
@@ -7850,13 +7860,13 @@ void FAsyncPackage2::EventDrivenCreateExport(const FAsyncPackageHeaderData& Head
 
 	if (!LoadClass)
 	{
-		UE_ASYNC_PACKAGE_LOG(Error, Desc, TEXT("CreateExport"), TEXT("Could not find class object (0x%llX - '%s') for %s"), Export.ClassIndex.Value(), *GetFNameFromPackageObjectIndex(Header, ImportStore, Export.ClassIndex).ToString(), *ObjectName.ToString());
+		UE_ASYNC_PACKAGE_LOG(Error, Desc, TEXT("CreateExport"), TEXT("Could not find class object (0x%llX - '%s') for %s"), Export.ClassIndex.Value(), *GetFNameFromPackageObjectIndex(Export.ClassIndex).ToString(), *ObjectName.ToString());
 		ExportObject.bExportLoadFailed = true;
 		return;
 	}
 	if (!ThisParent)
 	{
-		UE_ASYNC_PACKAGE_LOG(Error, Desc, TEXT("CreateExport"), TEXT("Could not find outer (0x%llX - '%s') object for %s"), Export.OuterIndex.Value(), *GetFNameFromPackageObjectIndex(Header, ImportStore, Export.OuterIndex).ToString(), *ObjectName.ToString());
+		UE_ASYNC_PACKAGE_LOG(Error, Desc, TEXT("CreateExport"), TEXT("Could not find outer (0x%llX - '%s') object for %s"), Export.OuterIndex.Value(), *GetFNameFromPackageObjectIndex(Export.OuterIndex).ToString(), *ObjectName.ToString());
 		ExportObject.bExportLoadFailed = true;
 		return;
 	}
@@ -7866,7 +7876,7 @@ void FAsyncPackage2::EventDrivenCreateExport(const FAsyncPackageHeaderData& Head
 		ExportObject.SuperObject = EventDrivenIndexToObject(Header, Export.SuperIndex, false);
 		if (!ExportObject.SuperObject)
 		{
-			UE_ASYNC_PACKAGE_LOG(Error, Desc, TEXT("CreateExport"), TEXT("Could not find SuperStruct object (0x%llX - '%s') for %s"), Export.SuperIndex.Value(), *GetFNameFromPackageObjectIndex(Header, ImportStore, Export.SuperIndex).ToString(), *ObjectName.ToString());
+			UE_ASYNC_PACKAGE_LOG(Error, Desc, TEXT("CreateExport"), TEXT("Could not find SuperStruct object (0x%llX - '%s') for %s"), Export.SuperIndex.Value(), *GetFNameFromPackageObjectIndex(Export.SuperIndex).ToString(), *ObjectName.ToString());
 			ExportObject.bExportLoadFailed = true;
 			return;
 		}
@@ -10312,11 +10322,22 @@ void FAsyncPackage2::ClearConstructedObjects()
 	}
 	ConstructedObjects.Empty();
 
-	for (FExportObject& Export : Data.Exports)
+	const int32 ExportCount = Data.Exports.Num();
+	for (int32 ExportIndex = 0; ExportIndex < ExportCount; ++ExportIndex)
 	{
+		FExportObject& Export = Data.Exports[ExportIndex];
 		if (Export.bWasFoundInMemory)
 		{
-			check(Export.Object);
+			checkf(Export.Object, TEXT("Export '%s' in package '%s' found in memory is missing:\n\t\tSuper: '%s', Template: '%s', LoadFailed: %s, Filtered: %s (%s), FoundInMemory: %s"),
+				*GetFNameFromPackageObjectIndex(FPackageObjectIndex::FromExportIndex(ExportIndex)).ToString(),
+				*Desc.PackagePathToLoad.GetPackageFName().ToString(),
+				(Export.SuperObject ? *Export.SuperObject->GetPathName() : TEXT("null")),
+				(Export.TemplateObject ? *Export.TemplateObject->GetPathName() : TEXT("null")),
+				(Export.bExportLoadFailed ? TEXT("true") : TEXT("false")),
+				(Export.bFiltered ? TEXT("true") : TEXT("false")),
+				LexToString(HeaderData.ExportMap[ExportIndex].FilterFlags),
+				(Export.bWasFoundInMemory ? TEXT("true") : TEXT("false")));
+
 			Export.Object->AtomicallyClearInternalFlags(EInternalObjectFlags_AsyncLoading | EInternalObjectFlags::Async);
 		}
 		else
