PR #12996: Fix the inner property for SerializeFromMismatchedTag in TArray, TMap, TSet, and TOptional

diff --git a/Engine/Source/Runtime/CoreUObject/Private/UObject/PropertyArray.cpp b/Engine/Source/Runtime/CoreUObject/Private/UObject/PropertyArray.cpp
--- a/Engine/Source/Runtime/CoreUObject/Private/UObject/PropertyArray.cpp
+++ b/Engine/Source/Runtime/CoreUObject/Private/UObject/PropertyArray.cpp
@@ -1240,6 +1240,7 @@ EConvertFromTypeResult FArrayProperty::ConvertFromType(const FPropertyTag& Tag,
 	ScriptArrayHelper.EmptyAndAddValues(ElementCount);
 
 	FPropertyTag InnerPropertyTag;
+	InnerPropertyTag.SetProperty(Inner);
 	InnerPropertyTag.SetType(Tag.GetType().GetParameter(0));
 	InnerPropertyTag.Name = Tag.Name;
 	InnerPropertyTag.ArrayIndex = 0;
diff --git a/Engine/Source/Runtime/CoreUObject/Private/UObject/PropertyMap.cpp b/Engine/Source/Runtime/CoreUObject/Private/UObject/PropertyMap.cpp
--- a/Engine/Source/Runtime/CoreUObject/Private/UObject/PropertyMap.cpp
+++ b/Engine/Source/Runtime/CoreUObject/Private/UObject/PropertyMap.cpp
@@ -1462,11 +1462,13 @@ EConvertFromTypeResult FMapProperty::ConvertFromType(const FPropertyTag& Tag, FS
 	};
 
 	FPropertyTag KeyPropertyTag;
+	KeyPropertyTag.SetProperty(KeyProp);
 	KeyPropertyTag.SetType(KeyType);
 	KeyPropertyTag.Name = Tag.Name;
 	KeyPropertyTag.ArrayIndex = 0;
 
 	FPropertyTag ValuePropertyTag;
+	ValuePropertyTag.SetProperty(ValueProp);
 	ValuePropertyTag.SetType(ValueType);
 	ValuePropertyTag.Name = Tag.Name;
 	ValuePropertyTag.ArrayIndex = 0;
diff --git a/Engine/Source/Runtime/CoreUObject/Private/UObject/PropertyOptional.cpp b/Engine/Source/Runtime/CoreUObject/Private/UObject/PropertyOptional.cpp
--- a/Engine/Source/Runtime/CoreUObject/Private/UObject/PropertyOptional.cpp
+++ b/Engine/Source/Runtime/CoreUObject/Private/UObject/PropertyOptional.cpp
@@ -503,6 +503,7 @@ EConvertFromTypeResult FOptionalProperty::ConvertFromType(const FPropertyTag& Ta
 				MarkUnset(ContainerPtrToValuePtr<void>(ContainerData));
 				return EConvertFromTypeResult::Converted;
 			}
+			ValueTag.SetProperty(GetValueProperty());
 			ValueTag.SetType(Tag.GetType().GetParameter(0));
 
 			// Mimic FPropertyTag::SerializeTaggedProperty storing bool values in the tag because implementations of ConvertFromType expect it there.
diff --git a/Engine/Source/Runtime/CoreUObject/Private/UObject/PropertySet.cpp b/Engine/Source/Runtime/CoreUObject/Private/UObject/PropertySet.cpp
--- a/Engine/Source/Runtime/CoreUObject/Private/UObject/PropertySet.cpp
+++ b/Engine/Source/Runtime/CoreUObject/Private/UObject/PropertySet.cpp
@@ -903,6 +903,7 @@ EConvertFromTypeResult FSetProperty::ConvertFromType(const FPropertyTag& Tag, FS
 	};
 
 	FPropertyTag InnerPropertyTag;
+	InnerPropertyTag.SetProperty(ElementProp);
 	InnerPropertyTag.SetType(Tag.GetType().GetParameter(0));
 	InnerPropertyTag.Name = Tag.Name;
 	InnerPropertyTag.ArrayIndex = 0;
