New function FCoreRedirects::GetHashOfRedirectsAffectingPackages to replace AppendHashOfRedirectsAffectingPackages. It returns one hash per package and will allow me to save the hash in the FCookDependency object in a later cl.

#rnx
#incrementalcook

diff --git a/Engine/Source/Editor/UnrealEd/Private/Cooker/CookPackageArtifacts.cpp b/Engine/Source/Editor/UnrealEd/Private/Cooker/CookPackageArtifacts.cpp
--- a/Engine/Source/Editor/UnrealEd/Private/Cooker/CookPackageArtifacts.cpp
+++ b/Engine/Source/Editor/UnrealEd/Private/Cooker/CookPackageArtifacts.cpp
@@ -183,7 +183,17 @@ FBuildDependencySet::ECurrentKeyResult FBuildDependencySet::TryCalculateCurrentK
 			{
 				PackageNames.Add(RedirectionDependency.GetPackageName());
 			}
-			FCoreRedirects::AppendHashOfRedirectsAffectingPackages(KeyBuilder, PackageNames);
+
+			TArray<FBlake3Hash> Hashes;
+			Hashes.SetNum(PackageNames.Num());
+			FCoreRedirects::GetHashOfRedirectsAffectingPackages(PackageNames, Hashes);
+
+			for (const FBlake3Hash& Hash : Hashes)
+			{
+				KeyBuilder.Update(&Hash, sizeof(Hash));
+			}
+
+			FCoreRedirects::AppendHashOfGlobalRedirects(KeyBuilder);
 			break;
 		}
 		default:
diff --git a/Engine/Source/Runtime/CoreUObject/Internal/UObject/CoreRedirects/RedirectionSummary.h b/Engine/Source/Runtime/CoreUObject/Internal/UObject/CoreRedirects/RedirectionSummary.h
--- a/Engine/Source/Runtime/CoreUObject/Internal/UObject/CoreRedirects/RedirectionSummary.h
+++ b/Engine/Source/Runtime/CoreUObject/Internal/UObject/CoreRedirects/RedirectionSummary.h
@@ -26,7 +26,11 @@ public:
 
 	void Add(const FCoreRedirect& CoreRedirect, bool bIsWildcardMatch);
 	void Remove(const FCoreRedirect& CoreRedirect, bool bIsWildcardMatch);
+
+	UE_DEPRECATED(5.6, "Use GetHashAffectingPackages instead.")
 	void AppendHashAffectingPackages(FBlake3& Hasher, TConstArrayView<FName> PackageNames);
+	
+	void GetHashAffectingPackages(const TConstArrayView<FName>& PackageNames, TArray<FBlake3Hash>& HashArray);
 	void AppendHashGlobal(FBlake3& Hasher);
 
 private:
diff --git a/Engine/Source/Runtime/CoreUObject/Private/UObject/CoreRedirects.cpp b/Engine/Source/Runtime/CoreUObject/Private/UObject/CoreRedirects.cpp
--- a/Engine/Source/Runtime/CoreUObject/Private/UObject/CoreRedirects.cpp
+++ b/Engine/Source/Runtime/CoreUObject/Private/UObject/CoreRedirects.cpp
@@ -2016,7 +2016,18 @@ void FCoreRedirects::AppendHashOfRedirectsAffectingPackages(FBlake3& Hasher, TCo
 	FCoreRedirectsContext& Context = FCoreRedirectsContext::GetThreadContext();
 	if (!!(Context.GetFlags() & FCoreRedirectsContext::EFlags::UseRedirectionSummary))
 	{
+		PRAGMA_DISABLE_DEPRECATION_WARNINGS
 		Context.GetRedirectionSummary().AppendHashAffectingPackages(Hasher, PackageNames);
+		PRAGMA_ENABLE_DEPRECATION_WARNINGS
+	}
+}
+
+void FCoreRedirects::GetHashOfRedirectsAffectingPackages(const TConstArrayView<FName>& PackageNames, TArray<FBlake3Hash>& Hashes)
+{
+	FCoreRedirectsContext& Context = FCoreRedirectsContext::GetThreadContext();
+	if (!!(Context.GetFlags() & FCoreRedirectsContext::EFlags::UseRedirectionSummary))
+	{
+		Context.GetRedirectionSummary().GetHashAffectingPackages(PackageNames, Hashes);
 	}
 }
 
diff --git a/Engine/Source/Runtime/CoreUObject/Private/UObject/RedirectionSummary.cpp b/Engine/Source/Runtime/CoreUObject/Private/UObject/RedirectionSummary.cpp
--- a/Engine/Source/Runtime/CoreUObject/Private/UObject/RedirectionSummary.cpp
+++ b/Engine/Source/Runtime/CoreUObject/Private/UObject/RedirectionSummary.cpp
@@ -85,31 +85,48 @@ void FRedirectionSummary::Remove(const FCoreRedirect& CoreRedirect, bool bIsWild
 
 void FRedirectionSummary::AppendHashAffectingPackages(FBlake3& Hasher, TConstArrayView<FName> PackageNames)
 {
+	TArray<FBlake3Hash> Hashes;
+	Hashes.SetNum(PackageNames.Num());
+	GetHashAffectingPackages(PackageNames, Hashes);
+
+	for (const FBlake3Hash& Hash : Hashes)
+	{
+		Hasher.Update(&Hash, sizeof(Hash));
+	}
+	
+	if (!PackageNames.IsEmpty())
+	{
+		AppendHashGlobal(Hasher);
+	}
+}
+
+void FRedirectionSummary::GetHashAffectingPackages(const TConstArrayView<FName>& PackageNames, TArray<FBlake3Hash>& Hashes)
+{
+	check(Hashes.Num() == PackageNames.Num());
+
 	if (PackageNames.IsEmpty())
 	{
 		return;
 	}
 
-	int32 Index = -1;
+	int32 Index = 0;
 	const FName* PackageNamesData = PackageNames.GetData();
 	int32 PackageNamesNum = PackageNames.Num();
 	bool bNeedWriteLock = false;
 
 	{
 		FReadScopeLock ScopeLock(Lock);
-		if (GlobalRedirects.TryAppendHashInReadLock(Hasher))
-		{
-			Index = 0;
-		}
-		else
-		{
-			bNeedWriteLock = true;
-		}
 		while (!bNeedWriteLock && Index < PackageNamesNum)
 		{
+			FBlake3 Hasher;
 			FRedirectContainer* Container = RedirectsForPackage.Find(PackageNamesData[Index]);
-			if (!Container || Container->TryAppendHashInReadLock(Hasher))
+			if (!Container)
+			{
+				++Index;
+			}
+			else if (Container->TryAppendHashInReadLock(Hasher))
 			{
+				Hashes[Index] = Hasher.Finalize();
 				++Index;
 			}
 			else
@@ -118,23 +135,21 @@ void FRedirectionSummary::AppendHashAffectingPackages(FBlake3& Hasher, TConstArr
 			}
 		}
 	}
+
 	if (!bNeedWriteLock)
 	{
 		return;
 	}
 
 	FWriteScopeLock WriteLock(Lock);
-	if (Index < 0)
-	{
-		GlobalRedirects.AppendHashInWriteLock(Hasher);
-		Index = 0;
-	}
 	while (Index < PackageNamesNum)
 	{
 		FRedirectContainer* Container = RedirectsForPackage.Find(PackageNamesData[Index]);
 		if (Container)
 		{
+			FBlake3 Hasher;
 			Container->AppendHashInWriteLock(Hasher);
+			Hashes[Index] = Hasher.Finalize();
 		}
 		++Index;
 	}
diff --git a/Engine/Source/Runtime/CoreUObject/Public/UObject/CoreRedirects.h b/Engine/Source/Runtime/CoreUObject/Public/UObject/CoreRedirects.h
--- a/Engine/Source/Runtime/CoreUObject/Public/UObject/CoreRedirects.h
+++ b/Engine/Source/Runtime/CoreUObject/Public/UObject/CoreRedirects.h
@@ -402,8 +402,16 @@ struct FCoreRedirects
 	 * redirecting from or to the package. Used in iterative cooking to invalidate the cooked version of packages when
 	 * CoreRedirects change. 
 	 */
-	static COREUOBJECT_API void AppendHashOfRedirectsAffectingPackages(FBlake3& Hasher,
-		TConstArrayView<FName> PackageNames);
+	UE_DEPRECATED(5.6, "Use GetHashOfRedirectsAffectingPackages(const TConstArrayView<FName> PackageNames, TArray<FBlake3Hash>& Hashes) instead.")
+	static COREUOBJECT_API void AppendHashOfRedirectsAffectingPackages(FBlake3& Hasher, TConstArrayView<FName> PackageNames);
+
+	/**
+	 * For each package in PackageNames, compute the hash of all redirects that affect the package, either
+	 * redirecting from or to the package. It doesn't include the GlobalRedirects.
+	 * Used in iterative cooking to invalidate the cooked version of packages when CoreRedirects change.
+	 */
+	static COREUOBJECT_API void GetHashOfRedirectsAffectingPackages(const TConstArrayView<FName>& PackageNames, TArray<FBlake3Hash>& Hashes);
+
 	/**
 	 * Append the hash of all redirects that can affect multiple packages, or for which the affected packages are unknown.
 	 * Used in iterative cooking to invalidate the cooked version of packages when CoreRedirects change.
