Avoid resolving remote objects when comparing weak object pointers during remote object serialization

diff --git a/Engine/Source/Runtime/CoreUObject/Private/UObject/PropertyWeakObjectPtr.cpp b/Engine/Source/Runtime/CoreUObject/Private/UObject/PropertyWeakObjectPtr.cpp
--- a/Engine/Source/Runtime/CoreUObject/Private/UObject/PropertyWeakObjectPtr.cpp
+++ b/Engine/Source/Runtime/CoreUObject/Private/UObject/PropertyWeakObjectPtr.cpp
@@ -107,6 +107,21 @@ void FWeakObjectProperty::SerializeItem( FStructuredArchive::FSlot Slot, void* V
 }
 
 #if UE_WITH_REMOTE_OBJECT_HANDLE
+bool FWeakObjectProperty::Identical(const void* A, const void* B, uint32 PortFlags) const
+{
+	if ((PortFlags & PPF_AvoidRemoteObjectMigration) != 0)
+	{
+		// With remote object handles enabled, weak pointers are equal if the globally unique remote ids of the objects they point to are identical
+		// which is the default behavior of FWeakObjectPtr::operator==
+		// This way we don't need to resolve the actual objects these weak pointers reference
+		FWeakObjectPtr ObjectA = A ? GetPropertyValue(A) : FWeakObjectPtr();
+		FWeakObjectPtr ObjectB = B ? GetPropertyValue(B) : FWeakObjectPtr();
+
+		return ObjectA == ObjectB;
+	}
+	return Super::Identical(A, B, PortFlags);
+}
+
 bool FWeakObjectProperty::NetSerializeItem(FArchive& Ar, UPackageMap* Map, void* Data, TArray<uint8>* MetaData) const
 {
 	const bool bUseRemoteObjectReference = Map && Map->IsUsingRemoteObjectReferences();
diff --git a/Engine/Source/Runtime/CoreUObject/Public/UObject/UnrealType.h b/Engine/Source/Runtime/CoreUObject/Public/UObject/UnrealType.h
--- a/Engine/Source/Runtime/CoreUObject/Public/UObject/UnrealType.h
+++ b/Engine/Source/Runtime/CoreUObject/Public/UObject/UnrealType.h
@@ -2991,6 +2991,7 @@ class COREUOBJECT_API FWeakObjectProperty : public TFObjectPropertyBase<FWeakObj
 	// FProperty interface
 	virtual void SerializeItem(FStructuredArchive::FSlot Slot, void* Value, void const* Defaults) const override;
 #if UE_WITH_REMOTE_OBJECT_HANDLE
+	virtual bool Identical(const void* A, const void* B, uint32 PortFlags) const override;
 	virtual bool NetSerializeItem(FArchive& Ar, UPackageMap* Map, void* Data, TArray<uint8>* MetaData = nullptr) const override;
 #endif
 protected:
