Removing deprecated EInternalObjectFlags::AsyncLoading flag to fix natvis not displaying EInternalObjectFlags properly as text

diff --git a/Engine/Source/Runtime/CoreUObject/Public/UObject/ObjectMacros.h b/Engine/Source/Runtime/CoreUObject/Public/UObject/ObjectMacros.h
--- a/Engine/Source/Runtime/CoreUObject/Public/UObject/ObjectMacros.h
+++ b/Engine/Source/Runtime/CoreUObject/Public/UObject/ObjectMacros.h
@@ -644,7 +644,6 @@ enum class EInternalObjectFlags : int32
 	// !!!!!!!!!!!!!!!!!!!!!!!!!!!!!!!!!!!!!!!!!!!!!!!!!!!!!!!!!!!!!!!!!!!!!!!!!!!!!!!!!!!!!!!!!!!!!!!!!!!!!!!!!!!!!!!!!!!!!!!!!!!!!!!!!!!
 	// DO NOT ADD composite flags to EInternalObjectFlags. Debugger visualisations have trouble displaying composite flag values as text.
 	// !!!!!!!!!!!!!!!!!!!!!!!!!!!!!!!!!!!!!!!!!!!!!!!!!!!!!!!!!!!!!!!!!!!!!!!!!!!!!!!!!!!!!!!!!!!!!!!!!!!!!!!!!!!!!!!!!!!!!!!!!!!!!!!!!!!
-	AsyncLoading UE_DEPRECATED(5.5, "Use EInternalObjectFlags_AsyncLoading instead.") = AsyncLoadingPhase1 | AsyncLoadingPhase2
 };
 ENUM_CLASS_FLAGS(EInternalObjectFlags);
 
