RemoteObject:  Adding debug functions to help figure out how FRemoteObjectIds can get lost across the network.

Use these in your watch window:
DebugFindRemoteObjectStub( (UObjectBase*)Pointer ) // will display the FRemoteObjectStub if it exists
DebugFindObjectLocallyFromRemoteId(2, 28595) // Will find the UObject* if it exists IN THE CURRENT PROCESS.  Change processes to find the UObject on the right server.

#rnx

diff --git a/Engine/Source/Runtime/CoreUObject/Private/UObject/RemoteObject.cpp b/Engine/Source/Runtime/CoreUObject/Private/UObject/RemoteObject.cpp
--- a/Engine/Source/Runtime/CoreUObject/Private/UObject/RemoteObject.cpp
+++ b/Engine/Source/Runtime/CoreUObject/Private/UObject/RemoteObject.cpp
@@ -164,7 +164,6 @@ void RegisterRemoteObjectId(FRemoteObjectId ObjectId, FRemoteServerId ResidentSe
 
 void RegisterSharedObject(UObject* Object)
 {
-	ensureMsgf(!Object->HasAnyFlags(EObjectFlags::RF_ClassDefaultObject | EObjectFlags::RF_ArchetypeObject), TEXT("We're about to set an ArchetypeObject %s as remote"), *GetNameSafe(Object));
 	FUObjectItem* ObjectItem = GUObjectArray.ObjectToObjectItem(Object);
 	ObjectItem->SetFlags(EInternalObjectFlags::RemoteReference);
 }
@@ -172,7 +171,7 @@ void RegisterSharedObject(UObject* Object)
 void MarkAsRemote(UObject* Object, FRemoteServerId DestinationServerId)
 {
 	static_assert(sizeof(FObjectHandle) == sizeof(UObject*));
-	checkf(!Object->IsA<UClass>(), TEXT("Attempted to Migrate UClass %s which is considered an asset and never allowed to migrate"), *GetNameSafe(Object));
+	checkf(!Object->IsTemplate(), TEXT("Attempted to Migrate Template Object '%s' which is considered an asset and never allowed to migrate"), *GetNameSafe(Object));
 
 	FUObjectItem* ObjectItem = GUObjectArray.ObjectToObjectItem(Object);
 	ObjectItem->SetFlags(EInternalObjectFlags::Remote);
@@ -554,3 +553,51 @@ FArchive& operator<<(FArchive& Ar, FRemoteServerId& Id)
 	Ar << Id.Id;
 	return Ar;
 }
+
+// Debugging functionality to help us find these objects in debug builds
+#if !UE_BUILD_SHIPPING && UE_WITH_REMOTE_OBJECT_HANDLE
+/**
+ * Put this in a Debug Watch Window on a specific UObject.  You may have to forcibly cast the UObject to UObjectBase*
+ * e.g. DebugFindRemoteObjectStub((UObjectBase*)Header.Class.DebugPtr)
+ */
+COREUOBJECT_API UE::RemoteObject::Handle::FRemoteObjectStub* DebugFindRemoteObjectStub(const UObjectBase* Object)
+{
+	if (!Object)
+	{
+		return nullptr;
+	}
+
+	uintptr_t Pointer = reinterpret_cast<uintptr_t>(Object);
+	if (Pointer & 0x1)
+	{
+		return reinterpret_cast<UE::RemoteObject::Handle::FRemoteObjectStub*>(Pointer & ~UPTRINT(1));
+	}
+
+	FRemoteObjectId ObjId { Object };
+	return UE::RemoteObject::Private::ObjectMaps->FindRemoteObjectStub(ObjId);
+}
+
+/**
+ * Attempt to find a UObject in the currently debugged process by its FRemoteObjectId constituents.
+ * Once you know a FRemoteObjectId, take its ServerId and SerialNumber and pass them into Debug Watch Window as arguments (in that order)
+ * e.g. DebugFindObjectLocallyFromRemoteId( 2, 1234 )
+ */
+COREUOBJECT_API UObject* DebugFindObjectLocallyFromRemoteId(uint16 ServerId, uint64 SerialNumber)
+{
+	return StaticFindObjectFastInternal(FRemoteObjectId(FRemoteServerId{ static_cast<uint32>(ServerId) }, SerialNumber));
+}
+
+/**
+ * Attempt to find a UObject in the currently debugged process by its FRemoteObjectId's Full uint64 Id
+ * Once you find a FRemoteObjectId, copy its Id and pass them into Debug Watch Window as the argument
+ * e.g. DebugFindObjectLocallyFromRemoteId( 1234567890 )
+ */
+COREUOBJECT_API UObject* DebugFindObjectLocallyFromRemoteId(uint64 FullId)
+{
+	constexpr uint64 SerialBitMask = (1ull << 54ull) - 1ull;
+	FRemoteObjectId RemoteId(FRemoteServerId{ static_cast<uint32>((FullId >> 54ull) & 1023ull) }, FullId & SerialBitMask);
+	ensure(RemoteId.GetIdNumber() == FullId);
+
+	return StaticFindObjectFastInternal(RemoteId);
+}
+#endif // !UE_BUILD_SHIPPING && UE_WITH_REMOTE_OBJECT_HANDLE
\ No newline at end of file
