- Support editor optional data in ZenStore (manual and auto)

diff --git a/Engine/Source/Runtime/Core/Public/IO/PackageStore.h b/Engine/Source/Runtime/Core/Public/IO/PackageStore.h
--- a/Engine/Source/Runtime/Core/Public/IO/PackageStore.h
+++ b/Engine/Source/Runtime/Core/Public/IO/PackageStore.h
@@ -40,6 +40,33 @@ struct FPackageStoreEntry
 #if WITH_EDITOR
 	TArrayView<const FPackageId> OptionalSegmentImportedPackageIds;
 	bool bHasOptionalSegment = false;
+	// This field is used to indicate that the package should load the optional chunk instead of the regular chunk. This is
+	// needed by the StorageServerPackgeStore (ie ZenStore) when loading "AutoOptional" assets, but the FilePackageStore 
+	// (ie .ucas files) does _not_ set this. This is because:
+	//    * AutoOptional is used to "silently" replace the normal asset with an asset that still has editor-only data in it
+	//    * When .ucas files are made, they create a .o.ucas file, which is mounted with a higher priority then the .ucas file
+	//    * When FilePackageStore loads, it will automatically read the editor-only version of the asset, _as if it was the normal asset_
+	//    * However, ZenStore has no ".o.ucas" file to silently read from instead of the regular file, so it must, at runtime,
+	//      request the editor version of the asset (the .o.uasset chunk in the store).
+	// This field is how we communicate out to the AsyncLoadingThread code that it should request the optional chunk instead of the
+	// regular chunk for a given asset. It is basically making a runtime decision that is handled offline by ucas files.
+	// Note that bHasOptionalSegment is always false for AutoOptional because of the "silently read optional as regular".
+	// So to sum up for the two main PackageStores:
+	//   ZenStore:
+	//     Manual Optional
+	//	     bHasOptionalSegment = true
+	//	     bReplaceChunkWithOptional = false
+	//	   AutoOptional
+	//	     bHasOptionalSegment = false
+	//	     bReplaceChunkWithOptional = true
+	//	 FileStore
+	//	   Manual Optional(same as Zen)
+	//	     bHasOptionalSegment = true
+	//	     bReplaceChunkWithOptional = false
+	//	   AutoOptional (_not_ same)
+	//	     bHasOptionalSegment = false
+	//	     bReplaceChunkWithOptional = false
+	bool bReplaceChunkWithOptional = false;
 #endif
 };
 
diff --git a/Engine/Source/Runtime/CoreUObject/Private/Serialization/AsyncLoading2.cpp b/Engine/Source/Runtime/CoreUObject/Private/Serialization/AsyncLoading2.cpp
--- a/Engine/Source/Runtime/CoreUObject/Private/Serialization/AsyncLoading2.cpp
+++ b/Engine/Source/Runtime/CoreUObject/Private/Serialization/AsyncLoading2.cpp
@@ -3622,6 +3622,7 @@ private:
 #if WITH_EDITOR
 	TOptional<FAsyncPackageHeaderData> OptionalSegmentHeaderData;
 	TOptional<FAsyncPackageSerializationState> OptionalSegmentSerializationState;
+	bool bRequestOptionalChunk = false;
 #endif
 #if ALT2_ENABLE_LINKERLOAD_SUPPORT
 	struct FLinkerLoadState
@@ -4647,6 +4648,7 @@ private:
 		const int32 ImportedPackagesCount = PackageStoreEntry.ImportedPackageIds.Num();
 		const uint64 ImportedPackageIdsMemSize = Align(sizeof(FPackageId) * ImportedPackagesCount, 8);
 #if WITH_EDITOR
+		const bool bRequestOptionalChunk = PackageStoreEntry.bReplaceChunkWithOptional;
 		const bool bHasOptionalSegment = PackageStoreEntry.bHasOptionalSegment;
 		const int32 OptionalSegmentImportedPackagesCount = PackageStoreEntry.OptionalSegmentImportedPackageIds.Num();
 		const uint64 OptionalSegmentImportedPackageIdsMemSize = Align(sizeof(FPackageId) * OptionalSegmentImportedPackagesCount, 8);
@@ -4713,6 +4715,8 @@ private:
 			HeaderData.ImportedAsyncPackagesView = Data.ImportedAsyncPackages.Left(ImportedPackagesCount);
 			OptionalSegmentHeaderData.ImportedAsyncPackagesView = Data.ImportedAsyncPackages.Right(OptionalSegmentImportedPackagesCount);
 		}
+		// track if later we need to request the optional chunk instead of regular
+		AsyncPackage->bRequestOptionalChunk = bRequestOptionalChunk;
 #endif
 		check(DataPtr - Data.MemoryBuffer0 == MemoryBufferSize);
 
@@ -6460,7 +6464,14 @@ void FAsyncPackage2::StartLoading(FAsyncLoadingThreadState2& ThreadState, FIoBat
 	int32 LocalPendingIoRequestsCounter = AsyncLoadingThread.PendingIoRequestsCounter.IncrementExchange() + 1;
 	TRACE_COUNTER_SET(AsyncLoadingPendingIoRequests, LocalPendingIoRequestsCounter);
 
-	SerializationState.IoRequest = IoBatch.ReadWithCallback(CreatePackageDataChunkId(Desc.PackageIdToLoad),
+
+#if WITH_EDITOR
+	const uint16 ChunkIndex = bRequestOptionalChunk ? 1 : 0;
+#else
+	const uint16 ChunkIndex = 0;
+#endif
+	FIoChunkId ChunkId = CreateIoChunkId(Desc.PackageIdToLoad.Value(), ChunkIndex, EIoChunkType::ExportBundleData);
+	SerializationState.IoRequest = IoBatch.ReadWithCallback(ChunkId,
 		ReadOptions,
 		Desc.Priority,
 		[this](TIoStatusOr<FIoBuffer> Result)
diff --git a/Engine/Source/Runtime/StorageServerClient/Private/StorageServerPackageStore.cpp b/Engine/Source/Runtime/StorageServerClient/Private/StorageServerPackageStore.cpp
--- a/Engine/Source/Runtime/StorageServerClient/Private/StorageServerPackageStore.cpp
+++ b/Engine/Source/Runtime/StorageServerClient/Private/StorageServerPackageStore.cpp
@@ -28,8 +28,30 @@ EPackageStoreEntryStatus FStorageServerPackageStoreBackend::GetPackageStoreEntry
 	const FStoreEntry* FindEntry = StoreEntriesMap.Find(PackageId);
 	if (FindEntry)
 	{
-		OutPackageStoreEntry.ImportedPackageIds = FindEntry->ImportedPackages;
 		OutPackageStoreEntry.ShaderMapHashes = FindEntry->ShaderMapHashes;
+#if WITH_EDITOR
+		// auto optional needs to request the optional chunk instead of regular, and because of that
+		// we use the optional imports as if they are regular imports, and we leave bHasOptionalSegment
+		// as false (this matches the FilePackageStore)
+		if (EnumHasAnyFlags(FindEntry->Flags, EPackageStoreEntryFlags::AutoOptional))
+		{
+			OutPackageStoreEntry.ImportedPackageIds = FindEntry->OptionalSegmentImportedPackageIds;
+			OutPackageStoreEntry.bReplaceChunkWithOptional = true;
+		}
+		// for manual optional, we report imported and optional imports as expected
+		else
+		{
+			OutPackageStoreEntry.ImportedPackageIds = FindEntry->ImportedPackages;
+			if (FindEntry->OptionalSegmentImportedPackageIds.Num() > 0)
+			{
+				OutPackageStoreEntry.OptionalSegmentImportedPackageIds = FindEntry->OptionalSegmentImportedPackageIds;
+				OutPackageStoreEntry.bHasOptionalSegment = true;
+			}
+		}
+#else
+		OutPackageStoreEntry.ImportedPackageIds = FindEntry->ImportedPackages;
+#endif
+
 		return EPackageStoreEntryStatus::Ok;
 	}
 	else
@@ -70,6 +92,10 @@ uint32 FStorageServerPackageStoreBackend::FAsyncInitRunnable::Run()
 		{
 			FStoreEntry& LocalEntry = Owner.StoreEntriesMap.FindOrAdd(PackageStoreEntryResource.GetPackageId());
 			LocalEntry.ImportedPackages = MoveTemp(PackageStoreEntryResource.ImportedPackageIds);
+#if WITH_EDITOR
+			LocalEntry.OptionalSegmentImportedPackageIds = MoveTemp(PackageStoreEntryResource.OptionalSegmentImportedPackageIds);
+			LocalEntry.Flags = PackageStoreEntryResource.Flags;
+#endif
 			LocalEntry.ShaderMapHashes = MoveTemp(PackageStoreEntryResource.ShaderMapHashes);
 		});
 
diff --git a/Engine/Source/Runtime/StorageServerClient/Private/StorageServerPackageStore.h b/Engine/Source/Runtime/StorageServerClient/Private/StorageServerPackageStore.h
--- a/Engine/Source/Runtime/StorageServerClient/Private/StorageServerPackageStore.h
+++ b/Engine/Source/Runtime/StorageServerClient/Private/StorageServerPackageStore.h
@@ -43,6 +43,10 @@ private:
 	{
 		TArray<FPackageId> ImportedPackages;
 		TArray<FSHAHash> ShaderMapHashes;
+#if WITH_EDITOR
+		EPackageStoreEntryFlags Flags;
+		TArray<FPackageId> OptionalSegmentImportedPackageIds;
+#endif
 	};
 	TMap<FPackageId, FStoreEntry> StoreEntriesMap;
 
