IncrementalCook: Fix property schemas used in class schemas to include hashing of IsEditorOnly. IncrementalCooks hashes the full list of editor properties, which remains the same if a property changes from usedingame to editoronly, and we need the hash to change when a property changes in that way because it impacts the cooked serialization and we need to recook the affected class instances.
#rnx
#incrementalcook

diff --git a/Engine/Source/Runtime/CoreUObject/Private/UObject/Property.cpp b/Engine/Source/Runtime/CoreUObject/Private/UObject/Property.cpp
--- a/Engine/Source/Runtime/CoreUObject/Private/UObject/Property.cpp
+++ b/Engine/Source/Runtime/CoreUObject/Private/UObject/Property.cpp
@@ -1176,6 +1176,10 @@ bool FProperty::SupportsNetSharedSerialization() const
 void FProperty::AppendSchemaHash(FBlake3& Builder, bool bSkipEditorOnly) const
 {
 	AppendHash(Builder, NamePrivate);
+	// Hash the EditorOnlyProperty so if a property changes from UsedInGame to EditorOnly it will still cause
+	// a change in the KeepEditorOnlyProperties hash.
+	uint8 FlagsRelevantToSchema = (IsEditorOnlyProperty() ? 1 : 0);
+	Builder.Update(&FlagsRelevantToSchema, sizeof(FlagsRelevantToSchema));
 	Builder.Update(&ArrayDim, sizeof(ArrayDim));
 	AppendHash(Builder, GetID());
 }
