SubObject Shadow Serialization for OS

#rnx

diff --git a/Engine/Source/Runtime/Core/Public/UObject/ObjectVersion.h b/Engine/Source/Runtime/Core/Public/UObject/ObjectVersion.h
--- a/Engine/Source/Runtime/Core/Public/UObject/ObjectVersion.h
+++ b/Engine/Source/Runtime/Core/Public/UObject/ObjectVersion.h
@@ -96,6 +96,9 @@ enum class EUnrealEngineObjectUE5Version : uint32
 	// Changed PackageFileSummary to write FIoHash PackageSavedHash instead of FGuid Guid
 	PACKAGE_SAVED_HASH,
 
+	// OS shadow serialization of subobjects
+	OS_SUB_OBJECT_SHADOW_SERIALIZATION,
+
 	// -----<new versions can be added before this line>-------------------------------------------------
 	// - this needs to be the last line (see note below)
 	AUTOMATIC_VERSION_PLUS_ONE,
diff --git a/Engine/Source/Runtime/CoreUObject/Private/UObject/Class.cpp b/Engine/Source/Runtime/CoreUObject/Private/UObject/Class.cpp
--- a/Engine/Source/Runtime/CoreUObject/Private/UObject/Class.cpp
+++ b/Engine/Source/Runtime/CoreUObject/Private/UObject/Class.cpp
@@ -1787,6 +1787,15 @@ void UStruct::SerializeVersionedTaggedProperties(FStructuredArchive::FSlot Slot,
 					{
 						UE_CLOG((UnderlyingArchive.IsPersistent() && FPlatformProperties::RequiresCookedData()), LogClass, Warning, TEXT("Skipping saved property %s of %s since it is no longer serializable for asset:  %s. (Maybe resave asset?)"), *Tag.Name.ToString(), *GetName(), *UnderlyingArchive.GetArchiveName());
 					}
+#if WITH_EDITORONLY_DATA
+					// Handle the sub object shadowing serialization,
+					// do not skip this property if loading into a loose property or into a placeholder
+					else if (!SerializeContext->bImpersonateProperties && Tag.OverrideOperation == EOverriddenPropertyOperation::SubObjectsShadowing)
+					{
+						UnderlyingArchive.Seek(StartOfProperty + Tag.Size);
+						bAdvanceProperty = true;
+					}
+#endif // WITH_EDITORONLY_DATA
 					else
 					{
 						FStructuredArchive::FSlot ValueSlot = PropertyRecord.EnterField(TEXT("Value"));
@@ -1883,7 +1892,7 @@ void UStruct::SerializeVersionedTaggedProperties(FStructuredArchive::FSlot Slot,
 								{
 									if (FOverriddenPropertySet* OverriddenProperties = FOverridableSerializationLogic::GetOverriddenProperties())
 									{
-										OverriddenProperties->SetOverriddenPropertyOperation(Tag.OverrideOperation, UnderlyingArchive.GetSerializedPropertyChain(), TempProperty);
+										OverriddenProperties->RestoreOverriddenPropertyOperation(Tag.OverrideOperation, UnderlyingArchive.GetSerializedPropertyChain(), TempProperty);
 									}
 								}
 							}
diff --git a/Engine/Source/Runtime/CoreUObject/Private/UObject/OverriddenPropertySet.cpp b/Engine/Source/Runtime/CoreUObject/Private/UObject/OverriddenPropertySet.cpp
--- a/Engine/Source/Runtime/CoreUObject/Private/UObject/OverriddenPropertySet.cpp
+++ b/Engine/Source/Runtime/CoreUObject/Private/UObject/OverriddenPropertySet.cpp
@@ -8,7 +8,6 @@
 #include "UObject/PropertyOptional.h"
 #include "Misc/ScopeExit.h"
 #include "UObject/InstanceDataObjectUtils.h"
-#include "UObject/PropertyBagRepository.h"
 #include "UObject/UObjectArchetypeHelper.h"
 #include "UObject/UObjectThreadContext.h"
 
@@ -21,7 +20,7 @@
 DEFINE_LOG_CATEGORY(LogOverridableObject);
 
 //----------------------------------------------------------------------//
-// FOverridableSerializationLogic
+// FOverridableSerializationLogicInternalAdapter
 //----------------------------------------------------------------------//
 
 struct FOverridableSerializationLogicInternalAdapter
@@ -51,61 +50,111 @@ namespace Private
 			FOverridableManager& Manager = FOverridableManager::Get();
 			FOverridableSerializationLogicInternalAdapter::SetCapability(FOverridableSerializationLogic::ECapabilities::T3DSerialization, bEnableT3D);
 		}));
-}
-
-ENUM_CLASS_FLAGS(FOverridableSerializationLogic::ECapabilities);
 
-FOverridableSerializationLogic::ECapabilities FOverridableSerializationLogic::Capabilities = FOverridableSerializationLogic::ECapabilities::None;
-thread_local bool FOverridableSerializationLogic::bUseOverridableSerialization = false;
-thread_local FOverriddenPropertySet* FOverridableSerializationLogic::OverriddenProperties = nullptr;
-thread_local FPropertyVisitorPath* FOverridableSerializationLogic::OverriddenPortTextPropertyPath = nullptr;
+	bool bEnableSubObjectsShadowSerialization = true;
+	FAutoConsoleVariableRef CVar_bEnableSubObjectsShadowSerialization(
+		TEXT("OverridableSerializationLogic.Capabilities.SubObjectsShadowSerialization"),
+		bEnableSubObjectsShadowSerialization,
+		TEXT("Enables shadow serialization of subobject"),
+		FConsoleVariableDelegate::CreateLambda([](IConsoleVariable* Variable)
+		{
+			FOverridableManager& Manager = FOverridableManager::Get();
+			FOverridableSerializationLogicInternalAdapter::SetCapability(FOverridableSerializationLogic::ECapabilities::SubObjectsShadowSerialization, bEnableSubObjectsShadowSerialization);
+		}));
 
-namespace Private
-{
 	struct FCapabilitiesAutoInitializer
 	{
 		FCapabilitiesAutoInitializer()
 		{
 			FOverridableSerializationLogicInternalAdapter::SetCapability(FOverridableSerializationLogic::ECapabilities::T3DSerialization, bEnableT3D);
+			FOverridableSerializationLogicInternalAdapter::SetCapability(FOverridableSerializationLogic::ECapabilities::SubObjectsShadowSerialization, bEnableSubObjectsShadowSerialization);
 		}
 	};
 
 	static FCapabilitiesAutoInitializer CapabilitiesAutoInitializer;
 }
 
+//----------------------------------------------------------------------//
+// FOverridableSerializationLogic
+//----------------------------------------------------------------------//
+
+ENUM_CLASS_FLAGS(FOverridableSerializationLogic::ECapabilities);
+
+FOverridableSerializationLogic::ECapabilities FOverridableSerializationLogic::Capabilities = FOverridableSerializationLogic::ECapabilities::None;
+thread_local bool FOverridableSerializationLogic::bUseOverridableSerialization = false;
+thread_local FOverriddenPropertySet* FOverridableSerializationLogic::OverriddenProperties = nullptr;
+thread_local FPropertyVisitorPath* FOverridableSerializationLogic::OverriddenPortTextPropertyPath = nullptr;
+
 EOverriddenPropertyOperation FOverridableSerializationLogic::GetOverriddenPropertyOperation(const int32 PortFlags, const FArchiveSerializedPropertyChain* CurrentPropertyChain, FProperty* Property, const void* DataPtr, const void* DefaultValue)
 {
 	checkf(bUseOverridableSerialization, TEXT("Nobody should use this method if it is not setup to use overridable serialization"));
+	if (!OverriddenProperties)
+	{
+		return EOverriddenPropertyOperation::None;
+	}
 
 	const FProperty* CurrentProperty = Property ? Property : (CurrentPropertyChain ? CurrentPropertyChain->GetPropertyFromStack(0) : nullptr);
-	if (CurrentProperty && CurrentProperty->HasAnyPropertyFlags(CPF_ExperimentalNeverOverriden))
+	checkf( CurrentProperty, TEXT("Expecting a property to get OS operation on") );
+
+	if (CurrentProperty->HasAnyPropertyFlags(CPF_ExperimentalNeverOverriden))
 	{
 		return EOverriddenPropertyOperation::None;
 	}
 
-	const EOverriddenPropertyOperation OverriddenOperation = OverriddenProperties ? OverriddenProperties->GetOverriddenPropertyOperation(CurrentPropertyChain, Property) : EOverriddenPropertyOperation::None;
+	const EOverriddenPropertyOperation OverriddenOperation = OverriddenProperties->GetOverriddenPropertyOperation(CurrentPropertyChain, Property);
 	if (OverriddenOperation != EOverriddenPropertyOperation::None)
 	{
 		return OverriddenOperation;
 	}
 
-	// It does not mean that if we have no record of an overriden operation that a subobject might have one, need to traverse all possible subobjects. 
-	if (CurrentProperty)
+	// In the case of a CDO owning default value, we might need to serialize it to keep its value.
+	if (DataPtr && DefaultValue && OverriddenProperties->IsCDOOwningProperty(*CurrentProperty))
 	{
-		// In the case of a CDO owning default value, we might need to serialize it to keep its value.
-		if (OverriddenProperties && OverriddenProperties->IsCDOOwningProperty(*CurrentProperty))
+		// Only need serialize this value if it is different from the default property value
+		if (!CurrentProperty->Identical(DataPtr, DefaultValue, PortFlags))
 		{
-			// Only need serialize this value if it is different from the default property value
-			if (!CurrentProperty->Identical(DataPtr, DefaultValue, PortFlags))
-			{
-				return 	EOverriddenPropertyOperation::Replace;
-			}
+			return 	EOverriddenPropertyOperation::Replace;
 		}
 	}
-	
+
+	if (ShouldPropertyShadowSerializeSubObject(CurrentProperty))
+	{
+		return EOverriddenPropertyOperation::SubObjectsShadowing;
+	}
+
 	return EOverriddenPropertyOperation::None;
 }
 
+bool FOverridableSerializationLogic::ShouldPropertyShadowSerializeSubObject(TNotNull<const FProperty*> Property)
+{
+	// Check if the shadow serialization of subobject is enabled
+	if (!HasCapabilities(ECapabilities::SubObjectsShadowSerialization))
+	{
+		return false;
+	}
+
+	// We shadow serialize every object property
+	if (CastField<FObjectPropertyBase>(Property))
+	{
+		return true;
+	}
+
+	// Otherwise check if the property is in the reference linked list
+	// @Todo optimized by caching the call to FProperty::ContainsObjectReference() maybe as a CPF_ContainsReferences?
+	checkf(Property->GetOwnerStruct(), TEXT("Expecting an owner struct for this type of property"));
+	FProperty* CurrentRefLink = Property->GetOwnerStruct()->RefLink;
+	while (CurrentRefLink != nullptr)
+	{
+		if (CurrentRefLink == Property)
+		{
+			return true;
+		}
+		CurrentRefLink = CurrentRefLink->NextRef;
+	}
+
+	return false;
+}
+
 bool FOverridableSerializationLogic::HasCapabilities(ECapabilities InCapabilities)
 {
 	return (Capabilities & InCapabilities) == InCapabilities;
@@ -119,7 +168,10 @@ FOverriddenPropertySet* FOverridableSerializationLogic::GetOverriddenPropertiesS
 EOverriddenPropertyOperation FOverridableSerializationLogic::GetOverriddenPropertyOperation(const FArchive& Ar, FProperty* Property /*= nullptr*/, uint8* DataPtr /*= nullptr*/, uint8* DefaultValue /*= nullptr*/)
 {
 	const FArchiveSerializedPropertyChain* CurrentPropertyChain = Ar.GetSerializedPropertyChain();
-	return GetOverriddenPropertyOperation(Ar.GetPortFlags(), CurrentPropertyChain, Property, DataPtr, DefaultValue);
+	const EOverriddenPropertyOperation Operation = GetOverriddenPropertyOperation(Ar.GetPortFlags(), CurrentPropertyChain, Property, DataPtr, DefaultValue);
+
+	// During transactions, we do not want any subobject shadow serialization
+	return Operation == EOverriddenPropertyOperation::SubObjectsShadowing && Ar.IsTransacting() ? EOverriddenPropertyOperation::None : Operation;
 }
 
 EOverriddenPropertyOperation FOverridableSerializationLogic::GetOverriddenPropertyOperationForPortText(const void* DataPtr, const void* DefaultValue, int32 PortFlags)
@@ -127,7 +179,10 @@ EOverriddenPropertyOperation FOverridableSerializationLogic::GetOverriddenProper
 	checkf(OverriddenPortTextPropertyPath, TEXT("Expecting an overridden port text path"));
 
 	const FArchiveSerializedPropertyChain CurrentPropertyChain = OverriddenPortTextPropertyPath->ToSerializedPropertyChain();
-	return GetOverriddenPropertyOperation(PortFlags,  &CurrentPropertyChain, nullptr, DataPtr, DefaultValue);
+	const EOverriddenPropertyOperation Operation = GetOverriddenPropertyOperation(PortFlags,  &CurrentPropertyChain, nullptr, DataPtr, DefaultValue);
+
+	// For now lets not support subobject shadow serialization until the copy and paste support loose property or placeholder
+	return Operation != EOverriddenPropertyOperation::SubObjectsShadowing ? Operation : EOverriddenPropertyOperation::None;
 }
 
 FPropertyVisitorPath* FOverridableSerializationLogic::GetOverriddenPortTextPropertyPath()
@@ -1379,10 +1434,12 @@ FOverriddenPropertyNode* FOverriddenPropertySet::SetOverriddenPropertyOperation(
 
 FOverriddenPropertyNode* FOverriddenPropertySet::RestoreOverriddenPropertyOperation(EOverriddenPropertyOperation Operation, const FArchiveSerializedPropertyChain* CurrentPropertyChain, FProperty* Property)
 {
-	// 'None' and 'Modified' operations are not needed to be restored on a property,
+	// 'None', 'Modified' and 'SubObjectShadowing' operations are not needed to be restored on a property,
 	// because 'None' is equal to the node not existing and
 	// 'Modified' will be restored when the sub property overrides will be restored successfully
-	if (Operation != EOverriddenPropertyOperation::None && Operation != EOverriddenPropertyOperation::Modified)
+	if (Operation != EOverriddenPropertyOperation::None && 
+		Operation != EOverriddenPropertyOperation::Modified &&
+		Operation != EOverriddenPropertyOperation::SubObjectsShadowing)
 	{
 		// Prevent marking as replaced the properties that are always overridden
 		if (!Property || Operation != EOverriddenPropertyOperation::Replace || !Property->HasAnyPropertyFlags(CPF_ExperimentalAlwaysOverriden))
diff --git a/Engine/Source/Runtime/CoreUObject/Private/UObject/PropertyArray.cpp b/Engine/Source/Runtime/CoreUObject/Private/UObject/PropertyArray.cpp
--- a/Engine/Source/Runtime/CoreUObject/Private/UObject/PropertyArray.cpp
+++ b/Engine/Source/Runtime/CoreUObject/Private/UObject/PropertyArray.cpp
@@ -130,6 +130,7 @@ void FArrayProperty::SerializeItem(FStructuredArchive::FSlot Slot, void* Value,
 	const bool bIsTextFormat = UnderlyingArchive.IsTextFormat();
 	const bool bUPS = UnderlyingArchive.UseUnversionedPropertySerialization();
 	bool bExperimentalOverridableLogic = HasAnyPropertyFlags(CPF_ExperimentalOverridableLogic);
+	FUObjectSerializeContext* Context = FUObjectThreadContext::Get().GetSerializeContext();
 	TOptional<FPropertyTag> MaybeInnerTag;
 
 	// Ensure that the Inner itself has been loaded before calling SerializeItem() on it
@@ -344,7 +345,7 @@ void FArrayProperty::SerializeItem(FStructuredArchive::FSlot Slot, void* Value,
 					}
 				}
 
-				// @Todo can be removed with an engine version, not being saved anymore on new objects
+				// We should not enter this code anymore but keeping for eventual support of modified values like struct
 				int32 NumModified = 0;
 				FStructuredArchive::FArray ModifiedArray = Record.EnterArray(TEXT("Modified"), NumModified);
 				if (NumModified != 0)
@@ -373,6 +374,46 @@ void FArrayProperty::SerializeItem(FStructuredArchive::FSlot Slot, void* Value,
 					}
 				}
 
+				// Support of subobject shadowed serialization
+				if (UnderlyingArchive.UEVer() >= EUnrealEngineObjectUE5Version::OS_SUB_OBJECT_SHADOW_SERIALIZATION)
+				{
+					int32 NumShadowed = 0;
+					FStructuredArchive::FArray ShadowedArray = Record.EnterArray(TEXT("Shadowed"), NumShadowed);
+					if (NumShadowed != 0)
+					{
+						if (!TempValueStorage)
+						{
+							TempValueStorage = (uint8*)FMemory::Malloc(InnerObjectProperty->GetElementSize());
+							InnerObjectProperty->InitializeValue(TempValueStorage);
+						}
+
+						for (int32 i = 0; i < NumShadowed; ++i)
+						{
+							{
+								FSerializedPropertyScope SerializedProperty(UnderlyingArchive, Inner, this);
+								SerializeContainerItem(ShadowedArray.EnterElement(), TempValueStorage);
+							}
+
+							// Only modifying property when loading loose properties or placeholders
+							if (Context->bImpersonateProperties)
+							{
+								int32 AddIndex = ArrayHelper.Num();
+								if (UObject* ShadowedObject = InnerObjectProperty->GetObjectPropertyValue(TempValueStorage))
+								{
+									int32 Index = FindObject(ShadowedObject->GetArchetype(), ShadowedObject, ArrayHelper);
+									if (Index == INDEX_NONE)
+									{
+										ArrayHelper.AddValue();
+										Index = AddIndex++;
+									}
+									InnerObjectProperty->SetObjectPropertyValue(ArrayHelper.GetRawPtr(Index), ShadowedObject);
+								}
+							}
+						}
+					}
+				}
+
+
 				int32 NumAdded = 0;
 				FStructuredArchive::FArray AddedArray = Record.EnterArray(TEXT("Added"), NumAdded);
 				if (NumAdded != 0)
@@ -417,6 +458,7 @@ void FArrayProperty::SerializeItem(FStructuredArchive::FSlot Slot, void* Value,
 			// Container for temporarily tracking some indices
 			TArray<int32> RemovedIndices;
 			TArray<int32> AddedIndices;
+			TArray<int32> ShadowIndices;
 
 			// We filter out external objects if:
 			//	- We're an array of objects.
@@ -446,6 +488,15 @@ void FArrayProperty::SerializeItem(FStructuredArchive::FSlot Slot, void* Value,
 
 				if (!bReplaceArray)
 				{
+					if (FOverridableSerializationLogic::ShouldPropertyShadowSerializeSubObject(this))
+					{
+						const int32 ArrayNum = ArrayHelper.Num();
+						for( int i = 0; i < ArrayNum; i++)
+						{
+							ShadowIndices.Add(i);
+						}
+					}
+
 					// Only array of instanced subobjects are handled here as sort of a set where the matching key is done using the archetype
 					checkf(InnerObjectProperty&& InnerObjectProperty->HasAnyPropertyFlags(CPF_PersistentInstance), TEXT("Expecting only arrays of instanced subobjects"));
 
@@ -502,6 +553,7 @@ void FArrayProperty::SerializeItem(FStructuredArchive::FSlot Slot, void* Value,
 											if (!bFilterExternalObjects)
 											{
 												AddedIndices.Add(Index);
+												ShadowIndices.Remove(Index);
 											}
 											else
 											{
@@ -509,6 +561,7 @@ void FArrayProperty::SerializeItem(FStructuredArchive::FSlot Slot, void* Value,
 												if(Object && !Object->IsPackageExternal())
 												{
 													AddedIndices.Add(Index);
+													ShadowIndices.Remove(Index);
 												}
 											}
 										}
@@ -553,10 +606,19 @@ void FArrayProperty::SerializeItem(FStructuredArchive::FSlot Slot, void* Value,
 					SerializeContainerItem(RemovedArray.EnterElement(), DefaultsArrayHelper.GetRawPtr(i));
 				}
 
-				// @Todo can be removed with an engine version, not being saved anymore on new objects
+				// We should not have modified elements anymore but keeping for eventual support of modified values like struct
 				int32 NumModified = 0;
 				FStructuredArchive::FArray ModifiedArray = Record.EnterArray(TEXT("Modified"), NumModified);
 
+				// Support of subobject shadowed serialization
+				// Introduced from EUnrealEngineObjectUE5Version::OS_SUB_OBJECT_SHADOW_SERIALIZATION
+				int32 NumShadowed = ShadowIndices.Num();
+				FStructuredArchive::FArray ShadowedArray = Record.EnterArray(TEXT("Shadowed"), NumShadowed);
+				for (int32 i : ShadowIndices)
+				{
+					SerializeContainerItem(ShadowedArray.EnterElement(), ArrayHelper.GetRawPtr(i));
+				}
+
 				int32 NumAdded = AddedIndices.Num();
 				FStructuredArchive::FArray AddedArray = Record.EnterArray(TEXT("Added"), NumAdded);
 				for (int32 i : AddedIndices)
@@ -592,8 +654,6 @@ void FArrayProperty::SerializeItem(FStructuredArchive::FSlot Slot, void* Value,
 	}
 	ArrayHelper.CountBytes(UnderlyingArchive);
 
-	FUObjectSerializeContext* Context = FUObjectThreadContext::Get().GetSerializeContext();
-
 	// Serialize a PropertyTag for the inner property of this array, allows us to validate the inner struct to see if it has changed
 	if (UnderlyingArchive.UEVer() < EUnrealEngineObjectUE5Version::PROPERTY_TAG_COMPLETE_TYPE_NAME &&
 		UnderlyingArchive.UEVer() >= VER_UE4_INNER_ARRAY_TAG_INFO &&
diff --git a/Engine/Source/Runtime/CoreUObject/Private/UObject/PropertyMap.cpp b/Engine/Source/Runtime/CoreUObject/Private/UObject/PropertyMap.cpp
--- a/Engine/Source/Runtime/CoreUObject/Private/UObject/PropertyMap.cpp
+++ b/Engine/Source/Runtime/CoreUObject/Private/UObject/PropertyMap.cpp
@@ -365,14 +365,15 @@ void FMapProperty::SerializeItem(FStructuredArchive::FSlot Slot, void* Value, co
 								KeyProp->SerializeItem(RemovedArray.EnterElement().EnterRecord().EnterField(TEXT("Key")), TempKeyValueStorage);
 							}
 
-							MapHelper.RemovePair(TempKeyValueStorage);
-
-							// Need to fetch the MapOverriddenPropertyNode every loop as the previous might have reallocated the node.
-							if (FOverriddenPropertyNode* MapOverriddenPropertyNode = OverriddenProperties ? OverriddenProperties->SetOverriddenPropertyOperation(EOverriddenPropertyOperation::Modified, UnderlyingArchive.GetSerializedPropertyChain(), /*Property*/nullptr) : nullptr)
+							if (MapHelper.RemovePair(TempKeyValueStorage))
 							{
-								// Rebuild the overridden info
-								FOverriddenPropertyNodeID RemovedKeyID = FOverriddenPropertyNodeID::FromMapKey(KeyProp, TempKeyValueStorage);
-								OverriddenProperties->SetSubPropertyOperation(EOverriddenPropertyOperation::Remove, *MapOverriddenPropertyNode, RemovedKeyID);
+								// Need to fetch the MapOverriddenPropertyNode every loop as the previous might have reallocated the node.
+								if (FOverriddenPropertyNode* MapOverriddenPropertyNode = OverriddenProperties ? OverriddenProperties->SetOverriddenPropertyOperation(EOverriddenPropertyOperation::Modified, UnderlyingArchive.GetSerializedPropertyChain(), /*Property*/nullptr) : nullptr)
+								{
+									// Rebuild the overridden info
+									FOverriddenPropertyNodeID RemovedKeyID = FOverriddenPropertyNodeID::FromMapKey(KeyProp, TempKeyValueStorage);
+									OverriddenProperties->SetSubPropertyOperation(EOverriddenPropertyOperation::Remove, *MapOverriddenPropertyNode, RemovedKeyID);
+								}
 							}
 						}
 					}
@@ -408,12 +409,53 @@ void FMapProperty::SerializeItem(FStructuredArchive::FSlot Slot, void* Value, co
 								ValueProp->SerializeItem(EntryRecord.EnterField(TEXT("Value")), ValuePtr);
 							}
 
-							// Need to fetch the MapOverriddenPropertyNode every loop as the previous might have reallocated the node.
-							if (FOverriddenPropertyNode* MapOverriddenPropertyNode = OverriddenProperties ? OverriddenProperties->SetOverriddenPropertyOperation(EOverriddenPropertyOperation::Modified, UnderlyingArchive.GetSerializedPropertyChain(), /*Property*/nullptr) : nullptr)
+							// Track only if we found the key in the array. Otherwise, skip it.
+							if (Index != INDEX_NONE)
 							{
-								// Rebuild the overridden info
-								FOverriddenPropertyNodeID ModifiedKeyID = FOverriddenPropertyNodeID::FromMapKey(KeyProp, TempKeyValueStorage);
-								OverriddenProperties->SetSubPropertyOperation(EOverriddenPropertyOperation::Modified, *MapOverriddenPropertyNode, ModifiedKeyID);
+								// Need to fetch the MapOverriddenPropertyNode every loop as the previous might have reallocated the node.
+								if (FOverriddenPropertyNode* MapOverriddenPropertyNode = OverriddenProperties ? OverriddenProperties->SetOverriddenPropertyOperation(EOverriddenPropertyOperation::Modified, UnderlyingArchive.GetSerializedPropertyChain(), /*Property*/nullptr) : nullptr)
+								{
+									// Rebuild the overridden info
+									FOverriddenPropertyNodeID ModifiedKeyID = FOverriddenPropertyNodeID::FromMapKey(KeyProp, TempKeyValueStorage);
+									OverriddenProperties->SetSubPropertyOperation(EOverriddenPropertyOperation::Modified, *MapOverriddenPropertyNode, ModifiedKeyID);
+								}
+							}
+						}
+					}
+
+					// Support of subobject shadowed serialization
+					if (UnderlyingArchive.UEVer() >= EUnrealEngineObjectUE5Version::OS_SUB_OBJECT_SHADOW_SERIALIZATION)
+					{
+						int32 NumShadowed = 0;
+						FStructuredArchive::FArray ShadowedArray = Record.EnterArray(TEXT("Shadowed"), NumShadowed);
+						if (NumShadowed != 0)
+						{
+							if (!TempKeyValueStorage)
+							{
+								TempKeyValueStorage = (uint8*)FMemory::Malloc(MapLayout.SetLayout.Size);
+								KeyProp->InitializeValue(TempKeyValueStorage);
+								ValueProp->InitializeValue(TempKeyValueStorage + MapLayout.ValueOffset);
+							}
+							for (int32 i = 0; i < NumShadowed; ++i)
+							{
+								FStructuredArchive::FRecord EntryRecord = ShadowedArray.EnterElement().EnterRecord();
+
+								// Read key into temporary storage
+								{
+									UE::FSerializedPropertyPathScope SerializedPropertyPath(Context, {UE::NAME_Key});
+									FSerializedPropertyScope SerializedProperty(UnderlyingArchive, KeyProp, this);
+									KeyProp->SerializeItem(EntryRecord.EnterField(TEXT("Key")), TempKeyValueStorage);
+								}
+
+								// Only modifying property when loading loose properties or placeholders, oterwise load in temp storage
+								void* ValuePtr = Context->bImpersonateProperties ? MapHelper.FindOrAdd(TempKeyValueStorage) : TempKeyValueStorage + MapLayout.ValueOffset;
+
+								// Deserialize value
+								{
+									UE::FSerializedPropertyPathScope SerializedPropertyPath(Context, {UE::NAME_Value});
+									FSerializedPropertyScope SerializedProperty(UnderlyingArchive, ValueProp, this);
+									ValueProp->SerializeItem(EntryRecord.EnterField(TEXT("Value")), ValuePtr);
+								}
 							}
 						}
 					}
@@ -466,6 +508,7 @@ void FMapProperty::SerializeItem(FStructuredArchive::FSlot Slot, void* Value, co
 				TArray<int32> RemovedIndices;
 				TArray<int32> AddedIndices;
 				TSet<int32> ModifiedIndices;
+				TSet<int32> ShadowedIndices;
 
 				bool bReplaceMap = false;
 				if (!Defaults || !UnderlyingArchive.DoDelta() || UnderlyingArchive.IsTransacting())
@@ -491,7 +534,15 @@ void FMapProperty::SerializeItem(FStructuredArchive::FSlot Slot, void* Value, co
 					{
 						checkf(!KeyProp->HasAnyPropertyFlags(CPF_PersistentInstance) || CastField<FClassProperty>(KeyProp) || !CastField<FObjectProperty>(KeyProp), TEXT("The key as an instanced sub object is NYI"));
 
-						if (OverriddenProperties && MapOverrideOp != EOverriddenPropertyOperation::None)
+						if (FOverridableSerializationLogic::ShouldPropertyShadowSerializeSubObject(this))
+						{
+							for (FScriptMapHelper::FIterator It(MapHelper); It; ++It)
+							{
+								ShadowedIndices.Add(It.GetInternalIndex());
+							}
+						}
+
+						if (OverriddenProperties)
 						{
 							checkf(Defaults, TEXT("Expecting overridable serialization to have defaults to compare to"));
 							FScriptMapHelper DefaultsMapHelper(this, Defaults);
@@ -519,6 +570,7 @@ void FMapProperty::SerializeItem(FStructuredArchive::FSlot Slot, void* Value, co
 											if (InternalIndex != INDEX_NONE)
 											{
 												AddedIndices.Add(InternalIndex);
+												ShadowedIndices.Remove(InternalIndex);
 											}
 											break;
 										}
@@ -528,6 +580,7 @@ void FMapProperty::SerializeItem(FStructuredArchive::FSlot Slot, void* Value, co
 											if (InternalIndex != INDEX_NONE)
 											{
 												ModifiedIndices.Add(InternalIndex);
+												ShadowedIndices.Remove(InternalIndex);
 											}
 											break;
 										}
@@ -541,24 +594,28 @@ void FMapProperty::SerializeItem(FStructuredArchive::FSlot Slot, void* Value, co
 					}
 				}
 
+				auto SerializePair = [this, Context, &UnderlyingArchive](FStructuredArchive::FArray& Array, uint8* PairPtr)
+				{
+					FStructuredArchive::FRecord EntryRecord = Array.EnterElement().EnterRecord();
+					{
+						UE::FSerializedPropertyPathScope SerializedPropertyPath(Context, {UE::NAME_Key});
+						FSerializedPropertyScope SerializedProperty(UnderlyingArchive, KeyProp, this);
+						KeyProp->SerializeItem(EntryRecord.EnterField(TEXT("Key")), PairPtr);
+					}
+					{
+						UE::FSerializedPropertyPathScope SerializedPropertyPath(Context, {UE::NAME_Value});
+						FSerializedPropertyScope SerializedProperty(UnderlyingArchive, ValueProp, this);
+						ValueProp->SerializeItem(EntryRecord.EnterField(TEXT("Value")), PairPtr + MapLayout.ValueOffset);
+					}
+				};
+
 				int32 NumReplaced = bReplaceMap ? MapHelper.Num() : INDEX_NONE;
 				FStructuredArchive::FArray ReplacedArray = Record.EnterArray(TEXT("Replaced"), NumReplaced);
 				if (bReplaceMap)
 				{
 					for (FScriptMapHelper::FIterator It(MapHelper); It; ++It)
 					{
-						FStructuredArchive::FRecord EntryRecord = ReplacedArray.EnterElement().EnterRecord();
-						uint8* PairPtr = MapHelper.GetPairPtr(It.GetInternalIndex());
-						{
-							UE::FSerializedPropertyPathScope SerializedPropertyPath(Context, {UE::NAME_Key});
-							FSerializedPropertyScope SerializedProperty(UnderlyingArchive, KeyProp, this);
-							KeyProp->SerializeItem(EntryRecord.EnterField(TEXT("Key")), PairPtr);
-						}
-						{
-							UE::FSerializedPropertyPathScope SerializedPropertyPath(Context, {UE::NAME_Value});
-							FSerializedPropertyScope SerializedProperty(UnderlyingArchive, ValueProp, this);
-							ValueProp->SerializeItem(EntryRecord.EnterField(TEXT("Value")), PairPtr + MapLayout.ValueOffset);
-						}
+						SerializePair(ReplacedArray, MapHelper.GetPairPtr(It.GetInternalIndex()));
 					}
 				}
 				else
@@ -580,18 +637,16 @@ void FMapProperty::SerializeItem(FStructuredArchive::FSlot Slot, void* Value, co
 					FStructuredArchive::FArray ModifiedArray = Record.EnterArray(TEXT("Modified"), NumModified);
 					for (int32 i : ModifiedIndices)
 					{
-						FStructuredArchive::FRecord EntryRecord = ModifiedArray.EnterElement().EnterRecord();
-						uint8* PairPtr = MapHelper.GetPairPtr(i);
-						{
-							UE::FSerializedPropertyPathScope SerializedPropertyPath(Context, {UE::NAME_Key});
-							FSerializedPropertyScope SerializedProperty(UnderlyingArchive, KeyProp, this);
-							KeyProp->SerializeItem(EntryRecord.EnterField(TEXT("Key")), PairPtr);
-						}
-						{
-							UE::FSerializedPropertyPathScope SerializedPropertyPath(Context, {UE::NAME_Value});
-							FSerializedPropertyScope SerializedProperty(UnderlyingArchive, ValueProp, this);
-							ValueProp->SerializeItem(EntryRecord.EnterField(TEXT("Value")), PairPtr + MapLayout.ValueOffset);
-						}
+						SerializePair(ModifiedArray, MapHelper.GetPairPtr(i));
+					}
+
+					// Support of subobject shadowed serialization
+					// Introduced from EUnrealEngineObjectUE5Version::OS_SUB_OBJECT_SHADOW_SERIALIZATION
+					int32 NumShadowed = ShadowedIndices.Num();
+					FStructuredArchive::FArray ShadowedArray = Record.EnterArray(TEXT("Shadowed"), NumShadowed);
+					for (int32 i : ShadowedIndices)
+					{
+						SerializePair(ShadowedArray, MapHelper.GetPairPtr(i));
 					}
 
 					// Added keys
@@ -599,18 +654,7 @@ void FMapProperty::SerializeItem(FStructuredArchive::FSlot Slot, void* Value, co
 					FStructuredArchive::FArray AddedArray = Record.EnterArray(TEXT("Added"), NumAdded);
 					for (int32 i : AddedIndices)
 					{
-						FStructuredArchive::FRecord EntryRecord = AddedArray.EnterElement().EnterRecord();
-						uint8* PairPtr = MapHelper.GetPairPtr(i);
-						{
-							UE::FSerializedPropertyPathScope SerializedPropertyPath(Context, {UE::NAME_Key});
-							FSerializedPropertyScope SerializedProperty(UnderlyingArchive, KeyProp, this);
-							KeyProp->SerializeItem(EntryRecord.EnterField(TEXT("Key")), PairPtr);
-						}
-						{
-							UE::FSerializedPropertyPathScope SerializedPropertyPath(Context, {UE::NAME_Value});
-							FSerializedPropertyScope SerializedProperty(UnderlyingArchive, ValueProp, this);
-							ValueProp->SerializeItem(EntryRecord.EnterField(TEXT("Value")), PairPtr + MapLayout.ValueOffset);
-						}
+						SerializePair(AddedArray, MapHelper.GetPairPtr(i));
 					}
 				}
 			}
diff --git a/Engine/Source/Runtime/CoreUObject/Private/UObject/PropertyTag.cpp b/Engine/Source/Runtime/CoreUObject/Private/UObject/PropertyTag.cpp
--- a/Engine/Source/Runtime/CoreUObject/Private/UObject/PropertyTag.cpp
+++ b/Engine/Source/Runtime/CoreUObject/Private/UObject/PropertyTag.cpp
@@ -128,24 +128,25 @@ static EPropertyTagExtension CalculatePropertyExtensionFlags(FArchive& Underlyin
 
 	if (UnderlyingArchive.IsSaving())
 	{
-		// OverridableInformation
-		const FOverriddenPropertySet* OverriddenProperties = FOverridableSerializationLogic::GetOverriddenProperties();
-		Tag.bExperimentalOverridableLogic = Tag.GetProperty()->HasAnyPropertyFlags(CPF_ExperimentalOverridableLogic);
-		if (OverriddenProperties || Tag.bExperimentalOverridableLogic)
+		if(FOverridableSerializationLogic::IsEnabled())
 		{
-			// If this property is marked as always overridden, emit op to replace any inherited/default value with our own. If this
-			// value is later read back in as a loose property, it will be properly treated as an override during write impersonation.
-			if (Tag.GetProperty()->HasAnyPropertyFlags(CPF_ExperimentalAlwaysOverriden))
-			{
-				Tag.OverrideOperation = EOverriddenPropertyOperation::Replace;
-			}
-			else
-			{
-				Tag.OverrideOperation = OverriddenProperties ? OverriddenProperties->GetOverriddenPropertyOperation(UnderlyingArchive.GetSerializedPropertyChain(), Tag.GetProperty()) : EOverriddenPropertyOperation::None;
-			}
+			Tag.OverrideOperation = FOverridableSerializationLogic::GetOverriddenPropertyOperation(UnderlyingArchive, Tag.GetProperty());
+			PropertyTagExtensions |= EPropertyTagExtension::OverridableInformation;
+		}
+		else
+		{
+			Tag.OverrideOperation = EOverriddenPropertyOperation::None;
+		}
 
+		if (Tag.GetProperty()->HasAnyPropertyFlags(CPF_ExperimentalOverridableLogic))
+		{
+			Tag.bExperimentalOverridableLogic = true;
 			PropertyTagExtensions |= EPropertyTagExtension::OverridableInformation;
 		}
+		else
+		{
+			Tag.bExperimentalOverridableLogic = false;
+		}
 	}
 
 	return PropertyTagExtensions;
diff --git a/Engine/Source/Runtime/CoreUObject/Private/UObject/SavePackage/PackageHarvester.cpp b/Engine/Source/Runtime/CoreUObject/Private/UObject/SavePackage/PackageHarvester.cpp
--- a/Engine/Source/Runtime/CoreUObject/Private/UObject/SavePackage/PackageHarvester.cpp
+++ b/Engine/Source/Runtime/CoreUObject/Private/UObject/SavePackage/PackageHarvester.cpp
@@ -741,9 +741,10 @@ void FPackageHarvester::TryHarvestExportInternal(UObject* InObject, UE::SavePack
 
 	HarvestExport(InObject);
 
+	// @todo: Remove this code once we know that the SubObjectsShadowSerialization works with SceneGraph and it is tested correctly.
 	// Objects that has overridable serialization enabled, needs to separately export their subobject independently of the property serializing it.
 	// This is because it needs to make a difference between the pointer that is overridden vs some properties inside the instanced sub object that are overridden.
-	if (FOverridableManager::Get().IsEnabled(InObject))
+	if (!FOverridableSerializationLogic::HasCapabilities(FOverridableSerializationLogic::ECapabilities::SubObjectsShadowSerialization) && FOverridableManager::Get().IsEnabled(InObject))
 	{
 		for (TPropertyValueIterator<FObjectProperty> ObjPropIt(InObject->GetClass(), InObject); ObjPropIt; ++ObjPropIt)
 		{
diff --git a/Engine/Source/Runtime/CoreUObject/Public/UObject/OverriddenPropertySet.h b/Engine/Source/Runtime/CoreUObject/Public/UObject/OverriddenPropertySet.h
--- a/Engine/Source/Runtime/CoreUObject/Public/UObject/OverriddenPropertySet.h
+++ b/Engine/Source/Runtime/CoreUObject/Public/UObject/OverriddenPropertySet.h
@@ -26,7 +26,10 @@ public:
 		None = 0,
 		// Enables serialization of override state into and from T3D
 		T3DSerialization = 1 << 0,
-		All = T3DSerialization
+		// Enables shadow serialization of subobject, any saved data can be skipped at load time unless loading in loose property or in a placeholder.
+		// This will keep reference to subobjects that might contain overrides.
+		SubObjectsShadowSerialization = 1 << 1,
+		All = T3DSerialization | SubObjectsShadowSerialization
 	};
 	
 	FRIEND_ENUM_CLASS_FLAGS(FOverridableSerializationLogic::ECapabilities)
@@ -114,6 +117,9 @@ public:
 	/** Call during the import text to reset property path */
 	COREUOBJECT_API static void ResetOverriddenPortTextPropertyPath();
 
+	/** To know if the specified property should shadow serialize its values for OS */
+	COREUOBJECT_API static bool ShouldPropertyShadowSerializeSubObject(TNotNull<const FProperty*> Property);
+
 private:
 
 	static EOverriddenPropertyOperation GetOverriddenPropertyOperation(const int32 PortFlags, const FArchiveSerializedPropertyChain* CurrentPropertyChain, FProperty* Property, const void* DataPtr, const void* DefaultValue);
@@ -172,6 +178,7 @@ enum class EOverriddenPropertyOperation : uint8
 	Replace,	/* everything has been overridden from this property down to every sub property/sub object*/
 	Add,		/* this element was added in the container */
 	Remove,		/* this element was removed from the container */
+	SubObjectsShadowing, /* only used to serialize/collect the subobjects, this saved data can be skipped at loading unless loading a loose property or in a placeholder which will keep reference to subobjects */
 };
 
 inline TOptional<EOverriddenPropertyOperation> GetOverriddenOperationFromString(const FString& OverriddenOperationString)
