Fix build break

#rnx

diff --git a/Engine/Source/Runtime/CoreUObject/Private/UObject/PropertyArray.cpp b/Engine/Source/Runtime/CoreUObject/Private/UObject/PropertyArray.cpp
--- a/Engine/Source/Runtime/CoreUObject/Private/UObject/PropertyArray.cpp
+++ b/Engine/Source/Runtime/CoreUObject/Private/UObject/PropertyArray.cpp
@@ -395,6 +395,7 @@ void FArrayProperty::SerializeItem(FStructuredArchive::FSlot Slot, void* Value,
 							}
 
 							// Only modifying property when loading loose properties or placeholders
+#if WITH_EDITORONLY_DATA
 							if (Context->bImpersonateProperties)
 							{
 								int32 AddIndex = ArrayHelper.Num();
@@ -409,6 +410,8 @@ void FArrayProperty::SerializeItem(FStructuredArchive::FSlot Slot, void* Value,
 									InnerObjectProperty->SetObjectPropertyValue(ArrayHelper.GetRawPtr(Index), ShadowedObject);
 								}
 							}
+#endif //  WITH_EDITORONLY_DATA
+
 						}
 					}
 				}
diff --git a/Engine/Source/Runtime/CoreUObject/Private/UObject/PropertyMap.cpp b/Engine/Source/Runtime/CoreUObject/Private/UObject/PropertyMap.cpp
--- a/Engine/Source/Runtime/CoreUObject/Private/UObject/PropertyMap.cpp
+++ b/Engine/Source/Runtime/CoreUObject/Private/UObject/PropertyMap.cpp
@@ -448,7 +448,11 @@ void FMapProperty::SerializeItem(FStructuredArchive::FSlot Slot, void* Value, co
 								}
 
 								// Only modifying property when loading loose properties or placeholders, oterwise load in temp storage
-								void* ValuePtr = Context->bImpersonateProperties ? MapHelper.FindOrAdd(TempKeyValueStorage) : TempKeyValueStorage + MapLayout.ValueOffset;
+								void* ValuePtr = 
+#if WITH_EDITORONLY_DATA
+								Context->bImpersonateProperties ? MapHelper.FindOrAdd(TempKeyValueStorage) :
+#endif // WITH_EDITORONLY_DATA
+								TempKeyValueStorage + MapLayout.ValueOffset;
 
 								// Deserialize value
 								{
