Compiled out ValidateGeneratedRepEnums and ValidateNetDeltaSerializeProperties based on DO_CHECK.

Checked server builds will still have this code compiled in, which is intended. The actual checks are disabled by default in shipping and enabled by default in all other configs.
Saves ~0.5MB ProgramSize on Win64.

diff --git a/Engine/Source/Programs/Shared/EpicGames.UHT/Exporters/CodeGen/UhtHeaderCodeGeneratorCppFile.cs b/Engine/Source/Programs/Shared/EpicGames.UHT/Exporters/CodeGen/UhtHeaderCodeGeneratorCppFile.cs
--- a/Engine/Source/Programs/Shared/EpicGames.UHT/Exporters/CodeGen/UhtHeaderCodeGeneratorCppFile.cs
+++ b/Engine/Source/Programs/Shared/EpicGames.UHT/Exporters/CodeGen/UhtHeaderCodeGeneratorCppFile.cs
@@ -2067,6 +2067,7 @@ namespace EpicGames.UHT.Exporters.CodeGen
 
 			if (classObj.ClassExportFlags.HasAnyFlags(UhtClassExportFlags.SelfHasReplicatedProperties))
 			{
+				builder.Append("#if VALIDATE_CLASS_REPS\r\n");
 				builder.Append("void ").Append(classObj.SourceName).Append("::ValidateGeneratedRepEnums(const TArray<struct FRepRecord>& ClassReps) const\r\n");
 				builder.Append("{\r\n");
 
@@ -2096,6 +2097,7 @@ namespace EpicGames.UHT.Exporters.CodeGen
 				builder.Append(";\r\n");
 				builder.Append("\tcheckf(bIsValid, TEXT(\"UHT Generated Rep Indices do not match runtime populated Rep Indices for properties in ").Append(classObj.SourceName).Append("\"));\r\n");
 				builder.Append("}\r\n");
+				builder.Append("#endif\r\n");
 			}
 			return builder;
 		}
diff --git a/Engine/Source/Programs/Shared/EpicGames.UHT/Exporters/CodeGen/UhtHeaderCodeGeneratorHFile.cs b/Engine/Source/Programs/Shared/EpicGames.UHT/Exporters/CodeGen/UhtHeaderCodeGeneratorHFile.cs
--- a/Engine/Source/Programs/Shared/EpicGames.UHT/Exporters/CodeGen/UhtHeaderCodeGeneratorHFile.cs
+++ b/Engine/Source/Programs/Shared/EpicGames.UHT/Exporters/CodeGen/UhtHeaderCodeGeneratorHFile.cs
@@ -1772,7 +1772,7 @@ namespace EpicGames.UHT.Exporters.CodeGen
 			}
 			builder.Append("\t}; \\\r\n");
 
-			builder.Append('\t').Append(api).Append("virtual void ValidateGeneratedRepEnums(const TArray<struct FRepRecord>& ClassReps) const override; \\\r\n");
+			builder.Append('\t').Append("DECLARE_VALIDATE_GENERATED_REP_ENUMS(").Append(api[0..^1]).Append(") \\\r\n");
 
 			return builder;
 		}
diff --git a/Engine/Source/Runtime/CoreUObject/Private/UObject/Class.cpp b/Engine/Source/Runtime/CoreUObject/Private/UObject/Class.cpp
--- a/Engine/Source/Runtime/CoreUObject/Private/UObject/Class.cpp
+++ b/Engine/Source/Runtime/CoreUObject/Private/UObject/Class.cpp
@@ -5332,19 +5332,115 @@ void UClass::Link(FArchive& Ar, bool bRelinkExistingProperties)
 	Super::Link(Ar, bRelinkExistingProperties);
 }
 
+#if VALIDATE_CLASS_REPS
 namespace UE::Net::Private
 {
-#if (UE_BUILD_SHIPPING)
-	static bool GValidateReplicatedProperties = false;
-#else 
-	static bool GValidateReplicatedProperties = true;
-#endif
+	static bool GValidateReplicatedProperties = !UE_BUILD_SHIPPING;
+	static FAutoConsoleVariable CVarValidateReplicatedPropertyRegistration(TEXT("net.ValidateReplicatedPropertyRegistration"), GValidateReplicatedProperties, TEXT("Warns if NetDeltaSerialize properties are not top-level"));
+
+	void ValidateNetDeltaSerializeProperties(const TArray<FRepRecord>& ClassReps)
+	{
+		// Find NetDeltaSerialize properties which are not top-level/parent properties (e.g. nested in another struct)
+		for (const FRepRecord& RepRecord : ClassReps)
+		{
+			const FProperty* BaseProp = RepRecord.Property;
+			TArray<const UStruct*, TInlineAllocator<16>> RecursiveStructList;
+
+			enum class ECheckStructType : uint8
+			{
+				TopLevel,	// Parent/top-level - NetDeltaSerialize is valid
+				Recursive	// Recursive (within other structs/arrays) - NetDeltaSerialize is not valid
+			};
+
+			auto CheckStructRecursive = [&BaseProp, &RecursiveStructList]
+			(auto CheckStructRecursive, const FProperty* InProp, ECheckStructType CheckType)
+				-> void
+				{
+					auto BadStruct = [&BaseProp](const FStructProperty* StructProp)
+						{
+							UE_LOG(LogClass, Warning, TEXT("Property %s contained nested NetDeltaSerialize struct '%s', ")
+								TEXT("when this is not supported. Only use NetDeltaSerialize structs at class level. ")
+								TEXT("Struct will replicate using non-delta serialization."),
+								ToCStr(BaseProp->GetPathName()), ToCStr(StructProp->GetPathName()));
+						};
+
+					const FStructProperty* ParamStructProp = nullptr;
+					bool bNonDeltaArrayStruct = false;
+
+					if (!EnumHasAnyFlags(InProp->PropertyFlags, CPF_RepSkip))
+					{
+						if (const FStructProperty* StructProp = CastField<FStructProperty>(InProp))
+						{
+							ParamStructProp = StructProp;
+						}
+						else if (const FArrayProperty* ArrayProp = CastField<FArrayProperty>(InProp))
+						{
+							if (const FStructProperty* ArrayStructProp = CastField<FStructProperty>(ArrayProp->Inner))
+							{
+								if (EnumHasAnyFlags(ArrayStructProp->Struct->StructFlags, STRUCT_NetDeltaSerializeNative))
+								{
+									BadStruct(ArrayStructProp);
+								}
+								else
+								{
+									ParamStructProp = ArrayStructProp;
+									bNonDeltaArrayStruct = true;
+								}
+							}
+						}
+					}
+
+
+					const bool bTopLevelNonDeltaStruct = ParamStructProp != nullptr && CheckType == ECheckStructType::TopLevel &&
+						(bNonDeltaArrayStruct || !EnumHasAnyFlags(ParamStructProp->Struct->StructFlags, STRUCT_NetDeltaSerializeNative));
+
+					const bool bFreshRecursiveStruct = ParamStructProp != nullptr && CheckType == ECheckStructType::Recursive &&
+						!RecursiveStructList.Contains(ParamStructProp->Struct);
 
-static FAutoConsoleVariable CVarValidateReplicatedPropertyRegistration(TEXT("net.ValidateReplicatedPropertyRegistration"), GValidateReplicatedProperties, TEXT("Warns if NetDeltaSerialize properties are not top-level"));
+					if (bTopLevelNonDeltaStruct || bFreshRecursiveStruct || bNonDeltaArrayStruct)
+					{
+						RecursiveStructList.Add(ParamStructProp->Struct);
 
-void ValidateNetDeltaSerializeProperties(UClass* ClassToValidate);
+						for (TFieldIterator<FProperty> It(ParamStructProp->Struct); It; ++It)
+						{
+							if (!EnumHasAnyFlags(It->PropertyFlags, CPF_RepSkip))
+							{
+								const FStructProperty* CurLevelStructProp = nullptr;
+
+								if (const FStructProperty* StructProp = CastField<FStructProperty>(*It))
+								{
+									CurLevelStructProp = StructProp;
+								}
+								else if (const FArrayProperty* ArrayProp = CastField<FArrayProperty>(*It))
+								{
+									if (const FStructProperty* ArrayStructProp = CastField<FStructProperty>(ArrayProp->Inner))
+									{
+										CurLevelStructProp = ArrayStructProp;
+									}
+								}
+
+								if (CurLevelStructProp != nullptr)
+								{
+									if (EnumHasAnyFlags(CurLevelStructProp->Struct->StructFlags, STRUCT_NetDeltaSerializeNative))
+									{
+										BadStruct(CurLevelStructProp);
+									}
+									else
+									{
+										CheckStructRecursive(CheckStructRecursive, CurLevelStructProp, ECheckStructType::Recursive);
+									}
+								}
+							}
+						}
+					}
+				};
+
+			CheckStructRecursive(CheckStructRecursive, BaseProp, ECheckStructType::TopLevel);
+		}
+	}
 
 } // end namespace UE::Net::Private
+#endif
 
 void UClass::SetUpRuntimeReplicationData()
 {
@@ -5429,127 +5525,23 @@ void UClass::SetUpRuntimeReplicationData()
 	}
 	check(ClassReps.Num() <= 65535);
 
-	if (bIsNativeClass && UE::Net::Private::GValidateReplicatedProperties)
-	{
-		GetDefaultObject()->ValidateGeneratedRepEnums(ClassReps);
-	}
-
 	NetFields.Shrink();
 		
 	Algo::SortBy(NetFields, &UField::GetFName, FNameLexicalLess());
 
 	ClassFlags |= CLASS_ReplicationDataIsSetUp;
 
+#if VALIDATE_CLASS_REPS
 	if (UE::Net::Private::GValidateReplicatedProperties)
 	{
-		UE::Net::Private::ValidateNetDeltaSerializeProperties(this);
-	}
-}
-
-namespace UE::Net::Private
-{
-
-void ValidateNetDeltaSerializeProperties(UClass* ClassToValidate)
-{
-	// Find NetDeltaSerialize properties which are not top-level/parent properties (e.g. nested in another struct)
-	for (const FRepRecord& RepRecord : ClassToValidate->ClassReps)
-	{
-		const FProperty* BaseProp = RepRecord.Property;
-		TArray<const UStruct*, TInlineAllocator<16>> RecursiveStructList;
-
-		enum class ECheckStructType : uint8
+		if (bIsNativeClass)
 		{
-			TopLevel,	// Parent/top-level - NetDeltaSerialize is valid
-			Recursive	// Recursive (within other structs/arrays) - NetDeltaSerialize is not valid
-		};
-
-		auto CheckStructRecursive = [&BaseProp, &RecursiveStructList]
-		(auto CheckStructRecursive, const FProperty* InProp, ECheckStructType CheckType)
-		-> void
-		{
-			auto BadStruct = [&BaseProp](const FStructProperty* StructProp)
-			{
-				UE_LOG(LogClass, Warning, TEXT("Property %s contained nested NetDeltaSerialize struct '%s', ")
-					TEXT("when this is not supported. Only use NetDeltaSerialize structs at class level. ")
-					TEXT("Struct will replicate using non-delta serialization."),
-					ToCStr(BaseProp->GetPathName()), ToCStr(StructProp->GetPathName()));
-			};
-
-			const FStructProperty* ParamStructProp = nullptr;
-			bool bNonDeltaArrayStruct = false;
-
-			if (!EnumHasAnyFlags(InProp->PropertyFlags, CPF_RepSkip))
-			{
-				if (const FStructProperty* StructProp = CastField<FStructProperty>(InProp))
-				{
-					ParamStructProp = StructProp;
-				}
-				else if (const FArrayProperty* ArrayProp = CastField<FArrayProperty>(InProp))
-				{
-					if (const FStructProperty* ArrayStructProp = CastField<FStructProperty>(ArrayProp->Inner))
-					{
-						if (EnumHasAnyFlags(ArrayStructProp->Struct->StructFlags, STRUCT_NetDeltaSerializeNative))
-						{
-							BadStruct(ArrayStructProp);
-						}
-						else
-						{
-							ParamStructProp = ArrayStructProp;
-							bNonDeltaArrayStruct = true;
-						}
-					}
-				}
-			}
-
-
-			const bool bTopLevelNonDeltaStruct = ParamStructProp != nullptr && CheckType == ECheckStructType::TopLevel &&
-			(bNonDeltaArrayStruct || !EnumHasAnyFlags(ParamStructProp->Struct->StructFlags, STRUCT_NetDeltaSerializeNative));
-
-			const bool bFreshRecursiveStruct = ParamStructProp != nullptr && CheckType == ECheckStructType::Recursive &&
-			!RecursiveStructList.Contains(ParamStructProp->Struct);
-
-			if (bTopLevelNonDeltaStruct || bFreshRecursiveStruct || bNonDeltaArrayStruct)
-			{
-				RecursiveStructList.Add(ParamStructProp->Struct);
-
-				for (TFieldIterator<FProperty> It(ParamStructProp->Struct); It; ++It)
-				{
-					if (!EnumHasAnyFlags(It->PropertyFlags, CPF_RepSkip))
-					{
-						const FStructProperty* CurLevelStructProp = nullptr;
-
-						if (const FStructProperty* StructProp = CastField<FStructProperty>(*It))
-						{
-							CurLevelStructProp = StructProp;
-						}
-						else if (const FArrayProperty* ArrayProp = CastField<FArrayProperty>(*It))
-						{
-							if (const FStructProperty* ArrayStructProp = CastField<FStructProperty>(ArrayProp->Inner))
-							{
-								CurLevelStructProp = ArrayStructProp;
-							}
-						}
-
-						if (CurLevelStructProp != nullptr)
-						{
-							if (EnumHasAnyFlags(CurLevelStructProp->Struct->StructFlags, STRUCT_NetDeltaSerializeNative))
-							{
-								BadStruct(CurLevelStructProp);
-							}
-							else
-							{
-								CheckStructRecursive(CheckStructRecursive, CurLevelStructProp, ECheckStructType::Recursive);
-							}
-						}
-					}
-				}
-			}
-		};
-
-		CheckStructRecursive(CheckStructRecursive, BaseProp, ECheckStructType::TopLevel);
+			GetDefaultObject()->ValidateGeneratedRepEnums(ClassReps);
+		}
+		UE::Net::Private::ValidateNetDeltaSerializeProperties(ClassReps);
 	}
+#endif
 }
-} // end namespace UE::Net::Private
 
 void UClass::InternalCreateDefaultObjectWrapper() const
 {
diff --git a/Engine/Source/Runtime/CoreUObject/Public/UObject/ObjectMacros.h b/Engine/Source/Runtime/CoreUObject/Public/UObject/ObjectMacros.h
--- a/Engine/Source/Runtime/CoreUObject/Public/UObject/ObjectMacros.h
+++ b/Engine/Source/Runtime/CoreUObject/Public/UObject/ObjectMacros.h
@@ -2273,3 +2273,13 @@ public:
 	**/
 	virtual void Restore() const=0;
 };
+
+#ifndef VALIDATE_CLASS_REPS
+#define VALIDATE_CLASS_REPS DO_CHECK
+#endif
+
+#if VALIDATE_CLASS_REPS
+#define DECLARE_VALIDATE_GENERATED_REP_ENUMS( API ) API virtual void ValidateGeneratedRepEnums(const TArray<struct FRepRecord>& ClassReps) const override;
+#else
+#define DECLARE_VALIDATE_GENERATED_REP_ENUMS( API )
+#endif
