[StructUtils] [[nodiscard]] for all functions 'Make'

diff --git a/Engine/Source/Runtime/CoreUObject/Public/StructUtils/InstancedStruct.h b/Engine/Source/Runtime/CoreUObject/Public/StructUtils/InstancedStruct.h
--- a/Engine/Source/Runtime/CoreUObject/Public/StructUtils/InstancedStruct.h
+++ b/Engine/Source/Runtime/CoreUObject/Public/StructUtils/InstancedStruct.h
@@ -117,7 +117,7 @@ public:
 
 	/** Creates a new FInstancedStruct from templated struct type. */
 	template<typename T>
-	static FInstancedStruct Make()
+	[[nodiscard]] static FInstancedStruct Make()
 	{
 		UE::StructUtils::CheckStructType<T>();
 
@@ -128,7 +128,7 @@ public:
 
 	/** Creates a new FInstancedStruct from templated struct. */
 	template<typename T>
-	static FInstancedStruct Make(const T& Struct)
+	[[nodiscard]] static FInstancedStruct Make(const T& Struct)
 	{
 		UE::StructUtils::CheckStructType<T>();
 
@@ -139,7 +139,7 @@ public:
 
 	/** Creates a new FInstancedStruct from the templated type and forward all arguments to constructor. */
 	template<typename T, typename... TArgs>
-	static inline FInstancedStruct Make(TArgs&&... InArgs)
+	[[nodiscard]] static FInstancedStruct Make(TArgs&&... InArgs)
 	{
 		UE::StructUtils::CheckStructType<T>();
 
@@ -372,7 +372,7 @@ public:
 
 	/** Creates a new TInstancedStruct from templated struct type. */
 	template<typename T = BaseStructT, typename = std::enable_if_t<std::is_base_of_v<BaseStructT, std::decay_t<T>>>>
-	static TInstancedStruct Make()
+	[[nodiscard]] static TInstancedStruct Make()
 	{
 		TInstancedStruct This;
 		This.InstancedStruct.InitializeAs(TBaseStructure<T>::Get(), nullptr);
@@ -381,7 +381,7 @@ public:
 
 	/** Creates a new TInstancedStruct from templated struct. */
 	template<typename T = BaseStructT, typename = std::enable_if_t<std::is_base_of_v<BaseStructT, std::decay_t<T>>>>
-	static TInstancedStruct Make(const T& Struct)
+	[[nodiscard]] static TInstancedStruct Make(const T& Struct)
 	{
 		TInstancedStruct This;
 		This.InstancedStruct.InitializeAs(TBaseStructure<T>::Get(), reinterpret_cast<const uint8*>(&Struct));
@@ -390,7 +390,7 @@ public:
 
 	/** Creates a new TInstancedStruct from the templated type and forward all arguments to constructor. */
 	template<typename T = BaseStructT, typename... TArgs, typename = std::enable_if_t<std::is_base_of_v<BaseStructT, std::decay_t<T>>>>
-	static TInstancedStruct Make(TArgs&&... InArgs)
+	[[nodiscard]] static TInstancedStruct Make(TArgs&&... InArgs)
 	{
 		TInstancedStruct This;
 		This.InstancedStruct.template InitializeAs<T>(Forward<TArgs>(InArgs)...);
diff --git a/Engine/Source/Runtime/CoreUObject/Public/StructUtils/SharedStruct.h b/Engine/Source/Runtime/CoreUObject/Public/StructUtils/SharedStruct.h
--- a/Engine/Source/Runtime/CoreUObject/Public/StructUtils/SharedStruct.h
+++ b/Engine/Source/Runtime/CoreUObject/Public/StructUtils/SharedStruct.h
@@ -246,7 +246,7 @@ struct COREUOBJECT_API FSharedStruct
 
 	/** Creates a new FSharedStruct from a templated struct type. This will create a new instance of the shared struct memory. */
 	template<typename T>
-	static FSharedStruct Make()
+	[[nodiscard]] static FSharedStruct Make()
 	{
 		FSharedStruct SharedStruct;
 		SharedStruct.InitializeAs<T>();
@@ -255,7 +255,7 @@ struct COREUOBJECT_API FSharedStruct
 
 	/** Creates a new FSharedStruct from templated struct instance. This will create a new instance of the shared struct memory. */
 	template<typename T>
-	static FSharedStruct Make(const T& Struct)
+	[[nodiscard]] static FSharedStruct Make(const T& Struct)
 	{
 		FSharedStruct SharedStruct;
 		SharedStruct.InitializeAs(Struct);
@@ -263,7 +263,7 @@ struct COREUOBJECT_API FSharedStruct
 	}
 
 	/** Creates a new FSharedStruct from struct type and optional instance memory. This will create a new instance of the shared struct memory. */
-	static FSharedStruct Make(const UScriptStruct* InScriptStruct, const uint8* InStructMemory = nullptr)
+	[[nodiscard]] static FSharedStruct Make(const UScriptStruct* InScriptStruct, const uint8* InStructMemory = nullptr)
 	{
 		FSharedStruct SharedStruct;
 		SharedStruct.InitializeAs(InScriptStruct, InStructMemory);
@@ -272,7 +272,7 @@ struct COREUOBJECT_API FSharedStruct
 
 	/** Creates a new FSharedStruct from the templated type and forward all arguments to constructor. This will create a new instance of the shared struct memory. */
 	template<typename T, typename... TArgs>
-	static FSharedStruct Make(TArgs&&... InArgs)
+	[[nodiscard]] static FSharedStruct Make(TArgs&&... InArgs)
 	{
 		FSharedStruct SharedStruct;
 		SharedStruct.InitializeAs<T>(Forward<TArgs>(InArgs)...);
@@ -431,21 +431,21 @@ struct TSharedStruct
 
 	/** Creates a new TSharedStruct. This will create a new instance of the shared struct memory. */
 	template<typename T = BaseStructT, typename = std::enable_if_t<std::is_base_of_v<BaseStructT, std::decay_t<T>>>>
-	static TSharedStruct Make()
+	[[nodiscard]] static TSharedStruct Make()
 	{
 		return FSharedStruct::Make<T>();
 	}
 
 	/** Creates a new TSharedStruct from templated struct instance. This will create a new instance of the shared struct memory. */
 	template<typename T = BaseStructT, typename = std::enable_if_t<std::is_base_of_v<BaseStructT, std::decay_t<T>>>>
-	static TSharedStruct Make(const T& Struct)
+	[[nodiscard]] static TSharedStruct Make(const T& Struct)
 	{
 		return FSharedStruct::Make<T>(Struct);
 	}
 
 	/** Creates a new TSharedStruct from the templated type and forward all arguments to constructor. This will create a new instance of the shared struct memory. */
 	template< typename... TArgs>
-	static TSharedStruct Make(TArgs&&... InArgs)
+	[[nodiscard]] static TSharedStruct Make(TArgs&&... InArgs)
 	{
 		return FSharedStruct::Make<BaseStructT>(Forward<TArgs>(InArgs)...);
 	}
@@ -622,7 +622,7 @@ struct COREUOBJECT_API FConstSharedStruct
 
 	/** Creates a new FSharedStruct from a templated struct type. This will create a new instance of the shared struct memory. */
 	template<typename T>
-	static FConstSharedStruct Make()
+	[[nodiscard]] static FConstSharedStruct Make()
 	{
 		FConstSharedStruct SharedStruct;
 		SharedStruct.InitializeAs<T>();
@@ -631,7 +631,7 @@ struct COREUOBJECT_API FConstSharedStruct
 
 	/** Creates a new FSharedStruct from templated struct instance. This will create a new instance of the shared struct memory. */
 	template<typename T>
-	static FConstSharedStruct Make(const T& Struct)
+	[[nodiscard]] static FConstSharedStruct Make(const T& Struct)
 	{
 		FConstSharedStruct SharedStruct;
 		SharedStruct.InitializeAs<T>(Struct);
@@ -639,7 +639,7 @@ struct COREUOBJECT_API FConstSharedStruct
 	}
 
 	/** Creates a new FSharedStruct from a struct type and optional data. This will create a new instance of the shared struct memory. */
-	static FConstSharedStruct Make(const UScriptStruct* InScriptStruct, const uint8* InStructMemory = nullptr)
+	[[nodiscard]] static FConstSharedStruct Make(const UScriptStruct* InScriptStruct, const uint8* InStructMemory = nullptr)
 	{
 		FConstSharedStruct SharedStruct;
 		SharedStruct.InitializeAs(InScriptStruct, InStructMemory);
@@ -648,7 +648,7 @@ struct COREUOBJECT_API FConstSharedStruct
 
 	/** Creates a new FSharedStruct from the templated type and forward all arguments to constructor. This will create a new instance of the shared struct memory. */
 	template<typename T, typename... TArgs>
-	static FConstSharedStruct Make(TArgs&&... InArgs)
+	[[nodiscard]] static FConstSharedStruct Make(TArgs&&... InArgs)
 	{
 		FConstSharedStruct SharedStruct;
 		SharedStruct.InitializeAs<T>(Forward<TArgs>(InArgs)...);
@@ -802,7 +802,7 @@ struct TConstSharedStruct
 
 	/** Creates a new TSharedStruct from a templated struct type. This will create a new instance of the shared struct memory. */
 	template<typename T = BaseStructT, typename = std::enable_if_t<std::is_base_of_v<BaseStructT, std::decay_t<T>>>>
-	static TConstSharedStruct Make()
+	[[nodiscard]] static TConstSharedStruct Make()
 	{
 		FConstSharedStruct SharedStruct;
 		SharedStruct.InitializeAs<BaseStructT>();
@@ -811,7 +811,7 @@ struct TConstSharedStruct
 
 	/** Creates a new TSharedStruct from templated struct instance. This will create a new instance of the shared struct memory. */
 	template<typename T = BaseStructT, typename = std::enable_if_t<std::is_base_of_v<BaseStructT, std::decay_t<T>>>>
-	static TConstSharedStruct Make(const T& Struct)
+	[[nodiscard]] static TConstSharedStruct Make(const T& Struct)
 	{
 		FConstSharedStruct SharedStruct;
 		SharedStruct.InitializeAs<BaseStructT>(Struct);
@@ -820,7 +820,7 @@ struct TConstSharedStruct
 
 	/** Creates a new TSharedStruct from a struct type and optional data. This will create a new instance of the shared struct memory. */
 	template<typename T = BaseStructT, typename = std::enable_if_t<std::is_base_of_v<BaseStructT, std::decay_t<T>>>>
-	static TConstSharedStruct Make(const uint8* InStructMemory = nullptr)
+	[[nodiscard]] static TConstSharedStruct Make(const uint8* InStructMemory = nullptr)
 	{
 		TConstSharedStruct SharedStruct;
 		SharedStruct.template Initialize<T>(InStructMemory);
@@ -829,7 +829,7 @@ struct TConstSharedStruct
 
 	/** Creates a new TSharedStruct from the templated type and forward all arguments to constructor. This will create a new instance of the shared struct memory. */
 	template<typename T = BaseStructT, typename... TArgs, typename = std::enable_if_t<std::is_base_of_v<BaseStructT, std::decay_t<T>>>>
-	static TConstSharedStruct Make(TArgs&&... InArgs)
+	[[nodiscard]] static TConstSharedStruct Make(TArgs&&... InArgs)
 	{
 		TConstSharedStruct SharedStruct;
 		SharedStruct.template Initialize<T>(Forward<TArgs>(InArgs)...);
diff --git a/Engine/Source/Runtime/CoreUObject/Public/StructUtils/StructView.h b/Engine/Source/Runtime/CoreUObject/Public/StructUtils/StructView.h
--- a/Engine/Source/Runtime/CoreUObject/Public/StructUtils/StructView.h
+++ b/Engine/Source/Runtime/CoreUObject/Public/StructUtils/StructView.h
@@ -40,7 +40,7 @@ public:
 
 	/** Creates a new FStructView from the templated struct. Note its not safe to make InStruct const ref as the original object may have been declared const */
 	template<typename T>
-	static FStructView Make(T& InStruct)
+	[[nodiscard]] static FStructView Make(T& InStruct)
 	{
 		UE::StructUtils::CheckStructType<T>();
 		return FStructView(TBaseStructure<T>::Get(), reinterpret_cast<uint8*>(&InStruct));
@@ -229,7 +229,7 @@ public:
 
 	/** Creates a new FConstStructView from the templated struct */
 	template<typename T>
-	static FConstStructView Make(const T& Struct)
+	[[nodiscard]] static FConstStructView Make(const T& Struct)
 	{
 		UE::StructUtils::CheckStructType<T>();
 		return FConstStructView(TBaseStructure<T>::Get(), reinterpret_cast<const uint8*>(&Struct));
