Wrapped CoreUObject test types in #if WITH_TESTS

#rnx

diff --git a/Engine/Source/Runtime/CoreUObject/Public/UObject/UObjectBase.h b/Engine/Source/Runtime/CoreUObject/Public/UObject/UObjectBase.h
--- a/Engine/Source/Runtime/CoreUObject/Public/UObject/UObjectBase.h
+++ b/Engine/Source/Runtime/CoreUObject/Public/UObject/UObjectBase.h
@@ -404,18 +404,6 @@ struct TRegistrationInfo
 	TVersion ReloadVersionInfo;
 };
 
-/**
- * Helper class to perform registration of object information.  It blindly forwards a call to RegisterCompiledInInfo
- */
-struct FRegisterCompiledInInfo
-{
-	template <typename ... Args>
-	FRegisterCompiledInInfo(Args&& ... args)
-	{
-		RegisterCompiledInInfo(std::forward<Args>(args)...);
-	}
-};
-
 /**
  * Reload version information for classes
  */
@@ -544,7 +532,6 @@ using FPackageRegistrationInfo = TRegistrationInfo<UPackage, FPackageReloadVersi
  */
 COREUOBJECT_API void RegisterCompiledInInfo(UPackage* (*InOuterRegister)(), const TCHAR* InPackageName, FPackageRegistrationInfo& InInfo, const FPackageReloadVersionInfo& InVersionInfo);
 
-
 /**
  * Register compiled in information for multiple classes, structures, and enumerations
  */
@@ -556,6 +543,18 @@ COREUOBJECT_API void RegisterModularObjectsProcessing();
 /** Must be called after a module has been loaded that contains UObject classes */
 COREUOBJECT_API void ProcessNewlyLoadedUObjects(FName InModuleName = NAME_None, bool bCanProcessNewlyLoadedObjects = true);
 
+/**
+ * Helper class to perform registration of object information.  It blindly forwards a call to RegisterCompiledInInfo
+ */
+struct FRegisterCompiledInInfo
+{
+	template <typename... ArgTypes>
+	FRegisterCompiledInInfo(ArgTypes&&... Args)
+	{
+		RegisterCompiledInInfo(Forward<ArgTypes>(Args)...);
+	}
+};
+
 /**
  * Final phase of UObject initialization. all auto register objects are added to the main data structures.
  */
diff --git a/Engine/Source/Runtime/CoreUObject/Tests/OptionalPropertyTestObject.h b/Engine/Source/Runtime/CoreUObject/Tests/OptionalPropertyTestObject.h
--- a/Engine/Source/Runtime/CoreUObject/Tests/OptionalPropertyTestObject.h
+++ b/Engine/Source/Runtime/CoreUObject/Tests/OptionalPropertyTestObject.h
@@ -4,6 +4,8 @@
 
 #include "OptionalPropertyTestObject.generated.h"
 
+#if WITH_TESTS
+
 UCLASS()
 class UOptionalPropertyTestObject : public UObject
 {
@@ -21,4 +23,6 @@ public:
 
     UPROPERTY()
     TOptional<int32> OptionalInt;
-};
\ No newline at end of file
+};
+
+#endif // WITH_TESTS
diff --git a/Engine/Source/Runtime/CoreUObject/Tests/UObject/PropertyPathFunctionsTest.h b/Engine/Source/Runtime/CoreUObject/Tests/UObject/PropertyPathFunctionsTest.h
--- a/Engine/Source/Runtime/CoreUObject/Tests/UObject/PropertyPathFunctionsTest.h
+++ b/Engine/Source/Runtime/CoreUObject/Tests/UObject/PropertyPathFunctionsTest.h
@@ -11,6 +11,8 @@
 #include "UObject/Object.h"
 #include "PropertyPathFunctionsTest.generated.h"
 
+#if WITH_TESTS
+
 USTRUCT()
 struct FTestPropertyPathFunctionsStructKey
 {
@@ -90,3 +92,5 @@ public:
 	UPROPERTY()
 	TOptional<FTestPropertyPathFunctionsStruct> StructOptional;
 };
+
+#endif // WITH_TESTS
diff --git a/Engine/Source/Runtime/CoreUObject/Tests/UObject/PropertyStateTrackingTest.h b/Engine/Source/Runtime/CoreUObject/Tests/UObject/PropertyStateTrackingTest.h
--- a/Engine/Source/Runtime/CoreUObject/Tests/UObject/PropertyStateTrackingTest.h
+++ b/Engine/Source/Runtime/CoreUObject/Tests/UObject/PropertyStateTrackingTest.h
@@ -7,6 +7,8 @@
 #include "UObject/Object.h"
 #include "PropertyStateTrackingTest.generated.h"
 
+#if WITH_TESTS
+
 UENUM()
 enum ETestInstanceDataObjectBird : uint8
 {
@@ -312,3 +314,5 @@ public:
 	UPROPERTY()
 	FTestInstanceDataObjectStruct Struct;
 };
+
+#endif // WITH_TESTS
