Add the value of dependency in FCookDependency class. It's a 32 bytes buffer.
Update the FCookDependency class to store/load the value to/from the artifact file.
Update the constructors, move and copy operators.
Update the functions UE::EditorDomain::TryAppendClassDigests, FCookDepedency::UpdateHash and FBuildDependencySet::TryCalculateCurrentKey to set the dependency value.

One other change I had to do for function dependencies : in FCookDepedency::UpdateHash, the function is called and we pass a FCookDependencyContext. This context has a pointer to the hasher of the entire key hash. Since I need to retrieve just the hash generated by the function, I added a new function to the context void* SetHasher(void* NewHasher) so I can pass a different hasher and restore it later on.

#rnx
#incrementalcook

diff --git a/Engine/Source/Editor/UnrealEd/Private/Cooker/CookPackageArtifacts.cpp b/Engine/Source/Editor/UnrealEd/Private/Cooker/CookPackageArtifacts.cpp
--- a/Engine/Source/Editor/UnrealEd/Private/Cooker/CookPackageArtifacts.cpp
+++ b/Engine/Source/Editor/UnrealEd/Private/Cooker/CookPackageArtifacts.cpp
@@ -118,7 +118,7 @@ FBuildDependencySet::ECurrentKeyResult FBuildDependencySet::TryCalculateCurrentK
 		{
 			++BatchEnd;
 		}
-		TConstArrayView<FCookDependency> Batch(Dependencies.GetData() + BatchStart, BatchEnd - BatchStart);
+		TArrayView<FCookDependency> Batch(Dependencies.GetData() + BatchStart, BatchEnd - BatchStart);
 		BatchStart += Batch.Num();
 
 		// Some CookDependency types can not handle Update being called, because their Update relies on functions
@@ -127,7 +127,7 @@ FBuildDependencySet::ECurrentKeyResult FBuildDependencySet::TryCalculateCurrentK
 		{
 		case ECookDependency::Package:
 		{
-			for (const FCookDependency& PackageDependency : Batch)
+			for (FCookDependency& PackageDependency : Batch)
 			{
 				FName DependencyPackageName = PackageDependency.GetPackageName();
 				UE::EditorDomain::FPackageDigest PackageDigest = EditorDomain->GetPackageDigest(DependencyPackageName);
@@ -137,7 +137,9 @@ FBuildDependencySet::ECurrentKeyResult FBuildDependencySet::TryCalculateCurrentK
 				}
 				if (PackageDigest.IsSuccessful())
 				{
-					KeyBuilder.Update(&PackageDigest.Hash, sizeof(PackageDigest.Hash));
+					PackageDependency.SetValue(PackageDigest.Hash);
+					KeyBuilder.Update(PackageDependency.GetRawValue(), FCookDependency::ValueSizeInBytes);
+					
 					continue;
 				}
 
@@ -147,13 +149,14 @@ FBuildDependencySet::ECurrentKeyResult FBuildDependencySet::TryCalculateCurrentK
 			}
 			break;
 		}
+
 		case ECookDependency::Config:
 		{
 #if UE_WITH_CONFIG_TRACKING
 			using namespace UE::ConfigAccessTracking;
 			FCookConfigAccessTracker& ConfigTracker = FCookConfigAccessTracker::Get();
 #endif
-			for (const FCookDependency& ConfigDependency : Batch)
+			for (FCookDependency& ConfigDependency : Batch)
 			{
 				FString Value;
 #if UE_WITH_CONFIG_TRACKING
@@ -163,7 +166,8 @@ FBuildDependencySet::ECurrentKeyResult FBuildDependencySet::TryCalculateCurrentK
 				KeyBuilder.Update(&Marker, sizeof(Marker));
 				if (!Value.IsEmpty())
 				{
-					KeyBuilder.Update(*Value, Value.Len() * sizeof(Value[0]));
+					ConfigDependency.SetValue(FUtf8String(Value));
+					KeyBuilder.Update(ConfigDependency.GetRawValue(), FCookDependency::ValueSizeInBytes);
 				}
 			}
 			break;
@@ -188,16 +192,21 @@ FBuildDependencySet::ECurrentKeyResult FBuildDependencySet::TryCalculateCurrentK
 			Hashes.SetNum(PackageNames.Num());
 			FCoreRedirects::GetHashOfRedirectsAffectingPackages(PackageNames, Hashes);
 
-			for (const FBlake3Hash& Hash : Hashes)
+			for(int Index = 0; Index < Hashes.Num(); ++Index)
 			{
-				KeyBuilder.Update(&Hash, sizeof(Hash));
+				const FBlake3Hash& Hash = Hashes[Index];
+				FCookDependency& RedirectionDependency = Batch[Index];
+
+				RedirectionDependency.SetValue(Hash);
+				KeyBuilder.Update(RedirectionDependency.GetRawValue(), FCookDependency::ValueSizeInBytes);
 			}
 
 			FCoreRedirects::AppendHashOfGlobalRedirects(KeyBuilder);
 			break;
 		}
+
 		default:
-			for (const FCookDependency& BatchDependency : Batch)
+			for (FCookDependency& BatchDependency : Batch)
 			{
 				BatchDependency.UpdateHash(Context);
 			}
diff --git a/Engine/Source/Editor/UnrealEd/Private/EditorDomain/EditorDomainUtils.cpp b/Engine/Source/Editor/UnrealEd/Private/EditorDomain/EditorDomainUtils.cpp
--- a/Engine/Source/Editor/UnrealEd/Private/EditorDomain/EditorDomainUtils.cpp
+++ b/Engine/Source/Editor/UnrealEd/Private/EditorDomain/EditorDomainUtils.cpp
@@ -459,13 +459,13 @@ FPackageDigest CalculatePackageDigest(const FAssetPackageData& PackageData, FNam
 	return Result;
 }
 
-bool TryAppendClassDigests(const TConstArrayView<UE::Cook::FCookDependency>& Dependencies, FBlake3& Writer, UE::Cook::FCookDependencyContext& Context)
+bool TryAppendClassDigests(const TArrayView<UE::Cook::FCookDependency>& Dependencies, FBlake3& Writer, UE::Cook::FCookDependencyContext& Context)
 {
 	FClassDigestMap& ClassDigests = GetClassDigests();
 
 	FReadScopeLock ClassDigestsScopeLock(ClassDigests.Lock);
 
-	for (const UE::Cook::FCookDependency& ClassDependency : Dependencies)
+	for (UE::Cook::FCookDependency& ClassDependency : Dependencies)
 	{
 		FStringView ClassPath = ClassDependency.GetClassPath();
 		FTopLevelAssetPath AssetPath(ClassPath);
@@ -507,7 +507,9 @@ bool TryAppendClassDigests(const TConstArrayView<UE::Cook::FCookDependency>& Dep
 			checkf(NativeData, TEXT("Classes are only stored in a ClosestNative field if they exist"));
 		}
 		check(NativeData->bNative);
-		Writer.Update(&NativeData->InclusiveSchemaHash, sizeof(NativeData->InclusiveSchemaHash));
+
+		ClassDependency.SetValue(NativeData->InclusiveSchemaHash);
+		Writer.Update(ClassDependency.GetRawValue(), UE::Cook::FCookDependency::ValueSizeInBytes);
 	}
 
 	return true;
diff --git a/Engine/Source/Editor/UnrealEd/Private/EditorDomain/EditorDomainUtils.h b/Engine/Source/Editor/UnrealEd/Private/EditorDomain/EditorDomainUtils.h
--- a/Engine/Source/Editor/UnrealEd/Private/EditorDomain/EditorDomainUtils.h
+++ b/Engine/Source/Editor/UnrealEd/Private/EditorDomain/EditorDomainUtils.h
@@ -94,7 +94,7 @@ FPackageDigest CalculatePackageDigest(IAssetRegistry& AssetRegistry, FName Packa
 FPackageDigest CalculatePackageDigest(const FAssetPackageData& PackageData, FName PackageName);
 
 /** FindOrAdd class digests for the given class dependencies, and append hashes onto the writer. Returns false if not found. */
-bool TryAppendClassDigests(const TConstArrayView<UE::Cook::FCookDependency>& Dependencies, FBlake3& Writer, UE::Cook::FCookDependencyContext& Context);
+bool TryAppendClassDigests(const TArrayView<UE::Cook::FCookDependency>& Dependencies, FBlake3& Writer, UE::Cook::FCookDependencyContext& Context);
 
 /** For any ClassNames not already in ClassDigests, look up their UStruct and add them. */
 void PrecacheClassDigests(TConstArrayView<FTopLevelAssetPath> ClassNames);
diff --git a/Engine/Source/Runtime/CoreUObject/Private/Cooker/CookDependency.cpp b/Engine/Source/Runtime/CoreUObject/Private/Cooker/CookDependency.cpp
--- a/Engine/Source/Runtime/CoreUObject/Private/Cooker/CookDependency.cpp
+++ b/Engine/Source/Runtime/CoreUObject/Private/Cooker/CookDependency.cpp
@@ -215,11 +215,12 @@ FCookDependency FCookDependency::RedirectionTarget(FName PackageName)
 
 FCookDependency::FCookDependency()
 	: Type(ECookDependency::None)
-{
-}
+	, HashedValue()
+{}
 
 FCookDependency::FCookDependency(ECookDependency InType)
 	: Type(InType)
+	, HashedValue()
 {
 	Construct();
 }
@@ -245,6 +246,7 @@ FCookDependency& FCookDependency::operator=(const FCookDependency& Other)
 {
 	Destruct();
 	Type = Other.Type;
+	HashedValue = Other.HashedValue;
 	Construct();
 	switch (Type)
 	{
@@ -289,6 +291,7 @@ FCookDependency& FCookDependency::operator=(FCookDependency&& Other)
 {
 	Destruct();
 	Type = Other.Type;
+	HashedValue = Other.HashedValue;
 	Construct();
 	switch (Type)
 	{
@@ -339,7 +342,12 @@ FString FCookDependency::GetConfigPath() const
 		ConfigAccessData->FullPathToString() : FString();
 }
 
-void FCookDependency::UpdateHash(FCookDependencyContext& Context) const
+const uint8* FCookDependency::GetRawValue() const
+{
+	return RawValue;
+}
+
+void FCookDependency::UpdateHash(FCookDependencyContext& Context)
 {
 	switch (GetType())
 	{
@@ -363,20 +371,26 @@ void FCookDependency::UpdateHash(FCookDependencyContext& Context) const
 		const int64 Size = Ar->TotalSize();
 		int64 Position = 0;
 		// Read in BufferSize chunks
+
+		FBlake3 Hasher;
 		while (Position < Size)
 		{
 			const int64 ReadNum = FMath::Min(Size - Position, (int64)Buffer.Num());
 			Ar->Serialize(Buffer.GetData(), static_cast<uint64>(ReadNum));
-			Context.Update(Buffer.GetData(), ReadNum);
+			Hasher.Update(Buffer.GetData(), ReadNum);
 			Position += ReadNum;
 		}
+
+		FBlake3Hash Hash = Hasher.Finalize();
+		SetValue(Hash);
+		Context.Update(GetRawValue(), ValueSizeInBytes);
+
 		break;
 	}
 	case ECookDependency::Function:
 	{
 		const TMap<FName, FCookDependencyFunction>& CookDependencyFunctions =
 			UE::Cook::Dependency::Private::GetCookDependencyFunctions();
-
 		const FCookDependencyFunction* Function = CookDependencyFunctions.Find(GetFunctionName());
 		if (!Function)
 		{
@@ -385,12 +399,23 @@ void FCookDependency::UpdateHash(FCookDependencyContext& Context) const
 				*GetFunctionName().ToString()));
 			return;
 		}
+
 		FCookDependencyContext::FErrorHandlerScope Scope = Context.ErrorHandlerScope([this](FString&& Message)
 			{
 				return FString::Printf(TEXT("FCookDependency::Function('%s') failed to UpdateHash: %s"),
 				*GetFunctionName().ToString(), *Message);
 			});
+
+		FBlake3 NewHasher;
+		void* OldHasher = Context.SetHasher(&NewHasher);
 		(*Function)(GetFunctionArgs(), Context);
+		
+		Context.SetHasher(OldHasher);
+
+		FBlake3Hash Hash = NewHasher.Finalize();
+		SetValue(Hash);
+		Context.Update(GetRawValue(), ValueSizeInBytes);
+
 		return;
 	}
 	case ECookDependency::TransitiveBuild:
@@ -413,7 +438,10 @@ void FCookDependency::UpdateHash(FCookDependencyContext& Context) const
 			return;
 		}
 		FString VariableAsString = VariableInstance->GetString();
-		Context.Update(VariableAsString.GetCharArray().GetData(), VariableAsString.GetCharArray().Num() * sizeof(FString::ElementType));
+
+		SetValue(FUtf8String(VariableAsString));
+		Context.Update(GetRawValue(), ValueSizeInBytes);
+
 		return;
 	}
 	case ECookDependency::Config:
@@ -446,12 +474,18 @@ void FCookDependency::UpdateHash(FCookDependencyContext& Context) const
 				PackageNames.Sort(FNameLexicalLess());
 				PackageNames.SetNum(Algo::Unique(PackageNames));
 				TStringBuilder<256> PackageNameStr;
+
+				FBlake3 Hasher;
 				for (FName PackageName : PackageNames)
 				{
 					PackageNameStr.Reset();
 					PackageNameStr << PackageName;
-					Context.Update(*PackageNameStr, sizeof(**PackageNameStr) * PackageNameStr.Len());
+					Hasher.Update(*PackageNameStr, sizeof(**PackageNameStr) * PackageNameStr.Len());
 				}
+
+				FBlake3Hash Hash = Hasher.Finalize();
+				SetValue(Hash);
+				Context.Update(GetRawValue(), ValueSizeInBytes);
 			}
 		}
 		return;
@@ -466,6 +500,44 @@ void FCookDependency::UpdateHash(FCookDependencyContext& Context) const
 	}
 }
 
+void FCookDependency::SetValue(const FIoHash& Hash)
+{
+	static_assert(sizeof(Hash) <= ValueSizeInBytes);
+	FMemory::Memcpy(RawValue, &Hash, sizeof(Hash));
+
+	if constexpr (sizeof(Hash) < ValueSizeInBytes)
+	{
+		FMemory::Memzero(RawValue + sizeof(Hash), ValueSizeInBytes - sizeof(Hash));
+	}
+}
+
+void FCookDependency::SetValue(const FBlake3Hash& Hash)
+{
+	static_assert(sizeof(Hash) == ValueSizeInBytes);
+	HashedValue = Hash;
+}
+
+void FCookDependency::SetValue(const FUtf8String& String)
+{
+	SIZE_T StringNumBytes = String.NumBytesWithoutNull();
+	if (StringNumBytes > ValueSizeInBytes)
+	{
+		FUtf8StringView View(String);
+
+		FBlake3 Hasher;
+		HashedValue = Hasher.HashBuffer(View.GetData(), View.NumBytes());
+	}
+	else
+	{
+		FMemory::Memcpy(RawValue, String.GetCharArray().GetData(), StringNumBytes);
+
+		if (StringNumBytes < ValueSizeInBytes)
+		{
+			FMemory::Memzero(RawValue + StringNumBytes, ValueSizeInBytes - StringNumBytes);
+		}
+	}
+}
+
 void FCookDependency::Construct()
 {
 	switch (Type)
@@ -597,6 +669,9 @@ void FCookDependency::Save(FCbWriter& Writer) const
 		checkNoEntry();
 		break;
 	}
+
+	Writer.AddBinary(FMemoryView(RawValue, ValueSizeInBytes));
+	
 	Writer.EndArray();
 }
 
@@ -619,7 +694,8 @@ bool FCookDependency::Load(FCbFieldView Value)
 	switch (LocalType)
 	{
 	case ECookDependency::None:
-		return true;
+		break;
+
 	case ECookDependency::File:
 	{
 		FUtf8StringView LocalFilename = Field.AsString();
@@ -628,8 +704,9 @@ bool FCookDependency::Load(FCbFieldView Value)
 			return false;
 		}
 		*this = FCookDependency::File(WriteToString<256>(LocalFilename));
-		return true;
 	}
+	break;
+
 	case ECookDependency::Function:
 	{
 		FName LocalFuncName;
@@ -649,10 +726,11 @@ bool FCookDependency::Load(FCbFieldView Value)
 		{
 			*this = FCookDependency::Function(LocalFuncName, FCbFieldIterator::CloneRange(Field));
 			FunctionData.Args.MakeOwned();
+			Field.Reset(); // In the old format, all the fields are arguments of the function so there is nothing to read after.
 		}
-
-		return true;
 	}
+	break;
+
 	case ECookDependency::TransitiveBuild:
 	{
 		FName LocalPackageName;
@@ -661,8 +739,9 @@ bool FCookDependency::Load(FCbFieldView Value)
 			return false;
 		}
 		*this = FCookDependency::TransitiveBuild(LocalPackageName);
-		return true;
 	}
+	break;
+
 	case ECookDependency::Package:
 	{
 		FName LocalPackageName;
@@ -671,8 +750,9 @@ bool FCookDependency::Load(FCbFieldView Value)
 			return false;
 		}
 		*this = FCookDependency::Package(LocalPackageName);
-		return true;
 	}
+	break;
+
 	case ECookDependency::ConsoleVariable:
 	{
 		FString LocalConsoleVariableName;
@@ -681,8 +761,9 @@ bool FCookDependency::Load(FCbFieldView Value)
 			return false;
 		}
 		*this = FCookDependency::ConsoleVariable(LocalConsoleVariableName);
-		return true;
 	}
+	break;
+
 	case ECookDependency::Config:
 	{
 		FString LocalConfigPath;
@@ -698,14 +779,16 @@ bool FCookDependency::Load(FCbFieldView Value)
 		{
 			*this = FCookDependency::Config(UE::ConfigAccessTracking::FConfigAccessData::Parse(LocalConfigPath));
 		}
-		return true;
 	}
+	break;
+
 	case ECookDependency::SettingsObject:
 	{
 		// Settings objects are not persistable; construct an empty SettingsObject dependency
 		*this = FCookDependency::SettingsObject(nullptr);
-		return true;
 	}
+	break;
+
 	case ECookDependency::NativeClass:
 	{
 		FString LocalClassPath;
@@ -714,8 +797,9 @@ bool FCookDependency::Load(FCbFieldView Value)
 			return false;
 		}
 		*this = FCookDependency::NativeClass(LocalClassPath);
-		return true;
 	}
+	break;
+
 	case ECookDependency::AssetRegistryQuery:
 	{
 		bool bValid;
@@ -736,8 +820,9 @@ bool FCookDependency::Load(FCbFieldView Value)
 			}
 			*this = FCookDependency::AssetRegistryQuery(MoveTemp(LocalARFilter));
 		}
-		return true;
 	}
+	break;
+
 	case ECookDependency::RedirectionTarget:
 	{
 		FName LocalPackageName;
@@ -746,13 +831,22 @@ bool FCookDependency::Load(FCbFieldView Value)
 			return false;
 		}
 		*this = FCookDependency::RedirectionTarget(LocalPackageName);
-		return true;
 	}
+	break;
+
 	default:
 		break;
 	}
-	checkNoEntry();
-	return false;
+
+	FMemoryView View = Field.AsBinaryView();
+	if (View.GetSize() != ValueSizeInBytes)
+	{
+		return false;
+	}
+
+	FMemory::Memcpy(RawValue, View.GetData(), ValueSizeInBytes);
+	
+	return true;
 }
 
 bool FCookDependency::ConfigAccessDataLessThan(const UE::ConfigAccessTracking::FConfigAccessData& A,
@@ -821,6 +915,13 @@ FCookDependencyContext::FErrorHandlerScope FCookDependencyContext::ErrorHandlerS
 	return FErrorHandlerScope(*this);
 }
 
+void* FCookDependencyContext::SetHasher(void* NewHasher)
+{
+	void* Temp = Hasher;
+	Hasher = NewHasher;
+	return Temp;
+}
+
 } // namespace UE::Cook
 
 namespace UE::Cook::Dependency::Private
diff --git a/Engine/Source/Runtime/CoreUObject/Public/Cooker/CookDependency.h b/Engine/Source/Runtime/CoreUObject/Public/Cooker/CookDependency.h
--- a/Engine/Source/Runtime/CoreUObject/Public/Cooker/CookDependency.h
+++ b/Engine/Source/Runtime/CoreUObject/Public/Cooker/CookDependency.h
@@ -47,6 +47,7 @@ public:
 	 * ignored. When called while testing iteratively skippable the package is marked modified and recooked.
 	 */
 	COREUOBJECT_API void ReportInvalidated();
+
 	/**
 	 * Reports failure to compute the hash (e.g. because a file cannot be read).
 	 * When called while calculating the initial hash the storage of the key fails and the package will be recooked on
@@ -79,6 +80,7 @@ public:
 		friend FCookDependencyContext;
 		FCookDependencyContext& Context;
 	};
+
 	/**
 	 * Add a function that will be removed when the return value goes out of scope, to modify error strings reported
 	 * inside the scope before passing them on to higher scopes or the error consumer.
@@ -93,6 +95,10 @@ public:
 	 * Get the name of the package being considered
 	 */
 	COREUOBJECT_API FName GetPackageName() const { return PackageName; }
+
+	/** Set a new hasher and return the old one. */
+	void* SetHasher(void* NewHasher);
+
 private:
 	TUniqueFunction<void(ELogVerbosity::Type, FString&&)> OnLog;
 	TUniqueFunction<void(ELogVerbosity::Type)> OnInvalidated;
@@ -299,6 +305,9 @@ public:
 	/** Returns the FARFilter if GetType() == AssetRegistryQuery, else nullptr. */
 	const FARFilter* GetARFilter() const;
 
+	/** Returns the value of the dependency. It's an array of size FCookDependency::VALUE_MAX_BYTES_SIZE */
+	COREUOBJECT_API const uint8* GetRawValue() const;
+
 	/**
 	 * Comparison operator for e.g. deterministic ordering of dependencies.
 	 * Uses persistent comparison data and is somewhat expensive.
@@ -309,7 +318,13 @@ public:
 	bool operator!=(const FCookDependency& Other) const;
 
 	/** Calculate the current hash of this CookDependency, and add it into Context. */
-	COREUOBJECT_API void UpdateHash(FCookDependencyContext& Context) const;
+	COREUOBJECT_API void UpdateHash(FCookDependencyContext& Context);
+
+	COREUOBJECT_API void SetValue(const FIoHash& Hash);
+	COREUOBJECT_API void SetValue(const FBlake3Hash& Hash);
+	COREUOBJECT_API void SetValue(const FUtf8String& String);
+
+	static const uint8 ValueSizeInBytes = 32;
 
 private:
 	explicit FCookDependency(ECookDependency InType);
@@ -358,6 +373,16 @@ private:
 		TUniquePtr<UE::ConfigAccessTracking::FConfigAccessData> ConfigAccessData;
 		TUniquePtr<FARFilter> ARFilter;
 	};
+
+	/** 
+	 * This is the value of the dependency stored as 32 bytes. Other types can be added to the union as long as it
+	 * doesn't go above 32 bytes.
+	 */
+	union
+	{
+		uint8 RawValue[ValueSizeInBytes];
+		FBlake3Hash HashedValue;
+	};
 };
 
 /**
