RemoteReference should not be copied when duplicating objects.

#rnx

diff --git a/Engine/Source/Runtime/CoreUObject/Private/UObject/UObjectGlobals.cpp b/Engine/Source/Runtime/CoreUObject/Private/UObject/UObjectGlobals.cpp
--- a/Engine/Source/Runtime/CoreUObject/Private/UObject/UObjectGlobals.cpp
+++ b/Engine/Source/Runtime/CoreUObject/Private/UObject/UObjectGlobals.cpp
@@ -3158,7 +3158,7 @@ UObject* StaticDuplicateObjectEx( FObjectDuplicationParameters& Parameters )
 	
 	UE_CLOG(FPlatformProperties::RequiresCookedData() && Parameters.SourceObject->HasAnyInternalFlags(EInternalObjectFlags_AsyncLoading), LogUObjectGlobals, Warning, TEXT("Duplicating object '%s' that's still being async loaded"), *Parameters.SourceObject->GetFullName());
 	// Make sure we're not duplicating the AsyncLoading, Async or LoaderImport internal flags, they will prevent the object from being gcd.
-	Parameters.InternalFlagMask &= ~(EInternalObjectFlags::Async | EInternalObjectFlags::LoaderImport | EInternalObjectFlags_AsyncLoading);
+	Parameters.InternalFlagMask &= ~(EInternalObjectFlags::Async | EInternalObjectFlags::LoaderImport | EInternalObjectFlags_AsyncLoading | EInternalObjectFlags::RemoteReference);
 
 	// We can't modify the loader from a transaction, so check for async loading and reset the 
 	// loaders in the open.
