When deserializing remote object data the engine will now attempt to load assets that don't exist on the receiving server.
Native classes and related objects (CDOs/UFunctions) as well as loaded content will have an 'Asset' server id assigned to the object ID to distinguish them from objects created at runtime
Refactored related code and moved F(Packed)RemoteObjectPathName code to its own file.
Added a new class flag to exlude instances of specific class from being assigned the 'asset server id'

diff --git a/Engine/Source/Programs/Shared/EpicGames.Core/UnrealEngineTypes.cs b/Engine/Source/Programs/Shared/EpicGames.Core/UnrealEngineTypes.cs
--- a/Engine/Source/Programs/Shared/EpicGames.Core/UnrealEngineTypes.cs
+++ b/Engine/Source/Programs/Shared/EpicGames.Core/UnrealEngineTypes.cs
@@ -159,10 +159,9 @@ namespace EpicGames.Core
 		Native = 0x00000080u,
 
 		/// <summary>
-		/// Don't export to C++ header.
+		/// Instances of this class can migrate to another server even if loaded from disk.
 		/// </summary>
-		[Obsolete("No longer used in the engine.")]
-		NoExport = 0x00000100u,
+		MigratingAsset = 0x00000100u,
 
 		/// <summary>
 		/// Do not allow users to create in the editor.
@@ -282,7 +281,7 @@ namespace EpicGames.Core
 		/// <summary>
 		/// Flags to inherit from base class
 		/// </summary>
-		Inherit = Transient | Optional | DefaultConfig | Config | PerObjectConfig | ConfigDoNotCheckDefaults | NotPlaceable | Const | HasInstancedReference |
+		Inherit = Transient | Optional | DefaultConfig | Config | PerObjectConfig | ConfigDoNotCheckDefaults | MigratingAsset | NotPlaceable | Const | HasInstancedReference |
 			Deprecated | DefaultToInstanced | GlobalUserConfig | ProjectUserConfig | NeedsDeferredDependencyLoading,
 
 		/// <summary>
@@ -303,7 +302,7 @@ namespace EpicGames.Core
 		/// <summary>
 		/// This is used as a mask for the flags put into generated code for "compiled in" classes.
 		/// </summary>
-		SaveInCompiledInClasses = Abstract | DefaultConfig | GlobalUserConfig | ProjectUserConfig | PerPlatformConfig | Config | Transient | Optional | Native | NotPlaceable | PerObjectConfig |
+		SaveInCompiledInClasses = Abstract | DefaultConfig | GlobalUserConfig | ProjectUserConfig | PerPlatformConfig | Config | Transient | Optional | Native | MigratingAsset | NotPlaceable | PerObjectConfig |
 			ConfigDoNotCheckDefaults | EditInlineNew | CollapseCategories | Interface | DefaultToInstanced | HasInstancedReference | Hidden | Deprecated |
 			HideDropDown | Intrinsic | Const | MinimalAPI | RequiredAPI | MatchedSerializers | NeedsDeferredDependencyLoading,
 	};
diff --git a/Engine/Source/Programs/Shared/EpicGames.UHT/Specifiers/UhtClassSpecifiers.cs b/Engine/Source/Programs/Shared/EpicGames.UHT/Specifiers/UhtClassSpecifiers.cs
--- a/Engine/Source/Programs/Shared/EpicGames.UHT/Specifiers/UhtClassSpecifiers.cs
+++ b/Engine/Source/Programs/Shared/EpicGames.UHT/Specifiers/UhtClassSpecifiers.cs
@@ -64,6 +64,13 @@ namespace EpicGames.UHT.Parsers
 			classObj.RemoveClassFlags(EClassFlags.NotPlaceable);
 		}
 
+		[UhtSpecifier(Extends = UhtTableNames.Class, ValueType = UhtSpecifierValueType.Legacy)]
+		private static void MigratingAssetSpecifier(UhtSpecifierContext specifierContext)
+		{
+			UhtClass classObj = (UhtClass)specifierContext.Type;
+			classObj.AddClassFlags(EClassFlags.MigratingAsset);
+		}
+
 		[UhtSpecifier(Extends = UhtTableNames.Class, ValueType = UhtSpecifierValueType.Legacy)]
 		private static void NotPlaceableSpecifier(UhtSpecifierContext specifierContext)
 		{
diff --git a/Engine/Source/Runtime/CoreUObject/Private/UObject/Class.cpp b/Engine/Source/Runtime/CoreUObject/Private/UObject/Class.cpp
--- a/Engine/Source/Runtime/CoreUObject/Private/UObject/Class.cpp
+++ b/Engine/Source/Runtime/CoreUObject/Private/UObject/Class.cpp
@@ -7919,9 +7919,21 @@ UScriptStruct* TBaseStructure<FRemoteObjectId>::Get()
 	return ScriptStruct;
 }
 
-UScriptStruct* TBaseStructure<FRemoteObjectPathName>::Get()
+UScriptStruct* TBaseStructure<FRemoteObjectTables>::Get()
 {
-	static UScriptStruct* ScriptStruct = StaticGetBaseStructureInternal(TEXT("RemoteObjectPathName"));
+	static UScriptStruct* ScriptStruct = StaticGetBaseStructureInternal(TEXT("RemoteObjectTables"));
+	return ScriptStruct;
+}
+
+UScriptStruct* TBaseStructure<FPackedRemoteObjectPathName>::Get()
+{
+	static UScriptStruct* ScriptStruct = StaticGetBaseStructureInternal(TEXT("PackedRemoteObjectPathName"));
+	return ScriptStruct;
+}
+
+UScriptStruct* TBaseStructure<FRemoteObjectBytes>::Get()
+{
+	static UScriptStruct* ScriptStruct = StaticGetBaseStructureInternal(TEXT("RemoteObjectBytes"));
 	return ScriptStruct;
 }
 
diff --git a/Engine/Source/Runtime/CoreUObject/Private/UObject/RemoteObject.cpp b/Engine/Source/Runtime/CoreUObject/Private/UObject/RemoteObject.cpp
--- a/Engine/Source/Runtime/CoreUObject/Private/UObject/RemoteObject.cpp
+++ b/Engine/Source/Runtime/CoreUObject/Private/UObject/RemoteObject.cpp
@@ -15,6 +15,7 @@
 #include "UObject/ObjectHandlePrivate.h"
 #include "UObject/UObjectHash.h"
 #include "UObject/UObjectMigrationContext.h"
+#include "UObject/UObjectThreadContext.h"
 #include "Misc/CommandLine.h"
 #include "Misc/CoreDelegates.h"
 #include <atomic>
@@ -26,12 +27,14 @@ namespace UE::RemoteObject::Private
 
 FRemoteServerId RemoteServerId;
 std::atomic<uint64> RemoteObjectSerialNumber(1);
+std::atomic<uint64> AssetObjectSerialNumber(1);
 int32 UnsafeToMigrateObjects = 0; // This should go into TLS
 
 class FRemoteObjectMaps
 {
 	mutable FTransactionallySafeCriticalSection ObjectMapCritical;
 	TMap<FRemoteObjectId, UE::RemoteObject::Handle::FRemoteObjectStub*> RemoteObjects;
+	TMap<FRemoteObjectId, FRemoteObjectPathName> AssetPaths;
 
 public:
 
@@ -98,11 +101,23 @@ public:
 		{
 			if (Elem.Value)
 			{
-				Elem.Value->PhysicsOverrideServerId = FRemoteServerId(0);
+				Elem.Value->PhysicsOverrideServerId = FRemoteServerId();
 			}
 		}
 #endif
 	}
+
+	void StoreAssetPath(UObject* InObject)
+	{
+		FRemoteObjectId ObjectId(InObject);
+		FRemoteObjectPathName& Path = AssetPaths.FindOrAdd(ObjectId);
+		Path = FRemoteObjectPathName(InObject);
+	}
+
+	FRemoteObjectPathName* FindAssetPath(FRemoteObjectId ObjectId)
+	{
+		return AssetPaths.Find(ObjectId);
+	}
 };
 FRemoteObjectMaps* ObjectMaps = nullptr;
 
@@ -118,9 +133,8 @@ void InitRemoteObjects()
 			ServerIdNumber = 1;
 		}
 	}
-	checkf(ServerIdNumber != 0, TEXT("Remote Server Id=%u is not valid"), ServerIdNumber);
+	checkf(ServerIdNumber != (uint32)ERemoteServerIdConstants::Invalid, TEXT("Remote Server Id=%u is not valid"), ServerIdNumber);
 	RemoteServerId = FRemoteServerId(ServerIdNumber);
-	checkf(RemoteServerId != UE::RemoteObject::Transfer::DatabaseId, TEXT("Remote Server Id=%s is reserved for database"), *RemoteServerId.ToString());
 
 	UE::RemoteObject::Transfer::InitRemoteObjectTransfer();
 
@@ -187,6 +201,18 @@ void MarkAsLocal(UObject* Object)
 	ensureMsgf(!Object->HasAnyFlags(EObjectFlags::RF_ClassDefaultObject | EObjectFlags::RF_ArchetypeObject), TEXT("We're about to set an ArchetypeObject %s as remote reference"), *GetNameSafe(Object));
 }
 
+void StoreAssetPath(UObject* Object)
+{
+	// Make sure the asset has a stub and that the stub knows the owner if this asset is the asset server (disk / content)
+	ObjectMaps->FindOrAddRemoteObjectStub(Object, FRemoteServerId(ERemoteServerIdConstants::Asset));
+	ObjectMaps->StoreAssetPath(Object);
+}
+
+FRemoteObjectPathName* FindAssetPath(FRemoteObjectId RemoteId)
+{
+	return ObjectMaps->FindAssetPath(RemoteId);
+}
+
 UE::RemoteObject::Handle::FRemoteObjectStub* FindRemoteObjectStub(FRemoteObjectId ObjectId)
 {
 	return ObjectMaps->FindRemoteObjectStub(ObjectId);
@@ -354,7 +380,7 @@ FRemoteServerId GetPhysicsServerId(const UObject* Object)
 	else
 	{
 		//If the object doesn't have stub, then the physics ID is invalid. It should not participate in push migration.
-		Result = FRemoteServerId(0);
+		Result = FRemoteServerId();
 	}
 #endif
 	return Result;
@@ -392,6 +418,15 @@ UObject* ResolveObject(const FRemoteObjectStub* Stub, ERemoteReferenceType RefTy
 
 	// This is a slightly faster version of IsRemote(FRemoteObjectId) because we already know a Stub exists and we are going to re-use the Object pointer
 	UObject* Object = StaticFindObjectFastInternal(Stub->Id);
+
+	if (!Object && Stub->OwningServerId == FRemoteServerId(ERemoteServerIdConstants::Asset))
+	{
+		if (FRemoteObjectPathName* AssetPath = FindAssetPath(Stub->Id))
+		{
+			Object = AssetPath->Resolve();
+		}
+	}
+
 	if (!IsSafeToMigrateObjects() && (Object || RefType == ERemoteReferenceType::Weak)) // Not a thread-safe test
 	{
 		// Begin/FinishDestroy overrides may attempt to access subobjects of objects that have been migrated in which case we
@@ -530,11 +565,78 @@ namespace UE::CoreUObject::Private
 
 } // namespace UE::CoreUObject::Private
 
-FRemoteObjectId FRemoteObjectId::Generate()
+#ifndef UE_WITH_REMOTE_ASSET_ID
+#define UE_WITH_REMOTE_ASSET_ID 1 // set this to 0 to disable remote asset IDs
+#endif
+
+FRemoteObjectId FRemoteObjectId::Generate(UObjectBase* InObject, EInternalObjectFlags InInitialFlags /*= EInternalObjectFlags::None*/)
 {
 	using namespace UE::RemoteObject::Private;
+	using namespace UE::CoreUObject::Private;
 
-	return FRemoteObjectId(RemoteServerId, RemoteObjectSerialNumber.fetch_add(1));
+	bool bIsAsset = false;
+#if UE_WITH_REMOTE_ASSET_ID
+	if (GIsInitialLoad || !UE::RemoteObject::GetGlobalServerId().IsValid() || !!(InInitialFlags & EInternalObjectFlags::Native) || !!(InObject->GetFlags() & RF_ArchetypeObject))
+	{
+		// Native objects (classes, CDOs etc) or objects loaded during initial load are always in memory and are considered assets any server can find locally
+		// Note that this first condition can not touch too much of UObject API because we might literarlly be constructing the first StaticClass() etc.
+		// Hopefull GIsInitialLoad and the Native flag will filter most of the initially created objects and most of the native objects will be constructed before we ever hit the 'else' block
+		bIsAsset = true;
+	}
+	else
+	{
+		FUObjectThreadContext& ThreadContext = FUObjectThreadContext::Get();
+		if (ThreadContext.AsyncPackageLoader || (ThreadContext.GetSerializeContext() && ThreadContext.GetSerializeContext()->GetBeginLoadCount() > 0) || !!(InObject->GetFlags() & RF_WasLoaded))
+		{
+			// If we're constructing this object when loading content then this object is an asset if:
+			// its class or any of its outers' class is NOT marked as MigratingAsset
+			// OR it's an archetype or subobject of an archetype 
+			// OR it's a subobject of a UStruct (class)
+			bool bIsMigratingAsset = false;
+			bool bIsClassOrArchetypeSubobject = false;
+
+			for (UObjectBase* OuterIt = InObject; OuterIt; OuterIt = FObjectHandleUtils::GetNonAccessTrackedOuterNoResolve(OuterIt))
+			{
+				UClass* Class = OuterIt->GetClass();
+				if (!!(OuterIt->GetFlags() & RF_ArchetypeObject) || Class->IsChildOf(UStruct::StaticClass()))
+				{
+					bIsClassOrArchetypeSubobject = true;
+					break;
+				}
+				if (Class->HasAnyClassFlags(CLASS_MigratingAsset))
+				{
+					bIsMigratingAsset = true;
+				}					
+			}
+
+			bIsAsset = bIsClassOrArchetypeSubobject || !bIsMigratingAsset;
+		}
+	}
+#endif // UE_WITH_REMOTE_ASSET_ID
+
+	FRemoteObjectId Result;
+	if (!bIsAsset)
+	{
+		Result = FRemoteObjectId(RemoteServerId, RemoteObjectSerialNumber.fetch_add(1));
+	}
+	else
+	{
+		Result = FRemoteObjectId(FRemoteServerId(ERemoteServerIdConstants::Asset), AssetObjectSerialNumber.fetch_add(1));
+	}
+
+	return Result;
+}
+
+FString FRemoteObjectId::ToString() const
+{
+	if (!IsAsset())
+	{
+		return FString::Printf(TEXT("%llu-%llu"), ServerId, SerialNumber);
+	}
+	else
+	{
+		return FString::Printf(TEXT("Asset-%llu"), SerialNumber);
+	}
 }
 
 FRemoteObjectId::FRemoteObjectId(const UObjectBase* Object)
diff --git a/Engine/Source/Runtime/CoreUObject/Private/UObject/RemoteObjectPathName.cpp b/Engine/Source/Runtime/CoreUObject/Private/UObject/RemoteObjectPathName.cpp
new file mode 100644
--- /dev/null
+++ b/Engine/Source/Runtime/CoreUObject/Private/UObject/RemoteObjectPathName.cpp
@@ -0,0 +1,156 @@
+// Copyright Epic Games, Inc. All Rights Reserved.
+
+#include "UObject/RemoteObjectPathName.h"
+#include "UObject/RemoteObject.h"
+#include "UObject/RemoteObjectTransfer.h"
+#include "UObject/RemoteObjectPrivate.h"
+#include "UObject/ObjectPtr.h"
+#include "UObject/PropertyPortFlags.h"
+#include "UObject/UObjectHash.h"
+#include "UObject/UObjectMigrationContext.h"
+#include "UObject/UnrealType.h"
+#include "UObject/Package.h"
+#include "Serialization/ArchiveUObject.h"
+#include "Serialization/MemoryWriter.h"
+#include "Serialization/MemoryReader.h"
+#include "Templates/Casts.h"
+#include "UObject/ObjectHandlePrivate.h"
+#include "Misc/Paths.h"
+
+namespace UE::RemoteObject::Private
+{
+
+template <typename PathType>
+FORCENOINLINE FString RemotePathNameToString(const PathType& InPathName, const FRemoteObjectTables& InTables, int32 InMinPathNameIndex = 0)
+{
+	TStringBuilder<256> Result;
+	for (int32 PathSegment = InPathName.Num() - 1; PathSegment >= InMinPathNameIndex; --PathSegment)
+	{
+		InPathName.GetSegmentName(PathSegment, InTables).AppendString(Result);
+		if (PathSegment > InMinPathNameIndex)
+		{
+			if (PathSegment != (InPathName.Num() - 2))
+			{
+				Result << TEXT('.');
+			}
+			else
+			{
+				// at this point the constructed path consists of /Package/Name.ObjectName so we need to use the subobject delimiter
+				// see UObjectBaseUtility::GetPathName
+				Result << SUBOBJECT_DELIMITER_CHAR;
+			}
+		}
+	}
+	return Result.ToString();
+}
+
+template <typename PathType>
+FORCENOINLINE UObject* LoadRemotePathName(const PathType& InPathName, const FRemoteObjectTables& InTables, int32 InNameIndex)
+{
+	FString LoadPathName = RemotePathNameToString<PathType>(InPathName, InTables, InNameIndex);
+	const int32 PackageNameIndex = InPathName.Num() - 1;
+	UObject* Object = nullptr;
+	if (InNameIndex == PackageNameIndex)
+	{
+		// We're loading the Outermost (Package) so we can't use StaticLoadObject because it assumes the name of the object we're trying to load 
+		// is just short package name (/Root/PackageName.PackageName) which is not true in case of Blueprint classes (/Root/PackageName.PackageName_C)
+		// Loading the package will load all objects inside of it but that's also true for StaticLoadObject.
+		Object = LoadPackage(nullptr, *LoadPathName, LOAD_None);
+	}
+	else
+	{
+		Object = StaticLoadObject(UObject::StaticClass(), nullptr, *LoadPathName);
+	}
+	if (!Object)
+	{
+		if (InNameIndex == 0)
+		{
+			UE_LOG(LogRemoteObject, Warning, TEXT("Failed to load asset object %s (%s)"), *LoadPathName, *InPathName.GetSegmentId(InNameIndex, InTables).ToString());
+		}
+		else
+		{
+			UE_LOG(LogRemoteObject, Warning, TEXT("Failed to load asset object %s which is an outer of remote object %s (%s)"),
+				*LoadPathName,
+				*RemotePathNameToString<PathType>(InPathName, InTables),
+				*InPathName.GetSegmentId(0, InTables).ToString());
+		}
+	}
+	return Object;
+}
+
+template <typename PathType>
+FORCENOINLINE UObject* ResolveRemotePathName(const PathType& InPathName, const FRemoteObjectTables& InTables)
+{
+	// Resolve a remote object path name starting with the outermost object and working towards the innermost one
+	UObject* Outer = nullptr;
+	UObject* Object = nullptr;
+	for (int32 PathNameIndex = InPathName.Num() - 1; PathNameIndex >= 0; --PathNameIndex)
+	{
+		Object = StaticFindObjectFast(UObject::StaticClass(), Outer, InPathName.GetSegmentName(PathNameIndex, InTables));
+		if (!Object)
+		{			
+			FRemoteObjectId RemoteId = InPathName.GetSegmentId(PathNameIndex, InTables);
+			if (RemoteId.IsAsset())
+			{
+				Object = LoadRemotePathName(InPathName, InTables, PathNameIndex);
+				if (!Object)
+				{
+					break;
+				}
+			}
+		}
+		Outer = Object;
+	}
+	return Object;
+}
+
+} // namespace UE::RemoteObject::Private
+
+void FRemoteObjectTables::Serialize(FArchive& Ar)
+{
+	Ar << Names;
+	Ar << RemoteIds;
+}
+
+UObject* FPackedRemoteObjectPathName::Resolve(const FRemoteObjectTables& InTables) const
+{
+	using namespace UE::RemoteObject::Private;
+	return ResolveRemotePathName(*this, InTables);
+}
+
+void FPackedRemoteObjectPathName::Serialize(FArchive& Ar)
+{
+	Ar << Names;
+	Ar << RemoteIds;
+}
+
+FString FPackedRemoteObjectPathName::ToString(const FRemoteObjectTables& InTables, int32 InMinPathSegmentIndex /*= 0*/) const
+{
+	using namespace UE::RemoteObject::Private;
+	checkf(InMinPathSegmentIndex >= 0, TEXT("Minimum path segment index must be > 0 (got %d)"), InMinPathSegmentIndex);
+	return RemotePathNameToString(*this, InTables, InMinPathSegmentIndex);
+}
+
+
+
+FRemoteObjectPathName::FRemoteObjectPathName(UObject* InObject)
+{
+	for (UObject* PathNameObject = InObject; PathNameObject; PathNameObject = PathNameObject->GetOuter())
+	{
+		Names.Add(PathNameObject->GetFName());
+		RemoteIds.Add(FRemoteObjectId(PathNameObject));
+	}
+}
+
+UObject* FRemoteObjectPathName::Resolve() const
+{
+	using namespace UE::RemoteObject::Private;
+	return ResolveRemotePathName(*this, *this);
+}
+
+FString FRemoteObjectPathName::ToString(int32 InMinPathSegmentIndex /*= 0*/) const
+{
+	using namespace UE::RemoteObject::Private;
+	checkf(InMinPathSegmentIndex >= 0, TEXT("Minimum path segment index must be > 0 (got %d)"), InMinPathSegmentIndex);
+	return RemotePathNameToString(*this, *this, InMinPathSegmentIndex);
+}
\ No newline at end of file
diff --git a/Engine/Source/Runtime/CoreUObject/Private/UObject/RemoteObjectPrivate.h b/Engine/Source/Runtime/CoreUObject/Private/UObject/RemoteObjectPrivate.h
--- a/Engine/Source/Runtime/CoreUObject/Private/UObject/RemoteObjectPrivate.h
+++ b/Engine/Source/Runtime/CoreUObject/Private/UObject/RemoteObjectPrivate.h
@@ -10,6 +10,8 @@
 
 DECLARE_LOG_CATEGORY_EXTERN(LogRemoteObject, Display, All);
 
+struct FRemoteObjectPathName;
+
 namespace UE::RemoteObject::Private
 {
 	/**
@@ -47,6 +49,16 @@ namespace UE::RemoteObject::Private
 	*/
 	FName GetServerBaseNameForUniqueName(const UClass* Class);
 
+	/**
+	* Stores FRemoteObjectPath for a remotely referenced asset that's about to be destroyed so that the engine knows it should load the asset when something requests it
+	*/
+	void StoreAssetPath(UObject* Object);
+
+	/**
+	* Attempts to find an FRemoteObjectPath for an object id representing an asset
+	*/
+	FRemoteObjectPathName* FindAssetPath(FRemoteObjectId RemoteId);
+
 	struct FUnsafeToMigrateScope
 	{
 		FUnsafeToMigrateScope();
@@ -56,6 +68,8 @@ namespace UE::RemoteObject::Private
 
 namespace UE::RemoteObject::Transfer::Private
 {
+
+
 	/**
 	* Stores unreachable objects to database
 	*/
diff --git a/Engine/Source/Runtime/CoreUObject/Private/UObject/RemoteObjectSerialization.cpp b/Engine/Source/Runtime/CoreUObject/Private/UObject/RemoteObjectSerialization.cpp
--- a/Engine/Source/Runtime/CoreUObject/Private/UObject/RemoteObjectSerialization.cpp
+++ b/Engine/Source/Runtime/CoreUObject/Private/UObject/RemoteObjectSerialization.cpp
@@ -32,13 +32,8 @@ static FAutoConsoleVariableRef CVarResetBorrowedObjects(
 	GRemoteObjectsMigrateFullHierarchy,
 	TEXT("Whether remote objects that were borrowed should be reset upon returning to their owner server instead of being reconstructed"));
 
-using FNameIndexType = FRemoteObjectPathName::FNameIndexType;
+using FNameIndexType = FPackedRemoteObjectPathName::FNameIndexType;
 
-FArchive& operator<<(FArchive& Ar, FRemoteObjectPathName& PathName)
-{
-	Ar << PathName.Names;
-	return Ar;
-}
 
 FArchive& operator<<(FArchive& Ar, FRemoteObjectBytes& Chunk)
 {
@@ -51,27 +46,29 @@ namespace UE::RemoteObject::Serialization::Disk
 
 void SerializeNameTables(FArchive& Ar, const FRemoteObjectData& InObjectData)
 {
-	FNameIndexType NumNames = IntCastChecked<FNameIndexType>(InObjectData.Names.Num());
+	FNameIndexType NumNames = IntCastChecked<FNameIndexType>(InObjectData.Tables.Names.Num());
 	Ar << NumNames;
-	for (const FName& Name : InObjectData.Names)
+	for (const FName& Name : InObjectData.Tables.Names)
 	{
 		FString NameString = Name.ToString();
 		Ar << NameString;
 	}
-	Ar << const_cast<TArray<FRemoteObjectPathName>&>(InObjectData.PathNames);
+	Ar << const_cast<TArray<FRemoteObjectId>&>(InObjectData.Tables.RemoteIds);
+	Ar << const_cast<TArray<FPackedRemoteObjectPathName>&>(InObjectData.PathNames);
 }
 
 void DeserializeNameTables(FArchive& Ar, FRemoteObjectData& OutObjectData)
 {
-	OutObjectData.Names.Reset();
+	OutObjectData.Tables.Names.Reset();
 	FNameIndexType NumNames = 0;
 	Ar << NumNames;
 	for (FNameIndexType NameIndex = 0; NameIndex < NumNames; ++NameIndex)
 	{
 		FString NameString;
 		Ar << NameString;
-		OutObjectData.Names.Add(FName(NameString, FNAME_Add));
+		OutObjectData.Tables.Names.Add(FName(NameString, FNAME_Add));
 	}
+	Ar << OutObjectData.Tables.RemoteIds;
 	Ar << OutObjectData.PathNames;
 }
 
@@ -179,29 +176,6 @@ protected:
 	FString SerializationScope;
 	TMap<FString, TMap<FProperty*, FMigratedPropertyStats>> ObjectPropertyStats;
 
-	FString PathNameToString(const FRemoteObjectPathName& PathName) const
-	{
-		TStringBuilder<256> Result;
-		for (int32 PathSegment = PathName.Names.Num() - 1; PathSegment >= 0; --PathSegment)
-		{
-			ObjectData.Names[PathName.Names[PathSegment]].AppendString(Result);
-			if (PathSegment > 0)
-			{
-				if (PathSegment != (PathName.Names.Num() - 2))
-				{
-					Result << TEXT('.');
-				}
-				else
-				{
-					// at this point the constructed path consists of /Package/Name.ObjectName so we need to use the subobject delimiter
-					// see UObjectBaseUtility::GetPathName
-					Result << SUBOBJECT_DELIMITER_CHAR;
-				}
-			}
-		}
-		return Result.ToString();
-	}
-
 	void DumpStatsToLog()
 	{
 		int64 NameTableSize = 0;
@@ -215,12 +189,12 @@ protected:
 		const int64 TotalSize = NameTableSize + ObjectData.GetNumBytes();
 
 		UE_LOG(LogRemoteObject, Log, TEXT("%s Object Data stats for %s %s (Object Data toal: %d, total: %lld):"), *GetArchiveName(), *FRemoteObjectId(RootObject).ToString(), *GetFullNameSafe(RootObject), ObjectData.GetNumBytes(), TotalSize);
-		UE_LOG(LogRemoteObject, Log, TEXT("  Name Table total size: %lld (Names: %d, Paths: %d)"), NameTableSize, ObjectData.Names.Num(), ObjectData.PathNames.Num());
+		UE_LOG(LogRemoteObject, Log, TEXT("  Name Table total size: %lld (Names: %d, RemoteIds: %d, Paths: %d)"), NameTableSize, ObjectData.Tables.Names.Num(), ObjectData.Tables.RemoteIds.Num(), ObjectData.PathNames.Num());
 
 		TArray<FString> SortedPathNames;
-		for (const FRemoteObjectPathName& PathName : ObjectData.PathNames)
+		for (const FPackedRemoteObjectPathName& PathName : ObjectData.PathNames)
 		{
-			SortedPathNames.Add(PathNameToString(PathName));
+			SortedPathNames.Add(PathName.ToString(ObjectData.Tables));
 		}
 		SortedPathNames.Sort();
 		for (const FString& PathName : SortedPathNames)
@@ -362,12 +336,15 @@ protected:
 	TSet<UObject*> SerializedObjects;
 	TMap<FName, FNameIndexType> NameMap;
 	TMap<UObject*, FNameIndexType> PathNameMap;
+	TMap<FRemoteObjectId, FNameIndexType> RemoteIdMap;
 	TMap<UObject*, UObject*>* ObjectToArchetypeMap = nullptr;
 
 public:
 
 	static FRemoteObjectReferenceInfo GetReferenceInfo(const FObjectPtr& ObjPtr, UObject* InRootObject)
 	{
+		using namespace UE::RemoteObject;
+
 		FRemoteObjectReferenceInfo Info;
 		Info.Id = ObjPtr.GetRemoteId();
 		if (ObjPtr)
@@ -377,7 +354,7 @@ public:
 				Info.Object = ObjPtr.Get();
 				Info.bIsSubobject = Info.Object->IsIn(InRootObject);
 
-				if (Info.Object == InRootObject || Info.bIsSubobject || (Info.Id.GetServerId().IsValid() && Info.Id.GetServerId() != UE::RemoteObject::GetGlobalServerId()))
+				if (Info.Object == InRootObject || Info.bIsSubobject || (Info.Id.GetServerId().IsValid() && Info.Id.GetServerId() != GetGlobalServerId() && !Info.Id.IsAsset()))
 				{
 					Info.Type = ERemoteReferenceType::IdOnly;
 				}
@@ -433,18 +410,18 @@ private:
 				Info.Id = FRemoteObjectId(Info.Object);
 			}
 
-			const bool bGloballyKnownAsset = Info.Object->HasAnyFlags(EObjectFlags::RF_ClassDefaultObject | EObjectFlags::RF_ArchetypeObject);
-			if (!bGloballyKnownAsset)
-			{
-				Info.Object->SetInternalFlags(EInternalObjectFlags::RemoteReference);
-			}
+			// Anything can be marked as a remote reference, even assets in which case we rely on this flag to be set so that GC calls
+			// StoreObjectToDatabase for any remotely referenced asset (and only for remotely referenced assets) that's about to be GC'd
+			Info.Object->SetInternalFlags(EInternalObjectFlags::RemoteReference);
 		}
 
 		*this << Info.Type;
 		// Always serialize unique id as objects may not exist on the other server and then we may need to pull them from this server
+		FNameIndexType IdIndex = 0;
 		if (Info.Type != ERemoteReferenceType::None)
 		{
-			*this << Info.Id;
+			IdIndex = AddRemoteIdToIdMap(Info.Id);
+			*this << IdIndex;
 		}
 
 		if (Info.Type == ERemoteReferenceType::IdOnly)
@@ -467,7 +444,7 @@ public:
 	FArchiveRemoteObjectWriter(UObject* InRootObject, FRemoteObjectData& OutObjectData, const FUObjectMigrationContext* InMigrationContext, const TCHAR* ArchiveName = nullptr)
 		: TArchiveRemoteObjectBase<FMemoryWriter>(OutObjectData, InMigrationContext, ArchiveName ? ArchiveName : TEXT("RemoteObjectWriter"))
 	{
-		static_assert(sizeof(FNameIndexType) == sizeof(FRemoteObjectPathName::FNameIndexType), "Name Index type must match the type used for storing remote object path name indices");
+		static_assert(sizeof(FNameIndexType) == sizeof(FPackedRemoteObjectPathName::FNameIndexType), "Name Index type must match the type used for storing remote object path name indices");
 
 		SetRootObject(InRootObject);
 		ObjectsToSerialize.Add(InRootObject);
@@ -494,13 +471,29 @@ public:
 		}
 		else
 		{
-			int32 NewIndex = ObjectData.Names.Add(Name);
+			int32 NewIndex = ObjectData.Tables.Names.Add(Name);
 			NameIndex = IntCastChecked<FNameIndexType>(NewIndex);
 			NameMap.Add(Name, NameIndex);
 		}
 		return NameIndex;
 	}
 
+	FNameIndexType AddRemoteIdToIdMap(FRemoteObjectId RemoteId)
+	{
+		FNameIndexType IdIndex = TNumericLimits<FNameIndexType>::Max();
+		if (FNameIndexType* ExistingIdIndex = RemoteIdMap.Find(RemoteId))
+		{
+			IdIndex = *ExistingIdIndex;
+		}
+		else
+		{
+			int32 NewIndex = ObjectData.Tables.RemoteIds.Add(RemoteId);
+			IdIndex = IntCastChecked<FNameIndexType>(NewIndex);
+			RemoteIdMap.Add(RemoteId, IdIndex);
+		}
+		return IdIndex;
+	}
+
 	FNameIndexType AddPathNameToNameMap(UObject* Object)
 	{
 		FNameIndexType PathNameIndex = TNumericLimits<FNameIndexType>::Max();
@@ -513,10 +506,12 @@ public:
 			int32 NewPathNameIndex = ObjectData.PathNames.AddDefaulted();
 			PathNameIndex = IntCastChecked<FNameIndexType>(NewPathNameIndex);
 
-			FRemoteObjectPathName& NewPathName = ObjectData.PathNames[PathNameIndex];
+			FPackedRemoteObjectPathName& NewPathName = ObjectData.PathNames[PathNameIndex];
+
 			// Store individual indices of FNames of every object in this object's outer chain
 			for (UObject* OuterChain = Object; OuterChain; OuterChain = UE::CoreUObject::Private::FObjectHandleUtils::GetNonAccessTrackedOuterNoResolve(OuterChain))
 			{
+				NewPathName.RemoteIds.Add(AddRemoteIdToIdMap(FRemoteObjectId(OuterChain)));
 				NewPathName.Names.Add(AddNameToNameMap(OuterChain->GetFName()));
 			}
 			PathNameMap.Add(Object, PathNameIndex);
@@ -685,6 +680,7 @@ FArchive& operator<<(FArchive& Ar, FRemoteObjectHeader& Header)
 class FArchiveRemoteObjectReader : public TArchiveRemoteObjectBase<FMemoryReader>
 {
 	const TArray<FName>& Names;
+	const TArray<FRemoteObjectId>& RemoteIds;
 	const TArray<UObject*>& ResolvedPathNameObjects;
 	const ERemoteObjectSerializationFlags DeserializeFlags = ERemoteObjectSerializationFlags::None;
 
@@ -696,7 +692,11 @@ class FArchiveRemoteObjectReader : public TArchiveRemoteObjectBase<FMemoryReader
 		if (Type != ERemoteReferenceType::None)
 		{
 			FRemoteObjectId ObjId;
-			*this << ObjId;
+			{
+				FNameIndexType IdIndex;
+				*this << IdIndex;
+				ObjId = RemoteIds[IdIndex];
+			}
 
 			bool bNeedsResolvingWithId = true;
 			if (Type == ERemoteReferenceType::PathName)
@@ -750,7 +750,8 @@ public:
 	 */
 	FArchiveRemoteObjectReader(FRemoteObjectData& InObjectData, const TArray<UObject*>& InResolvedPathNames, const FUObjectMigrationContext* InMigrationContext, ERemoteObjectSerializationFlags InDeserializeFlags)
 		: TArchiveRemoteObjectBase<FMemoryReader>(InObjectData, InMigrationContext, TEXT("RemoteObjectReader"))
-		, Names(InObjectData.Names)
+		, Names(InObjectData.Tables.Names)
+		, RemoteIds(InObjectData.Tables.RemoteIds)
 		, ResolvedPathNameObjects(InResolvedPathNames)
 		, DeserializeFlags(InDeserializeFlags)
 	{
@@ -998,20 +999,10 @@ UObject* ConstructRemoteObject(const FRemoteObjectHeader& Header, FRemoteObjectC
 }
 
 void ResolvePathNames(const FRemoteObjectData& InObjectData, TArray<UObject*>& OutResolvedObjects)
-{
-	for (const FRemoteObjectPathName& PathName : InObjectData.PathNames)
+{	
+	for (const FPackedRemoteObjectPathName& PathName : InObjectData.PathNames)
 	{
-		UObject* Outer = nullptr;
-		UObject* Object = nullptr;
-		for (int32 PathNameIndex = PathName.Names.Num() - 1; PathNameIndex >= 0; --PathNameIndex)
-		{
-			Object = StaticFindObjectFast(UObject::StaticClass(), Outer, InObjectData.Names[PathName.Names[PathNameIndex]]);
-			if (!Object)
-			{
-				break;
-			}
-			Outer = Object;
-		}
+		UObject* Object = PathName.Resolve(InObjectData.Tables);
 		OutResolvedObjects.Add(Object);
 	}
 }
diff --git a/Engine/Source/Runtime/CoreUObject/Private/UObject/RemoteObjectTransfer.cpp b/Engine/Source/Runtime/CoreUObject/Private/UObject/RemoteObjectTransfer.cpp
--- a/Engine/Source/Runtime/CoreUObject/Private/UObject/RemoteObjectTransfer.cpp
+++ b/Engine/Source/Runtime/CoreUObject/Private/UObject/RemoteObjectTransfer.cpp
@@ -21,7 +21,7 @@
 
 namespace UE::RemoteObject::Transfer
 {
-const FRemoteServerId DatabaseId(0x3ff);
+const FRemoteServerId DatabaseId(ERemoteServerIdConstants::Database);
 static TAutoConsoleVariable<bool> CVarRemoveMissingMigratedComponents(TEXT("DSTM.RemoveMissingMigratedComponents"), false, TEXT("When we receive less SubObjects (Components) than we expect, remove the extras that we have locally (assume they were removed on purpose)"));
 
 UObject* MigrateObjectFromObjectData(FRemoteObjectData& ObjectData, const FUObjectMigrationContext& MigrationContext);
@@ -744,6 +744,29 @@ void StoreObjectToDatabase(UObject* Object)
 	using namespace UE::RemoteObject::Private;
 	using namespace UE::RemoteObject::Transfer;
 
+	FRemoteObjectId ObjectId(Object);
+	if (ObjectId.IsAsset())
+	{
+		// We don't need to store assets to database since they already exist on disk so we just store their pathname so that we can reload them if needed
+		StoreAssetPath(Object);
+		return;
+	}
+
+	FRemoteServerId LocalServerId = GetGlobalServerId();
+	FRemoteServerId OwnerServerId = UE::RemoteObject::Handle::GetOwnerServerId(Object);
+	if (LocalServerId != OwnerServerId)
+	{
+		// We don't own this, and it's definitely remotely owned, just don't save it to the database
+		if (UE::RemoteObject::Handle::IsRemote(Object))
+		{
+			return;
+		}
+
+		ensureMsgf(false, TEXT("%hs was about to store non-locally-owned Object %s to Server %s's Database; making it remote instead"), __func__, *GetNameSafe(Object), *LocalServerId.ToString());
+		MarkAsRemote(Object, OwnerServerId);
+		return;
+	}
+
 	do
 	{
 		// Object could be a subobject of its outer that is going to be destroyed later (and have this function invoked on)
@@ -757,20 +780,7 @@ void StoreObjectToDatabase(UObject* Object)
 		{
 			// Let's setup a fake migration context for storing this
 			FRemoteObjectId RootObjectId = FRemoteObjectId { OutermostUnreachableRemoteReference };
-			FRemoteServerId LocalServerId = GetGlobalServerId();
-			FRemoteServerId OwnerServerId = UE::RemoteObject::Handle::GetOwnerServerId(Object);
-			if (LocalServerId != OwnerServerId)
-			{
-				// We don't own this, and it's definitely remotely owned, just don't save it to the database
-				if (UE::RemoteObject::Handle::IsRemote(Object))
-				{
-					return;
-				}
 
-				ensureMsgf(false, TEXT("%hs was about to store non-locally-owned Object %s to Server %s's Database; making it remote instead"), __func__, *GetNameSafe(Object), *LocalServerId.ToString());
-				MarkAsRemote(Object, OwnerServerId); 
-				return;
-			}
 
 			// Setup a migration context for "push migrating to the new owner: the Database Server"
 			FUObjectMigrationContext StoreToDatabaseContext{ .ObjectId = RootObjectId, .RemoteServerId = DatabaseId, .OwnerServerId = DatabaseId, .PhysicsServerId = DatabaseId, .RequestId = PushMigrationRequestId, .MigrationSide = EObjectMigrationSide::Send };
@@ -784,12 +794,9 @@ void StoreObjectToDatabase(UObject* Object)
 
 		checkf(SentObjects.Num() > 0, TEXT("PeerId:%s unable to store object %s to DatabaseId:%s"), *GetGlobalServerId().ToString(), *FRemoteObjectId(Object).ToString(), *DatabaseId.ToString());
 
-		// At this point the entire leaked object tree should be remote with the exception that its current owner server is the database
-		for (UObject* SentObject : SentObjects)
-		{
-			// MarkAsRemote also clears EInternalObjectFlags::RemoteReference flag
-			MarkAsRemote(SentObject, DatabaseId);
-		}
+		// MarkAsRemote clears EInternalObjectFlags::RemoteReference flag which will prevent re-enetring this loop with the same outermost object but 
+		// we still want to separately process its subobjects (so that they're stored as separate entries in the database) so only mark the root object as remote
+		MarkAsRemote(OutermostUnreachableRemoteReference, DatabaseId);
 
 		// It's possible that OutermostUnreachableRemoteReference no longer referenced this Object and Object hasn't been stored to disk
 		// in which case we need to repeat this loop until Object is marked as remote (and is no longer marked as RemoteReference).
@@ -848,23 +855,25 @@ UObject* FRemoteObjectReference::Resolve() const
 
 bool FRemoteObjectReference::NetSerialize(FArchive& Ar, class UPackageMap* Map, bool& bOutSuccess)
 {
+	using namespace UE::RemoteObject::Serialization;
+
 	// Serialize remote object id
 	Ar << *this;
 
-	// Serialize the pathname of the object if possible which will then be used on the receiving end to find assets loaded by both servers in memory
-	FString PathName;
+	// Serialize the pathname of the object if possible which will then be used on the receiving end to find assets loaded by both servers in memory	
+	FRemoteObjectPathName PathName;
 	if (Ar.IsSaving() && ObjectId.IsValid())
 	{		
 		if (UObject* ExistingObject = StaticFindObjectFastInternal(ObjectId))
 		{
 			UE::RemoteObject::Private::RegisterSharedObject(ExistingObject);
-			PathName = ExistingObject->GetPathName();
+			PathName = FRemoteObjectPathName(ExistingObject);
 		}
 	}
 
 	Ar << PathName;
 
-	if (Ar.IsLoading() && PathName.Len())
+	if (Ar.IsLoading() && PathName.Num())
 	{
 		// We currently can't differentiate between assets that exist on both servers which we don't want to migrate
 		// so try to find the object in memory first by its id and if it doesn't exist try to find the object by pathname
@@ -872,7 +881,7 @@ bool FRemoteObjectReference::NetSerialize(FArchive& Ar, class UPackageMap* Map,
 		UObject* ExistingObject = StaticFindObjectFastInternal(ObjectId);
 		if (!ExistingObject)
 		{
-			ExistingObject = StaticFindObject(UObject::StaticClass(), nullptr, *PathName);
+			ExistingObject = PathName.Resolve();
 			if (ExistingObject)
 			{
 				ObjectId = FRemoteObjectId(ExistingObject);
diff --git a/Engine/Source/Runtime/CoreUObject/Private/UObject/UObjectArray.cpp b/Engine/Source/Runtime/CoreUObject/Private/UObject/UObjectArray.cpp
--- a/Engine/Source/Runtime/CoreUObject/Private/UObject/UObjectArray.cpp
+++ b/Engine/Source/Runtime/CoreUObject/Private/UObject/UObjectArray.cpp
@@ -240,7 +240,7 @@ void FUObjectArray::AllocateUObjectIndex(UObjectBase* Object, EInternalObjectFla
 #if UE_WITH_REMOTE_OBJECT_HANDLE
 	if (!RemoteId.IsValid())
 	{
-		RemoteId = FRemoteObjectId::Generate();
+		RemoteId = FRemoteObjectId::Generate(Object, InitialFlags);
 	}
 #endif
 
diff --git a/Engine/Source/Runtime/CoreUObject/Private/UObject/UObjectBase.cpp b/Engine/Source/Runtime/CoreUObject/Private/UObject/UObjectBase.cpp
--- a/Engine/Source/Runtime/CoreUObject/Private/UObject/UObjectBase.cpp
+++ b/Engine/Source/Runtime/CoreUObject/Private/UObject/UObjectBase.cpp
@@ -540,7 +540,7 @@ void UObjectBase::Register(class UClass* (*StaticClassFn)(), const TCHAR* Packag
 	FPendingRegistrant* PendingRegistration = new FPendingRegistrant(this);
 	PendingRegistrants.Add(this, FPendingRegistrantInfo(StaticClassFn, InName, PackageName
 #if UE_WITH_REMOTE_OBJECT_HANDLE
-		, FRemoteObjectId::Generate()
+		, FRemoteObjectId::Generate(this)
 #endif
 	));
 
diff --git a/Engine/Source/Runtime/CoreUObject/Public/UObject/Class.h b/Engine/Source/Runtime/CoreUObject/Public/UObject/Class.h
--- a/Engine/Source/Runtime/CoreUObject/Public/UObject/Class.h
+++ b/Engine/Source/Runtime/CoreUObject/Public/UObject/Class.h
@@ -4559,11 +4559,25 @@ template<> struct TBaseStructure<FRemoteObjectId>
 	static COREUOBJECT_API UScriptStruct* Get();
 };
 
-template<> struct TBaseStructure<FRemoteObjectPathName>
+struct FRemoteObjectTables;
+template<> struct TBaseStructure<FRemoteObjectTables>
 {
 	static COREUOBJECT_API UScriptStruct* Get();
 };
 
+struct FPackedRemoteObjectPathName;
+template<> struct TBaseStructure<FPackedRemoteObjectPathName>
+{
+	static COREUOBJECT_API UScriptStruct* Get();
+};
+
+struct FRemoteObjectBytes;
+template<> struct TBaseStructure<FRemoteObjectBytes>
+{
+	static COREUOBJECT_API UScriptStruct* Get();
+};
+
+struct FRemoteObjectData;
 template<> struct TBaseStructure<FRemoteObjectData>
 {
 	static COREUOBJECT_API UScriptStruct* Get();
diff --git a/Engine/Source/Runtime/CoreUObject/Public/UObject/NoExportTypes.h b/Engine/Source/Runtime/CoreUObject/Public/UObject/NoExportTypes.h
--- a/Engine/Source/Runtime/CoreUObject/Public/UObject/NoExportTypes.h
+++ b/Engine/Source/Runtime/CoreUObject/Public/UObject/NoExportTypes.h
@@ -3406,10 +3406,13 @@ struct FRemoteObjectId
 };
 
 USTRUCT(noexport, BlueprintType)
-struct FRemoteObjectPathName
+struct FPackedRemoteObjectPathName
 {
 	GENERATED_BODY()
 
+	UPROPERTY()
+	TArray<uint16> RemoteIds;
+
 	UPROPERTY()
 	TArray<uint16> Names;
 };
@@ -3421,16 +3424,26 @@ struct FRemoteObjectBytes
 	TArray<uint8> Bytes;
 };
 
+USTRUCT(noexport, BlueprintType)
+struct FRemoteObjectTables
+{
+	UPROPERTY()
+	TArray<FName> Names;
+
+	UPROPERTY()
+	TArray<FRemoteObjectId> RemoteIds;
+};
+
 USTRUCT(noexport, BlueprintType)
 struct FRemoteObjectData
 {
 	GENERATED_BODY()
 
 	UPROPERTY()
-	TArray<FName> Names;
+	FRemoteObjectTables Tables;
 
 	UPROPERTY()
-	TArray<FRemoteObjectPathName> PathNames;
+	TArray<FPackedRemoteObjectPathName> PathNames;
 
 	UPROPERTY()
 	TArray<FRemoteObjectBytes> Bytes;
diff --git a/Engine/Source/Runtime/CoreUObject/Public/UObject/ObjectMacros.h b/Engine/Source/Runtime/CoreUObject/Public/UObject/ObjectMacros.h
--- a/Engine/Source/Runtime/CoreUObject/Public/UObject/ObjectMacros.h
+++ b/Engine/Source/Runtime/CoreUObject/Public/UObject/ObjectMacros.h
@@ -204,8 +204,8 @@ enum EClassFlags
 	CLASS_ProjectUserConfig	  = 0x00000040u,
 	/** Class is a native class - native interfaces will have CLASS_Native set, but not RF_MarkAsNative */
 	CLASS_Native			  = 0x00000080u,
-	/** Don't export to C++ header. */
-	CLASS_NoExport UE_DEPRECATED(5.1, "CLASS_NoExport should no longer be used. It is no longer being set by engine code.") = 0x00000100u,
+	/** Instances of this class can migrate to another server even if loaded from disk. */
+	CLASS_MigratingAsset	  = 0x00000100u,
 	/** Do not allow users to create in the editor. */
 	CLASS_NotPlaceable        = 0x00000200u,
 	/** Handle object configuration on a per-object basis, rather than per-class. */
@@ -266,7 +266,7 @@ enum EClassFlags
 ENUM_CLASS_FLAGS(EClassFlags);
 
 /** Flags to inherit from base class */
-#define CLASS_Inherit ((EClassFlags)(CLASS_Transient | CLASS_Optional | CLASS_DefaultConfig | CLASS_Config | CLASS_PerObjectConfig | CLASS_ConfigDoNotCheckDefaults | CLASS_NotPlaceable \
+#define CLASS_Inherit ((EClassFlags)(CLASS_Transient | CLASS_Optional | CLASS_DefaultConfig | CLASS_Config | CLASS_PerObjectConfig | CLASS_ConfigDoNotCheckDefaults | CLASS_MigratingAsset | CLASS_NotPlaceable \
 						| CLASS_Const | CLASS_HasInstancedReference | CLASS_Deprecated | CLASS_DefaultToInstanced | CLASS_GlobalUserConfig | CLASS_ProjectUserConfig | CLASS_PerPlatformConfig | CLASS_NeedsDeferredDependencyLoading))
 
 /** These flags will be cleared by the compiler when the class is parsed during script compilation */
@@ -289,6 +289,7 @@ ENUM_CLASS_FLAGS(EClassFlags);
 	CLASS_Transient | \
 	CLASS_Optional | \
 	CLASS_Native | \
+	CLASS_MigratingAsset | \
 	CLASS_NotPlaceable | \
 	CLASS_PerObjectConfig | \
 	CLASS_PerPlatformConfig | \
diff --git a/Engine/Source/Runtime/CoreUObject/Public/UObject/RemoteObjectPathName.h b/Engine/Source/Runtime/CoreUObject/Public/UObject/RemoteObjectPathName.h
new file mode 100644
--- /dev/null
+++ b/Engine/Source/Runtime/CoreUObject/Public/UObject/RemoteObjectPathName.h
@@ -0,0 +1,91 @@
+// Copyright Epic Games, Inc. All Rights Reserved.
+
+#pragma once
+
+#include "RemoteObjectTypes.h"
+
+/**
+* Structure that holds unique Names and RemoteIds associated with remote object data or pathnames (noexport type)
+*/
+struct FRemoteObjectTables
+{
+	TArray<FName> Names;
+	TArray<FRemoteObjectId> RemoteIds;
+
+	FORCEINLINE int32 Num() const
+	{
+		return Names.Num();
+	}
+
+	COREUOBJECT_API void Serialize(FArchive& Ar);
+};
+
+inline FArchive& operator<<(FArchive& Ar, FRemoteObjectTables& Tables)
+{
+	Tables.Serialize(Ar);
+	return Ar;
+}
+
+
+/**
+* Structure that represents remote object pathname (noexport type)
+* Stores pathnames as an array of indices (into a sidecar FRemoteObjectTables object) of FNames and associated FRemoteObjectIds
+* Names and Ids are stored from the innermost object first to the outermost object last
+*/
+struct FPackedRemoteObjectPathName
+{
+	using FNameIndexType = uint16;
+	TArray<FNameIndexType> RemoteIds;
+	TArray<FNameIndexType> Names;
+
+	FORCEINLINE int32 Num() const
+	{
+		return Names.Num();
+	}
+
+	FORCEINLINE FName GetSegmentName(int32 InSegmentIndex, const FRemoteObjectTables& InTables) const
+	{
+		return InTables.Names[Names[InSegmentIndex]];
+	}
+
+	FORCEINLINE FRemoteObjectId GetSegmentId(int32 InSegmentIndex, const FRemoteObjectTables& InTables) const
+	{
+		return InTables.RemoteIds[RemoteIds[InSegmentIndex]];
+	}
+
+	COREUOBJECT_API UObject* Resolve(const FRemoteObjectTables& InTables) const;
+	COREUOBJECT_API FString ToString(const FRemoteObjectTables& InTables, int32 InMinPathSegmentIndex = 0) const;
+	COREUOBJECT_API void Serialize(FArchive& Ar);
+};
+
+inline FArchive& operator<<(FArchive& Ar, FPackedRemoteObjectPathName& PathName)
+{
+	PathName.Serialize(Ar);
+	return Ar;
+}
+
+/**
+* Structure that represents remote object pathname (noexport type)
+* Stores pathnames as an array of FNames and associated FRemoteObjectIds
+* Names and Ids are stored from the innermost object first to the outermost object last
+*/
+struct FRemoteObjectPathName : public FRemoteObjectTables
+{
+	FRemoteObjectPathName() = default;
+	FRemoteObjectPathName(FRemoteObjectPathName&&) = default;
+	explicit FRemoteObjectPathName(UObject* InObject);
+	FRemoteObjectPathName& operator=(FRemoteObjectPathName&&) = default;
+
+	FORCEINLINE FName GetSegmentName(int32 InSegmentIndex, const FRemoteObjectTables&) const
+	{
+		return Names[InSegmentIndex];
+	}
+
+	FORCEINLINE FRemoteObjectId GetSegmentId(int32 InSegmentIndex, const FRemoteObjectTables&) const
+	{
+		return RemoteIds[InSegmentIndex];
+	}
+
+	COREUOBJECT_API UObject* Resolve() const;
+	COREUOBJECT_API FString ToString(int32 InMinPathSegmentIndex = 0) const;
+};
\ No newline at end of file
diff --git a/Engine/Source/Runtime/CoreUObject/Public/UObject/RemoteObjectTransfer.h b/Engine/Source/Runtime/CoreUObject/Public/UObject/RemoteObjectTransfer.h
--- a/Engine/Source/Runtime/CoreUObject/Public/UObject/RemoteObjectTransfer.h
+++ b/Engine/Source/Runtime/CoreUObject/Public/UObject/RemoteObjectTransfer.h
@@ -4,13 +4,40 @@
 
 #include "CoreTypes.h"
 #include "Misc/AssertionMacros.h"
-#include "UObject/RemoteObjectTypes.h"
-#include "UObject/Class.h"
+#include "UObject/RemoteObjectPathName.h"
 #include "Async/Async.h"
 #include "UObject/RemoteExecutor.h"
 
 struct FUObjectMigrationContext;
 
+/**
+* Structure that holds serialized remote object data chunk (< 64KB of data) (noexport type)
+*/
+struct FRemoteObjectBytes
+{
+	TArray<uint8> Bytes;
+};
+
+/**
+* Structure that holds remote object memory (noexport type)
+*/
+struct FRemoteObjectData
+{
+	FRemoteObjectTables Tables;
+	TArray<FPackedRemoteObjectPathName> PathNames;
+	TArray<FRemoteObjectBytes> Bytes;
+
+	inline int32 GetNumBytes() const
+	{
+		int32 Num = 0;
+		for (const FRemoteObjectBytes& Chunk : Bytes)
+		{
+			Num += Chunk.Bytes.Num();
+		}
+		return Num;
+	}
+};
+
 namespace UE::RemoteObject::Transfer
 {
 	/** Information for performing a migration (send) an object to a remote server */
diff --git a/Engine/Source/Runtime/CoreUObject/Public/UObject/RemoteObjectTypes.h b/Engine/Source/Runtime/CoreUObject/Public/UObject/RemoteObjectTypes.h
--- a/Engine/Source/Runtime/CoreUObject/Public/UObject/RemoteObjectTypes.h
+++ b/Engine/Source/Runtime/CoreUObject/Public/UObject/RemoteObjectTypes.h
@@ -6,19 +6,39 @@
 #include "Templates/TypeHash.h"
 #include "Containers/Array.h"
 #include "Containers/UnrealString.h"
+#include "UObject/ObjectMacros.h"
 
+class UObjectBase;
 class UObject;
 class FArchive;
 
+enum class ERemoteServerIdConstants : uint32
+{
+	Invalid = 0,
+	Max = (1 << 10) - 1,
+
+	Database = Max,
+	Asset = Max - 1,
+
+	// Add new reserved server IDs above this line in a descending order (so the next would be Something = Max - 2)
+
+	FirstReservedPlusOne,
+	FirstReserved = FirstReservedPlusOne - 1
+};
+
 struct FRemoteServerId
 {
 	friend struct FRemoteObjectId;
 
 	FRemoteServerId() = default;
+	explicit FRemoteServerId(ERemoteServerIdConstants InId)
+		: Id((uint32)InId)
+	{
+	}
 	explicit FRemoteServerId(uint32 InId)
 		: Id(InId)
 	{
-		checkf(Id < (1 << 10), TEXT("Remote server id can only be a 10-bit value, got: %u"), Id);
+		checkf(Id < (uint32)ERemoteServerIdConstants::FirstReserved, TEXT("Remote server id can not be greater than %u, got: %u"), (uint32)ERemoteServerIdConstants::FirstReserved - 1, Id);
 	}
 	uint32 GetIdNumber() const
 	{
@@ -56,7 +76,16 @@ struct FRemoteServerId
 	friend COREUOBJECT_API FArchive& operator<<(FArchive& Ar, FRemoteServerId& Id);
 
 private:
-	uint32 Id = 0;
+
+	/** 
+	 * Constructor used exclusively by FRemoteObjectId::GetServerId() to bypass range changes in the other constructors 
+	 */
+	FRemoteServerId(uint32 InId, ERemoteServerIdConstants /*Unused*/)
+		: Id(InId)
+	{
+	}
+
+	uint32 Id = (uint32)ERemoteServerIdConstants::Invalid;
 };
 
 struct FRemoteObjectId
@@ -125,15 +154,17 @@ public:
 	
 	FORCEINLINE FRemoteServerId GetServerId() const
 	{
-		return FRemoteServerId(ServerId);
+		return FRemoteServerId(ServerId, ERemoteServerIdConstants::Invalid);
 	}
 
-	FORCEINLINE FString ToString() const
+	FORCEINLINE bool IsAsset() const
 	{
-		return FString::Printf(TEXT("%llu-%llu"), ServerId, SerialNumber);
+		return GetServerId() == FRemoteServerId(ERemoteServerIdConstants::Asset);
 	}
 
-	COREUOBJECT_API static FRemoteObjectId Generate();
+	COREUOBJECT_API FString ToString() const;
+
+	COREUOBJECT_API static FRemoteObjectId Generate(UObjectBase* InObject, EInternalObjectFlags InInitialFlags = EInternalObjectFlags::None);
 
 	friend COREUOBJECT_API FArchive& operator<<(FArchive& Ar, FRemoteObjectId& Id);
 };
@@ -142,40 +173,3 @@ FORCEINLINE uint32 GetTypeHash(const FRemoteObjectId& ObjectId)
 {
 	return ObjectId.GetTypeHash();
 }
-
-/**
-* Structure that represents remote object pathname (noexport type)
-*/
-struct FRemoteObjectPathName
-{
-	using FNameIndexType = uint16;
-	TArray<FNameIndexType> Names;
-};
-
-/**
-* Structure that holds serialized remote object data chunk (< 64KB of data) (noexport type)
-*/
-struct FRemoteObjectBytes
-{
-	TArray<uint8> Bytes;
-};
-
-/**
-* Structure that holds remote object memory (noexport type)
-*/
-struct FRemoteObjectData
-{
-	TArray<FName> Names;
-	TArray<FRemoteObjectPathName> PathNames;
-	TArray<FRemoteObjectBytes> Bytes;
-
-	inline int32 GetNumBytes() const
-	{
-		int32 Num = 0;
-		for (const FRemoteObjectBytes& Chunk : Bytes)
-		{
-			Num += Chunk.Bytes.Num();
-		}
-		return Num;
-	}
-};
diff --git a/Engine/Source/Runtime/Engine/Classes/GameFramework/Actor.h b/Engine/Source/Runtime/Engine/Classes/GameFramework/Actor.h
--- a/Engine/Source/Runtime/Engine/Classes/GameFramework/Actor.h
+++ b/Engine/Source/Runtime/Engine/Classes/GameFramework/Actor.h
@@ -249,7 +249,7 @@ struct TIsContiguousContainer<TInlineComponentArray<T, NumElements>>
  * @see https://docs.unrealengine.com/Programming/UnrealArchitecture/Actors/ActorLifecycle
  * @see UActorComponent
  */
-UCLASS(BlueprintType, Blueprintable, config=Engine, meta=(ShortTooltip="An Actor is an object that can be placed or spawned in the world."), MinimalAPI)
+UCLASS(BlueprintType, Blueprintable, MigratingAsset, config=Engine, meta=(ShortTooltip="An Actor is an object that can be placed or spawned in the world."), MinimalAPI)
 class AActor : public UObject
 {
 	GENERATED_BODY()
