Avoid resolving remote objects inside of FClassProperty::Identical

diff --git a/Engine/Source/Runtime/CoreUObject/Private/UObject/PropertyClass.cpp b/Engine/Source/Runtime/CoreUObject/Private/UObject/PropertyClass.cpp
--- a/Engine/Source/Runtime/CoreUObject/Private/UObject/PropertyClass.cpp
+++ b/Engine/Source/Runtime/CoreUObject/Private/UObject/PropertyClass.cpp
@@ -201,8 +201,15 @@ bool FClassProperty::Identical( const void* A, const void* B, uint32 PortFlags )
 	TObjectPtr<UObject> ObjectA = A ? GetObjectPtrPropertyValue(A) : TObjectPtr<UObject>();
 	TObjectPtr<UObject> ObjectB = B ? GetObjectPtrPropertyValue(B) : TObjectPtr<UObject>();
 
-	check(ObjectA == nullptr || ObjectA.IsA<UClass>());
-	check(ObjectB == nullptr || ObjectB.IsA<UClass>());
+#if DO_CHECK && UE_WITH_REMOTE_OBJECT_HANDLE
+	// Object.IsA<UClass>() will try to resolve a remote object so skip these checks if we're running with remote object support and serializing remote objects
+	if (!(PortFlags & PPF_AvoidRemoteObjectMigration))
+#endif
+	{
+		check(ObjectA == nullptr || ObjectA.IsA<UClass>());
+		check(ObjectB == nullptr || ObjectB.IsA<UClass>());
+	}
+
 	return (ObjectA == ObjectB);
 }
 
