Fixing compilation error

diff --git a/Engine/Source/Runtime/CoreUObject/Private/UObject/RemoteObject.cpp b/Engine/Source/Runtime/CoreUObject/Private/UObject/RemoteObject.cpp
--- a/Engine/Source/Runtime/CoreUObject/Private/UObject/RemoteObject.cpp
+++ b/Engine/Source/Runtime/CoreUObject/Private/UObject/RemoteObject.cpp
@@ -129,7 +129,7 @@ void InitServerId()
 	{
 		if (!FParse::Value(FCommandLine::Get(), TEXT("LocalPeerId="), ServerId))
 		{
-			RemoteServerId == FRemoteServerId(ERemoteServerIdConstants::FirstValid);
+			RemoteServerId = FRemoteServerId(ERemoteServerIdConstants::FirstValid);
 		}
 	}
 	if (!ServerId.IsEmpty())
