[CoreUObject]

- Deprecate RF_KeepForCooker to free up some flag space now that this flag isn't used anymore.

#rnx

diff --git a/Engine/Source/Runtime/CoreUObject/Public/UObject/ObjectMacros.h b/Engine/Source/Runtime/CoreUObject/Public/UObject/ObjectMacros.h
--- a/Engine/Source/Runtime/CoreUObject/Public/UObject/ObjectMacros.h
+++ b/Engine/Source/Runtime/CoreUObject/Public/UObject/ObjectMacros.h
@@ -560,7 +560,7 @@ enum EObjectFlags
 	// The group of flags tracks the stages of the lifetime of a uobject
 	RF_NeedInitialization		=0x00000200,	///< This object has not completed its initialization process. Cleared when ~FObjectInitializer completes
 	RF_NeedLoad					=0x00000400,	///< During load, indicates object needs loading.
-	RF_KeepForCooker			=0x00000800,	///< Keep this object during garbage collection because it's still being used by the cooker
+	RF_KeepForCooker UE_DEPRECATED(5.6, "This flag is now obsolete as the functionality is unnecessary.")			=0x00000800,	///< Keep this object during garbage collection because it's still being used by the cooker
 	RF_NeedPostLoad				=0x00001000,	///< Object needs to be postloaded.
 	RF_NeedPostLoadSubobjects	=0x00002000,	///< During load, indicates that the object still needs to instance subobjects and fixup serialized component references
 	RF_NewerVersionExists		=0x00004000,	///< Object has been consigned to oblivion due to its owner package being reloaded, and a newer version currently exists
