RemoteObjects:  Better logging for some test cases

#rnx

#virtualized

diff --git a/Engine/Source/Runtime/CoreUObject/Private/UObject/RemoteObjectTransfer.cpp b/Engine/Source/Runtime/CoreUObject/Private/UObject/RemoteObjectTransfer.cpp
--- a/Engine/Source/Runtime/CoreUObject/Private/UObject/RemoteObjectTransfer.cpp
+++ b/Engine/Source/Runtime/CoreUObject/Private/UObject/RemoteObjectTransfer.cpp
@@ -361,7 +361,8 @@ UObject* MigrateObjectFromObjectData(FRemoteObjectData& ObjectData, const FUObje
 
 void MigrateObjectFromRemoteServer(FRemoteObjectId ObjectId, FRemoteServerId CurrentOwnerServerId, UObject* DestinationOuter)
 {
-	UE_LOG(LogRemoteObject, Verbose, TEXT("Aborting transaction, need obj id %s on server %s"),
+	UE_LOG(LogRemoteObject, Verbose, TEXT("Aborting transaction, server %s needs obj id %s from server %s"),
+		*UE::RemoteObject::GetGlobalServerId().ToString(),
 		*ObjectId.ToString(),
 		*CurrentOwnerServerId.ToString());
 
