PropertyBag: Skip refcount checks for property bag instances when exiting the application

Containing instances were all being cleaned up during the last GC, but GC ordering could mean that spurious dangling refcounts were reported in the final purge.

diff --git a/Engine/Source/Runtime/CoreUObject/Private/StructUtils/PropertyBag.cpp b/Engine/Source/Runtime/CoreUObject/Private/StructUtils/PropertyBag.cpp
--- a/Engine/Source/Runtime/CoreUObject/Private/StructUtils/PropertyBag.cpp
+++ b/Engine/Source/Runtime/CoreUObject/Private/StructUtils/PropertyBag.cpp
@@ -3844,7 +3844,7 @@ void UPropertyBag::DestroyStruct(void* Dest, int32 ArrayDim) const
 void UPropertyBag::FinishDestroy()
 {
 	const int32 Count = RefCount.load(); 
-	if (Count > 0)
+	if (Count > 0 && !IsEngineExitRequested())
 	{
 		UE_LOG(LogCore, Error, TEXT("PropertyBag: Expecting RefCount to be zero on destructor, but it is %d."), Count);
 	}
