PropertyBag: Hash each property for user defined structs

#rnx

diff --git a/Engine/Source/Runtime/CoreUObject/Private/StructUtils/PropertyBag.cpp b/Engine/Source/Runtime/CoreUObject/Private/StructUtils/PropertyBag.cpp
--- a/Engine/Source/Runtime/CoreUObject/Private/StructUtils/PropertyBag.cpp
+++ b/Engine/Source/Runtime/CoreUObject/Private/StructUtils/PropertyBag.cpp
@@ -54,6 +54,7 @@ FCustomVersionRegistration GPropertyBagCustomVersion(FPropertyBagCustomVersion::
 
 namespace UE::StructUtils::Private
 {
+	uint64 CalcPropertyDescArrayHash(const TConstArrayView<FPropertyBagPropertyDesc> Descs);
 
 	bool CanCastTo(const UStruct* From, const UStruct* To)
 	{
@@ -66,6 +67,36 @@ namespace UE::StructUtils::Private
 		return CityHash64(reinterpret_cast<const char*>(GetData(PathName)), PathName.Len() * sizeof(TCHAR));
 	}
 
+	uint64 CalcValueObjectHash(const UObject* InValueObject)
+	{
+		uint64 ValueTypeObjectHash = GetObjectHash(InValueObject);
+#if WITH_EDITOR
+		// for user defined structures we need to hash each property. after changing a user defined structure
+		// we may find the same property bag again if the hash is not sufficient, even though the memory layout
+		// has changed. hashing the path name of the user defined structure is not enough.
+		if(const UUserDefinedStruct* UserDefinedStruct = Cast<UUserDefinedStruct>(InValueObject))
+		{
+			for(const FField* Property = UserDefinedStruct->ChildProperties; Property != nullptr; Property = Property->Next)
+			{
+				ValueTypeObjectHash = CityHash128to64({ValueTypeObjectHash, GetTypeHash(Property)});
+				if(const FStructProperty* StructProperty = CastField<FStructProperty>(Property))
+				{
+					if(StructProperty->Struct)
+					{
+						ValueTypeObjectHash = CityHash128to64({ValueTypeObjectHash, CalcValueObjectHash(StructProperty->Struct)});
+					}
+				}
+			}
+		}
+		// for property bags we recurse
+		else if(const UPropertyBag* PropertyBag = Cast<UPropertyBag>(InValueObject))
+		{
+			ValueTypeObjectHash = CityHash128to64({ValueTypeObjectHash, CalcPropertyDescArrayHash(PropertyBag->GetPropertyDescs())});
+		}
+#endif
+		return ValueTypeObjectHash;
+	}
+
 	uint64 CalcPropertyDescHash(const FPropertyBagPropertyDesc& Desc)
 	{
 #if WITH_EDITORONLY_DATA
@@ -73,7 +104,7 @@ namespace UE::StructUtils::Private
 #else
 		const uint32 Hashes[] = { GetTypeHash(Desc.ID), GetTypeHash(Desc.Name), GetTypeHash(Desc.ValueType), GetTypeHash(Desc.ContainerTypes), GetTypeHash(Desc.PropertyFlags) };
 #endif
-		return CityHash64WithSeed(reinterpret_cast<const char*>(Hashes), sizeof(Hashes), GetObjectHash(Desc.ValueTypeObject));
+		return CityHash64WithSeed(reinterpret_cast<const char*>(Hashes), sizeof(Hashes), CalcValueObjectHash(Desc.ValueTypeObject));
 	}
 
 	uint64 CalcPropertyDescArrayHash(const TConstArrayView<FPropertyBagPropertyDesc> Descs)
