[ZenLoader]
- Increase hang detection timeout when running instrumented

diff --git a/Engine/Source/Runtime/CoreUObject/Private/Serialization/AsyncLoading2.cpp b/Engine/Source/Runtime/CoreUObject/Private/Serialization/AsyncLoading2.cpp
--- a/Engine/Source/Runtime/CoreUObject/Private/Serialization/AsyncLoading2.cpp
+++ b/Engine/Source/Runtime/CoreUObject/Private/Serialization/AsyncLoading2.cpp
@@ -497,7 +497,12 @@ static FAutoConsoleVariableRef CVarGReloadPackagesWithGCedExports(
 );
 #endif
 
+#if USING_INSTRUMENTATION
+static float GStallDetectorTimeout = 1200.0f; // 10x factor when instrumentation is enabled
+#else
 static float GStallDetectorTimeout = 120.0f;
+#endif
+
 static FAutoConsoleVariableRef CVarGStallDetectorTimeout(
 	TEXT("s.StallDetectorTimeout"),
 	GStallDetectorTimeout,
