Added console command for outputing instance data object statistics (object counts and memory).

#rnx

diff --git a/Engine/Source/Runtime/CoreUObject/Internal/UObject/PropertyBagRepository.h b/Engine/Source/Runtime/CoreUObject/Internal/UObject/PropertyBagRepository.h
--- a/Engine/Source/Runtime/CoreUObject/Internal/UObject/PropertyBagRepository.h
+++ b/Engine/Source/Runtime/CoreUObject/Internal/UObject/PropertyBagRepository.h
@@ -16,6 +16,16 @@ namespace UE
 
 #if WITH_EDITORONLY_DATA
 
+#if STATS
+struct FPropertyBagRepositoryStats
+{
+	size_t IDOMemoryBytes;
+	int32 NumIDOs;
+	int32 NumIDOsWithLooseProperties;
+	int32 NumPlaceholderTypes;
+};
+#endif
+
 // Singleton class tracking property bag association with objects
 class FPropertyBagRepository : public FGCObject
 {
@@ -108,6 +118,10 @@ public:
 
 	COREUOBJECT_API const UObject* FindInstanceForDataObject(const UObject* InstanceDataObject) const;
 
+	#if STATS
+	UE_INTERNAL void GatherStats(FPropertyBagRepositoryStats& Stats);
+	#endif
+
 	// FGCObject interface
 	virtual void AddReferencedObjects(FReferenceCollector& Collector) override;
 	virtual FString GetReferencerName() const override;
diff --git a/Engine/Source/Runtime/CoreUObject/Private/UObject/InstanceDataObjectUtils.cpp b/Engine/Source/Runtime/CoreUObject/Private/UObject/InstanceDataObjectUtils.cpp
--- a/Engine/Source/Runtime/CoreUObject/Private/UObject/InstanceDataObjectUtils.cpp
+++ b/Engine/Source/Runtime/CoreUObject/Private/UObject/InstanceDataObjectUtils.cpp
@@ -13,6 +13,7 @@
 #include "Serialization/ObjectReader.h"
 #include "Serialization/ObjectWriter.h"
 #include "Serialization/StructuredArchive.h"
+#include "Serialization/ArchiveCountMem.h"
 #include "UObject/Class.h"
 #include "UObject/EnumProperty.h"
 #include "UObject/Field.h"
@@ -259,6 +260,9 @@ namespace UE
 	static const FName NAME_IDOMapKey(ANSITEXTVIEW("Key"));
 	static const FName NAME_IDOMapValue(ANSITEXTVIEW("Value"));
 
+	static TMap<FBlake3Hash, TWeakObjectPtr<UInstanceDataObjectClass>> IDOClassCache;
+	static FSharedMutex IDOClassCacheMutex;
+
 	static void OnInstanceDataObjectSupportChanged(IConsoleVariable*);
 
 	bool bEverCreatedIDO = false;
@@ -405,6 +409,108 @@ namespace UE
 		})
 	);
 
+	static void OutputIDOStats(FOutputDevice& OutputDevice)
+	{
+		#if STATS
+		struct MemoryMetric
+		{
+			const TCHAR* Unit;
+			double Value;
+		};
+
+		const auto ConvertToMemoryMetric = [](size_t MemoryBytes) -> MemoryMetric
+		{
+			MemoryMetric Metric = {};
+
+			size_t GB = 1024ULL * 1024ULL * 1024ULL;
+			size_t MB = 1024ULL * 1024ULL;
+			size_t KB = 1024ULL;
+			Metric.Value = (double)MemoryBytes;
+
+			if (MemoryBytes >= GB)
+			{
+				Metric.Value /= (double)GB;
+				Metric.Unit = TEXT("GB");
+			}
+			else if (MemoryBytes >= MB)
+			{
+				Metric.Value /= (double)MB;
+				Metric.Unit = TEXT("MB");
+			}
+			else if (MemoryBytes >= KB)
+			{
+				Metric.Value /= (double)KB;
+				Metric.Unit = TEXT("KB");
+			}
+			else
+			{
+				Metric.Unit = TEXT("bytes");
+			}
+
+			return Metric;
+		};
+
+		FPropertyBagRepositoryStats Stats;
+		FPropertyBagRepository::Get().GatherStats(Stats);
+
+		int32 NumIDOClasses = 0;
+		size_t ClassMemoryBytes = 0;
+		size_t CDOMemoryBytes = 0;
+		{
+			TUniqueLock Lock(IDOClassCacheMutex);
+
+			for (const TPair<FBlake3Hash, TWeakObjectPtr<UInstanceDataObjectClass>>& Pair : IDOClassCache)
+			{
+				UInstanceDataObjectClass* Class = Pair.Value.Get();
+
+				if (Class)
+				{
+					++NumIDOClasses;
+
+					{
+						FArchiveCountMem MemoryCount(Class);
+						ClassMemoryBytes += MemoryCount.GetMax();
+					}
+
+					if (UObject* CDO = Class->GetDefaultObject(/*bCreateIfNeeded*/false))
+					{
+						FArchiveCountMem MemoryCount(CDO);
+						CDOMemoryBytes += MemoryCount.GetMax();
+					}
+				}
+				
+			}
+		}
+
+		size_t TotalMemoryBytes = Stats.IDOMemoryBytes + ClassMemoryBytes + CDOMemoryBytes;
+
+		MemoryMetric TotalMemory = ConvertToMemoryMetric(TotalMemoryBytes);
+		MemoryMetric ObjectMemory = ConvertToMemoryMetric(Stats.IDOMemoryBytes);
+		MemoryMetric ClassMemory = ConvertToMemoryMetric(ClassMemoryBytes);
+		MemoryMetric CDOMemory = ConvertToMemoryMetric(CDOMemoryBytes);
+
+		OutputDevice.Logf(TEXT("Number of IDOs = %d"), Stats.NumIDOs);
+		OutputDevice.Logf(TEXT("Number of IDOs with loose properties = %d"), Stats.NumIDOsWithLooseProperties);
+		OutputDevice.Logf(TEXT("Number of IDO classes = %u"), NumIDOClasses);
+		OutputDevice.Logf(TEXT("Number of placeholder types = %u"), Stats.NumPlaceholderTypes);
+		OutputDevice.Logf(TEXT("Total IDO memory = %.2f %s"), TotalMemory.Value, TotalMemory.Unit);
+		OutputDevice.Logf(TEXT("    IDO object memory = %.2f %s"), ObjectMemory.Value, ObjectMemory.Unit);
+		OutputDevice.Logf(TEXT("    IDO class memory = %.2f %s"), ClassMemory.Value, ClassMemory.Unit);
+		OutputDevice.Logf(TEXT("    IDO CDO memory = %.2f %s"), CDOMemory.Value, CDOMemory.Unit);
+		#else
+		OutputDevice.Log(TEXT("Stats not enabled on current build"));
+		#endif // STATS
+	}
+
+	FAutoConsoleCommand DumpIDOStatsCommand(
+		TEXT("IDO.DumpStats"),
+		TEXT("Prints statistics for all current Instance Data Objects."),
+		FConsoleCommandWithArgsAndOutputDeviceDelegate::CreateLambda([](const TArray<FString>& Args, FOutputDevice& OutputDevice)
+		{
+			OutputIDOStats(OutputDevice);
+		})
+	);
+
 	static TSet<FString> GetExcludedUnknownPropertyTypes()
 	{
 		TArray<FString> Result;
@@ -1158,9 +1264,7 @@ namespace UE
 			Key = KeyBuilder.Finalize();
 		}
 
-		static TMap<FBlake3Hash, TWeakObjectPtr<UInstanceDataObjectClass>> ClassCache;
-		static FSharedMutex ClassCacheMutex;
-		if (TSharedLock Lock(ClassCacheMutex); UClass* Class = ClassCache.FindRef(Key).Get())
+		if (TSharedLock Lock(IDOClassCacheMutex); UClass* Class = IDOClassCache.FindRef(Key).Get())
 		{
 			return Class;
 		}
@@ -1175,15 +1279,15 @@ namespace UE
 
 		CopyTaggedProperties(OwnerClass->GetDefaultObject(), NewClass->GetDefaultObject());
 
-		TUniqueLock Lock(ClassCacheMutex);
-		if (UClass* Class = ClassCache.FindRef(Key).Get())
+		TUniqueLock Lock(IDOClassCacheMutex);
+		if (UClass* Class = IDOClassCache.FindRef(Key).Get())
 		{
 			return Class;
 		}
 
-		CleanUpInstanceDataObjectTypeCache(ClassCache);
+		CleanUpInstanceDataObjectTypeCache(IDOClassCache);
 
-		ClassCache.Add(Key, NewClass);
+		IDOClassCache.Add(Key, NewClass);
 		return NewClass;
 	}
 
diff --git a/Engine/Source/Runtime/CoreUObject/Private/UObject/PropertyBagRepository.cpp b/Engine/Source/Runtime/CoreUObject/Private/UObject/PropertyBagRepository.cpp
--- a/Engine/Source/Runtime/CoreUObject/Private/UObject/PropertyBagRepository.cpp
+++ b/Engine/Source/Runtime/CoreUObject/Private/UObject/PropertyBagRepository.cpp
@@ -10,6 +10,7 @@
 #include "Serialization/ObjectReader.h"
 #include "Serialization/ObjectWriter.h"
 #include "Serialization/ArchiveProxy.h"
+#include "Serialization/ArchiveCountMem.h"
 #include "UObject/GarbageCollection.h"
 #include "UObject/InstanceDataObjectUtils.h"
 #include "UObject/LinkerLoad.h"
@@ -68,6 +69,11 @@ public:
 		return PlaceholderTypes.Contains(Type);
 	}
 
+	int32 Num() const
+	{
+		return PlaceholderTypes.Num();
+	}
+
 protected:
 	FPropertyBagPlaceholderTypeRegistry() = default;
 	FPropertyBagPlaceholderTypeRegistry(const FPropertyBagPlaceholderTypeRegistry&) = delete;
@@ -1003,6 +1009,35 @@ void FPropertyBagRepository::RemovePropertyBagPlaceholderType(UStruct* Placehold
 	FPropertyBagPlaceholderTypeRegistry::Get().Remove(PlaceholderType);
 }
 
+#if STATS
+void FPropertyBagRepository::GatherStats(FPropertyBagRepositoryStats& Stats)
+{
+	FMemory::Memset(Stats, 0);
+
+	Stats.NumPlaceholderTypes = FPropertyBagPlaceholderTypeRegistry::Get().Num();
+
+	FPropertyBagRepositoryLock LockRepo(this);
+
+	for (const TPair<const UObject*, FPropertyBagAssociationData>& Pair : AssociatedData)
+	{ 
+		const FPropertyBagAssociationData& BagData = Pair.Value;
+
+		if (BagData.InstanceDataObject != nullptr)
+		{
+			++Stats.NumIDOs;
+
+			FArchiveCountMem MemoryCount(BagData.InstanceDataObject);
+			Stats.IDOMemoryBytes += MemoryCount.GetMax();
+
+			if (BagData.bNeedsFixup)
+			{
+				++Stats.NumIDOsWithLooseProperties;
+			}
+		}
+	}
+}
+#endif // STATS
+
 } // UE
 
 #endif // WITH_EDITORONLY_DATA
