Universally Unique Object Names now mark a bit in their FName number for quick identification checks.
Introduced Universally Unique Object Name identification functions to EntityFramework.

diff --git a/Engine/Source/Runtime/CoreUObject/Private/UObject/UObjectGlobals.cpp b/Engine/Source/Runtime/CoreUObject/Private/UObject/UObjectGlobals.cpp
--- a/Engine/Source/Runtime/CoreUObject/Private/UObject/UObjectGlobals.cpp
+++ b/Engine/Source/Runtime/CoreUObject/Private/UObject/UObjectGlobals.cpp
@@ -2954,47 +2954,45 @@ FName MakeUniqueObjectName(UObject* Parent, const UClass* Class, FName InBaseNam
 	return TestName;
 }
 
-FName MakeUniversallyUniqueObjectName(UObject* InOuter, FName InBaseName)
+FUniversallyUniqueObjectNameGenerator::FUniversallyUniqueObjectNameGenerator(const FGuid& Guid)
 {
-	// Generate a unique generation context for this session.
-	static struct FContext
-	{
-		FContext()
-		{
-			const FGuid Guid = FGuid::NewGuid();
-			
-			uint64 Value = ((uint64)Guid.A << 32) | Guid.C;
-			uint32 Index = 0;
+	uint64 Value = ((uint64)Guid.A << 32) | Guid.C;
+	uint32 Index = 0;
 
-			while (Value)
-			{
-				const uint8 Remainder = uint8(Value % 36);
-				SessionIdBase36[Index++] = NibbleToTCharLower(Remainder);
-				check(Index < (sizeof(SessionIdBase36) / sizeof(SessionIdBase36[0])));
-				Value /= 36;
-			}
+	while (Value)
+	{
+		const uint8 Remainder = uint8(Value % 36);
+		SessionIdBase36[Index++] = NibbleToTCharLower(Remainder);
+		check(Index < (sizeof(SessionIdBase36) / sizeof(SessionIdBase36[0])));
+		Value /= 36;
+	}
 
-			SessionIdBase36[Index] = 0;
+	SessionIdBase36[Index] = 0;
 
-			Seed = Guid.D;
-		}
+	Seed = Guid.D;
+}
 
-		uint32 GetNewValue() const
+uint32 FUniversallyUniqueObjectNameGenerator::GetNewValue() const
+{
+	for (uint32 CurValue = Seed.load(std::memory_order_relaxed);;)
+	{
+		const uint32 NewValue = CurValue * 1664525 + 1013904223;
+		if (Seed.compare_exchange_weak(CurValue, NewValue))
 		{
-			for (uint32 CurValue = Seed.load(std::memory_order_relaxed);;)
-			{
-				const uint32 NewValue = CurValue * 1664525 + 1013904223;
-				if (Seed.compare_exchange_weak(CurValue, NewValue))
-				{
-					return NewValue;
-				}
-			}
+			return NewValue;
 		}
+	}
+}
+
+FName MakeUniversallyUniqueObjectName(UObject* InOuter, FName InBaseName, FUniversallyUniqueObjectNameGenerator const* Generator)
+{
+	// Generate a unique generation context for this session.
+	static FUniversallyUniqueObjectNameGenerator const Context(FGuid::NewGuid());
+	if (!Generator)
+	{
+		Generator = &Context;
+	}
 
-		TCHAR SessionIdBase36[14];
-		mutable std::atomic<uint32> Seed;
-	} const Context;
-	
 	TStringBuilder<NAME_SIZE> StringBuilder;
 
 	// Append the provided base name if present.
@@ -3004,22 +3002,29 @@ FName MakeUniversallyUniqueObjectName(UObject* InOuter, FName InBaseName)
 		StringBuilder.AppendChar(TEXT('_'));
 	}
 
-	StringBuilder += Context.SessionIdBase36;
+	StringBuilder += Generator->SessionIdBase36;
 
 	FName OutObjectName;
 	do
 	{
 		// Store the random part in the lowest 31 bits of the name number below. This is to minimize the total names 
 		// generated for globally unique names. The name number bit 31 is reserved to preserve the fast path names
-		// generation (see GFastPathUniqueNameGeneration).
-		int32 OutObjectIndex = Context.GetNewValue() & 0x7fffffff;
-		OutObjectName = FName(*StringBuilder, OutObjectIndex);		
+		// generation (see GFastPathUniqueNameGeneration). We also set bit 30 as an indicator of a universally unique
+		// object name.
+		const int32 OutObjectIndex = (Generator->GetNewValue() & 0x3fffffff) | 0x40000000;
+		OutObjectName = FName(*StringBuilder, OutObjectIndex);
 	}
 	while (StaticFindObjectFast(nullptr, InOuter, OutObjectName));
 
 	return OutObjectName;
 }
 
+bool HasUniversallyUniqueObjectNameIdentifier(FName Name)
+{
+	const int32 Number = Name.GetNumber();
+	return (Number & 0x40000000) != 0;
+}
+
 FName MakeObjectNameFromDisplayLabel(const FString& DisplayLabel, const FName CurrentObjectName)
 {
 	FString GeneratedName = SlugStringForValidName(DisplayLabel);
diff --git a/Engine/Source/Runtime/CoreUObject/Public/UObject/UObjectGlobals.h b/Engine/Source/Runtime/CoreUObject/Public/UObject/UObjectGlobals.h
--- a/Engine/Source/Runtime/CoreUObject/Public/UObject/UObjectGlobals.h
+++ b/Engine/Source/Runtime/CoreUObject/Public/UObject/UObjectGlobals.h
@@ -977,16 +977,38 @@ ENUM_CLASS_FLAGS(EUniqueObjectNameOptions);
  */
 COREUOBJECT_API FName MakeUniqueObjectName( UObject* Outer, const UClass* Class, FName BaseName = NAME_None, EUniqueObjectNameOptions Options = EUniqueObjectNameOptions::None);
 
+/**
+ * Generator class that can be used as a session to generate universally unique object names.
+ */
+struct FUniversallyUniqueObjectNameGenerator
+{
+	COREUOBJECT_API FUniversallyUniqueObjectNameGenerator(const FGuid& Seed);
+	COREUOBJECT_API uint32 GetNewValue() const;
+
+	TCHAR SessionIdBase36[14];
+	mutable std::atomic<uint32> Seed;
+};
+
 /**
  * Create a universally unique name by combining a base name, a GUID and a timestamp.
  * The object name returned is guaranteed not to locally and globally exist.
  *
  * @param	Outer		the outer for the object that needs to be named
  * @param	BaseName	optional base name to use when generating the globally unique object name.
+ * @param	Generator	optional generator session that can be passed to not use instead of the global one.
  * 
  * @return	name is the form BaseName_##, where ## a universally unique identifier making sure that name is universally unique.
+ */
+COREUOBJECT_API FName MakeUniversallyUniqueObjectName(UObject* Outer, FName InBaseName = NAME_None, FUniversallyUniqueObjectNameGenerator const* Generator = nullptr);
+
+/**
+ * Checks if the given FName is marked as a Universally Unique Object Name.
+ * Note: This is an early out check. This function may return true for FNames that are not Universally Unique Object Names.
+ *
+ * @param	Name	The name to check.
+ * @return	True if the FName is marked as a Universally Unique Object Name, otherwise false.
   */
-COREUOBJECT_API FName MakeUniversallyUniqueObjectName(UObject* Outer, FName InBaseName = NAME_None);
+COREUOBJECT_API bool HasUniversallyUniqueObjectNameIdentifier(FName Name);
 
 /**
  * Given a display label string, generates an FName slug that is a valid FName for that label.
