[ZenLoader]
- Partial loading will now add dynamic imports instead of relying on postload groups to avoid potential deadlocks and improve efficiency.
- Add a unittest that would deadlock when sync loading an object from both a partial load and normal load.

diff --git a/Engine/Source/Runtime/CoreUObject/Private/Serialization/AsyncLoading2.cpp b/Engine/Source/Runtime/CoreUObject/Private/Serialization/AsyncLoading2.cpp
--- a/Engine/Source/Runtime/CoreUObject/Private/Serialization/AsyncLoading2.cpp
+++ b/Engine/Source/Runtime/CoreUObject/Private/Serialization/AsyncLoading2.cpp
@@ -3624,6 +3624,11 @@ private:
 	FAsyncPackageData Data;
 	FAsyncPackageHeaderData HeaderData;
 	FAsyncPackageSerializationState SerializationState;
+
+#if WITH_PARTIAL_REQUEST_DURING_RECURSION
+	TSet<FAsyncPackage2*> AdditionalImportedAsyncPackages;
+#endif
+
 #if WITH_EDITOR
 	TOptional<FAsyncPackageHeaderData> OptionalSegmentHeaderData;
 	TOptional<FAsyncPackageSerializationState> OptionalSegmentSerializationState;
@@ -5079,6 +5084,14 @@ void FAsyncLoadingThread2::UpdatePackagePriorityRecursive(FAsyncLoadingThreadSta
 			UpdatePackagePriorityRecursive(ThreadState, ImportedPackage, NewPriority);
 		}
 	}
+
+#if WITH_PARTIAL_REQUEST_DURING_RECURSION
+	for (FAsyncPackage2* ImportedPackage : Package->AdditionalImportedAsyncPackages)
+	{
+		UpdatePackagePriorityRecursive(ThreadState, ImportedPackage, NewPriority);
+	}
+#endif
+
 	UpdatePackagePriority(ThreadState, Package);
 }
 
@@ -5290,11 +5303,6 @@ void FAsyncLoadingThread2::IncludePackageInSyncLoadContextRecursive(FAsyncLoadin
 
 	Package->SyncLoadContextId = ContextId;
 
-	// When using the partial loading feature, don't try to upgrade postload groups into higher priority sync context as
-	// this is exactly what we're trying to prevent. Package that are merged together into a postload group are supposed
-	// to stay at the lower syncloadcontextid, that way, it allows to exit the flush when serialization is done and let
-	// postload run with the ones of the caller.
-#if !WITH_PARTIAL_REQUEST_DURING_RECURSION
 	FAsyncLoadingPostLoadGroup* PostLoadGroup = Package->PostLoadGroup ? Package->PostLoadGroup : Package->DeferredPostLoadGroup;
 	if (PostLoadGroup && PostLoadGroup->SyncLoadContextId < ContextId)
 	{
@@ -5307,7 +5315,7 @@ void FAsyncLoadingThread2::IncludePackageInSyncLoadContextRecursive(FAsyncLoadin
 			}
 		}
 	}
-#endif
+
 	for (FAsyncPackage2* ImportedPackage : Package->Data.ImportedAsyncPackages)
 	{
 		if (ImportedPackage && ImportedPackage->SyncLoadContextId < ContextId)
@@ -5315,6 +5323,17 @@ void FAsyncLoadingThread2::IncludePackageInSyncLoadContextRecursive(FAsyncLoadin
 			IncludePackageInSyncLoadContextRecursive(ThreadState, ContextId, ImportedPackage);
 		}
 	}
+
+#if WITH_PARTIAL_REQUEST_DURING_RECURSION
+	for (FAsyncPackage2* ImportedPackage : Package->AdditionalImportedAsyncPackages)
+	{
+		if (ImportedPackage->SyncLoadContextId < ContextId)
+		{
+			IncludePackageInSyncLoadContextRecursive(ThreadState, ContextId, ImportedPackage);
+		}
+	}
+#endif
+
 	if (Package->Desc.Priority < MAX_int32)
 	{
 		Package->Desc.Priority = MAX_int32;
@@ -6319,6 +6338,12 @@ void FAsyncPackage2::ImportPackagesRecursive(FAsyncLoadingThreadState2& ThreadSt
 	}
 	bHasStartedImportingPackages = true;
 
+#if WITH_PARTIAL_REQUEST_DURING_RECURSION
+	// No additional imports should be present when ImportPackagesRecursive is called since this should only be called during
+	// the early phases of loading.
+	check(AdditionalImportedAsyncPackages.IsEmpty());
+#endif
+
 	if (Data.ImportedAsyncPackages.IsEmpty())
 	{
 		return;
@@ -8292,25 +8317,36 @@ void FAsyncPackage2::ConditionalReleasePartialRequests(FAsyncLoadingThreadState2
 				AsyncLoadingThread.RemovePendingRequests(ThreadState, { RequestId });
 			}
 
-			// We want to skip postloads inside recursive loads that happens during serialization and preload
-			// since it is mostly what BP expects and how the old LinkerLoad was implemented.
-			// In some cases, trying to postload right away during recursion could end-up in data loss (i.e. UE-190649).
-			// The expectation is that postload should be run as part of the package causing this load to occur.
-			// As a more general rule, we want any synchronous load during deserialization to act like
-			// an import. We want the deserialization steps done before returning from the load but
-			// we keep postload for after deserialization is done on the referencer too so that postload 
-			// sees everything initialized in the case where it also has a reference back to the referencer.
-			// Using a partial request flag allows the flush to return as soon as exports are done which
-			// gives the opportunity to the caller to get a pointer to an object inside the package even if it's not fully loaded yet.
-			if (AsyncPackageLoadingState == EAsyncPackageLoadingState2::ExportsDone && CurrentSyncLoadContext->RequestingPackage->PostLoadGroup)
+			switch (AsyncPackageLoadingState)
 			{
-				UE_LOG(LogStreaming, Log, TEXT("Merging postload groups of package %s with requester package %s"),
-					*Desc.UPackageName.ToString(),
-					*CurrentSyncLoadContext->RequestingPackage->Desc.UPackageName.ToString());
+				case EAsyncPackageLoadingState2::ExportsDone:
+				{
+					// As a general rule, we want any synchronous load during deserialization to act like an import in order
+					// to make sure the requested package is part of the loading of the requesting package even if the
+					// sync load was released to avoid a deadlock. This will make sure that by the time the requesting package
+					// finishes loading, the requested package will be done too.
+					bool bIsAlreadyInSet = false;
+					CurrentSyncLoadContext->RequestingPackage->AdditionalImportedAsyncPackages.FindOrAdd(this, &bIsAlreadyInSet);
+					if (!bIsAlreadyInSet)
+					{
+						this->AddRef();
 
-				// Do not adjust sync load context, we want to be able to exit the current one even if the caller has not finished yet.
-				const bool bUpdateSyncLoadContext = false;
-				AsyncLoadingThread.MergePostLoadGroups(ThreadState, PostLoadGroup, CurrentSyncLoadContext->RequestingPackage->PostLoadGroup, bUpdateSyncLoadContext);
+						UE_LOG(LogStreaming, Display, TEXT("Package %s is adding a dynamic import to package %s because of a recursive sync load"),
+							*CurrentSyncLoadContext->RequestingPackage->Desc.UPackageName.ToString(),
+							*Desc.UPackageName.ToString()
+						);
+					}
+				}
+				break;
+				case EAsyncPackageLoadingState2::PostLoad:
+					[[fallthrough]];
+				case EAsyncPackageLoadingState2::DeferredPostLoad:
+					// We don't need to add a dynamic import if the package is already in postload since it wouldn't provide any benefits...
+				break;
+				default:
+					// If anything else begins to call ConditionalReleasePartialRequests, we'll need to consider what needs to be done here
+					checkNoEntry();
+				break;
 			}
 		}
 	}
@@ -8463,52 +8499,78 @@ FAsyncPackage2* FAsyncPackage2::UpdateDependenciesStateRecursive(FAsyncLoadingTh
 	Context.S.Push(this);
 	Context.P.Push(this);
 
-	for (FAsyncPackage2* ImportedPackage : Data.ImportedAsyncPackages)
-	{
-		if (!ImportedPackage)
+	auto ProcessImportedPackage = 
+		[&ThreadState, &Context, &WaitingForPackage](FAsyncPackage2* ImportedPackage) -> bool
 		{
-			continue;
-		}
+			if (!ImportedPackage)
+			{
+				return true;
+			}
 
-		if (ThreadState.PackagesOnStack.Contains(ImportedPackage))
-		{
-			continue;
-		}
+			if (ThreadState.PackagesOnStack.Contains(ImportedPackage))
+			{
+				return true;
+			}
 
-		FAllDependenciesState& ImportedPackageState = ImportedPackage->*Context.StateMemberPtr;
-		if (ImportedPackageState.bAllDone)
-		{
-			continue;
-		}
+			FAllDependenciesState& ImportedPackageState = ImportedPackage->*Context.StateMemberPtr;
+			if (ImportedPackageState.bAllDone)
+			{
+				return true;
+			}
 
-		if (ImportedPackage->AsyncPackageLoadingState < Context.WaitForPackageState)
+			if (ImportedPackage->AsyncPackageLoadingState < Context.WaitForPackageState)
+			{
+				WaitingForPackage = ImportedPackage;
+				return false;
+			}
+
+			ImportedPackageState.UpdateTick(Context.CurrentTick);
+			if (ImportedPackageState.PreOrderNumber < 0)
+			{
+				WaitingForPackage = ImportedPackage->UpdateDependenciesStateRecursive(ThreadState, Context);
+				if (WaitingForPackage)
+				{
+					return false;
+				}
+			}
+			else if (!ImportedPackageState.bAssignedToStronglyConnectedComponent)
+			{
+				while ((Context.P.Top()->*Context.StateMemberPtr).PreOrderNumber > ImportedPackageState.PreOrderNumber)
+				{
+					Context.P.Pop();
+				}
+			}
+			if (ImportedPackageState.WaitingForPackage)
+			{
+				WaitingForPackage = ImportedPackageState.WaitingForPackage;
+				return false;
+			}
+
+			return true;
+		};
+		
+	bool bContinueToProcessPackage = true;
+	for (FAsyncPackage2* ImportedPackage : Data.ImportedAsyncPackages)
+	{
+		if (!ProcessImportedPackage(ImportedPackage))
 		{
-			WaitingForPackage = ImportedPackage;
+			bContinueToProcessPackage = false;
 			break;
 		}
+	}
 
-		ImportedPackageState.UpdateTick(Context.CurrentTick);
-		if (ImportedPackageState.PreOrderNumber < 0)
+#if WITH_PARTIAL_REQUEST_DURING_RECURSION
+	if (bContinueToProcessPackage)
+	{
+		for (FAsyncPackage2* ImportedPackage : AdditionalImportedAsyncPackages)
 		{
-			WaitingForPackage = ImportedPackage->UpdateDependenciesStateRecursive(ThreadState, Context);
-			if (WaitingForPackage)
+			if (!ProcessImportedPackage(ImportedPackage))
 			{
 				break;
 			}
 		}
-		else if (!ImportedPackageState.bAssignedToStronglyConnectedComponent)
-		{
-			while ((Context.P.Top()->*Context.StateMemberPtr).PreOrderNumber > ImportedPackageState.PreOrderNumber)
-			{
-				Context.P.Pop();
-			}
-		}
-		if (ImportedPackageState.WaitingForPackage)
-		{
-			WaitingForPackage = ImportedPackageState.WaitingForPackage;
-			break;
-		}
 	}
+#endif
 
 	if (Context.P.Top() == this)
 	{
@@ -10268,6 +10330,14 @@ void FAsyncPackage2::ClearImportedPackages()
 			ImportedAsyncPackage->ReleaseRef();
 		}
 	}
+
+#if WITH_PARTIAL_REQUEST_DURING_RECURSION
+	TSet<FAsyncPackage2*> LocalAdditionalImportedAsyncPackages = MoveTemp(AdditionalImportedAsyncPackages);
+	for (FAsyncPackage2* ImportedAsyncPackage : LocalAdditionalImportedAsyncPackages)
+	{
+		ImportedAsyncPackage->ReleaseRef();
+	}
+#endif
 }
 
 void FAsyncPackage2::MoveConstructedObjectsToPhase2()
@@ -11321,10 +11391,20 @@ void FAsyncPackage2::DumpStateImpl(TSet<FAsyncPackage2*>& Set, int32 Indent, TMu
 	{
 		if (Import && Import->AsyncPackageLoadingState < EAsyncPackageLoadingState2::DeferredDelete)
 		{
-			UE_LOG(LogStreaming, Warning, TEXT(" %s Imports %s package %s state %s"), FCString::Spc(Indent), GetPackageType(Import), *Import->Desc.UPackageName.ToString(), LexToString(Import->AsyncPackageLoadingState));
+			UE_LOG(LogStreaming, Warning, TEXT(" %s imports %s package %s state %s"), FCString::Spc(Indent), GetPackageType(Import), *Import->Desc.UPackageName.ToString(), LexToString(Import->AsyncPackageLoadingState));
 		}
 	}
 
+#if WITH_PARTIAL_REQUEST_DURING_RECURSION
+	for (FAsyncPackage2* Import : AdditionalImportedAsyncPackages)
+	{
+		if (Import->AsyncPackageLoadingState < EAsyncPackageLoadingState2::DeferredDelete)
+		{
+			UE_LOG(LogStreaming, Warning, TEXT(" %s dynamically imports %s package %s state %s"), FCString::Spc(Indent), GetPackageType(Import), *Import->Desc.UPackageName.ToString(), LexToString(Import->AsyncPackageLoadingState));
+		}
+	}
+#endif
+
 	if (AllDependenciesSetupState.WaitingForPackage)
 	{
 		UE_LOG(LogStreaming, Warning, TEXT(" %s AllDependenciesSetupState is waiting on %s"), FCString::Spc(Indent), *FormatPackage(AllDependenciesSetupState.WaitingForPackage));
@@ -11351,6 +11431,21 @@ void FAsyncPackage2::DumpStateImpl(TSet<FAsyncPackage2*>& Set, int32 Indent, TMu
 			}
 		}
 	}
+
+#if WITH_PARTIAL_REQUEST_DURING_RECURSION
+	for (FAsyncPackage2* Import : AdditionalImportedAsyncPackages)
+	{
+		if (Import && Import->AsyncPackageLoadingState < EAsyncPackageLoadingState2::DeferredDelete)
+		{
+			if (!Set.Contains(Import))
+			{
+				UE_LOG(LogStreaming, Warning, TEXT(""));
+				Import->DumpStateImpl(Set, 0, MappedNodes);
+			}
+		}
+	}
+#endif
+
 }
 
 #endif // ALT2_DUMP_STATE_ON_HANG
diff --git a/Engine/Source/Runtime/Engine/Private/Tests/Loading/AsyncLoadingTests_RecursiveLoads.cpp b/Engine/Source/Runtime/Engine/Private/Tests/Loading/AsyncLoadingTests_RecursiveLoads.cpp
--- a/Engine/Source/Runtime/Engine/Private/Tests/Loading/AsyncLoadingTests_RecursiveLoads.cpp
+++ b/Engine/Source/Runtime/Engine/Private/Tests/Loading/AsyncLoadingTests_RecursiveLoads.cpp
@@ -32,7 +32,10 @@ bool FLoadingTests_RecursiveLoads_FromSerialize::RunTest(const FString& Paramete
 				if (UObject* Obj = Object->SoftReference.LoadSynchronous())
 				{
 					TestTrue(TEXT("Recursive loads in serialize should be deserialized"), !Obj->HasAnyFlags(RF_NeedLoad));
-					TestTrue(TEXT("Recursive loads in serialize should have their postload delayed until the caller is also ready to postload"), Obj->HasAnyFlags(RF_NeedPostLoad));
+					if (!IsInGameThread())
+					{
+						TestTrue(TEXT("Recursive loads in serialize skip thread-unsafe postloads when run from the ALT"), Obj->HasAnyFlags(RF_NeedPostLoad));
+					}
 				}
 			}
 		}
@@ -240,4 +243,90 @@ bool FLoadingTests_RecursiveLoads_FullFlushFrom_PostLoad::RunTest(const FString&
 	return true;
 }
 
+IMPLEMENT_CUSTOM_SIMPLE_AUTOMATION_TEST(
+	FLoadingTests_RecursiveLoads_FromBothSerializeAndPostLoad,
+	FLoadingTests_RecursiveLoads_Base,
+	TEXT("System.Engine.Loading.RecursiveLoads.FromBothSerializeAndPostLoad"),
+	EAutomationTestFlags::EditorContext | EAutomationTestFlags::EngineFilter
+)
+
+bool FLoadingTests_RecursiveLoads_FromBothSerializeAndPostLoad::RunTest(const FString& Parameters)
+{
+	FLoadingTestsScope LoadingTestScope(this,
+		[](FLoadingTestsScope& TestScope)
+		{
+			TestScope.DefaultMutateObjects();
+		}
+	);
+
+	std::atomic<bool> PartialLoadAchieved = false;
+	UE::FManualResetEvent Event;
+	int32 SerializeCount = 0;
+	// On serialize we try to force load, this should add the newly loaded package as a dynamic import of the package that requested them.
+	UAsyncLoadingTests_Shared::OnSerialize.BindLambda(
+		[this, &LoadingTestScope, &PartialLoadAchieved, &Event, &SerializeCount](FArchive& Ar, UAsyncLoadingTests_Shared* Object)
+		{
+			if (Ar.IsLoading() && Object->GetPathName() == FLoadingTestsScope::ObjectPath1)
+			{
+				SerializeCount++;
+				check(SerializeCount == 1);
+
+				// Loading Object2 while being in Object1 serialize should trigger partial load of Object2
+				if (UObject* Obj = Object->SoftReference.LoadSynchronous())
+				{
+					TestTrue(TEXT("Recursive loads in serialize should be deserialized"), !Obj->HasAnyFlags(RF_NeedLoad));
+					if (!IsInGameThread())
+					{
+						TestTrue(TEXT("Recursive loads in serialize skip thread-unsafe postloads when run from the ALT"), Obj->HasAnyFlags(RF_NeedPostLoad));
+					}
+
+					PartialLoadAchieved = true;
+
+					// When everything runs on GT, we will have no choice but to use the whole time
+					// but when running with ALT, we will be able to resolve earlier as the GT will
+					// unlock us sooner.
+					Event.WaitFor(UE::FMonotonicTimeSpan::FromSeconds(1.0f)); 
+				}
+			}
+		}
+	);
+
+	// Once in postload of object3, we now force load the same softref a second time, expecting to be able to postload it without deadlocking because of the merged postload groups.
+	UAsyncLoadingTests_Shared::OnPostLoad.BindLambda(
+		[this, &LoadingTestScope](UAsyncLoadingTests_Shared* Object)
+		{
+			if (Object->GetPathName() == FLoadingTestsScope::ObjectPath3)
+			{
+				// Trying to load Package2 completely while being in package 3 postload should be able to succeed
+				// even if the Package2 is now a dynamic import of Package1.
+				if (UObject* Obj = LoadObject<UObject>(nullptr, LoadingTestScope.ObjectPath2))
+				{
+					TestTrue(TEXT("Recursive loads in postload should be deserialized"), !Obj->HasAnyFlags(RF_NeedLoad));
+					TestTrue(TEXT("Recursive loads in postload should be able to postload"), !Obj->HasAnyFlags(RF_NeedPostLoad));
+				}
+			}
+		}
+	);
+
+	int32 Request1 = LoadPackageAsync(LoadingTestScope.PackagePath1);
+
+	// Let the first package run until we reach the serialization part then we'll 
+	// back-off to start another package and finish it.
+	while (!PartialLoadAchieved)
+	{
+		ProcessAsyncLoadingUntilComplete([&PartialLoadAchieved](){ return PartialLoadAchieved.load(); }, 0.1f);
+	}
+
+	// Hopefully, the Object1 and Object2 are still being loaded when we reach this point.
+	int32 Request2 = LoadPackageAsync(LoadingTestScope.PackagePath3);
+
+	// Unlock the loading thread faster than the timeout if we reach here first.
+	Event.Notify();
+
+	FlushAsyncLoading(Request2);
+	FlushAsyncLoading(Request1);
+
+	return true;
+}
+
 #endif // WITH_DEV_AUTOMATION_TESTS
diff --git a/Engine/Source/Runtime/Engine/Private/Tests/Loading/AsyncLoadingTests_Shared.cpp b/Engine/Source/Runtime/Engine/Private/Tests/Loading/AsyncLoadingTests_Shared.cpp
--- a/Engine/Source/Runtime/Engine/Private/Tests/Loading/AsyncLoadingTests_Shared.cpp
+++ b/Engine/Source/Runtime/Engine/Private/Tests/Loading/AsyncLoadingTests_Shared.cpp
@@ -21,6 +21,9 @@ void FLoadingTestsScope::CreateObjects()
 
 	Package2 = CreatePackage();
 	Object2 = NewObject<UAsyncLoadingTests_Shared>(Package2, ObjectName, RF_Public | RF_Standalone);
+	
+	Package3 = CreatePackage();
+	Object3 = NewObject<UAsyncLoadingTests_Shared>(Package3, ObjectName, RF_Public | RF_Standalone);
 }
 
 void FLoadingTestsScope::DefaultMutateObjects()
diff --git a/Engine/Source/Runtime/Engine/Private/Tests/Loading/AsyncLoadingTests_Shared.h b/Engine/Source/Runtime/Engine/Private/Tests/Loading/AsyncLoadingTests_Shared.h
--- a/Engine/Source/Runtime/Engine/Private/Tests/Loading/AsyncLoadingTests_Shared.h
+++ b/Engine/Source/Runtime/Engine/Private/Tests/Loading/AsyncLoadingTests_Shared.h
@@ -84,12 +84,16 @@ public:
 	static constexpr const TCHAR* ObjectPath1 = TEXT("/Engine/LoadingTestsScope_Package1.TestObject");
 	static constexpr const TCHAR* PackagePath2 = TEXT("/Engine/LoadingTestsScope_Package2");
 	static constexpr const TCHAR* ObjectPath2 = TEXT("/Engine/LoadingTestsScope_Package2.TestObject");
+	static constexpr const TCHAR* PackagePath3 = TEXT("/Engine/LoadingTestsScope_Package3");
+	static constexpr const TCHAR* ObjectPath3 = TEXT("/Engine/LoadingTestsScope_Package3.TestObject");
 
 	UPackage* Package1 = nullptr;
 	UPackage* Package2 = nullptr;
+	UPackage* Package3 = nullptr;
 
 	UAsyncLoadingTests_Shared* Object1 = nullptr;
 	UAsyncLoadingTests_Shared* Object2 = nullptr;
+	UAsyncLoadingTests_Shared* Object3 = nullptr;
 
 	UPackage* CreatePackage()
 	{
