Avoid creating and destroying an FText every time a BlueprintExceptionInfo is made.

Every time execTracepoint is reached, we create a BlueprintExceptionInfo with no description. Despite having no text assigned to it, the exception info structure's FText member still bumps the refcount of a globally-shared "EmptyText" object. When AutoRTFM is on, each of these refcount bumps enqueues an OnCommit handler. These are individually cheap, but can add up to be quite expensive.

We now hold a `TOptional<FText>` object instead. This will be slightly more efficient across the board, and much more efficient when AutoRTFM is on.

diff --git a/Engine/Source/Programs/AutoRTFMTests/Private/UECoreTests.cpp b/Engine/Source/Programs/AutoRTFMTests/Private/UECoreTests.cpp
--- a/Engine/Source/Programs/AutoRTFMTests/Private/UECoreTests.cpp
+++ b/Engine/Source/Programs/AutoRTFMTests/Private/UECoreTests.cpp
@@ -7,6 +7,7 @@
 #include "Async/ParallelFor.h"
 #include "AssetRegistry/AssetDataTagMap.h"
 #include "AutoRTFMTesting.h"
+#include "Blueprint/BlueprintExceptionInfo.h"
 #include "Delegates/IDelegateInstance.h"
 #include "HAL/IConsoleManager.h"
 #include "HAL/MallocLeakDetection.h"
@@ -3032,3 +3033,19 @@ TEST_CASE("UECore.GetOSVersion")
 		Version = FPlatformMisc::GetOSVersion();
 	});
 }
+
+TEST_CASE("BlueprintExceptionInfo.Tracepoint")
+{
+	AutoRTFM::Testing::Commit([&]
+	{
+		// We should be able to construct millions of tracepoints within a transaction without causing 
+		// the task array to overflow.
+		for (int Count = 0; Count < 50'000'000; ++Count)
+		{
+			FBlueprintExceptionInfo TracepointExceptionInfo(EBlueprintExceptionType::Tracepoint);
+			
+			// REQUIRE is actually too slow here, as it goes into the open and back.
+			check(TracepointExceptionInfo.GetType() == EBlueprintExceptionType::Tracepoint);
+		}
+	});
+}
diff --git a/Engine/Source/Runtime/CoreUObject/Public/Blueprint/BlueprintExceptionInfo.h b/Engine/Source/Runtime/CoreUObject/Public/Blueprint/BlueprintExceptionInfo.h
--- a/Engine/Source/Runtime/CoreUObject/Public/Blueprint/BlueprintExceptionInfo.h
+++ b/Engine/Source/Runtime/CoreUObject/Public/Blueprint/BlueprintExceptionInfo.h
@@ -3,6 +3,7 @@
 #pragma once
 
 #include "Internationalization/Text.h"
+#include "Misc/Optional.h"
 
 // Kinds of Blueprint exceptions
 namespace EBlueprintExceptionType
@@ -43,9 +44,20 @@ public:
 
 	const FText& GetDescription() const
 	{
-		return Description;
+		if (!Description.IsSet())
+		{
+			Description = FText();
+		}
+		return Description.GetValue();
 	}
-protected:
+
+private:
 	EBlueprintExceptionType::Type EventType;
-	FText Description;
+	
+	// We use TOptional<FText> here as an optimization for the common case of an exception with no text.
+	// For instance, every tracepoint in the code creates a FBlueprintExceptionInfo with no description.
+	// Constructing an empty FText will take a reference on the shared FText::GetEmpty() string, which
+	// is especially expensive in AutoRTFM because we need to track reference-count updates in order to
+	// undo them. If the caller actually inspects the empty description, we synthesize it on demand.
+	mutable TOptional<FText> Description;
 };
