dstm: multi-server commit

if a transaction gets flagged via TransactionRequiresMultiServerCommit, then when we get to the end of the transaction, just before we commit, we initiate the multi-server commit protocol to negotiate returning borrowed objects to their owners, etc. this ends up (carefully) performing a fair amount of work in the Open, including network IO

the multi-server commit flow goes like this:

- provoking server broadcasts 'begin' message to remotes
- provoking server performs ExecuteMultiServerCommit on each subsystem (subsystems may send messages out)
- provoking server broadcasts 'ready' message to remotes
- remotes reply with either 'ready response' or 'abort'
- after all remotes have responded, if any remote replied with 'abort', broadcast 'abandon' to remotes and abort the transaction
- otherwise, perform CommitMultiServerCommit on each subsystem
- broadcast 'end' message to remotes

the remote servers, on receiving the 'begin' message, enter a mode where most messages get deferred into a queue so that we only process messages belonging to the multi-server commit. we also build up a list of deferred actions (like object deserialize) to execute when we are told to 'end'

changes:

- FWeakObjectPtr no longer triggers a resolve when being assigned from a FObjectPtr, we just copy the RemoteId over and lazily resolve it if needed
- added 'net.DeferTransactionalSend' cvar to toggle the behavior of whether low level net sends get deferred or executed immediately while in a transaction
- added a virtual IsMigrationRoot() to UObject, overridden to 'true' within AActor, so that we do a more correct job in FindCanonicalRootObjectForSerialization

diff --git a/Engine/Plugins/Online/OnlineSubsystemUtils/Source/OnlineSubsystemUtils/Private/IpConnection.cpp b/Engine/Plugins/Online/OnlineSubsystemUtils/Source/OnlineSubsystemUtils/Private/IpConnection.cpp
--- a/Engine/Plugins/Online/OnlineSubsystemUtils/Source/OnlineSubsystemUtils/Private/IpConnection.cpp
+++ b/Engine/Plugins/Online/OnlineSubsystemUtils/Source/OnlineSubsystemUtils/Private/IpConnection.cpp
@@ -412,7 +412,7 @@ void UIpConnection::LowLevelSend(void* Data, int32 CountBits, FOutPacketTraits&
 		return;
 	}
 
-	if (AutoRTFM::IsTransactional())
+	if (AutoRTFM::IsTransactional() && Driver->ShouldDeferTransactionalSend())
 	{
 		// Inside an AutoRTFM transaction, we must postpone network sends until commit time, to account for the
 		// possibility of aborting the transaction. (An abort will prevent any network traffic from occurring.)
diff --git a/Engine/Source/Runtime/CoreUObject/Private/UObject/RemoteExecutor.cpp b/Engine/Source/Runtime/CoreUObject/Private/UObject/RemoteExecutor.cpp
--- a/Engine/Source/Runtime/CoreUObject/Private/UObject/RemoteExecutor.cpp
+++ b/Engine/Source/Runtime/CoreUObject/Private/UObject/RemoteExecutor.cpp
@@ -9,7 +9,13 @@ DEFINE_LOG_CATEGORY(LogRemoteExec);
 namespace UE::RemoteExecutor
 {
 TDelegate<void()> TickNetworkDelegate;
-TDelegate<TTuple<FName, FRemoteWorkPriority, TFunction<void(void)>>(void)> FetchNextDeferredRPCDelegate;
+TDelegate<void(FRemoteTransactionId, FRemoteWorkPriority, const TArray<FRemoteServerId>& /*MultiServerCommitRemoteServers*/)> BeginMultiServerCommitDelegate;
+TDelegate<void(FRemoteTransactionId, const TArray<FRemoteServerId>& /*MultiServerCommitRemoteServers*/)> ReadyMultiServerCommitDelegate;
+TDelegate<void(FRemoteTransactionId, const TArray<FRemoteServerId>& /*MultiServerCommitRemoteServers*/)> AbandonMultiServerCommitDelegate;
+TDelegate<void(FRemoteTransactionId, const TArray<FRemoteServerId>& /*MultiServerCommitRemoteServers*/)> EndMultiServerCommitDelegate;
+TDelegate<void(FRemoteTransactionId, FRemoteServerId /*ServerId*/)> ReadyRemoteMultiServerCommitDelegate;
+TDelegate<void(FRemoteTransactionId, FRemoteServerId /*ServerId*/)> AbortRemoteMultiServerCommitDelegate;
+TDelegate<TOptional<TTuple<FName, FRemoteWorkPriority, bool, TFunction<void(void)>>>(void)> FetchNextDeferredRPCDelegate;
 
 TMulticastDelegate<void(FRemoteTransactionId, FName)> OnTransactionStartingDelegate;
 TMulticastDelegate<void(FRemoteTransactionId)> OnTransactionCompletedDelegate;
@@ -34,9 +40,17 @@ struct FRemoteExecutorWork
 	FName Name;
 	FRemoteTransactionId RequestId;
 	FRemoteWorkPriority Priority;
+	bool bIsTransactional = false;
+
 	TFunction<void(void)> Work;
+
 	TArray<void*> SubsystemRequests;
+	bool bRequestCreated = false;
+	
 	uint32 ExecutionAttempts = 0;
+	
+	bool bRequiresMultiServerCommit = false;
+	FString RequiresMultiServerCommitReason;
 };
 
 class FRemoteExecutor
@@ -48,9 +62,23 @@ class FRemoteExecutor
 	TArray<FRemoteExecutorWork> PendingWorks;
 
 public:
+	FRemoteExecutorWork* ExecutingWork = nullptr;
+
 	ERemoteExecutorAbortReason AbortReason = ERemoteExecutorAbortReason::Unspecified;
 	FString AbortReasonDescription;
 
+	// tracking data for servicing a remote multi-server commit
+	FRemoteServerId ActiveRemoteMultiServerCommitServerId;
+	FRemoteTransactionId ActiveRemoteMultiServerCommitRequestId;
+	FRemoteWorkPriority ActiveRemoteMultiServerCommitPriority;
+	TArray<TFunction<void()>> ActiveRemoteMultiServerCommitDeferredActions;
+	bool bActiveRemoteMultiServerCommitReady = false;
+
+	// tracking data for executing our local multi-server commit
+	FRemoteTransactionId MultiServerCommitRequestId;
+	TArray<FRemoteServerId> MultiServerCommitReadyServers;
+	bool bMultiServerCommitRequiresAbort = false;
+
 	void RegisterSubsystem(FRemoteSubsystemBase* Subsystem)
 	{
 		Subsystems.Add(Subsystem);
@@ -62,35 +90,50 @@ public:
 		return FRemoteTransactionId((RawRequestId++ % 0x800000u) + 1000u);
 	}
 
-	void EnqueueWork(FName WorkName, FRemoteWorkPriority InWorkPriority, const TFunction<void(void)>& InWork)
+	void EnqueueWork(FName WorkName, FRemoteWorkPriority InWorkPriority, bool bIsTransactional, const TFunction<void(void)>& InWork)
 	{
 		FRemoteExecutorWork& NewWork = PendingWorks.Emplace_GetRef();
 		NewWork.Name = WorkName;
 		NewWork.Priority = InWorkPriority;
 		NewWork.Work = InWork;
+		NewWork.bIsTransactional = bIsTransactional;
+		NewWork.RequestId = GenerateNextTransactionId();
 	}
 
 	void ExecutePendingWork()
 	{
-		uint32 LocalIterationNumber = 0;
+		check(ExecutingWork == nullptr);
+
 		double LastStallPrintTime = FPlatformTime::Seconds();
 
-		while (PendingWorks.Num() > 0)
+		for (int32 LocalIterationNumber = 0; ; LocalIterationNumber++)
 		{
-			LocalIterationNumber++;
+			double Now = FPlatformTime::Seconds();
+			bool bPrintStall = false;
 
-			UE_LOG(LogRemoteExec, VeryVerbose, TEXT("ExecutePendingWork iteration %u"), LocalIterationNumber);
+			if ((Now - LastStallPrintTime) > 1.0)
+			{
+				LastStallPrintTime = Now;
+				bPrintStall = true;
+			}
 
 			// pump the network
 			UE::RemoteExecutor::TickNetworkDelegate.ExecuteIfBound();
 
-			// fetch pending RPC work
+			// fetch and enqueue all RPCs to execute from the network
 			for (;;)
 			{
-				TTuple<FName, FRemoteWorkPriority, TFunction<void(void)>> NewWork = UE::RemoteExecutor::FetchNextDeferredRPCDelegate.Execute();
-				if (NewWork.Get<1>().IsValid())
+				TOptional<TTuple<FName, FRemoteWorkPriority, bool, TFunction<void(void)>>> NewWorkOpt = UE::RemoteExecutor::FetchNextDeferredRPCDelegate.Execute();
+				if (NewWorkOpt.IsSet())
 				{
-					EnqueueWork(NewWork.Get<0>(), NewWork.Get<1>(), NewWork.Get<2>());
+					TTuple<FName, FRemoteWorkPriority, bool, TFunction<void(void)>>& NewWork = NewWorkOpt.GetValue();
+					EnqueueWork(NewWork.Get<0>(), NewWork.Get<1>(), NewWork.Get<2>(), NewWork.Get<3>());
+
+					UE_LOG(LogRemoteExec, VeryVerbose, TEXT("ExecutePendingWork[%d] enqueued deferred RPC '%s' %s %s"),
+						LocalIterationNumber,
+						*NewWork.Get<0>().ToString(),
+						*NewWork.Get<1>().ToString(),
+						NewWork.Get<2>() ? TEXT("(transactional)") :  TEXT("(non-transactional)"));
 				}
 				else
 				{
@@ -98,177 +141,432 @@ public:
 				}
 			}
 
+			// check if we're actively servicing a remote multi-server commit
+			if (ActiveRemoteMultiServerCommitRequestId.IsValid())
+			{
+				check(ActiveRemoteMultiServerCommitServerId.IsValid());
+				check(ActiveRemoteMultiServerCommitPriority.IsValid());
+
+				if (bPrintStall)
+				{
+					UE_LOG(LogRemoteExec, Verbose, TEXT("ExecutePendingWork[%d] Waiting on handling remote multi server commit %s %s"),
+						LocalIterationNumber,
+						*ActiveRemoteMultiServerCommitServerId.ToString(),
+						*ActiveRemoteMultiServerCommitRequestId.ToString());
+				}
+
+				// if we have an active remote multi-server commit that we are servicing, we pause executing any local work
+				// until that is complete
+				continue;
+			}
+
+			//
+			// if we get here and there is no work pending, break out and finish
+			//
+			if (PendingWorks.Num() == 0)
+			{
+				break;
+			}
+
 			for (int32 PendingWorkIndex = 0; PendingWorkIndex < PendingWorks.Num(); PendingWorkIndex++)
 			{
 				FRemoteExecutorWork& PendingWork = PendingWorks[PendingWorkIndex];
 
-				UE_LOG(LogRemoteExec, VeryVerbose, TEXT("PendingWork[%d] '%s' RequestId %s %s"),
+				UE_LOG(LogRemoteExec, VeryVerbose, TEXT("ExecutePendingWork[%d] : PendingWork[%d] '%s' RequestId %s %s"),
+					LocalIterationNumber,
 					PendingWorkIndex,
 					*PendingWork.Name.ToString(),
 					*PendingWork.RequestId.ToString(),
 					*PendingWork.Priority.ToString());
 			}
 
+			// round robin through executing all pending work
 			for(int32 PendingWorkIndex = 0; PendingWorkIndex < PendingWorks.Num(); )
 			{
-				FRemoteExecutorWork& ExecutingWork = PendingWorks[PendingWorkIndex];
+				check(ExecutingWork == nullptr);
+				TGuardValue<FRemoteExecutorWork*> ExecutingWorkSave(ExecutingWork, &PendingWorks[PendingWorkIndex]);
 
-				UE_LOG(LogRemoteExec, VeryVerbose, TEXT("-- ExecuteTransactional : iteration %u request %s %s --"),
+				UE_LOG(LogRemoteExec, VeryVerbose, TEXT("ExecutePendingWork[%d,%s] : Executing request %s %s --"),
 					LocalIterationNumber,
-					*ExecutingWork.RequestId.ToString(),
-					*ExecutingWork.Priority.ToString());
+					*ExecutingWork->RequestId.ToString(),
+					*ExecutingWork->Name.ToString(),
+					*ExecutingWork->Priority.ToString());
 
-				if (!ExecutingWork.RequestId.IsValid())
+				if (ExecutingWork->bIsTransactional)
 				{
-					// new request
-					ExecutingWork.RequestId = GenerateNextTransactionId();
-
-					for (int32 SubsystemIndex = 0; SubsystemIndex < Subsystems.Num(); SubsystemIndex++)
+					if (!ExecutingWork->bRequestCreated)
 					{
-						FRemoteSubsystemBase* Subsystem = Subsystems[SubsystemIndex];
-						UE_LOG(LogRemoteExec, VeryVerbose, TEXT("\tcreate request[%s]"), Subsystem->NameForDebug());
-						Subsystem->CreateRequest(ExecutingWork.RequestId, ExecutingWork.Priority);
-
-						Subsystem->SetActiveRequest(ExecutingWork.RequestId);
+						// new request
+						ExecutingWork->bRequestCreated = true;
 
-						Subsystem->BeginRequest();
-					}
-				}
-				else
-				{
-					for (int32 SubsystemIndex = 0; SubsystemIndex < Subsystems.Num(); SubsystemIndex++)
-					{
-						FRemoteSubsystemBase* Subsystem = Subsystems[SubsystemIndex];
-						Subsystem->SetActiveRequest(ExecutingWork.RequestId);
-					}
-				}
-
-				double Now = FPlatformTime::Seconds();
-				bool bPrintStall = false;
-
-				if ((Now - LastStallPrintTime) > 1.0)
-				{
-					LastStallPrintTime = Now;
-					bPrintStall = true;
-				}
+						for (int32 SubsystemIndex = 0; SubsystemIndex < Subsystems.Num(); SubsystemIndex++)
+						{
+							FRemoteSubsystemBase* Subsystem = Subsystems[SubsystemIndex];
+							UE_LOG(LogRemoteExec, VeryVerbose, TEXT("ExecutePendingWork[%d,%s] : Creating request[%s]"),
+								LocalIterationNumber,
+								*ExecutingWork->RequestId.ToString(),
+								Subsystem->NameForDebug());
+							Subsystem->CreateRequest(ExecutingWork->RequestId, ExecutingWork->Priority);
 
-				// tick the subsystems
-				for (int32 SubsystemIndex = 0; SubsystemIndex < Subsystems.Num(); SubsystemIndex++)
-				{
-					FRemoteSubsystemBase* Subsystem = Subsystems[SubsystemIndex];
-					Subsystem->Tick();
-				}
+							Subsystem->SetActiveRequest(ExecutingWork->RequestId);
 
-				int32 SubsystemReadyCount = 0;
-				for (int32 SubsystemIndex = 0; SubsystemIndex < Subsystems.Num(); SubsystemIndex++)
-				{
-					FRemoteSubsystemBase* Subsystem = Subsystems[SubsystemIndex];
-					if (Subsystem->AreDependenciesSatisfied())
-					{
-						SubsystemReadyCount++;
+							Subsystem->BeginRequest();
+						}
 					}
 					else
 					{
-						if (bPrintStall)
+						for (int32 SubsystemIndex = 0; SubsystemIndex < Subsystems.Num(); SubsystemIndex++)
 						{
-							UE_LOG(LogRemoteExec, VeryVerbose, TEXT("\tsubsystem %s not ready..."), Subsystem->NameForDebug());
+							FRemoteSubsystemBase* Subsystem = Subsystems[SubsystemIndex];
+							Subsystem->SetActiveRequest(ExecutingWork->RequestId);
 						}
 					}
-				}
 
-				if (SubsystemReadyCount != Subsystems.Num())
-				{
-					if (bPrintStall)
+					Now = FPlatformTime::Seconds();
+					if ((Now - LastStallPrintTime) > 1.0)
 					{
-						UE_LOG(LogRemoteExec, VeryVerbose, TEXT("\t%d of %d subsystems not ready..."), Subsystems.Num() - SubsystemReadyCount, Subsystems.Num());
+						LastStallPrintTime = Now;
+						bPrintStall = true;
 					}
 
-					PendingWorkIndex++;
-					continue;
-				}
-
-				// all of the subsystems are ready, try to perform the work
-				AbortReason = ERemoteExecutorAbortReason::Unspecified;
-				AbortReasonDescription = TEXT("");
-
-				ExecutingWork.ExecutionAttempts++;
+					// tick the subsystems
+					for (int32 SubsystemIndex = 0; SubsystemIndex < Subsystems.Num(); SubsystemIndex++)
+					{
+						FRemoteSubsystemBase* Subsystem = Subsystems[SubsystemIndex];
+						Subsystem->Tick();
+					}
 
-				FRemoteTransactionId RequestId = ExecutingWork.RequestId;
-				UE::RemoteExecutor::OnTransactionStartingDelegate.Broadcast(RequestId, ExecutingWork.Name);
-				AutoRTFM::ETransactionResult WorkTransactResult = AutoRTFM::Transact([&ExecutingWork]()
+					int32 SubsystemReadyCount = 0;
+					for (int32 SubsystemIndex = 0; SubsystemIndex < Subsystems.Num(); SubsystemIndex++)
 					{
-						// You can set a breakpoint inside the abort handler below to see what caused the abort,
-						// or you can rely on the log output which only catches livelocks
-						const uint32 ExecutionAttempts = ExecutingWork.ExecutionAttempts;
-						UE_AUTORTFM_ONABORT(ExecutionAttempts)
+						FRemoteSubsystemBase* Subsystem = Subsystems[SubsystemIndex];
+						if (Subsystem->AreDependenciesSatisfied())
 						{
-							constexpr uint32 MaxExecutionAttemptsBeforeLivelockWarning = 200;
-							if ((ExecutionAttempts >= MaxExecutionAttemptsBeforeLivelockWarning) && (ExecutionAttempts % MaxExecutionAttemptsBeforeLivelockWarning) == 0)
+							SubsystemReadyCount++;
+						}
+						else
+						{
+							if (bPrintStall)
 							{
-								UE_LOG(LogRemoteExec, Display, TEXT("vvv Transaction Aborted %u Times. Dumping Callstack. vvv"), ExecutionAttempts);
-								FDebug::DumpStackTraceToLog(ELogVerbosity::Display);
-								UE_LOG(LogRemoteExec, Display, TEXT("^^^ Transaction Aborted %u Times ^^^"), ExecutionAttempts);
+								UE_LOG(LogRemoteExec, VeryVerbose, TEXT("ExecutePendingWork[%d,%s] : subsystem %s not ready..."),
+									LocalIterationNumber,
+									*ExecutingWork->RequestId.ToString(),
+									Subsystem->NameForDebug());
 							}
-						};
+						}
+					}
 
-						ExecutingWork.Work();
-					});
-				bool bWorkComplete = false;
-				bool bWorkTransactionAborted = false;
+					if (SubsystemReadyCount != Subsystems.Num())
+					{
+						if (bPrintStall)
+						{
+							UE_LOG(LogRemoteExec, VeryVerbose, TEXT("ExecutePendingWork[%d,%s] : %d of %d subsystems not ready..."),
+								LocalIterationNumber,
+								*ExecutingWork->RequestId.ToString(),
+								Subsystems.Num() - SubsystemReadyCount, Subsystems.Num());
+						}
 
-				if (WorkTransactResult == AutoRTFM::ETransactionResult::AbortedByRequest)
-				{
-					bWorkTransactionAborted = true;
+						PendingWorkIndex++;
+						continue;
+					}
 
-					UE::RemoteExecutor::OnTransactionAbortedDelegate.Broadcast(RequestId);
+					// all of the subsystems are ready, try to perform the work
+					AbortReason = ERemoteExecutorAbortReason::Unspecified;
+					AbortReasonDescription = TEXT("");
 
-					if (AbortReason == ERemoteExecutorAbortReason::RequiresDependencies)
-					{
-						UE_LOG(LogRemoteExec, Verbose, TEXT("\twork ABORTED for required dependency '%s'"), *AbortReasonDescription);
-					}
-					else if (AbortReason == ERemoteExecutorAbortReason::AbandonWork)
+					ExecutingWork->ExecutionAttempts++;
+
+					FRemoteTransactionId RequestId = ExecutingWork->RequestId;
+					UE::RemoteExecutor::OnTransactionStartingDelegate.Broadcast(RequestId, ExecutingWork->Name);
+					AutoRTFM::ETransactionResult WorkTransactResult =
+						AutoRTFM::Transact(
+							[this, LocalIterationNumber]()
+							{
+								// You can set a breakpoint inside the abort handler below to see what caused the abort,
+								// or you can rely on the log output which only catches livelocks
+								const uint32 ExecutionAttempts = ExecutingWork->ExecutionAttempts;
+								UE_AUTORTFM_ONABORT(ExecutionAttempts)
+								{
+									constexpr uint32 MaxExecutionAttemptsBeforeLivelockWarning = 200;
+									if ((ExecutionAttempts >= MaxExecutionAttemptsBeforeLivelockWarning) && (ExecutionAttempts % MaxExecutionAttemptsBeforeLivelockWarning) == 0)
+									{
+										UE_LOG(LogRemoteExec, Display, TEXT("vvv Transaction Aborted %u Times. Dumping Callstack. vvv"), ExecutionAttempts);
+										FDebug::DumpStackTraceToLog(ELogVerbosity::Display);
+										UE_LOG(LogRemoteExec, Display, TEXT("^^^ Transaction Aborted %u Times ^^^"), ExecutionAttempts);
+									}
+								};
+
+								ExecutingWork->Work();
+
+								if (ExecutingWork->bRequiresMultiServerCommit)
+								{
+									// we're done with the work and about to commit - 
+									// first we need to send borrowed objects back
+									// to their owner and ask them if they are ready
+									// to commit or not...
+									UE_LOG(LogRemoteExec, VeryVerbose, TEXT("ExecutePendingWork[%d,%s] : STARTING multi-server commit because %s"),
+										LocalIterationNumber,
+										*ExecutingWork->RequestId.ToString(),
+										*ExecutingWork->RequiresMultiServerCommitReason);
+
+									check (!MultiServerCommitRequestId.IsValid());
+									check (MultiServerCommitReadyServers.Num() == 0);
+									check (bMultiServerCommitRequiresAbort == false);
+
+									UE_AUTORTFM_OPEN
+									{
+										MultiServerCommitRequestId = ExecutingWork->RequestId;
+										bMultiServerCommitRequiresAbort = false;
+									};
+
+									// first collect the list of servers that need to be involved
+									TArray<FRemoteServerId> MultiServerCommitRemoteServers;
+									for (int32 SubsystemIndex = 0; SubsystemIndex < Subsystems.Num(); SubsystemIndex++)
+									{
+										UE_LOG(LogRemoteExec, VeryVerbose, TEXT("ExecutePendingWork[%d,%s] : BeginMultiServerCommit subsystem %d..."),
+											LocalIterationNumber,
+											*ExecutingWork->RequestId.ToString(),
+											SubsystemIndex);
+
+										FRemoteSubsystemBase* Subsystem = Subsystems[SubsystemIndex];
+										Subsystem->BeginMultiServerCommit(MultiServerCommitRemoteServers);
+									}
+
+									UE_LOG(LogRemoteExec, VeryVerbose, TEXT("ExecutePendingWork[%d,%s] : BeginMultiServerCommit DONE (%d servers)"),
+										LocalIterationNumber,
+										*ExecutingWork->RequestId.ToString(),
+										MultiServerCommitRemoteServers.Num());
+
+									for (FRemoteServerId ServerId : MultiServerCommitRemoteServers)
+									{
+										UE_LOG(LogRemoteExec, VeryVerbose, TEXT("ExecutePendingWork[%d,%s] : BeginMultiServerCommit server %s"),
+											LocalIterationNumber,
+											*ExecutingWork->RequestId.ToString(),
+											*ServerId.ToString());
+									}
+
+									// signal to the relevant servers that they need to help us do a multi-server commit
+									UE_AUTORTFM_OPEN
+									{
+										UE::RemoteExecutor::BeginMultiServerCommitDelegate.ExecuteIfBound(ExecutingWork->RequestId, ExecutingWork->Priority, MultiServerCommitRemoteServers);
+									};
+
+									// ask each subsystem to send any necessary data as part of this commit
+									for (int32 SubsystemIndex = 0; SubsystemIndex < Subsystems.Num(); SubsystemIndex++)
+									{
+										UE_LOG(LogRemoteExec, VeryVerbose, TEXT("ExecutePendingWork[%d,%s] : ExecuteMultiServerCommit subsystem %d..."),
+											LocalIterationNumber,
+											*ExecutingWork->RequestId.ToString(),
+											SubsystemIndex);
+
+										FRemoteSubsystemBase* Subsystem = Subsystems[SubsystemIndex];
+										Subsystem->ExecuteMultiServerCommit();
+									}
+
+									UE_AUTORTFM_OPEN
+									{
+										// tell each server that we're done sending commit data and we are waiting for them
+										// to respond with whether they are ready or we need to abort and retry
+										UE::RemoteExecutor::ReadyMultiServerCommitDelegate.ExecuteIfBound(ExecutingWork->RequestId, MultiServerCommitRemoteServers);
+									};
+
+									// tick the network until everything is ready
+									double LastPrintTime = FPlatformTime::Seconds();
+
+									for (;;)
+									{
+										double Now = FPlatformTime::Seconds();
+										if ((Now - LastPrintTime) > 1.0)
+										{
+											LastPrintTime = Now;
+											UE_LOG(LogRemoteExec, VeryVerbose, TEXT("ExecutePendingWork[%d,%s] : WAITING for multi-server commit..."),
+												LocalIterationNumber,
+												*ExecutingWork->RequestId.ToString());
+										}
+										
+										if (MultiServerCommitRemoteServers.Num() == MultiServerCommitReadyServers.Num())
+										{
+											// we got a response from each server, do a sanity check to ensure the list of ready
+											// servers is the list we expected to have										
+											UE_AUTORTFM_OPEN
+											{
+												MultiServerCommitRemoteServers.Sort();
+												MultiServerCommitReadyServers.Sort();
+											};
+
+											for (int32 ServerIndex = 0; ServerIndex < MultiServerCommitRemoteServers.Num(); ServerIndex++)
+											{
+												FRemoteServerId ServerId = MultiServerCommitRemoteServers[ServerIndex];
+												FRemoteServerId ReadyServerId = MultiServerCommitReadyServers[ServerIndex];
+
+												if (ServerId != ReadyServerId)
+												{
+													UE_LOG(LogRemoteExec, VeryVerbose, TEXT("ExecutePendingWork[%d,%s] : Multi-server commit expected server %s but got %s"),
+														LocalIterationNumber,
+														*ExecutingWork->RequestId.ToString(),
+														*ServerId.ToString(),
+														*ReadyServerId.ToString());
+												}
+
+												check(ServerId == ReadyServerId);
+											}
+
+											if (bMultiServerCommitRequiresAbort)
+											{
+												UE_LOG(LogRemoteExec, Verbose, TEXT("ExecutePendingWork[%d,%s] : Multi-server commit ALL servers READY, but we are flagged with ABORT"),
+													LocalIterationNumber,
+													*ExecutingWork->RequestId.ToString());
+
+												UE_AUTORTFM_OPEN
+												{
+													for (int32 SubsystemIndex = 0; SubsystemIndex < Subsystems.Num(); SubsystemIndex++)
+													{
+														FRemoteSubsystemBase* Subsystem = Subsystems[SubsystemIndex];
+														Subsystem->AbortMultiServerCommit();
+													}
+
+													// aborted, tell all servers to abandon this commit
+													UE::RemoteExecutor::AbandonMultiServerCommitDelegate.ExecuteIfBound(ExecutingWork->RequestId, MultiServerCommitRemoteServers);
+
+													MultiServerCommitRequestId = FRemoteTransactionId::Invalid();
+													MultiServerCommitReadyServers.Reset();
+													bMultiServerCommitRequiresAbort = false;
+												};
+
+												UE::RemoteExecutor::AbortTransactionRequiresDependencies(TEXT("bMultiServerCommitRequiresAbort"));
+												return;
+											}
+											else
+											{
+												// the commit was accepted by every server, notify each subsystem that we are committing
+												UE_LOG(LogRemoteExec, VeryVerbose, TEXT("ExecutePendingWork[%d,%s] : Multi-server commit ALL servers READY, COMMITTING"),
+													LocalIterationNumber,
+													*ExecutingWork->RequestId.ToString());
+
+												for (int32 SubsystemIndex = 0; SubsystemIndex < Subsystems.Num(); SubsystemIndex++)
+												{
+													FRemoteSubsystemBase* Subsystem = Subsystems[SubsystemIndex];
+													Subsystem->CommitMultiServerCommit();
+												}
+
+												break;
+											}
+										}
+
+										UE_AUTORTFM_OPEN
+										{
+											UE::RemoteExecutor::TickNetworkDelegate.ExecuteIfBound();
+										};
+									}
+
+									UE_LOG(LogRemoteExec, VeryVerbose, TEXT("ExecutePendingWork[%d,%s] : DONE with multi-server commit"),
+										LocalIterationNumber,
+										*ExecutingWork->RequestId.ToString());
+
+									UE_AUTORTFM_OPEN
+									{
+										MultiServerCommitRequestId = FRemoteTransactionId::Invalid();
+										MultiServerCommitReadyServers.Reset();
+										bMultiServerCommitRequiresAbort = false;
+									};
+
+									UE_AUTORTFM_OPEN
+									{
+										// tell each server that we are committed and they should commit
+										UE::RemoteExecutor::EndMultiServerCommitDelegate.ExecuteIfBound(ExecutingWork->RequestId, MultiServerCommitRemoteServers);
+									};
+								};
+							});
+
+					bool bWorkComplete = false;
+					bool bWorkTransactionAborted = false;
+
+					if (WorkTransactResult == AutoRTFM::ETransactionResult::AbortedByRequest)
 					{
-						UE_LOG(LogRemoteExec, Verbose, TEXT("\twork ABORTED for ABANDONMENT '%s'"), *AbortReasonDescription);
-						bWorkComplete = true;
+						bWorkTransactionAborted = true;
+
+						UE::RemoteExecutor::OnTransactionAbortedDelegate.Broadcast(RequestId);
+
+						if (AbortReason == ERemoteExecutorAbortReason::RequiresDependencies)
+						{
+							UE_LOG(LogRemoteExec, Verbose, TEXT("ExecutePendingWork[%d,%s] : Work ABORTED for required dependency '%s'"),
+								LocalIterationNumber,
+								*ExecutingWork->RequestId.ToString(),
+								*AbortReasonDescription);
+						}
+						else if (AbortReason == ERemoteExecutorAbortReason::AbandonWork)
+						{
+							UE_LOG(LogRemoteExec, Verbose, TEXT("ExecutePendingWork[%d,%s] : Work ABORTED for ABANDONMENT '%s'"),
+								LocalIterationNumber,
+								*ExecutingWork->RequestId.ToString(),
+								*AbortReasonDescription);
+							bWorkComplete = true;
+						}
+						else
+						{
+							UE_LOG(LogRemoteExec, Verbose, TEXT("ExecutePendingWork[%d,%s] : Work ABORTED for unknown reason '%s'"),
+								LocalIterationNumber,
+								*ExecutingWork->RequestId.ToString(),
+								*AbortReasonDescription);
+						}
 					}
 					else
 					{
-						UE_LOG(LogRemoteExec, Verbose, TEXT("\twork ABORTED for unknown reason '%s'"), *AbortReasonDescription);
-					}
-				}
-				else
-				{
-					check(AbortReason == ERemoteExecutorAbortReason::Unspecified);
+						check(AbortReason == ERemoteExecutorAbortReason::Unspecified);
 
-					UE_LOG(LogRemoteExec, VeryVerbose, TEXT("\twork COMPLETED after %u attempts"), ExecutingWork.ExecutionAttempts);
-					UE::RemoteExecutor::OnTransactionCompletedDelegate.Broadcast(RequestId);
+						UE_LOG(LogRemoteExec, VeryVerbose, TEXT("ExecutePendingWork[%d,%s] : Work COMPLETED after %u attempts"),
+							LocalIterationNumber,
+							*ExecutingWork->RequestId.ToString(),
+							ExecutingWork->ExecutionAttempts);
+						UE::RemoteExecutor::OnTransactionCompletedDelegate.Broadcast(RequestId);
 
-					bWorkComplete = true;
-				}
+						bWorkComplete = true;
+					}
 
-				if (bWorkComplete)
-				{
-					for (int32 SubsystemIndex = 0; SubsystemIndex < Subsystems.Num(); SubsystemIndex++)
+					if (bWorkComplete)
 					{
-						FRemoteSubsystemBase* Subsystem = Subsystems[SubsystemIndex];
-						UE_LOG(LogRemoteExec, VeryVerbose, TEXT("\tend request[%s]"), Subsystem->NameForDebug());
-						Subsystem->EndRequest(!bWorkTransactionAborted);
-						Subsystem->ClearActiveRequest();
-					}
+						for (int32 SubsystemIndex = 0; SubsystemIndex < Subsystems.Num(); SubsystemIndex++)
+						{
+							FRemoteSubsystemBase* Subsystem = Subsystems[SubsystemIndex];
+							UE_LOG(LogRemoteExec, VeryVerbose, TEXT("ExecutePendingWork[%d,%s] : End request[%s]"),
+								LocalIterationNumber,
+								*ExecutingWork->RequestId.ToString(),
+								Subsystem->NameForDebug());
+							Subsystem->EndRequest(!bWorkTransactionAborted);
+							Subsystem->ClearActiveRequest();
+						}
 
-					for (int32 SubsystemIndex = 0; SubsystemIndex < Subsystems.Num(); SubsystemIndex++)
+						for (int32 SubsystemIndex = 0; SubsystemIndex < Subsystems.Num(); SubsystemIndex++)
+						{
+							FRemoteSubsystemBase* Subsystem = Subsystems[SubsystemIndex];
+							Subsystem->DestroyRequest(ExecutingWork->RequestId);
+						}
+
+						// remove this work from the list and do not increment PendingWorkIndex
+						PendingWorks.RemoveAt(PendingWorkIndex, EAllowShrinking::No);
+					}
+					else
 					{
-						FRemoteSubsystemBase* Subsystem = Subsystems[SubsystemIndex];
-						Subsystem->DestroyRequest(ExecutingWork.RequestId);
+						// we will revisit this later, move on to the next work
+						PendingWorkIndex++;
 					}
+				}
+				else
+				{
+					// execute non-transactional work
+					ExecutingWork->Work();
 
 					// remove this work from the list and do not increment PendingWorkIndex
 					PendingWorks.RemoveAt(PendingWorkIndex, EAllowShrinking::No);
 				}
-				else
+
+				// it's possible this work ended with trying to perform a multi-server
+				// commit that got aborted because a remote multi-server commit of higher
+				// priority came in. if so, we need to break out of processing further work
+				// so we can service the remote multi-server commit
+				if (ActiveRemoteMultiServerCommitRequestId.IsValid())
 				{
-					// we will revisit this later, move on to the next work
-					PendingWorkIndex++;
+					break;
 				}
 			}
 		}
@@ -282,7 +580,7 @@ void ExecuteTransactionalInternal(FName WorkName, FRemoteWorkPriority InWorkPrio
 #if UE_WITH_REMOTE_OBJECT_HANDLE
 	check(!AutoRTFM::IsClosed());
 
-	GRemoteExecutor.EnqueueWork(WorkName, InWorkPriority, [&InWork]() { InWork(); });
+	GRemoteExecutor.EnqueueWork(WorkName, InWorkPriority, true, [&InWork]() { InWork(); });
 	GRemoteExecutor.ExecutePendingWork();
 #else
 	InWork();
@@ -298,6 +596,8 @@ void RegisterRemoteSubsystem(FRemoteSubsystemBase* Subsystem)
 
 void AbortTransactionRequiresDependencies(FStringView Description)
 {
+	check(AutoRTFM::IsClosed());
+
 	UE_AUTORTFM_OPEN
 	{
 		GRemoteExecutor.AbortReason = ERemoteExecutorAbortReason::RequiresDependencies;
@@ -307,8 +607,20 @@ void AbortTransactionRequiresDependencies(FStringView Description)
 	AutoRTFM::AbortTransaction();
 }
 
+void RollbackTransactionRequiresDependencies(FStringView Description)
+{
+	check(!AutoRTFM::IsClosed());
+
+	GRemoteExecutor.AbortReason = ERemoteExecutorAbortReason::RequiresDependencies;
+	GRemoteExecutor.AbortReasonDescription = Description;
+
+	AutoRTFM::ForTheRuntime::RollbackTransaction();
+}
+
 void AbortTransactionAndAbandonWork(FStringView Description)
 {
+	check(AutoRTFM::IsClosed());
+
 	UE_AUTORTFM_OPEN
 	{
 		GRemoteExecutor.AbortReason = ERemoteExecutorAbortReason::AbandonWork;
@@ -318,12 +630,242 @@ void AbortTransactionAndAbandonWork(FStringView Description)
 	AutoRTFM::AbortTransaction();
 }
 
+void TransactionRequiresMultiServerCommit(FStringView Description)
+{
+	if (GRemoteExecutor.ExecutingWork && !GRemoteExecutor.ExecutingWork->bRequiresMultiServerCommit)
+	{
+		GRemoteExecutor.ExecutingWork->bRequiresMultiServerCommit = true;
+		GRemoteExecutor.ExecutingWork->RequiresMultiServerCommitReason = Description;
+
+		UE_LOG(LogRemoteExec, VeryVerbose, TEXT("TransactionRequiresMultiServerCommit ACTIVATED because: %s"), *GRemoteExecutor.ExecutingWork->RequiresMultiServerCommitReason);
+	}
+}
+
+void BeginRemoteMultiServerCommit(FRemoteServerId ServerId, FRemoteTransactionId RequestId, FRemoteWorkPriority RequestPriority)
+{
+	bool bAcceptRequest = false;
+	bool bAbortLocalCommit = false;
+
+	if (GRemoteExecutor.ActiveRemoteMultiServerCommitRequestId.IsValid())
+	{
+		if (GRemoteExecutor.bActiveRemoteMultiServerCommitReady)
+		{
+			// we already told the remote server that we are READY, so we are locked in for a moment until we finish
+			UE_LOG(LogRemoteExec, VeryVerbose, TEXT("BeginRemoteMultiServerCommit %s from %s %s DENYING because we are already READY with remote multi-server commit %s %s"),
+				*RequestId.ToString(),
+				*ServerId.ToString(),
+				*RequestPriority.ToString(),
+				*GRemoteExecutor.ActiveRemoteMultiServerCommitRequestId.ToString(),
+				*GRemoteExecutor.ActiveRemoteMultiServerCommitPriority.ToString());
+		}
+		else if (IsHigherPriority(RequestPriority, GRemoteExecutor.ActiveRemoteMultiServerCommitPriority))
+		{
+			UE_LOG(LogRemoteExec, VeryVerbose, TEXT("BeginRemoteMultiServerCommit %s from %s %s ACCEPTING because our remote multi-server commit %s is lower priority %s"),
+				*RequestId.ToString(),
+				*ServerId.ToString(),
+				*RequestPriority.ToString(),
+				*GRemoteExecutor.ActiveRemoteMultiServerCommitRequestId.ToString(),
+				*GRemoteExecutor.ActiveRemoteMultiServerCommitPriority.ToString());
+
+			bAcceptRequest = true;
+		}
+		else
+		{
+			UE_LOG(LogRemoteExec, VeryVerbose, TEXT("BeginRemoteMultiServerCommit %s from %s %s DENYING because we are servicing higher priority remote multi-server commit %s %s"),
+				*RequestId.ToString(),
+				*ServerId.ToString(),
+				*RequestPriority.ToString(),
+				*GRemoteExecutor.ActiveRemoteMultiServerCommitRequestId.ToString(),
+				*GRemoteExecutor.ActiveRemoteMultiServerCommitPriority.ToString());
+		}
+	}
+	else if (GRemoteExecutor.MultiServerCommitRequestId.IsValid())
+	{
+		// we are in a local multi-server commit - should we abort it in favor of this remote one?
+		check(GRemoteExecutor.ExecutingWork);
+		if (IsHigherPriority(RequestPriority, GRemoteExecutor.ExecutingWork->Priority))
+		{
+			UE_LOG(LogRemoteExec, VeryVerbose, TEXT("BeginRemoteMultiServerCommit %s from %s %s ACCEPTING because our local multi-server commit %s is lower priority %s"),
+				*RequestId.ToString(),
+				*ServerId.ToString(),
+				*RequestPriority.ToString(),
+				*GRemoteExecutor.MultiServerCommitRequestId.ToString(),
+				*GRemoteExecutor.ExecutingWork->Priority.ToString());
+
+			bAcceptRequest = true;
+			bAbortLocalCommit = true;
+		}
+		else
+		{
+			UE_LOG(LogRemoteExec, VeryVerbose, TEXT("BeginRemoteMultiServerCommit %s from %s %s DENYING because we are locally in multi-server commit %s %s"),
+				*RequestId.ToString(),
+				*ServerId.ToString(),
+				*RequestPriority.ToString(),
+				*GRemoteExecutor.MultiServerCommitRequestId.ToString(),
+				*GRemoteExecutor.ExecutingWork->Priority.ToString());
+		}
+	}
+	else
+	{
+		// we aren't currently in a local multi-server commit, and we aren't
+		// servicing a remote multi-server commit, so accept this
+		UE_LOG(LogRemoteExec, VeryVerbose, TEXT("BeginRemoteMultiServerCommit %s %s from %s ACCEPTED"),
+			*RequestId.ToString(),
+			*RequestPriority.ToString(),
+			*ServerId.ToString());
+
+		bAcceptRequest = true;
+	}
+
+	if (bAcceptRequest)
+	{
+		// do we have to first abandon one we're working on?
+		if (GRemoteExecutor.ActiveRemoteMultiServerCommitRequestId.IsValid())
+		{
+			AbortRemoteMultiServerCommitDelegate.Execute(GRemoteExecutor.ActiveRemoteMultiServerCommitRequestId, GRemoteExecutor.ActiveRemoteMultiServerCommitServerId);
+
+			GRemoteExecutor.ActiveRemoteMultiServerCommitServerId = FRemoteServerId{};
+			GRemoteExecutor.ActiveRemoteMultiServerCommitRequestId = FRemoteTransactionId::Invalid();
+			GRemoteExecutor.ActiveRemoteMultiServerCommitPriority = FRemoteWorkPriority{};
+			GRemoteExecutor.bActiveRemoteMultiServerCommitReady = false;
+			GRemoteExecutor.ActiveRemoteMultiServerCommitDeferredActions.Reset();
+		}
+
+		check(!GRemoteExecutor.ActiveRemoteMultiServerCommitServerId.IsValid());
+		check(!GRemoteExecutor.ActiveRemoteMultiServerCommitRequestId.IsValid());
+		check(!GRemoteExecutor.ActiveRemoteMultiServerCommitPriority.IsValid());
+		check(!GRemoteExecutor.bActiveRemoteMultiServerCommitReady);
+		check(GRemoteExecutor.ActiveRemoteMultiServerCommitDeferredActions.Num() == 0);
+
+		GRemoteExecutor.ActiveRemoteMultiServerCommitServerId = ServerId;
+		GRemoteExecutor.ActiveRemoteMultiServerCommitRequestId = RequestId;
+		GRemoteExecutor.ActiveRemoteMultiServerCommitPriority = RequestPriority;
+	}
+	else
+	{
+		AbortRemoteMultiServerCommitDelegate.Execute(RequestId, ServerId);
+	}
+
+	if (bAbortLocalCommit)
+	{
+		// we can't immediately abort the transaction, set this flag for it to properly shut down gracefully
+		GRemoteExecutor.bMultiServerCommitRequiresAbort = true;
+	}
+}
+
+void EndRemoteMultiServerCommit(FRemoteServerId ServerId, FRemoteTransactionId RequestId)
+{
+	UE_LOG(LogRemoteExec, VeryVerbose, TEXT("EndRemoteMultiServerCommit : %s %s"), *ServerId.ToString(), *RequestId.ToString());
+
+	check(GRemoteExecutor.ActiveRemoteMultiServerCommitServerId == ServerId);
+	check(GRemoteExecutor.ActiveRemoteMultiServerCommitRequestId == RequestId);
+
+	for (TFunction<void()>& DeferredAction : GRemoteExecutor.ActiveRemoteMultiServerCommitDeferredActions)
+	{
+		DeferredAction();
+	}
+
+	GRemoteExecutor.ActiveRemoteMultiServerCommitServerId = FRemoteServerId{};
+	GRemoteExecutor.ActiveRemoteMultiServerCommitRequestId = FRemoteTransactionId::Invalid();
+	GRemoteExecutor.ActiveRemoteMultiServerCommitPriority = FRemoteWorkPriority{};
+	GRemoteExecutor.bActiveRemoteMultiServerCommitReady = false;
+	GRemoteExecutor.ActiveRemoteMultiServerCommitDeferredActions.Reset();
+}
+
+void AbandonRemoteMultiServerCommit(FRemoteServerId ServerId, FRemoteTransactionId RequestId)
+{
+	if ((GRemoteExecutor.ActiveRemoteMultiServerCommitServerId == ServerId) &&
+		GRemoteExecutor.ActiveRemoteMultiServerCommitRequestId == RequestId)
+	{
+		UE_LOG(LogRemoteExec, VeryVerbose, TEXT("AbandonRemoteMultiServerCommit : %s %s"), *ServerId.ToString(), *RequestId.ToString());
+
+		GRemoteExecutor.ActiveRemoteMultiServerCommitServerId = FRemoteServerId{};
+		GRemoteExecutor.ActiveRemoteMultiServerCommitRequestId = FRemoteTransactionId::Invalid();
+		GRemoteExecutor.ActiveRemoteMultiServerCommitPriority = FRemoteWorkPriority{};
+		GRemoteExecutor.bActiveRemoteMultiServerCommitReady = false;
+		GRemoteExecutor.ActiveRemoteMultiServerCommitDeferredActions.Reset();
+	}
+	else
+	{
+		UE_LOG(LogRemoteExec, VeryVerbose, TEXT("AbandonRemoteMultiServerCommit : %s %s IGNORING"), *ServerId.ToString(), *RequestId.ToString());
+	}
+}
+
+void EnqueueRemoteMultiServerCommitAction(FRemoteServerId ServerId, FRemoteTransactionId RequestId, const TFunction<void()>& Action)
+{
+#if UE_WITH_REMOTE_OBJECT_HANDLE
+	// THIS has to sit in a holding pen...
+	if ((GRemoteExecutor.ActiveRemoteMultiServerCommitServerId == ServerId) &&
+		GRemoteExecutor.ActiveRemoteMultiServerCommitRequestId == RequestId)
+	{
+		UE_LOG(LogRemoteExec, Verbose, TEXT("EnqueueRemoteMultiServerCommitAction Enqueueing action from %s request %s"), *ServerId.ToString(), *RequestId.ToString());
+
+		GRemoteExecutor.ActiveRemoteMultiServerCommitDeferredActions.Add(Action);
+	}
+	else
+	{
+		// else we've received this message for a different multi-server commit that we no longer have active, so just abandon
+		// since they would have been told previously that we were aborting it
+		UE_LOG(LogRemoteExec, VeryVerbose, TEXT("EnqueueRemoteMultiServerCommitAction : IGNORING action from %s request %s"), *ServerId.ToString(), *RequestId.ToString());
+	}
+#endif
+}
+
+void ReadyMultiServerCommitResponse(FRemoteServerId ServerId, FRemoteTransactionId RequestId)
+{
+	if (GRemoteExecutor.MultiServerCommitRequestId == RequestId)
+	{
+		UE_LOG(LogRemoteExec, VeryVerbose, TEXT("ReadyMultiServerCommitResponse got ready server: %s"), *ServerId.ToString());
+		GRemoteExecutor.MultiServerCommitReadyServers.Add(ServerId);
+	}
+	else
+	{
+		UE_LOG(LogRemoteExec, VeryVerbose, TEXT("ReadyMultiServerCommitResponse ignoring %s because we are working on %s"), *RequestId.ToString(), *GRemoteExecutor.MultiServerCommitRequestId.ToString());
+	}
+}
+
+void AbortMultiServerCommit(FRemoteServerId ServerId, FRemoteTransactionId RequestId)
+{
+	if (GRemoteExecutor.MultiServerCommitRequestId == RequestId)
+	{
+		UE_LOG(LogRemoteExec, VeryVerbose, TEXT("AbortMultiServerCommit got valid: %s"), *RequestId.ToString());
+		GRemoteExecutor.bMultiServerCommitRequiresAbort = true;
+		GRemoteExecutor.MultiServerCommitReadyServers.Add(ServerId);
+	}
+	else
+	{
+		UE_LOG(LogRemoteExec, VeryVerbose, TEXT("AbortMultiServerCommit ignoring %s because we are working on %s"), *RequestId.ToString(), *GRemoteExecutor.MultiServerCommitRequestId.ToString());
+	}
+}
+
+void ReadyRemoteMultiServerCommit(FRemoteServerId ServerId, FRemoteTransactionId RequestId)
+{
+	if ((GRemoteExecutor.ActiveRemoteMultiServerCommitServerId == ServerId) &&
+		(GRemoteExecutor.ActiveRemoteMultiServerCommitRequestId == RequestId))
+	{
+		UE_LOG(LogRemoteExec, VeryVerbose, TEXT("ReadyRemoteMultiServerCommit %s from %s"), *RequestId.ToString(), *ServerId.ToString());
+
+		check(!GRemoteExecutor.bActiveRemoteMultiServerCommitReady);
+		GRemoteExecutor.bActiveRemoteMultiServerCommitReady = true;
+
+		// if everything passed OK we need to send a message back saying we're ready
+		// TODO: this is where we'll need to hook code in that verifies everything
+		// we received is able to be accepted
+		ReadyRemoteMultiServerCommitDelegate.Execute(RequestId, ServerId);
+	}
+}
+
 FRemoteWorkPriority CreateRootWorkPriority()
 {
 	FRemoteWorkPriority RootWorkPriority = FRemoteWorkPriority::CreateRootWorkPriority(UE::RemoteObject::GetGlobalServerId(), GRemoteExecutor.GenerateNextTransactionId());
 	return RootWorkPriority;
 }
 
+void ExecutePendingWork()
+{
+	GRemoteExecutor.ExecutePendingWork();
+}
+
 void ExecuteTransactionalWithExplicitPriority(FName WorkName, FRemoteWorkPriority WorkPriority, const TFunctionRef<void(void)>& Work)
 {
 #if UE_WITH_REMOTE_OBJECT_HANDLE
diff --git a/Engine/Source/Runtime/CoreUObject/Private/UObject/RemoteObject.cpp b/Engine/Source/Runtime/CoreUObject/Private/UObject/RemoteObject.cpp
--- a/Engine/Source/Runtime/CoreUObject/Private/UObject/RemoteObject.cpp
+++ b/Engine/Source/Runtime/CoreUObject/Private/UObject/RemoteObject.cpp
@@ -176,7 +176,7 @@ void InitRemoteObjects()
 
 	if (!UE::RemoteExecutor::FetchNextDeferredRPCDelegate.IsBound())
 	{
-		UE::RemoteExecutor::FetchNextDeferredRPCDelegate.BindStatic([]() { return TTuple<FName, FRemoteWorkPriority, TFunction<void(void)>>(); });
+		UE::RemoteExecutor::FetchNextDeferredRPCDelegate.BindStatic([]() { return TOptional<TTuple<FName, FRemoteWorkPriority, bool, TFunction<void(void)>>>(); });
 	}
 }
 
@@ -289,14 +289,17 @@ bool IsRemote(FRemoteObjectId ObjectId)
 	{
 		return false;
 	}
+
 	if (UObject* Object = StaticFindObjectFastInternal(ObjectId))
 	{
 		return IsRemote(Object);
 	}
+
 	if (FindRemoteObjectStub(ObjectId))
 	{
 		return true;
 	}
+
 	const FRemoteServerId ServerId = ObjectId.GetServerId();
 	// Invalid server Id means local native classes which are created before a server has a chance to have an id assigned
 	return ServerId.IsValid() && ServerId != UE::RemoteObject::GetGlobalServerId();
diff --git a/Engine/Source/Runtime/CoreUObject/Private/UObject/RemoteObjectSerialization.cpp b/Engine/Source/Runtime/CoreUObject/Private/UObject/RemoteObjectSerialization.cpp
--- a/Engine/Source/Runtime/CoreUObject/Private/UObject/RemoteObjectSerialization.cpp
+++ b/Engine/Source/Runtime/CoreUObject/Private/UObject/RemoteObjectSerialization.cpp
@@ -338,6 +338,7 @@ protected:
 	TMap<UObject*, FNameIndexType> PathNameMap;
 	TMap<FRemoteObjectId, FNameIndexType> RemoteIdMap;
 	TMap<UObject*, UObject*>* ObjectToArchetypeMap = nullptr;
+	TSet<UObject*>* ReferencedObjectsSet = nullptr;
 
 public:
 
@@ -412,7 +413,16 @@ private:
 
 			// Anything can be marked as a remote reference, even assets in which case we rely on this flag to be set so that GC calls
 			// StoreObjectToDatabase for any remotely referenced asset (and only for remotely referenced assets) that's about to be GC'd
-			Info.Object->SetInternalFlags(EInternalObjectFlags::RemoteReference);
+			if (!Info.Object->HasAnyInternalFlags(EInternalObjectFlags::RemoteReference))
+			{
+				if (ReferencedObjectsSet)
+				{
+					UE_AUTORTFM_OPEN
+					{
+						ReferencedObjectsSet->Add(Info.Object);
+					};
+				}
+			}
 		}
 
 		*this << Info.Type;
@@ -441,13 +451,15 @@ private:
 
 public:
 
-	FArchiveRemoteObjectWriter(UObject* InRootObject, FRemoteObjectData& OutObjectData, const FUObjectMigrationContext* InMigrationContext, const TCHAR* ArchiveName = nullptr)
+	FArchiveRemoteObjectWriter(UObject* InRootObject, FRemoteObjectData& OutObjectData, const FUObjectMigrationContext* InMigrationContext, const TCHAR* ArchiveName = nullptr, TSet<UObject*>* OutReferencedObjectsSet = nullptr)
 		: TArchiveRemoteObjectBase<FMemoryWriter>(OutObjectData, InMigrationContext, ArchiveName ? ArchiveName : TEXT("RemoteObjectWriter"))
 	{
 		static_assert(sizeof(FNameIndexType) == sizeof(FPackedRemoteObjectPathName::FNameIndexType), "Name Index type must match the type used for storing remote object path name indices");
 
 		SetRootObject(InRootObject);
 		ObjectsToSerialize.Add(InRootObject);
+
+		ReferencedObjectsSet = OutReferencedObjectsSet;
 	}
 
 	virtual UObject* GetArchetypeFromLoader(const UObject* Obj) override
@@ -1062,7 +1074,7 @@ void SerializeObjectDataInternal(FArchiveRemoteObjectWriter& Ar, UObject* Reques
 	}
 }
 
-FRemoteObjectData SerializeObjectData(UObject* InObject, TSet<UObject*>& OutObjects, const FUObjectMigrationContext* MigrationContext)
+FRemoteObjectData SerializeObjectData(UObject* InObject, TSet<UObject*>& OutObjects, TSet<UObject*>& OutReferencedObjects, const FUObjectMigrationContext* MigrationContext)
 {
 	UObject* Object = InObject;
 	if (GRemoteObjectsMigrateFullHierarchy)
@@ -1074,7 +1086,7 @@ FRemoteObjectData SerializeObjectData(UObject* InObject, TSet<UObject*>& OutObje
 	FRemoteObjectId RequestedObjectId(InObject);
 
 	{
-		FArchiveRemoteObjectWriter Ar(Object, ObjectData, MigrationContext);
+		FArchiveRemoteObjectWriter Ar(Object, ObjectData, MigrationContext, nullptr, &OutReferencedObjects);
 		Ar.SetMigratingRemoteObjects(true);
 		SerializeObjectDataInternal(Ar, InObject, RequestedObjectId, OutObjects);
 	}
@@ -1281,15 +1293,30 @@ int32 DeserializeObjectData(FRemoteObjectData& ObjectData,
 
 UObject* FindCanonicalRootObjectForSerialization(UObject* Object)
 {
-	while (Object && Object->IsDefaultSubobject())
+	// find the outermost migration root
+	UObject* Cursor = Object;
+
+	// if we walk the outer chain and don't happen to find 
+	// any migration roots, default to using the object itself
+	UObject* Result = Object;
+
+#if UE_WITH_REMOTE_OBJECT_HANDLE
+	// walk the cursor up the entire Outer chain and update
+	// Result with the outermost Outer that is a migration root
+	// (this covers the case where we find a migration root
+	// nested in another, we pick the outermost one)
+	while (Cursor)
 	{
-		Object = Object->GetOuter();
-	}
+		if (Cursor->IsMigrationRoot())
+		{
+			Result = Cursor;
+		}
 
-	// Sanity check because by definition default subobjects should always have an outer
-	checkf(Object, TEXT("Attempted to migrate default subobject (%s) hierarchy but it doesn't have an outer or all outers are also default subobjects"), *GetPathNameSafe(Object));
+		Cursor = Cursor->GetOuter();
+	}
+#endif
 
-	return Object;
+	return Result;
 }
 
 } // namespace UE::RemoteObject::Serialization
\ No newline at end of file
diff --git a/Engine/Source/Runtime/CoreUObject/Private/UObject/RemoteObjectSerialization.h b/Engine/Source/Runtime/CoreUObject/Private/UObject/RemoteObjectSerialization.h
--- a/Engine/Source/Runtime/CoreUObject/Private/UObject/RemoteObjectSerialization.h
+++ b/Engine/Source/Runtime/CoreUObject/Private/UObject/RemoteObjectSerialization.h
@@ -72,10 +72,11 @@ namespace UE::RemoteObject::Serialization
 	* Serializes an object and its subobject (or if the object is a default subobject, its parent and the parent's subobjects)
 	* @param InObject Object to be serialized
 	* @param OutObjects All objects that have been serialized (Object and its subobjects and/or parent)
+	* @param OutReferencedObjects Keeps track of all objects that need to be tagged with RemoteReference
 	* @param MigrationContext Contains the meta data of the current migration request
 	* @return Remote object data representing the serialized objects
 	*/
-	FRemoteObjectData SerializeObjectData(UObject* InObject, TSet<UObject*>& OutObjects, const FUObjectMigrationContext* MigrationContext);
+	FRemoteObjectData SerializeObjectData(UObject* InObject, TSet<UObject*>& OutObjects, TSet<UObject*>& OutReferencedObjects, const FUObjectMigrationContext* MigrationContext);
 
 	/**
 	* Deserializes remote object data
diff --git a/Engine/Source/Runtime/CoreUObject/Private/UObject/RemoteObjectTransfer.cpp b/Engine/Source/Runtime/CoreUObject/Private/UObject/RemoteObjectTransfer.cpp
--- a/Engine/Source/Runtime/CoreUObject/Private/UObject/RemoteObjectTransfer.cpp
+++ b/Engine/Source/Runtime/CoreUObject/Private/UObject/RemoteObjectTransfer.cpp
@@ -51,9 +51,16 @@ public:
 	FRemoteTransactionId RequestId;
 	FRemoteWorkPriority Priority;
 
+	// this array fills up with object ids that we have ever touched
+	// while running the transaction
 	TArray<FRemoteObjectId> RequiredObjects;
 	TArray<FRemoteObjectId> RequiredObjectsCanonicalRoot;
 
+	// this array gets cleared before running the 
+	// transaction and only tracks the objects touched 
+	// during the most recent run
+	TArray<FRemoteObjectId> UsedObjects;
+
 	FRemoteTransactionId GetRequestId() const { return RequestId; }
 	FRemoteWorkPriority GetPriority() const { return Priority; }
 };
@@ -63,7 +70,12 @@ class FRemoteObjectTransferQueue : public FRemoteSubsystem<FRemoteObjectRequest>
 	// these are objects that we have outstanding requests for (across all requests)
 	TArray<FRemoteObjectId> PendingObjectRequests;
 
+	TSet<UObject*> MultiServerCommitSentObjects;
+	TSet<UObject*> MultiServerCommitReferencedObjects;
+
 public:
+	bool bInMultiServerCommit = false;
+
 	const TCHAR* NameForDebug() final { return TEXT("RemoteObjectTransferQueue"); }
 
 	void BeginRequest() final
@@ -78,6 +90,9 @@ public:
 
 		check(ActiveRequest->RequiredObjects.Num() == ActiveRequest->RequiredObjectsCanonicalRoot.Num());
 
+		// reset the list of used objects each time before we run
+		ActiveRequest->UsedObjects.SetNum(0, EAllowShrinking::No);
+
 		UE_LOG(LogRemoteObject, VeryVerbose, TEXT("\t%d required objs (%d pending requests):"), ActiveRequest->RequiredObjects.Num(), PendingObjectRequests.Num());
 		for (int32 RequiredObjectIndex = 0; RequiredObjectIndex < ActiveRequest->RequiredObjects.Num(); RequiredObjectIndex++)
 		{
@@ -168,6 +183,165 @@ public:
 		return true;
 	}
 
+	void BeginMultiServerCommit(TArray<FRemoteServerId>& OutMultiServerCommitRemoteServers) final
+	{
+		check(ActiveRequest);
+		check(!bInMultiServerCommit);
+		bInMultiServerCommit = true;
+
+		UE_LOG(LogRemoteObject, VeryVerbose, TEXT("FRemoteObjectTransferQueue: BeginMultiServerCommit has %d used objects"), ActiveRequest->UsedObjects.Num());
+
+		// are any of our used objects borrowed?
+		for (FRemoteObjectId UsedObjectId : ActiveRequest->UsedObjects)
+		{
+			UE_LOG(LogRemoteObject, VeryVerbose, TEXT("FRemoteObjectTransferQueue: used obj: %s remote: %d"), *UsedObjectId.ToString(), Handle::IsRemote(UsedObjectId));
+
+			// if we used it this run and we're about to commit, we expect that the object is local
+			check(!Handle::IsRemote(UsedObjectId));
+
+			UObject* UsedObject = StaticFindObjectFastInternal(UsedObjectId);
+			check(UsedObject);
+
+			FRemoteServerId OwnerServerId = Handle::GetOwnerServerId(UsedObject);
+			if (OwnerServerId != GetGlobalServerId())
+			{
+				if (OutMultiServerCommitRemoteServers.Find(OwnerServerId) == INDEX_NONE)
+				{
+					OutMultiServerCommitRemoteServers.Add(OwnerServerId);
+				}
+			}
+		}
+	}
+
+	void ExecuteMultiServerCommit() final
+	{
+		using namespace UE::RemoteObject::Serialization;
+		
+		check(MultiServerCommitSentObjects.Num() == 0);
+		check(MultiServerCommitReferencedObjects.Num() == 0);
+		check(bInMultiServerCommit);
+
+		TSet<UObject*> BorrowedObjectsToReturn;
+
+		// are any of our used objects borrowed?
+		for (FRemoteObjectId ObjectId : ActiveRequest->UsedObjects)
+		{
+			if (Handle::IsRemote(ObjectId))
+			{
+				UE_LOG(LogRemoteObject, VeryVerbose, TEXT("FRemoteObjectTransferQueue: ExecuteMultiServerCommit UNEXPECTED: object %s is remote"), *ObjectId.ToString());
+			}
+
+			// if we used it this run and we're about to commit, we expect that the object is local
+			check(!Handle::IsRemote(ObjectId));
+
+			UObject* Object = StaticFindObjectFastInternal(ObjectId);
+			UObject* RootObject = FindCanonicalRootObjectForSerialization(Object);
+			
+			FRemoteServerId RootOwnerServerId = Handle::GetOwnerServerId(RootObject);
+			if (RootOwnerServerId != GetGlobalServerId())
+			{
+				FRemoteObjectId RootObjectId = UE::CoreUObject::Private::FObjectHandleUtils::GetRemoteId(RootObject);
+
+				UE_LOG(LogRemoteObject, VeryVerbose, TEXT("FRemoteObjectTransferQueue: ExecuteMultiServerCommit found borrowed object %s (root %s)"), *ObjectId.ToString(), *RootObjectId.ToString());
+				BorrowedObjectsToReturn.Add(RootObject);
+			}
+		}
+
+		for (UObject* Object : BorrowedObjectsToReturn)
+		{
+			FRemoteObjectId ObjectId = UE::CoreUObject::Private::FObjectHandleUtils::GetRemoteId(Object);
+			FRemoteServerId OwnerServerId = Handle::GetOwnerServerId(Object);
+
+			UE_LOG(LogRemoteObject, VeryVerbose, TEXT("FRemoteObjectTransferQueue: ExecuteMultiServerCommit sending root object %s back to %s"), *ObjectId.ToString(), *OwnerServerId.ToString());
+
+			// NOTE: we need to run SerializeObjectData in the CLOSED because it turns out some of the
+			// serialization code ends up mutating the object, so we need to be able to undo those
+			// mutations if we need to abort past this point
+			FUObjectMigrationContext MigrationContext
+			{
+				.ObjectId = ObjectId,
+				.RemoteServerId = OwnerServerId,
+				.OwnerServerId = OwnerServerId,
+				.PhysicsServerId = OwnerServerId,
+				.RequestId = ActiveRequest->RequestId,
+				.MigrationSide = EObjectMigrationSide::Send
+			};
+
+			TSet<UObject*> SentObjects;
+			FRemoteObjectData ObjectData = SerializeObjectData(Object, SentObjects, MultiServerCommitReferencedObjects, &MigrationContext);
+			FMigrateSendParams SendRemoteObjectParams = { .MigrationContext = MigrationContext, .ObjectData = ObjectData, .bForMultiServerCommit = true };
+
+			UE_AUTORTFM_OPEN
+			{
+				SendRemoteObject(SendRemoteObjectParams);
+
+				for (UObject* SO : SentObjects)
+				{
+					FRemoteObjectId SOId = UE::CoreUObject::Private::FObjectHandleUtils::GetRemoteId(SO);
+					UE_LOG(LogRemoteObject, VeryVerbose, TEXT("FRemoteObjectTransferQueue: ExecuteMultiServerCommit sent object %s (subobject %s %s)"), *ObjectId.ToString(), *SOId.ToString(), *SO->GetName());
+				}
+
+				MultiServerCommitSentObjects.Append(SentObjects);
+			};
+		}
+	}
+
+	void AbortMultiServerCommit() final
+	{
+		UE_AUTORTFM_OPEN
+		{
+			MultiServerCommitSentObjects.Reset();
+			MultiServerCommitReferencedObjects.Reset();
+		};
+
+		bInMultiServerCommit = false;
+	}
+
+	void CommitMultiServerCommit() final
+	{
+		using namespace UE::RemoteObject::Serialization;
+		using namespace UE::RemoteObject::Private;
+
+		UE_LOG(LogRemoteObject, VeryVerbose, TEXT("FRemoteObjectTransferQueue: CommitMultiServerCommit"));
+
+		for (UObject* ReferencedObject : MultiServerCommitReferencedObjects)
+		{
+			FRemoteObjectId ReferencedObjectId = UE::CoreUObject::Private::FObjectHandleUtils::GetRemoteId(ReferencedObject);
+			UE_LOG(LogRemoteObject, VeryVerbose, TEXT("FRemoteObjectTransferQueue: CommitMultiServerCommit setting RemoteReference %s"), *ReferencedObjectId.ToString());
+
+			ReferencedObject->SetInternalFlags(EInternalObjectFlags::RemoteReference);
+		}
+
+		UE_AUTORTFM_OPEN
+		{
+			MultiServerCommitReferencedObjects.Reset();
+		};
+
+		FUnsafeToMigrateScope UnsafeToMigrateScope;
+		for (UObject* SentObject : MultiServerCommitSentObjects)
+		{
+			FRemoteServerId OwnerServerId = Handle::GetOwnerServerId(SentObject);
+			FRemoteServerId PhysicsServerId = Handle::GetPhysicsServerId(SentObject);
+			FRemoteObjectId SentObjectId = UE::CoreUObject::Private::FObjectHandleUtils::GetRemoteId(SentObject);
+
+			FUObjectMigrationContext MigrationContext{ .ObjectId = SentObjectId, .RemoteServerId = OwnerServerId, .OwnerServerId = OwnerServerId, .PhysicsServerId = PhysicsServerId, .RequestId = ActiveRequest->RequestId, .MigrationSide = EObjectMigrationSide::Send };
+
+			UE_LOG(LogRemoteObject, VeryVerbose, TEXT("FRemoteObjectTransferQueue: CommitMultiServerCommit postmigrate %s"), *SentObjectId.ToString());
+
+#if UE_WITH_REMOTE_OBJECT_HANDLE
+			SentObject->PostMigrate(MigrationContext);
+#endif
+			MarkAsRemote(SentObject, OwnerServerId);
+		}
+
+		UE_AUTORTFM_OPEN
+		{
+			MultiServerCommitSentObjects.Reset();
+		};
+
+		bInMultiServerCommit = false;
+	}
+
 	void EndRequest(bool bTransactionCommitted) final
 	{
 	}
@@ -371,6 +545,13 @@ void MigrateObjectFromRemoteServer(FRemoteObjectId ObjectId, FRemoteServerId Cur
 	// ensure the request is added in the Open so after we abort it is preserved
 	UE_AUTORTFM_OPEN
 	{
+		// check bInMultiServerCommit here to ensure the UsedObjects list doesn't 
+		// accidentally mutate during the actual multi-server commit process
+		if (!GTransferQueue.bInMultiServerCommit && GTransferQueue.ActiveRequest->UsedObjects.Find(ObjectId) == INDEX_NONE)
+		{
+			GTransferQueue.ActiveRequest->UsedObjects.Add(ObjectId);
+		}
+
 		// add it to the active request's list of required objects
 		if (GTransferQueue.ActiveRequest->RequiredObjects.Find(ObjectId) == INDEX_NONE)
 		{
@@ -397,6 +578,8 @@ void MigrateObjectFromRemoteServer(FRemoteObjectId ObjectId, FRemoteServerId Cur
 			GTransferQueue.ActiveRequest->RequiredObjects.Add(ObjectId);
 			GTransferQueue.ActiveRequest->RequiredObjectsCanonicalRoot.Add(CanonicalRootObjectId);
 
+			GTransferQueue.ActiveRequest->UsedObjects.Add(ObjectId);
+
 			OnObjectTouchedDelegate.Broadcast(GTransferQueue.ActiveRequest->RequestId, ObjectId);
 		}
 		
@@ -424,12 +607,24 @@ void TouchResidentObject(UObject* Object)
 		return;
 	}
 
+	// check bInMultiServerCommit here to ensure the UsedObjects list doesn't 
+	// accidentally mutate during the actual multi-server commit process
+	if (GTransferQueue.bInMultiServerCommit)
+	{
+		return;
+	}
+
 	if (AutoRTFM::IsClosed() && GTransferQueue.ActiveRequest)
 	{
 		UE_AUTORTFM_OPEN
 		{
 			const FRemoteObjectId ObjectId = UE::CoreUObject::Private::FObjectHandleUtils::GetRemoteId(Object);
 
+			if (GTransferQueue.ActiveRequest->UsedObjects.Find(ObjectId) == INDEX_NONE)
+			{
+				GTransferQueue.ActiveRequest->UsedObjects.Add(ObjectId);
+			}
+
 			// add it to the active request's list of required objects
 			if (GTransferQueue.ActiveRequest->RequiredObjects.Find(ObjectId) == INDEX_NONE)
 			{
@@ -568,11 +763,17 @@ void MigrateObjectToRemoteServerWithExplicitPriority(FRemoteTransactionId Reques
 
 		// send the object
 		TSet<UObject*> SentObjects;
-		FRemoteObjectData ObjectData = SerializeObjectData(Object, SentObjects, &MigrationContext);
+		TSet<UObject*> ReferencedObjects;
+		FRemoteObjectData ObjectData = SerializeObjectData(Object, SentObjects, ReferencedObjects, &MigrationContext);
 
 		UE_LOG(LogRemoteObject, VeryVerbose, TEXT("     : MigrateObjectToRemoteServer(%s) serialization complete obj %s (%d bytes)"), *RequestId.ToString(), *ObjectId.ToString(), ObjectData.GetNumBytes());
 
-		for(UObject* SO : SentObjects)
+		for (UObject* RO : ReferencedObjects)
+		{
+			RO->SetInternalFlags(EInternalObjectFlags::RemoteReference);
+		}
+
+		for (UObject* SO : SentObjects)
 		{
 			FRemoteObjectId SubObjectId = UE::CoreUObject::Private::FObjectHandleUtils::GetRemoteId(SO);
 			UE_LOG(LogRemoteObject, VeryVerbose, TEXT("         sent subobject : %s [%s]"), *SubObjectId.ToString(), *GetNameSafe(SO));
@@ -697,6 +898,13 @@ void PrefetchRemoteObject(FRemoteObjectId ObjectId, UObject* DestinationOuter)
 	
 	UE_AUTORTFM_OPEN
 	{
+		// check bInMultiServerCommit here to ensure the UsedObjects list doesn't 
+		// accidentally mutate during the actual multi-server commit process
+		if (!GTransferQueue.bInMultiServerCommit && GTransferQueue.ActiveRequest->UsedObjects.Find(ObjectId) == INDEX_NONE)
+		{
+			GTransferQueue.ActiveRequest->UsedObjects.Add(ObjectId);
+		}
+
 		// it's possible a previous prefetch or strong fetch already added this to the RequiredObjects list
 		bAdded = (GTransferQueue.ActiveRequest->RequiredObjects.Find(ObjectId) == INDEX_NONE);
 
@@ -778,6 +986,7 @@ void StoreObjectToDatabase(UObject* Object)
 		UObject* OutermostUnreachableRemoteReference = GetOutermostUnreachableRemoteReference(Object);
 
 		TSet<UObject*> SentObjects;
+		TSet<UObject*> ReferencedObjects;
 		{
 			// Let's setup a fake migration context for storing this
 			FRemoteObjectId RootObjectId = FRemoteObjectId { OutermostUnreachableRemoteReference };
@@ -787,7 +996,7 @@ void StoreObjectToDatabase(UObject* Object)
 			FUObjectMigrationContext StoreToDatabaseContext{ .ObjectId = RootObjectId, .RemoteServerId = DatabaseId, .OwnerServerId = DatabaseId, .PhysicsServerId = DatabaseId, .RequestId = PushMigrationRequestId, .MigrationSide = EObjectMigrationSide::Send };
 
 			// Do the actual serialization
-			FRemoteObjectData ObjectData = SerializeObjectData(OutermostUnreachableRemoteReference, SentObjects, &StoreToDatabaseContext);
+			FRemoteObjectData ObjectData = SerializeObjectData(OutermostUnreachableRemoteReference, SentObjects, ReferencedObjects, &StoreToDatabaseContext);
 
 			const FMigrateSendParams Params { .MigrationContext = StoreToDatabaseContext, .ObjectData = ObjectData };
 			StoreRemoteObjectDataDelegate.Execute(Params);
@@ -795,6 +1004,11 @@ void StoreObjectToDatabase(UObject* Object)
 
 		checkf(SentObjects.Num() > 0, TEXT("PeerId:%s unable to store object %s to DatabaseId:%s"), *GetGlobalServerId().ToString(), *FRemoteObjectId(Object).ToString(), *DatabaseId.ToString());
 
+		for (UObject* ReferencedObject : ReferencedObjects)
+		{
+			ReferencedObject->SetInternalFlags(EInternalObjectFlags::RemoteReference);
+		}
+
 		// MarkAsRemote clears EInternalObjectFlags::RemoteReference flag which will prevent re-enetring this loop with the same outermost object but 
 		// we still want to separately process its subobjects (so that they're stored as separate entries in the database) so only mark the root object as remote
 		MarkAsRemote(OutermostUnreachableRemoteReference, DatabaseId);
diff --git a/Engine/Source/Runtime/CoreUObject/Private/UObject/WeakObjectPtr.cpp b/Engine/Source/Runtime/CoreUObject/Private/UObject/WeakObjectPtr.cpp
--- a/Engine/Source/Runtime/CoreUObject/Private/UObject/WeakObjectPtr.cpp
+++ b/Engine/Source/Runtime/CoreUObject/Private/UObject/WeakObjectPtr.cpp
@@ -31,14 +31,18 @@ void FWeakObjectPtr::operator=(FObjectPtr ObjectPtr)
 	if (ObjectPtr // && UObjectInitialized() we might need this at some point, but it is a speed hit we would prefer to avoid
 		)
 	{
-		const UObject* Object = ObjectPtr.Get();
-		ObjectIndex = GUObjectArray.ObjectToIndex((UObjectBase*)Object);
-		ObjectSerialNumber = GUObjectArray.AllocateSerialNumber(ObjectIndex);
 #if UE_WITH_REMOTE_OBJECT_HANDLE
-		ObjectRemoteId = GUObjectArray.GetRemoteId(ObjectIndex);
-#endif
+		ObjectRemoteId = ObjectPtr.GetRemoteId();
 
-		checkSlow(SerialNumbersMatch());
+		// if the object is local, fill in the index and serial number immediately
+		if (!ObjectPtr.IsRemote())
+#endif
+		{
+			const UObject* Object = ObjectPtr.Get();
+			ObjectIndex = GUObjectArray.ObjectToIndex((UObjectBase*)Object);
+			ObjectSerialNumber = GUObjectArray.AllocateSerialNumber(ObjectIndex);
+			checkSlow(SerialNumbersMatch());
+		}
 	}
 	else
 	{
diff --git a/Engine/Source/Runtime/CoreUObject/Public/UObject/Object.h b/Engine/Source/Runtime/CoreUObject/Public/UObject/Object.h
--- a/Engine/Source/Runtime/CoreUObject/Public/UObject/Object.h
+++ b/Engine/Source/Runtime/CoreUObject/Public/UObject/Object.h
@@ -548,6 +548,8 @@ public:
 	 * Called when an object is about to be migrated and after it's been migrated to another process.
 	 */
 	virtual void PostMigrate(const struct FUObjectMigrationContext& MigrationContext) {}
+
+	virtual bool IsMigrationRoot() const { return false; }
 #endif
 
 	/**
diff --git a/Engine/Source/Runtime/CoreUObject/Public/UObject/RemoteExecutor.h b/Engine/Source/Runtime/CoreUObject/Public/UObject/RemoteExecutor.h
--- a/Engine/Source/Runtime/CoreUObject/Public/UObject/RemoteExecutor.h
+++ b/Engine/Source/Runtime/CoreUObject/Public/UObject/RemoteExecutor.h
@@ -108,6 +108,11 @@ public:
 	virtual void BeginRequest() = 0;
 	virtual void Tick() = 0;
 	virtual bool AreDependenciesSatisfied() const = 0;
+	virtual void BeginMultiServerCommit(TArray<FRemoteServerId>& OutMultiServerCommitRemoteServers) = 0;
+	virtual void ExecuteMultiServerCommit() = 0;
+	virtual void AbortMultiServerCommit() = 0;
+	virtual void CommitMultiServerCommit() = 0;
+
 	virtual void EndRequest(bool bTransactionCommitted) = 0;
 };
 
@@ -186,9 +191,15 @@ namespace UE::RemoteExecutor
 {
 	/** Delegate that pumps the network and executes RPC handlers */
 	extern COREUOBJECT_API TDelegate<void()> TickNetworkDelegate;
+	extern COREUOBJECT_API TDelegate<void(FRemoteTransactionId, FRemoteWorkPriority, const TArray<FRemoteServerId>&)> BeginMultiServerCommitDelegate;
+	extern COREUOBJECT_API TDelegate<void(FRemoteTransactionId, const TArray<FRemoteServerId>&)> ReadyMultiServerCommitDelegate;
+	extern COREUOBJECT_API TDelegate<void(FRemoteTransactionId, const TArray<FRemoteServerId>&)> AbandonMultiServerCommitDelegate;
+	extern COREUOBJECT_API TDelegate<void(FRemoteTransactionId, const TArray<FRemoteServerId>&)> EndMultiServerCommitDelegate;
+	extern COREUOBJECT_API TDelegate<void(FRemoteTransactionId, FRemoteServerId)> ReadyRemoteMultiServerCommitDelegate;
+	extern COREUOBJECT_API TDelegate<void(FRemoteTransactionId, FRemoteServerId)> AbortRemoteMultiServerCommitDelegate;
 
 	/** Fetch the next deferred RPC request **/
-	extern COREUOBJECT_API TDelegate<TTuple<FName, FRemoteWorkPriority, TFunction<void(void)>>(void)> FetchNextDeferredRPCDelegate;
+	extern COREUOBJECT_API TDelegate<TOptional<TTuple<FName, FRemoteWorkPriority, bool, TFunction<void(void)>>>(void)> FetchNextDeferredRPCDelegate;
 
 	/** Delegate executed when a ExecuteTransactional starts an AutoRTFM transaction */
 	extern COREUOBJECT_API TMulticastDelegate<void(FRemoteTransactionId, FName)> OnTransactionStartingDelegate;
@@ -203,9 +214,22 @@ namespace UE::RemoteExecutor
 
 	COREUOBJECT_API void ExecuteTransactional(FName WorkName, const TFunctionRef<void(void)>& Work);
 	COREUOBJECT_API void ExecuteTransactionalWithExplicitPriority(FName WorkName, FRemoteWorkPriority WorkPriority, const TFunctionRef<void(void)>& Work);
+	COREUOBJECT_API void ExecutePendingWork();
 
 	COREUOBJECT_API void AbortTransactionRequiresDependencies(FStringView Description);
+	COREUOBJECT_API void RollbackTransactionRequiresDependencies(FStringView Description);
 	COREUOBJECT_API void AbortTransactionAndAbandonWork(FStringView Description);
+
+	COREUOBJECT_API void TransactionRequiresMultiServerCommit(FStringView Description);
+	
+	COREUOBJECT_API void BeginRemoteMultiServerCommit(FRemoteServerId ServerId, FRemoteTransactionId RequestId, FRemoteWorkPriority RequestPriority);
+	COREUOBJECT_API void ReadyRemoteMultiServerCommit(FRemoteServerId ServerId, FRemoteTransactionId RequestId);
+	COREUOBJECT_API void AbandonRemoteMultiServerCommit(FRemoteServerId ServerId, FRemoteTransactionId RequestId);
+	COREUOBJECT_API void EndRemoteMultiServerCommit(FRemoteServerId ServerId, FRemoteTransactionId RequestId);
+	COREUOBJECT_API void EnqueueRemoteMultiServerCommitAction(FRemoteServerId ServerId, FRemoteTransactionId RequestId, const TFunction<void()>& Action);
+
+	COREUOBJECT_API void ReadyMultiServerCommitResponse(FRemoteServerId ServerId, FRemoteTransactionId RequestId);
+	COREUOBJECT_API void AbortMultiServerCommit(FRemoteServerId ServerId, FRemoteTransactionId RequestId);
 	
 	/** Establishes a synchronization barrier across all servers */
 	COREUOBJECT_API void GlobalServerSync(FName SyncName);
diff --git a/Engine/Source/Runtime/CoreUObject/Public/UObject/RemoteObjectTransfer.h b/Engine/Source/Runtime/CoreUObject/Public/UObject/RemoteObjectTransfer.h
--- a/Engine/Source/Runtime/CoreUObject/Public/UObject/RemoteObjectTransfer.h
+++ b/Engine/Source/Runtime/CoreUObject/Public/UObject/RemoteObjectTransfer.h
@@ -48,6 +48,9 @@ namespace UE::RemoteObject::Transfer
 
 		/** The serialized data of the object being sent */
 		FRemoteObjectData ObjectData;
+
+		/** Whether this migration is part of a multi-server commit */
+		bool bForMultiServerCommit = false;
 	};
 
 	constexpr FRemoteTransactionId PushMigrationRequestId = FRemoteTransactionId(0xFFFFFFFFu);
@@ -134,7 +137,7 @@ namespace UE::RemoteObject::Transfer
 	COREUOBJECT_API void InitRemoteObjectTransfer();
 
 	/** Delegate that transfers object data to another server */
-	extern COREUOBJECT_API TDelegate<void(const FMigrateSendParams& Params)> RemoteObjectTransferDelegate;
+	extern COREUOBJECT_API TDelegate<void(const FMigrateSendParams& /*Params*/)> RemoteObjectTransferDelegate;
 
 	/** Delegate that handles an object request being denied */
 	extern COREUOBJECT_API TDelegate<void(FRemoteTransactionId /*RequestId*/, FRemoteObjectId /*ObjectId*/, FRemoteServerId /*DestinationServerId*/)> RemoteObjectDeniedTransferDelegate;
diff --git a/Engine/Source/Runtime/CoreUObject/Public/UObject/WeakObjectPtr.h b/Engine/Source/Runtime/CoreUObject/Public/UObject/WeakObjectPtr.h
--- a/Engine/Source/Runtime/CoreUObject/Public/UObject/WeakObjectPtr.h
+++ b/Engine/Source/Runtime/CoreUObject/Public/UObject/WeakObjectPtr.h
@@ -377,6 +377,13 @@ protected:
 private:
 	FORCEINLINE int32 GetObjectIndex_Private() const
 	{
+#if UE_WITH_REMOTE_OBJECT_HANDLE
+		using namespace UE::Core::Private;
+		if ((ObjectIndex == InvalidWeakObjectIndex) && ObjectRemoteId.IsValid())
+		{
+			TryResolveRemoteObject();
+		}
+#endif
 		return ObjectIndex;
 	}
 
diff --git a/Engine/Source/Runtime/Engine/Classes/Engine/NetDriver.h b/Engine/Source/Runtime/Engine/Classes/Engine/NetDriver.h
--- a/Engine/Source/Runtime/Engine/Classes/Engine/NetDriver.h
+++ b/Engine/Source/Runtime/Engine/Classes/Engine/NetDriver.h
@@ -2342,6 +2342,16 @@ public:
 		return bReplicateTransactionally;
 	}
 
+	void SetDeferTransactionalSend(bool bInDeferTransactionalSend)
+	{
+		bDeferTransactionalSend = bInDeferTransactionalSend;
+	}
+
+	bool ShouldDeferTransactionalSend() const
+	{
+		return bDeferTransactionalSend;
+	}
+
 	/** Returns true if references to UObjects should be serialized as FRemoteObjectReference */
 	bool IsUsingRemoteObjectReferences() const { return bUseRemoteObjectRefs; }
 
@@ -2507,6 +2517,9 @@ private:
 
 	/** When using remote objects with UE_WITH_REMOTE_OBJECT_HANDLE, determines whether or not references to UObjects (either from TObjectPtr or TWeakObjectPtr) are serialized as FRemoteObjectReference */
 	bool bUseRemoteObjectRefs = false;
+	
+	/** When using remote objects with UE_WITH_REMOTE_OBJECT_HANDLE, determines whether or not we defer the actual network send until the transaction commits */
+	bool bDeferTransactionalSend = true;
 
 	/** Server support for NMT_JoinNoPawn and NMT_JoinNoPawnSplit. */
 	bool bSupportNoPawnConnection = false;
diff --git a/Engine/Source/Runtime/Engine/Classes/GameFramework/Actor.h b/Engine/Source/Runtime/Engine/Classes/GameFramework/Actor.h
--- a/Engine/Source/Runtime/Engine/Classes/GameFramework/Actor.h
+++ b/Engine/Source/Runtime/Engine/Classes/GameFramework/Actor.h
@@ -2336,6 +2336,7 @@ public:
 	ENGINE_API virtual void PostLoadSubobjects( FObjectInstancingGraph* OuterInstanceGraph ) override;
 #if UE_WITH_REMOTE_OBJECT_HANDLE
 	ENGINE_API virtual void PostMigrate(const struct FUObjectMigrationContext& MigrationContext) override;
+	ENGINE_API virtual bool IsMigrationRoot() const override;
 #endif
 	ENGINE_API virtual void BeginDestroy() override;
 #if !UE_STRIP_DEPRECATED_PROPERTIES
diff --git a/Engine/Source/Runtime/Engine/Private/Actor.cpp b/Engine/Source/Runtime/Engine/Private/Actor.cpp
--- a/Engine/Source/Runtime/Engine/Private/Actor.cpp
+++ b/Engine/Source/Runtime/Engine/Private/Actor.cpp
@@ -1393,6 +1393,12 @@ void AActor::PostMigrate(const FUObjectMigrationContext& MigrationContext)
 		}
 	}
 }
+
+bool AActor::IsMigrationRoot() const
+{
+	return true;
+}
+
 #endif // UE_WITH_REMOTE_OBJECT_HANDLE
 
 void AActor::ProcessEvent(UFunction* Function, void* Parameters)
