resolving some static analysis warnings

diff --git a/Engine/Source/Runtime/CoreUObject/Private/UObject/RemoteExecutor.cpp b/Engine/Source/Runtime/CoreUObject/Private/UObject/RemoteExecutor.cpp
--- a/Engine/Source/Runtime/CoreUObject/Private/UObject/RemoteExecutor.cpp
+++ b/Engine/Source/Runtime/CoreUObject/Private/UObject/RemoteExecutor.cpp
@@ -185,12 +185,13 @@ public:
 			{
 				check(ExecutingWork == nullptr);
 				TGuardValue<FRemoteExecutorWork*> ExecutingWorkSave(ExecutingWork, &PendingWorks[PendingWorkIndex]);
+				check(ExecutingWork);
 
 				UE_LOG(LogRemoteExec, VeryVerbose, TEXT("ExecutePendingWork[%d,%s] : Executing request %s %s --"),
 					LocalIterationNumber,
 					*ExecutingWork->RequestId.ToString(),
 					*ExecutingWork->Name.ToString(),
-					*ExecutingWork->Priority.ToString());
+					*ExecutingWork->Priority.ToString()); // -V522
 
 				if (ExecutingWork->bIsTransactional)
 				{
@@ -487,14 +488,14 @@ public:
 
 						UE::RemoteExecutor::OnTransactionAbortedDelegate.Broadcast(RequestId);
 
-						if (AbortReason == ERemoteExecutorAbortReason::RequiresDependencies)
+						if (AbortReason == ERemoteExecutorAbortReason::RequiresDependencies) // -V547
 						{
 							UE_LOG(LogRemoteExec, Verbose, TEXT("ExecutePendingWork[%d,%s] : Work ABORTED for required dependency '%s'"),
 								LocalIterationNumber,
 								*ExecutingWork->RequestId.ToString(),
 								*AbortReasonDescription);
 						}
-						else if (AbortReason == ERemoteExecutorAbortReason::AbandonWork)
+						else if (AbortReason == ERemoteExecutorAbortReason::AbandonWork) // -V547
 						{
 							UE_LOG(LogRemoteExec, Verbose, TEXT("ExecutePendingWork[%d,%s] : Work ABORTED for ABANDONMENT '%s'"),
 								LocalIterationNumber,
