Safeguard against orphan SavedPins crashing when opening the view.

#rnx
#tests Verified manually the metadata after renaming WBP, also checked with moving the WBP to another folder

diff --git a/Engine/Plugins/Runtime/ModelViewViewModel/Source/ModelViewViewModelEditor/Private/Widgets/SMVVMViewBindingListView.cpp b/Engine/Plugins/Runtime/ModelViewViewModel/Source/ModelViewViewModelEditor/Private/Widgets/SMVVMViewBindingListView.cpp
--- a/Engine/Plugins/Runtime/ModelViewViewModel/Source/ModelViewViewModelEditor/Private/Widgets/SMVVMViewBindingListView.cpp
+++ b/Engine/Plugins/Runtime/ModelViewViewModel/Source/ModelViewViewModelEditor/Private/Widgets/SMVVMViewBindingListView.cpp
@@ -686,6 +686,11 @@ TSharedRef<ITableRow> SBindingsList::GenerateEntryRow(TSharedPtr<FBindingEntry>
 			case FBindingEntry::ERowType::EventParameter:
 			case FBindingEntry::ERowType::ConditionParameter:
 			{
+				if (!Entry->GetBindingParameterId().IsValid())
+				{
+					ensureMsgf(false, TEXT("Corrupted Binding Parameter."));
+					return SNew(STableRow<TSharedPtr<FBindingEntry>>, OwnerTable);
+				}
 				Row = SNew(UE::MVVM::BindingEntry::SFunctionParameterRow, OwnerTable, WeakBlueprintEditor.Pin(), MVVMExtensionPtr->GetWidgetBlueprint(), Entry);
 				break;
 			}
diff --git a/Engine/Source/Runtime/CoreUObject/Private/UObject/MetaData.cpp b/Engine/Source/Runtime/CoreUObject/Private/UObject/MetaData.cpp
--- a/Engine/Source/Runtime/CoreUObject/Private/UObject/MetaData.cpp
+++ b/Engine/Source/Runtime/CoreUObject/Private/UObject/MetaData.cpp
@@ -70,6 +70,7 @@ FMetaDataUtilities::FMoveMetadataHelperContext::FMoveMetadataHelperContext(UObje
 		UPackage* Package = SourceObject->GetOutermost();
 		check(Package);
 		OldPackage = Package;
+		OldObjectPath = FSoftObjectPath::ConstructFromObject(SourceObject);
 		OldObject = SourceObject;
 		bShouldSearchChildren = bSearchChildren;
 	}
@@ -82,7 +83,7 @@ FMetaDataUtilities::FMoveMetadataHelperContext::~FMoveMetadataHelperContext()
 	{
 		FMetaData& NewMetaData = OldObject->GetPackage()->GetMetaData();
 		FMetaData& OldMetaData = OldPackage->GetMetaData();
-		TMap< FName, FString>* OldObjectMetaData = OldMetaData.ObjectMetaDataMap.Find(OldObject);
+		TMap< FName, FString>* OldObjectMetaData = OldMetaData.ObjectMetaDataMap.Find(OldObjectPath);
 		if (OldObjectMetaData != NULL)
 		{
 			NewMetaData.SetObjectValues(OldObject, *OldObjectMetaData);
@@ -96,7 +97,9 @@ FMetaDataUtilities::FMoveMetadataHelperContext::~FMoveMetadataHelperContext()
 			for ( auto ChildIt = Children.CreateConstIterator(); ChildIt; ++ChildIt )
 			{
 				UObject* Child = *ChildIt;
-				TMap<FName, FString>* OldObjectMetaDataPtr = OldMetaData.ObjectMetaDataMap.Find(Child);
+				FSoftObjectPath ChildPath = FSoftObjectPath::ConstructFromObject(Child);
+				ChildPath.SetPath(OldObjectPath.GetAssetPath(), ChildPath.GetSubPathUtf8String());
+				TMap<FName, FString>* OldObjectMetaDataPtr = OldMetaData.ObjectMetaDataMap.Find(ChildPath);
 				if ( OldObjectMetaDataPtr )
 				{
 					NewMetaData.SetObjectValues(Child, *OldObjectMetaDataPtr);
diff --git a/Engine/Source/Runtime/CoreUObject/Public/UObject/MetaData.h b/Engine/Source/Runtime/CoreUObject/Public/UObject/MetaData.h
--- a/Engine/Source/Runtime/CoreUObject/Public/UObject/MetaData.h
+++ b/Engine/Source/Runtime/CoreUObject/Public/UObject/MetaData.h
@@ -253,6 +253,9 @@ private:
 		/** Cache a pointer to the object so we can do the search on the old metadata. */
 		UObject* OldObject;
 
+		/** Cache the old path so we can compare to the metadata key */
+		FSoftObjectPath OldObjectPath;
+
 		/** When true, search children as well. */
 		bool bShouldSearchChildren;
 	};
