Removed StaticFindObject overload for VersePaths.

There is presently no efficient or reliable mapping between Verse paths and top level object paths since Verse paths omit the short package name.

diff --git a/Engine/Source/Runtime/CoreUObject/Private/UObject/UObjectGlobals.cpp b/Engine/Source/Runtime/CoreUObject/Private/UObject/UObjectGlobals.cpp
--- a/Engine/Source/Runtime/CoreUObject/Private/UObject/UObjectGlobals.cpp
+++ b/Engine/Source/Runtime/CoreUObject/Private/UObject/UObjectGlobals.cpp
@@ -84,7 +84,6 @@
 #include "Templates/GuardValueAccessors.h"
 #include "UObject/UObjectGlobalsInternal.h"
 #include "Serialization/AsyncPackageLoader.h"
-#include "Containers/VersePath.h"
 #include "AutoRTFM/AutoRTFM.h"
 #include "UObject/PropertyOptional.h"
 #include "UObject/ObjectHandlePrivate.h"
@@ -98,9 +97,6 @@
 #include "UObject/VerseValueProperty.h"
 #endif
 
-#include "Interfaces/IPluginManager.h"
-
-
 DEFINE_LOG_CATEGORY(LogUObjectGlobals);
 
 int32 GAllowUnversionedContentInEditor = 0;
@@ -633,32 +629,6 @@ UObject* StaticFindObjectSafe(UClass* Class, FTopLevelAssetPath ObjectPath, bool
 	}
 }
 
-FORCENOINLINE UObject* StaticFindObject(UClass* Class, const UE::Core::FVersePath& VersePath)
-{
-	IPluginManager& PluginManager = IPluginManager::Get();
-
-	FName   PackageName;
-	FString LeafPath;
-	if (!PluginManager.TrySplitVersePath(VersePath, PackageName, LeafPath))
-	{
-		return nullptr;
-	}
-
-	// Make sure there's a trailing leafname - assume that's the root level object
-	int32 LastSlashIndex;
-	if (!LeafPath.FindLastChar(TEXT('/'), LastSlashIndex))
-	{
-		return nullptr;
-	}
-
-	LeafPath[LastSlashIndex] = TEXT('.');
-
-	FString ObjectPath = FString::Printf(TEXT("/%s%s"), *PackageName.ToString(), *LeafPath);
-
-	UObject* Result = StaticFindObject(Class, nullptr, *ObjectPath);
-	return Result;
-}
-
 //
 // Find an object; can't fail.
 //
diff --git a/Engine/Source/Runtime/CoreUObject/Public/UObject/UObjectGlobals.h b/Engine/Source/Runtime/CoreUObject/Public/UObject/UObjectGlobals.h
--- a/Engine/Source/Runtime/CoreUObject/Public/UObject/UObjectGlobals.h
+++ b/Engine/Source/Runtime/CoreUObject/Public/UObject/UObjectGlobals.h
@@ -16,7 +16,6 @@
 #include "Containers/StringFwd.h"
 #include "Containers/StringView.h"
 #include "Containers/UnrealString.h"
-#include "Containers/VersePathFwd.h"
 #include "CoreGlobals.h"
 #include "CoreTypes.h"
 #include "Delegates/Delegate.h"
@@ -391,17 +390,6 @@ COREUOBJECT_API UObject* StaticFindObjectChecked( UClass* Class, UObject* InOute
 /** Internal version of StaticFindObject that will not assert on GIsSavingPackage or IsGarbageCollectingAndLockingUObjectHashTables() */
 COREUOBJECT_API UObject* StaticFindObjectSafe( UClass* Class, UObject* InOuter, const TCHAR* Name, bool ExactClass=false );
 
-/**
- * Tries to find an object in memory, using a Verse path.
- *
- * @param	VersePath		The path to the object to find.
- * @param	ObjectPath		FName pair representing the outer package object and the inner top level object (asset)
- * @param	ExactClass		Whether to require an exact match with the passed in class
- *
- * @return	Returns a pointer to the found object or nullptr if none could be found
- */
-COREUOBJECT_API UObject* StaticFindObject(UClass* Class, const UE::Core::FVersePath& VersePath);
-
 /**
  * Tries to find an object in memory. This version uses FTopLevelAssetPath to find the object.
  * Version of StaticFindObject that will not assert on GIsSavingPackage or IsGarbageCollectingAndLockingUObjectHashTables()
