RemoteObjects:  Fixing an issue where once an object is Garbage Collected (Stored to Database), it can no longer be fetched from the Remote Server that's requesting it.  We now restore it locally and then immediately send it to the requesting Server.

diff --git a/Engine/Source/Runtime/CoreUObject/Private/UObject/RemoteObjectSerialization.cpp b/Engine/Source/Runtime/CoreUObject/Private/UObject/RemoteObjectSerialization.cpp
--- a/Engine/Source/Runtime/CoreUObject/Private/UObject/RemoteObjectSerialization.cpp
+++ b/Engine/Source/Runtime/CoreUObject/Private/UObject/RemoteObjectSerialization.cpp
@@ -72,9 +72,9 @@ void DeserializeNameTables(FArchive& Ar, FRemoteObjectData& OutObjectData)
 	Ar << OutObjectData.PathNames;
 }
 
-FString GenerateRemoteObjectFilename(FRemoteObjectId ObjectId, FRemoteServerId DestinationServerId)
+FString GenerateRemoteObjectFilename(FRemoteObjectId ObjectId, FRemoteServerId OwnerServerId)
 {
-	return FPaths::Combine(*FPaths::ProjectSavedDir(), *(FString::Printf(TEXT("%s_%s.remote"), *ObjectId.ToString(), *DestinationServerId.ToString())));
+	return FPaths::Combine(*FPaths::ProjectSavedDir(), *(FString::Printf(TEXT("%s-%s_%s.remote"), *UE::RemoteObject::GetGlobalServerId().ToString(), *ObjectId.ToString(), *OwnerServerId.ToString())));
 }
 
 void LoadObjectFromDisk(const FUObjectMigrationContext& MigrationContext)
diff --git a/Engine/Source/Runtime/CoreUObject/Private/UObject/RemoteObjectTransfer.cpp b/Engine/Source/Runtime/CoreUObject/Private/UObject/RemoteObjectTransfer.cpp
--- a/Engine/Source/Runtime/CoreUObject/Private/UObject/RemoteObjectTransfer.cpp
+++ b/Engine/Source/Runtime/CoreUObject/Private/UObject/RemoteObjectTransfer.cpp
@@ -360,9 +360,23 @@ public:
 		}
 		else
 		{
-			// We are intercepting a request to pull-migrate since we believe the object resides in our database, migrate it "from" the database
-			FUObjectMigrationContext MigrationContext { .ObjectId = ObjectId, .RemoteServerId = DatabaseId, .OwnerServerId = DatabaseId, .PhysicsServerId = DatabaseId, .RequestId = RequestId, .MigrationSide = EObjectMigrationSide::Receive };
+			UE_LOG(LogRemoteObject, Log, TEXT("%hs: Restoring ObjectId %s from Database and Forwarding to %s"), __func__, *ObjectId.ToString(), *DestinationServerId.ToString());
+			// We are intercepting a request to pull-migrate since we believe the object resides in our database, migrate it "from" the database to us
+			// Inside the database call, it will fix-up these migration context parameters to make us the owner
+			FUObjectMigrationContext MigrationContext { .ObjectId = ObjectId, .RemoteServerId = DatabaseId, .OwnerServerId = DatabaseId, .PhysicsServerId = DatabaseId, .RequestId = RequestId, .MigrationSide = EObjectMigrationSide::Receive};
 			RestoreRemoteObjectDataDelegate.Execute(MigrationContext);
+
+#if UE_WITH_REMOTE_OBJECT_HANDLE
+			// We assume we have it restored, now forward it
+			TWeakObjectPtr WeakObjPtr { ObjectId };
+			ensureMsgf(!WeakObjPtr.IsExplicitlyNull() && WeakObjPtr.IsValid(true) && !WeakObjPtr.IsRemote(), TEXT("Expected RemoteObjectId %s to be restored from Database and thus locally owned"), *ObjectId.ToString());
+#endif
+
+			// If we didn't expect it locally, send it on to the final destination
+			if (DestinationServerId != GetGlobalServerId())
+			{
+				MigrateObjectToRemoteServer(RequestId, ObjectId, DestinationServerId);
+			}
 		}
 	}
 
@@ -972,6 +986,8 @@ void StoreObjectToDatabase(UObject* Object)
 			return;
 		}
 
+		// This ensure should never fire.  We shouldn't be able to store a non-locally-owned Object into the Database (that is the job of the Owning Server).  A scenario where I've seen this:
+		// We borrowed an Object that was Marked As Garbage and didn't return it to the Owning Server before a GC.
 		ensureMsgf(false, TEXT("%hs was about to store non-locally-owned Object %s to Server %s's Database; making it remote instead"), __func__, *GetNameSafe(Object), *LocalServerId.ToString());
 		MarkAsRemote(Object, OwnerServerId);
 		return;
diff --git a/Engine/Source/Runtime/Engine/Private/Actor.cpp b/Engine/Source/Runtime/Engine/Private/Actor.cpp
--- a/Engine/Source/Runtime/Engine/Private/Actor.cpp
+++ b/Engine/Source/Runtime/Engine/Private/Actor.cpp
@@ -1391,6 +1391,14 @@ void AActor::PostMigrate(const FUObjectMigrationContext& MigrationContext)
 			bool bTickFunctionsRegisteredOnOwningServer = bTickFunctionsRegistered;
 			RegisterAllActorTickFunctions(false, false);
 			bTickFunctionsRegistered = bTickFunctionsRegisteredOnOwningServer;
+
+			// We will always add ourselves to the list when borrowed, even if destroyed
+			// The reason is that we never want to GC it, and always want to return it during a sync point
+			if (bActorIsBeingDestroyed)
+			{
+				constexpr bool bOnlyAddIfUnique = true;
+				GetLevel()->TryAddActorToList(this, bOnlyAddIfUnique);
+			}
 		}
 	}
 }
