[AsyncPackageLoader]

- Restore zenloader as the default loader for loose cooked file loading
- Configures EngineTest to use loose cooked data rather than ZenStore when packaging. ZenStore isn't a default option for users, so this allows us to validate the defaults users see.

#rnx

diff --git a/Engine/Source/Runtime/CoreUObject/Private/Serialization/AsyncPackageLoader.cpp b/Engine/Source/Runtime/CoreUObject/Private/Serialization/AsyncPackageLoader.cpp
--- a/Engine/Source/Runtime/CoreUObject/Private/Serialization/AsyncPackageLoader.cpp
+++ b/Engine/Source/Runtime/CoreUObject/Private/Serialization/AsyncPackageLoader.cpp
@@ -17,7 +17,7 @@
 #include "HAL/IConsoleManager.h"
 #include "AutoRTFM/AutoRTFM.h"
 
-static bool GUseOldLoaderAsFallback = true;
+static bool GUseOldLoaderAsFallback = false;
 static FAutoConsoleVariableRef CVarGUseOldLoaderAsFallback(
 	TEXT("s.UseOldLoaderAsFallback"),
 	GUseOldLoaderAsFallback,
