Avoid redundant time checks in async loading game thread code, this can speed up flushed loads by 10%

diff --git a/Engine/Source/Runtime/CoreUObject/Private/Serialization/AsyncLoading2.cpp b/Engine/Source/Runtime/CoreUObject/Private/Serialization/AsyncLoading2.cpp
--- a/Engine/Source/Runtime/CoreUObject/Private/Serialization/AsyncLoading2.cpp
+++ b/Engine/Source/Runtime/CoreUObject/Private/Serialization/AsyncLoading2.cpp
@@ -3117,27 +3117,32 @@ struct FAsyncLoadingThreadState2
 		bool bUseTimeLimit = false;
 		double TimeLimit = 0.0;
 		double StartTime = 0.0;
+		FAsyncLoadingThreadState2& ThreadState;
 
-		FTimeLimitScope(bool bInUseTimeLimit, double InTimeLimit)
+		FTimeLimitScope(FAsyncLoadingThreadState2& InThreadState, bool bInUseTimeLimit, double InTimeLimit)
+			: ThreadState(InThreadState)
 		{
-			FAsyncLoadingThreadState2* ThreadState = FAsyncLoadingThreadState2::Get();
+			bUseTimeLimit = ThreadState.bUseTimeLimit;
+			TimeLimit = ThreadState.TimeLimit;
+			StartTime = ThreadState.StartTime;
 
-			bUseTimeLimit = ThreadState->bUseTimeLimit;
-			TimeLimit = ThreadState->TimeLimit;
-			StartTime = ThreadState->StartTime;
-
-			ThreadState->bUseTimeLimit = bInUseTimeLimit;
-			ThreadState->TimeLimit = InTimeLimit;
-			ThreadState->StartTime = FPlatformTime::Seconds();
+			ThreadState.bUseTimeLimit = bInUseTimeLimit;
+			ThreadState.TimeLimit = InTimeLimit;
+			if (bInUseTimeLimit)
+			{
+				ThreadState.StartTime = FPlatformTime::Seconds();
+			}
+			else
+			{
+				ThreadState.StartTime = 0.0;
+			}
 		}
 		
 		~FTimeLimitScope()
 		{
-			FAsyncLoadingThreadState2* ThreadState = FAsyncLoadingThreadState2::Get();
-
-			ThreadState->bUseTimeLimit = bUseTimeLimit;
-			ThreadState->TimeLimit = TimeLimit;
-			ThreadState->StartTime = StartTime;
+			ThreadState.bUseTimeLimit = bUseTimeLimit;
+			ThreadState.TimeLimit = TimeLimit;
+			ThreadState.StartTime = StartTime;
 		}
 	};
 
@@ -9447,10 +9452,9 @@ EAsyncPackageState::Type FAsyncLoadingThread2::TickAsyncLoadingFromGameThread(FA
 	{
 		// Use a time limit scope to restore the time limit to the old values when exiting the scope
 		// This is required to ensure a reentrant call here doesn't overwrite time limits permanently.
-		FAsyncLoadingThreadState2::FTimeLimitScope TimeLimitScope(bUseTimeLimit, TimeLimit);
+		FAsyncLoadingThreadState2::FTimeLimitScope TimeLimitScope(ThreadState, bUseTimeLimit, TimeLimit);
 
 		const bool bIsMultithreaded = FAsyncLoadingThread2::IsMultithreaded();
-		double TickStartTime = FPlatformTime::Seconds();
 
 		if (!bIsMultithreaded)
 		{
@@ -9458,10 +9462,11 @@ EAsyncPackageState::Type FAsyncLoadingThread2::TickAsyncLoadingFromGameThread(FA
 		}
 		UpdateSyncLoadContext(ThreadState);
 
+		Result = ProcessLoadedPackagesFromGameThread(ThreadState, bDidSomething, FlushRequestIDs);
+		if (bUseTimeLimit && !GIsEditor)
 		{
-			Result = ProcessLoadedPackagesFromGameThread(ThreadState, bDidSomething, FlushRequestIDs);
-			double TimeLimitUsedForProcessLoaded = FPlatformTime::Seconds() - TickStartTime;
-			UE_CLOG(!GIsEditor && bUseTimeLimit && TimeLimitUsedForProcessLoaded > .1f, LogStreaming, Warning, TEXT("Took %6.2fms to ProcessLoadedPackages"), float(TimeLimitUsedForProcessLoaded) * 1000.0f);
+			double TimeLimitUsedForProcessLoaded = FPlatformTime::Seconds() - ThreadState.StartTime;
+			UE_CLOG(TimeLimitUsedForProcessLoaded > 0.1, LogStreaming, Warning, TEXT("Took %6.2fms to ProcessLoadedPackages"), float(TimeLimitUsedForProcessLoaded) * 1000.0f);
 		}
 
 		if (!bIsMultithreaded && Result != EAsyncPackageState::TimeOut)
