RemoteObjects:  FRemoteObjectId now reuses FRemoteServerId for operations (such as ToString)

#rnx

diff --git a/Engine/Source/Runtime/CoreUObject/Private/UObject/RemoteObject.cpp b/Engine/Source/Runtime/CoreUObject/Private/UObject/RemoteObject.cpp
--- a/Engine/Source/Runtime/CoreUObject/Private/UObject/RemoteObject.cpp
+++ b/Engine/Source/Runtime/CoreUObject/Private/UObject/RemoteObject.cpp
@@ -641,14 +641,7 @@ FRemoteObjectId FRemoteObjectId::Generate(UObjectBase* InObject, EInternalObject
 
 FString FRemoteObjectId::ToString() const
 {
-	if (!IsAsset())
-	{
-		return FString::Printf(TEXT("%llu-%llu"), ServerId, SerialNumber);
-	}
-	else
-	{
-		return FString::Printf(TEXT("Asset-%llu"), SerialNumber);
-	}
+	return FString::Printf(TEXT("%s-%llu"), *GetServerId().ToString(), SerialNumber);
 }
 
 FRemoteObjectId::FRemoteObjectId(const UObjectBase* Object)
@@ -688,6 +681,21 @@ FRemoteServerId::FRemoteServerId(const FString& InText)
 	}
 }
 
+FString FRemoteServerId::ToString() const
+{
+	switch (Id)
+	{
+		case (uint32)ERemoteServerIdConstants::Asset:
+			return TEXT("Asset");
+
+		case (uint32)ERemoteServerIdConstants::Database:
+			return TEXT("Database");
+
+		default:
+			return FString::FromInt(Id);
+	}
+}
+
 FArchive& operator<<(FArchive& Ar, FRemoteServerId& Id)
 {
 	Ar << Id.Id;
diff --git a/Engine/Source/Runtime/CoreUObject/Public/UObject/RemoteObjectTypes.h b/Engine/Source/Runtime/CoreUObject/Public/UObject/RemoteObjectTypes.h
--- a/Engine/Source/Runtime/CoreUObject/Public/UObject/RemoteObjectTypes.h
+++ b/Engine/Source/Runtime/CoreUObject/Public/UObject/RemoteObjectTypes.h
@@ -42,15 +42,12 @@ struct FRemoteServerId
 		checkf(Id < (uint32)ERemoteServerIdConstants::FirstReserved, TEXT("Remote server id can not be greater than %u, got: %u"), (uint32)ERemoteServerIdConstants::FirstReserved - 1, Id);
 	}
 	explicit COREUOBJECT_API FRemoteServerId(const FString& InText);
+	COREUOBJECT_API FString ToString() const;
 
 	uint32 GetIdNumber() const
 	{
 		return Id;
 	}
-	FString ToString() const
-	{
-		return FString::FromInt(Id);
-	}
 	bool IsValid() const
 	{
 		return Id != 0;
@@ -75,6 +72,14 @@ struct FRemoteServerId
 	{
 		Id = Other.Id;
 	}
+	bool IsAsset() const
+	{
+		return (Id == (uint32)ERemoteServerIdConstants::Asset);
+	}
+	bool IsDatabase() const
+	{
+		return (Id == (uint32)ERemoteServerIdConstants::Database);
+	}
 
 	friend COREUOBJECT_API FArchive& operator<<(FArchive& Ar, FRemoteServerId& Id);
 
@@ -162,7 +167,7 @@ public:
 
 	FORCEINLINE bool IsAsset() const
 	{
-		return GetServerId() == FRemoteServerId(ERemoteServerIdConstants::Asset);
+		return GetServerId().IsAsset();
 	}
 
 	COREUOBJECT_API FString ToString() const;
