Added borrowed objects to the ignore list for physics queries to enable borrow mutations.
Updated the test level to have some borrow mutate tests.
Also fixed an issue where the enlarged aabb wasn't being used for determining which servers to query.

#virtualized

diff --git a/Engine/Source/Runtime/CoreUObject/Private/UObject/RemoteObjectTransfer.cpp b/Engine/Source/Runtime/CoreUObject/Private/UObject/RemoteObjectTransfer.cpp
--- a/Engine/Source/Runtime/CoreUObject/Private/UObject/RemoteObjectTransfer.cpp
+++ b/Engine/Source/Runtime/CoreUObject/Private/UObject/RemoteObjectTransfer.cpp
@@ -836,6 +836,19 @@ void RegisterRemoteObjectId(FRemoteObjectId Id, FRemoteServerId ResidentServerId
 	UE::RemoteObject::Private::RegisterRemoteObjectId(Id, ResidentServerId);
 }
 
+void GetAllBorrowedObjects(TArray<FRemoteObjectId>& OutBorrowedObjectIds)
+{
+	check(GTransferQueue.ActiveRequest);
+
+	for (const FRemoteObjectId& Id : GTransferQueue.ActiveRequest->UsedObjects)
+	{
+		if (!UE::RemoteObject::Handle::IsOwned(Id))
+		{
+			OutBorrowedObjectIds.Add(Id);
+		}
+	}
+}
+
 void MigrateObjectToRemoteServer(FRemoteTransactionId RequestId, FRemoteObjectId ObjectId, FRemoteServerId DestinationServerId)
 {
 	FRemoteWorkPriority RootWorkPriority = UE::RemoteExecutor::CreateRootWorkPriority();
diff --git a/Engine/Source/Runtime/CoreUObject/Public/UObject/RemoteObjectTransfer.h b/Engine/Source/Runtime/CoreUObject/Public/UObject/RemoteObjectTransfer.h
--- a/Engine/Source/Runtime/CoreUObject/Public/UObject/RemoteObjectTransfer.h
+++ b/Engine/Source/Runtime/CoreUObject/Public/UObject/RemoteObjectTransfer.h
@@ -123,6 +123,11 @@ namespace UE::RemoteObject::Transfer
 	*/
 	COREUOBJECT_API void RegisterRemoteObjectId(FRemoteObjectId Id, FRemoteServerId ResidentServerId);
 
+	/**
+	* Returns the list of all object IDs currently borrowed from another server.
+	*/
+	COREUOBJECT_API void GetAllBorrowedObjects(TArray<FRemoteObjectId>& OutBorrowedObjectIds);
+
 	/**
 	* Registers object for sharing, marking it as owned by the current server
 	* @param Object to make shareable
