Fixed crash when deleting a verse component when added to prefabs.
Also fixed a verbose log of entity unable to remove from parent

#rnx

diff --git a/Engine/Source/Runtime/CoreUObject/Private/UObject/OverridableManager.cpp b/Engine/Source/Runtime/CoreUObject/Private/UObject/OverridableManager.cpp
--- a/Engine/Source/Runtime/CoreUObject/Private/UObject/OverridableManager.cpp
+++ b/Engine/Source/Runtime/CoreUObject/Private/UObject/OverridableManager.cpp
@@ -411,7 +411,7 @@ void FOverridableManager::AddReferencedObjects(FReferenceCollector& Collector)
 #if WITH_EDITORONLY_DATA
 	// This isn't great but this object implements this FGCObject interface just to so 
 	// that the replace references archives go through this object
-	if (!IsGarbageCollecting())
+	if (GIsReinstancing)
 	{
 		const TMap<const UObjectBase *, FOverriddenPropertyAnnotation>& AnnotationMap = OverriddenObjectAnnotations.GetAnnotationMap();
 		for (const auto& Pair : AnnotationMap)
