EditorOptional

diff --git a/Engine/Source/Runtime/CoreUObject/Private/UObject/ObjectEditorOptionalSupport.cpp b/Engine/Source/Runtime/CoreUObject/Private/UObject/ObjectEditorOptionalSupport.cpp
new file mode 100644
--- /dev/null
+++ b/Engine/Source/Runtime/CoreUObject/Private/UObject/ObjectEditorOptionalSupport.cpp
@@ -0,0 +1,77 @@
+// Copyright Epic Games, Inc. All Rights Reserved.
+
+#include "UObject/ObjectEditorOptionalSupport.h"
+#include "CoreMinimal.h"
+#include "UObject/UObjectThreadContext.h"
+#include "Serialization/CustomVersion.h"
+#include "UObject/UObjectSerializeContext.h"
+#include "UObject/Class.h"
+
+
+#if WITH_EDITORONLY_DATA
+
+namespace UE::EditorOptional
+{
+
+void ConditionalUpgradeObject(FArchive& Ar, TNotNull<UObject*> SecondaryObject, const FGuid& VersionGuid, int Version)
+{
+	// attempt to fixup old content from before the class was split between runtime and EditorOptional classes
+	// we use custom version to determine that the data is old, or always upgrade if a custom version is not
+	// given to check against
+	if (Ar.IsLoading() && (!VersionGuid.IsValid() || Ar.CustomVer(VersionGuid) < Version))
+	{
+		int64 ScriptStartOffset = -1;
+		int64 ScriptEndOffset = -1;
+		int64 CurrentOffset = Ar.Tell();
+		UObject* MainObject = nullptr; 
+		{
+			// grab the threads serialization context, which we use to find where the current runtime object's
+			// ScriptProperties begin, so we can load those same set of properties in the EO object
+			// ScriptProperties are dynamically looked up, so the moved properties would have already been
+			// skipped over before, and then we load the same set of properties again, this time loading the
+			// ones that moved into the EO object
+			// doing this in a scope since we are going to make a new context, just for extra safety
+			FUObjectSerializeContext* SerializeContext = FUObjectThreadContext::Get().GetSerializeContext();
+			bool bIsValid = SerializeContext != nullptr && SerializeContext->SerializedObjectScriptStartOffset != -1; 
+			checkf(bIsValid, TEXT("Unable to retrieve a valid ScriptStartOffset from the active serialization context. Make sure to call UE::EditorOptional::ConditionalUpgradeObject() from an object's Serialize function _after_ calling Super::Serialize()."));
+			
+			ScriptStartOffset = SerializeContext->SerializedObjectScriptStartOffset;
+			ScriptEndOffset = SerializeContext->SerializedObjectScriptEndOffset;
+			MainObject = SerializeContext->SerializedObject;
+		}
+		
+		// make a new context so when we leave the scope, state is cleaned up
+		FScopedObjectSerializeContext RewindContext(SecondaryObject, Ar);
+		
+		// rewind and load
+		Ar.Seek(ScriptStartOffset);
+		SecondaryObject->SerializeScriptProperties(Ar);
+		
+		// verify it ended up where we expected
+		checkf(Ar.Tell() == ScriptEndOffset, TEXT("Unexpected offset in file after loading ScriptProperties into the EditorOptional object"));
+		
+		Ar.Seek(CurrentOffset);
+	}
+}
+
+void UpgradeObject(FArchive& Ar, TNotNull<UObject*> SecondaryObject)
+{
+	ConditionalUpgradeObject(Ar, SecondaryObject, FGuid(), -1);
+}
+
+
+UObject* CreateEditorOptionalObject(TNotNull<UObject*> MainObject, TNotNull<const UClass*> EditorOptionalClass, const TCHAR* OverrideName)
+{
+	checkf(EditorOptionalClass->HasAllClassFlags(CLASS_Optional), TEXT("The class (%s) used with UE::EditorOptional::CreateEditorOptionalObject() was not marked with UCLASS(Optional)"), *EditorOptionalClass->GetName());
+	
+	// get name and flags for the EO object
+	const TCHAR* EditorOptionalName = OverrideName ? OverrideName : TEXT("EditorOptionalData");
+	const EObjectFlags EditorOptionalFlags = MainObject->GetMaskedFlags(RF_PropagateToSubObjects);
+	
+	// create the EO object inside the main object
+	return NewObject<UObject>(MainObject, EditorOptionalClass, EditorOptionalName, EditorOptionalFlags);
+}
+
+}
+#endif
+
diff --git a/Engine/Source/Runtime/CoreUObject/Public/UObject/ObjectEditorOptionalSupport.h b/Engine/Source/Runtime/CoreUObject/Public/UObject/ObjectEditorOptionalSupport.h
new file mode 100644
--- /dev/null
+++ b/Engine/Source/Runtime/CoreUObject/Public/UObject/ObjectEditorOptionalSupport.h
@@ -0,0 +1,32 @@
+// Copyright Epic Games, Inc. All Rights Reserved.
+
+#pragma once
+
+#include "CoreTypes.h"
+#include "Misc/NotNull.h"
+
+#if WITH_EDITORONLY_DATA
+
+class FArchive;
+class UObject;
+class UClass;
+struct FGuid;
+
+namespace UE::EditorOptional
+{
+
+	COREUOBJECT_API void ConditionalUpgradeObject(FArchive& Ar, TNotNull<UObject*> SecondaryObject, const FGuid& VersionGuid, int Version);
+
+	COREUOBJECT_API void UpgradeObject(FArchive& Ar, TNotNull<UObject*> SecondaryObject);
+
+	COREUOBJECT_API UObject* CreateEditorOptionalObject(TNotNull<UObject*> MainObject, TNotNull<const UClass*> EditorOptionalClass, const TCHAR* OverrideName = nullptr);
+
+	template<typename EditorOptionalClass>
+	EditorOptionalClass* CreateEditorOptionalObject(TNotNull<UObject*> MainObject)
+	{
+		// use the non-templated version of this function, and typecast the result (we don't need CastChecked since we know it's the right class)
+		return Cast<EditorOptionalClass>(CreateEditorOptionalObject(MainObject, EditorOptionalClass::StaticClass()));
+	}
+}
+
+#endif
diff --git a/Engine/Source/Runtime/Engine/Private/Materials/MaterialExpressions.cpp b/Engine/Source/Runtime/Engine/Private/Materials/MaterialExpressions.cpp
--- a/Engine/Source/Runtime/Engine/Private/Materials/MaterialExpressions.cpp
+++ b/Engine/Source/Runtime/Engine/Private/Materials/MaterialExpressions.cpp
@@ -320,6 +320,7 @@
 #include "Curves/CurveLinearColorAtlas.h"
 #include "Curves/CurveLinearColor.h"
 #include "MaterialExpressionSettings.h"
+#include "UObject/ObjectEditorOptionalSupport.h"
 
 #define LOCTEXT_NAMESPACE "MaterialExpression"
 
@@ -13291,7 +13292,10 @@ namespace MaterialFunctionInterface
 void UMaterialFunctionInterface::PostInitProperties()
 {
 #if WITH_EDITORONLY_DATA
-	EditorOnlyData = CreateEditorOnlyData();
+	// use the non-templated CreateEditorOnlyData because we need to use the virtual to get the class of the EOData
+	// additionally, pass in an overridden name because we have existing EOData in the wild that must load correctly
+	UObject* EOData = UE::EditorOptional::CreateEditorOptionalObject(this, GetEditorOnlyDataClass(), *MaterialFunctionInterface::GetEditorOnlyDataName(*GetName()));
+	EditorOnlyData = CastChecked<UMaterialFunctionInterfaceEditorOnlyData>(EOData);
 #endif
 	Super::PostInitProperties();
 
@@ -13628,16 +13632,6 @@ const UClass* UMaterialFunctionInterface::GetEditorOnlyDataClass() const
 	return UMaterialFunctionInterfaceEditorOnlyData::StaticClass();
 }
 
-UMaterialFunctionInterfaceEditorOnlyData* UMaterialFunctionInterface::CreateEditorOnlyData()
-{
-	const UClass* EditorOnlyClass = GetEditorOnlyDataClass();
-	check(EditorOnlyClass);
-	check(EditorOnlyClass->HasAllClassFlags(CLASS_Optional));
-
-	const FString EditorOnlyName = MaterialFunctionInterface::GetEditorOnlyDataName(*GetName());
-	const EObjectFlags EditorOnlyFlags = GetMaskedFlags(RF_PropagateToSubObjects);
-	return NewObject<UMaterialFunctionInterfaceEditorOnlyData>(this, EditorOnlyClass, *EditorOnlyName, EditorOnlyFlags);
-}
 #endif // WITH_EDITORONLY_DATA
 
 ///////////////////////////////////////////////////////////////////////////////
diff --git a/Engine/Source/Runtime/Engine/Private/Materials/MaterialInterface.cpp b/Engine/Source/Runtime/Engine/Private/Materials/MaterialInterface.cpp
--- a/Engine/Source/Runtime/Engine/Private/Materials/MaterialInterface.cpp
+++ b/Engine/Source/Runtime/Engine/Private/Materials/MaterialInterface.cpp
@@ -46,6 +46,7 @@
 #include "ShaderCompiler.h"
 #include "Materials/MaterialInsights.h"
 #include "UnrealEngine.h"
+#include "UObject/ObjectEditorOptionalSupport.h"
 
 #if WITH_EDITOR
 #include "ObjectCacheEventSink.h"
@@ -735,7 +736,10 @@ void UMaterialInterface::AddReferencedObjects(UObject* InThis, FReferenceCollect
 void UMaterialInterface::PostInitProperties()
 {
 #if WITH_EDITORONLY_DATA
-	EditorOnlyData = CreateEditorOnlyData();
+	// use the non-templated CreateEditorOnlyData because we need to use the virtual to get the class of the EOData
+	// additionally, pass in an overridden name because we have existing EOData in the wild that must load correctly
+	UObject* EOData = UE::EditorOptional::CreateEditorOptionalObject(this, GetEditorOnlyDataClass(), *UE::MaterialInterface::GetEditorOnlyDataName(*GetName()));
+	EditorOnlyData = CastChecked<UMaterialInterfaceEditorOnlyData>(EOData);
 #endif
 	Super::PostInitProperties();
 }
@@ -2041,17 +2045,6 @@ const UClass* UMaterialInterface::GetEditorOnlyDataClass() const
 	return UMaterialInterfaceEditorOnlyData::StaticClass();
 }
 
-
-UMaterialInterfaceEditorOnlyData* UMaterialInterface::CreateEditorOnlyData()
-{
-	const UClass* EditorOnlyClass = GetEditorOnlyDataClass();
-	check(EditorOnlyClass);
-	check(EditorOnlyClass->HasAllClassFlags(CLASS_Optional));
-
-	const FString EditorOnlyName = UE::MaterialInterface::GetEditorOnlyDataName(*GetName());
-	const EObjectFlags EditorOnlyFlags = GetMaskedFlags(RF_PropagateToSubObjects);
-	return NewObject<UMaterialInterfaceEditorOnlyData>(this, EditorOnlyClass, *EditorOnlyName, EditorOnlyFlags);
-}
 #endif // WITH_EDITORONLY_DATA
 
 bool UMaterialInterface::Rename(const TCHAR* NewName, UObject* NewOuter, ERenameFlags Flags)
diff --git a/Engine/Source/Runtime/Engine/Public/Materials/MaterialFunctionInterface.h b/Engine/Source/Runtime/Engine/Public/Materials/MaterialFunctionInterface.h
--- a/Engine/Source/Runtime/Engine/Public/Materials/MaterialFunctionInterface.h
+++ b/Engine/Source/Runtime/Engine/Public/Materials/MaterialFunctionInterface.h
@@ -66,7 +66,6 @@ protected:
 	TObjectPtr<UMaterialFunctionInterfaceEditorOnlyData> EditorOnlyData;
 
 	ENGINE_API virtual const UClass* GetEditorOnlyDataClass() const;
-	ENGINE_API virtual UMaterialFunctionInterfaceEditorOnlyData* CreateEditorOnlyData();
 
 public:
 	UMaterialFunctionInterfaceEditorOnlyData* GetEditorOnlyData() { return EditorOnlyData; }
diff --git a/Engine/Source/Runtime/Engine/Public/Materials/MaterialInterface.h b/Engine/Source/Runtime/Engine/Public/Materials/MaterialInterface.h
--- a/Engine/Source/Runtime/Engine/Public/Materials/MaterialInterface.h
+++ b/Engine/Source/Runtime/Engine/Public/Materials/MaterialInterface.h
@@ -302,7 +302,6 @@ protected:
 	TObjectPtr<UMaterialInterfaceEditorOnlyData> EditorOnlyData;
 
 	ENGINE_API virtual const UClass* GetEditorOnlyDataClass() const;
-	ENGINE_API UMaterialInterfaceEditorOnlyData* CreateEditorOnlyData();
 
 public:
 	virtual UMaterialInterfaceEditorOnlyData* GetEditorOnlyData() { return EditorOnlyData; }
