dstm object migration cleanup and fixes

- rewrote how object migration requests are handled. instead of always immediately replying with either the object or a deny message, we keep track of each server that requested an object and at which priority. as soon as the request is highest priority and we can send the object, we do that and also forward any other requests to the next server. this eliminates a lot of network polling where we constantly request an object and get denied

- removed the RequestId from the migration APIs because we only need to know which server needs an object, and at which priority - the specific request id is irrelevant and unused

- added TransferObjectOwnershipToRemoteServer() to replace calling 'migrate' with the PushMigrationRequestId special constant, also fixed some callsites that were migrating an object and transferring ownership when we only meant to migrate it

- simplified GetObjectMigrationRecvType and GetObjectMigrationSendType - now we just look at each object, the current owner, and what the owner in the migration context says it should be to determine if it's a borrow or an ownership transfer

- remove some of the duplicate code between BeginMultiServerCommit and ExecuteMultiServerCommit that was already out of sync - now we build the list of objects to return in Begin_ and simply use that list in Execute_ so we don't need the same logic to run twice

- FUObjectMigrationContext now has the full MultiServerCommitRequestId in place of the bool bForMultiServerCommit because for that case we actually do need to send the request id since it's used to validate that we're talking about the same transaction

- TransferObjectOwnershipToRemoteServer can no longer do nothing if we have higher priority work that wants to use that object. now it always does the ownership transfer and we will re-request it for borrow if other work needs it

- mark UPlayer::ConsoleCommand as ALWAYS_OPEN to fix an issue where if you mistype a command we end up going down an AutoRTFM-unsafe path and crash

diff --git a/Engine/Source/Runtime/CoreUObject/Internal/UObject/UObjectMigrationContext.h b/Engine/Source/Runtime/CoreUObject/Internal/UObject/UObjectMigrationContext.h
--- a/Engine/Source/Runtime/CoreUObject/Internal/UObject/UObjectMigrationContext.h
+++ b/Engine/Source/Runtime/CoreUObject/Internal/UObject/UObjectMigrationContext.h
@@ -66,14 +66,11 @@ struct COREUOBJECT_API FUObjectMigrationContext
 	/** The Server that this Object's Physics reside on */
 	FRemoteServerId PhysicsServerId;
 
-	/** The RequestId that corresponds to this migration request */
-	FRemoteTransactionId RequestId;
-
 	/** Which side of the migration does this context apply to? Sending or Receiving? */
 	EObjectMigrationSide MigrationSide;
 
-	/** Whether this migration is part of a multi-server commit */
-	bool bForMultiServerCommit = false;
+	/** If valid, the request id of the multi-server commit */
+	FRemoteTransactionId MultiServerCommitRequestId;
 
 	/** Are the values of this context valid? */
 	bool IsValid() const;
diff --git a/Engine/Source/Runtime/CoreUObject/Private/UObject/RemoteExecutor.cpp b/Engine/Source/Runtime/CoreUObject/Private/UObject/RemoteExecutor.cpp
--- a/Engine/Source/Runtime/CoreUObject/Private/UObject/RemoteExecutor.cpp
+++ b/Engine/Source/Runtime/CoreUObject/Private/UObject/RemoteExecutor.cpp
@@ -160,6 +160,13 @@ public:
 				continue;
 			}
 
+			// subsystem ticking
+			for (int32 SubsystemIndex = 0; SubsystemIndex < Subsystems.Num(); SubsystemIndex++)
+			{
+				FRemoteSubsystemBase* Subsystem = Subsystems[SubsystemIndex];
+				Subsystem->TickSubsystem();
+			}
+
 			//
 			// if we get here and there is no work pending, break out and finish
 			//
@@ -234,7 +241,7 @@ public:
 					for (int32 SubsystemIndex = 0; SubsystemIndex < Subsystems.Num(); SubsystemIndex++)
 					{
 						FRemoteSubsystemBase* Subsystem = Subsystems[SubsystemIndex];
-						Subsystem->Tick();
+						Subsystem->TickRequest();
 					}
 
 					int32 SubsystemReadyCount = 0;
@@ -937,6 +944,11 @@ bool IsHigherOrEqualPriority(FRemoteWorkPriority Lhs, FRemoteWorkPriority Rhs)
 	return IsHigherPriority(Lhs, Rhs) || IsEqualPriority(Lhs, Rhs);
 }
 
+bool operator==(FRemoteWorkPriority Lhs, FRemoteWorkPriority Rhs)
+{
+	return Lhs.PackedData == Rhs.PackedData;
+}
+
 FArchive& operator<<(FArchive& Ar, FRemoteWorkPriority& Priority)
 {
 	Ar << Priority.PackedData;
diff --git a/Engine/Source/Runtime/CoreUObject/Private/UObject/RemoteObject.cpp b/Engine/Source/Runtime/CoreUObject/Private/UObject/RemoteObject.cpp
--- a/Engine/Source/Runtime/CoreUObject/Private/UObject/RemoteObject.cpp
+++ b/Engine/Source/Runtime/CoreUObject/Private/UObject/RemoteObject.cpp
@@ -154,12 +154,12 @@ void InitRemoteObjects()
 	if (!UE::RemoteObject::Transfer::RequestRemoteObjectDelegate.IsBound())
 	{
 		UE::RemoteObject::Transfer::RequestRemoteObjectDelegate.BindLambda(
-			[](FRemoteTransactionId RequestId, FRemoteWorkPriority RequestPriority, FRemoteObjectId ObjectId, FRemoteServerId LastKnownResidentServerId, FRemoteServerId DestinationServerId)
+			[](FRemoteWorkPriority RequestPriority, FRemoteObjectId ObjectId, FRemoteServerId LastKnownResidentServerId, FRemoteServerId DestinationServerId)
 			{
 				// Turns a request into an immediate load
 				FUObjectMigrationContext MigrationContext {
 					.ObjectId = ObjectId, .RemoteServerId = DestinationServerId, .OwnerServerId = LastKnownResidentServerId,
-					.PhysicsServerId = LastKnownResidentServerId, .RequestId = RequestId, .MigrationSide = EObjectMigrationSide::Receive
+					.PhysicsServerId = LastKnownResidentServerId, .MigrationSide = EObjectMigrationSide::Receive
 				};
 				UE::RemoteObject::Serialization::Disk::LoadObjectFromDisk(MigrationContext);
 			}
diff --git a/Engine/Source/Runtime/CoreUObject/Private/UObject/RemoteObjectSerialization.cpp b/Engine/Source/Runtime/CoreUObject/Private/UObject/RemoteObjectSerialization.cpp
--- a/Engine/Source/Runtime/CoreUObject/Private/UObject/RemoteObjectSerialization.cpp
+++ b/Engine/Source/Runtime/CoreUObject/Private/UObject/RemoteObjectSerialization.cpp
@@ -92,7 +92,7 @@ void LoadObjectFromDisk(const FUObjectMigrationContext& MigrationContext)
 	// We have transferred ownership from the Database to the local server
 	FRemoteServerId LocalServerId = UE::RemoteObject::GetGlobalServerId();
 	const FRemoteServerId& DatabaseId = UE::RemoteObject::Transfer::DatabaseId;
-	UE::RemoteObject::Transfer::OnObjectDataReceived(MigrationContext.RequestId, LocalServerId, LocalServerId, MigrationContext.ObjectId, DatabaseId, ObjectData);
+	UE::RemoteObject::Transfer::OnObjectDataReceived(LocalServerId, LocalServerId, MigrationContext.ObjectId, DatabaseId, ObjectData);
 }
 
 void SaveObjectToDisk(const UE::RemoteObject::Transfer::FMigrateSendParams& Params)
diff --git a/Engine/Source/Runtime/CoreUObject/Private/UObject/RemoteObjectTransfer.cpp b/Engine/Source/Runtime/CoreUObject/Private/UObject/RemoteObjectTransfer.cpp
--- a/Engine/Source/Runtime/CoreUObject/Private/UObject/RemoteObjectTransfer.cpp
+++ b/Engine/Source/Runtime/CoreUObject/Private/UObject/RemoteObjectTransfer.cpp
@@ -65,14 +65,51 @@ public:
 	FRemoteWorkPriority GetPriority() const { return Priority; }
 };
 
+class FObjectMigrationRequest
+{
+public:
+	FRemoteObjectId ObjectId;
+	FRemoteServerId DestinationServerId;
+	FRemoteWorkPriority RequestPriority;
+
+	FString ToString() const
+	{
+		return FString::Format(TEXT("[obj %s dest %s %s]"),
+			{*ObjectId.ToString(),
+			*DestinationServerId.ToString(),
+			*RequestPriority.ToString()});
+	}
+};
+
+class FObjectMigrationRequests
+{
+public:
+	FRemoteObjectId ObjectId;
+
+	// IndividualRequests is sorted by priority (highest first)
+	TArray<FObjectMigrationRequest> IndividualRequests;
+};
+
+class FPendingObjectRequest
+{
+public:
+	FRemoteObjectId ObjectId;
+	FRemoteWorkPriority RequestPriority;
+};
+
 class FRemoteObjectTransferQueue : public FRemoteSubsystem<FRemoteObjectRequest>
 {
 	// these are objects that we have outstanding requests for (across all requests)
-	TArray<FRemoteObjectId> PendingObjectRequests;
+	TArray<FPendingObjectRequest> PendingObjectRequests;
 
+	TSet<UObject*> MultiServerCommitObjectsToReturn;
 	TSet<UObject*> MultiServerCommitSentObjects;
 	TSet<UObject*> MultiServerCommitReferencedObjects;
 
+	// this is a list of objects that we have locally that 
+	// other servers are asking for
+	TArray<FObjectMigrationRequests> PendingObjectMigrationRequests;
+
 public:
 	bool bInMultiServerCommit = false;
 
@@ -82,7 +119,205 @@ public:
 	{
 	}
 
-	void Tick() final
+	void TickSubsystem() final
+	{
+		using namespace UE::RemoteObject::Private;
+		using namespace UE::RemoteObject::Handle;
+		using namespace UE::RemoteObject::Serialization;
+
+		// go through the list of pending object migration requests and see if we can satisfy any of them
+		for (int32 MigrationRequestsIndex = 0; MigrationRequestsIndex < PendingObjectMigrationRequests.Num(); )
+		{
+			FObjectMigrationRequests& MigrationRequests = PendingObjectMigrationRequests[MigrationRequestsIndex];
+
+			UE_LOG(LogRemoteObject, VeryVerbose, TEXT("TickObjectMigrations processing (%d) requests for obj %s"),
+					MigrationRequests.IndividualRequests.Num(),
+					*MigrationRequests.ObjectId.ToString());
+
+			UObject* Object = StaticFindObjectFastInternal(MigrationRequests.ObjectId);
+
+			if (Object && !Handle::IsRemote(Object))
+			{
+				UE_LOG(LogRemoteObject, VeryVerbose, TEXT("TickObjectMigrations obj %s (is local)"),
+					*MigrationRequests.ObjectId.ToString());
+
+				// look at the list of requests for this object and pick the one with the highest priority for arbitration
+				// (the list is sorted, so the first element is always the highest priority)
+				FObjectMigrationRequest* HighestPriorityRequest = &MigrationRequests.IndividualRequests[0];
+
+				// verify the list is sorted
+				for (int32 MigrationRequestIndex = 1; MigrationRequestIndex < MigrationRequests.IndividualRequests.Num(); MigrationRequestIndex++)
+				{
+					FObjectMigrationRequest* Request = &MigrationRequests.IndividualRequests[MigrationRequestIndex];
+					check(!IsHigherPriority(Request->RequestPriority, HighestPriorityRequest->RequestPriority));
+				}
+
+				UE_LOG(LogRemoteObject, VeryVerbose, TEXT("TickObjectMigrations obj %s highest priority request is %s"),
+					*MigrationRequests.ObjectId.ToString(),
+					*HighestPriorityRequest->RequestPriority.ToString());
+
+				// find the canonical root object id to use for arbitration
+				UObject* const RootObject = FindCanonicalRootObjectForSerialization(Object);
+				const FRemoteObjectId CanonicalRootObjectId = UE::CoreUObject::Private::FObjectHandleUtils::GetRemoteId(RootObject);
+
+				// now look over all active transactions and see if this request is higher priority than all of them
+				bool bObjectLocked = false;
+
+				for (int32 RequestIndex = 0; RequestIndex < GetRequestCount(); RequestIndex++)
+				{
+					FRemoteObjectRequest* ExistingRequest = GetRequestByIndex(RequestIndex);
+
+					if (IsHigherPriority(ExistingRequest->GetPriority(), HighestPriorityRequest->RequestPriority))
+					{
+						// ExistingRequest is higher priority, now see if it actually needs this object
+						for (int32 RequiredObjectIndex = 0; RequiredObjectIndex < ExistingRequest->RequiredObjectsCanonicalRoot.Num(); RequiredObjectIndex++)
+						{
+							FRemoteObjectId RequiredObjectId = ExistingRequest->RequiredObjectsCanonicalRoot[RequiredObjectIndex];
+
+							if (RequiredObjectId == CanonicalRootObjectId)
+							{
+								// we found an existing request that is higher priority than the highest priority request,
+								// so we have to give it up
+								bObjectLocked = true;
+								break;
+							}
+						}
+					}
+
+					if (bObjectLocked)
+					{
+						break;
+					}
+				}
+
+				if (!bObjectLocked)
+				{
+					// send the object
+					UE_LOG(LogRemoteObject, VeryVerbose, TEXT("TickObjectMigrations sending obj %s to %s %s"),
+						*HighestPriorityRequest->ObjectId.ToString(),
+						*HighestPriorityRequest->DestinationServerId.ToString(),
+						*HighestPriorityRequest->RequestPriority.ToString());
+
+					FRemoteServerId OwnerServerId = Handle::GetOwnerServerId(RootObject);
+					FRemoteServerId PhysicsServerId = Handle::GetPhysicsServerId(RootObject);
+					FUObjectMigrationContext MigrationContext{ .ObjectId = HighestPriorityRequest->ObjectId, .RemoteServerId = HighestPriorityRequest->DestinationServerId, .OwnerServerId = OwnerServerId, .PhysicsServerId = PhysicsServerId, .MigrationSide = EObjectMigrationSide::Send };
+
+					TSet<UObject*> SentObjects;
+					TSet<UObject*> ReferencedObjects;
+					FRemoteObjectData ObjectData = SerializeObjectData(Object, SentObjects, ReferencedObjects, &MigrationContext);
+
+					checkf(SentObjects.Num() > 0, TEXT("unable to migrate object %s to %s"), *HighestPriorityRequest->ObjectId.ToString(), *HighestPriorityRequest->DestinationServerId.ToString());
+
+					for (UObject* RO : ReferencedObjects)
+					{
+						RO->SetInternalFlags(EInternalObjectFlags::RemoteReference);
+					}
+
+					for (UObject* SO : SentObjects)
+					{
+						FRemoteObjectId SubObjectId = UE::CoreUObject::Private::FObjectHandleUtils::GetRemoteId(SO);
+						UE_LOG(LogRemoteObject, VeryVerbose, TEXT("TickObjectMigrations sending obj %s sent subobject : %s [%s]"),
+							*HighestPriorityRequest->ObjectId.ToString(),
+							*SubObjectId.ToString(),
+							*GetNameSafe(SO));
+					}
+
+					// Notify that object data has been migrated (actually will be but SendObjectData does a MoveTemp on the data)
+					OnObjectDataSentDelegate.Broadcast(ObjectData, MigrationContext);
+
+					FMigrateSendParams SendRemoteObjectParams = { .MigrationContext = MigrationContext, .ObjectData = ObjectData };
+					SendRemoteObject(SendRemoteObjectParams);
+
+					{
+						// PostMigrate callbacks may be manipulating migrated objects and may try to migrate them back so prevent that from happening since
+						// they all will be destroyed anyway
+						FUnsafeToMigrateScope UnsafeToMigrateScope;
+
+						// Notify that objects have been migrated to a remote server
+						OnObjectsSentDelegate.Broadcast(SentObjects, MigrationContext);
+
+						for (UObject* SentObject : SentObjects)
+						{
+#if UE_WITH_REMOTE_OBJECT_HANDLE
+							SentObject->PostMigrate(MigrationContext);
+#endif
+							MarkAsRemote(SentObject, HighestPriorityRequest->DestinationServerId);
+							UE::RemoteObject::Handle::ChangeOwnerServerId(SentObject, OwnerServerId);
+						}
+					}
+
+					// for every other request, forward the request on to the new server
+					for (int32 MigrationRequestIndex = 0; MigrationRequestIndex < MigrationRequests.IndividualRequests.Num(); MigrationRequestIndex++)
+					{
+						FObjectMigrationRequest* MigrationRequest = &MigrationRequests.IndividualRequests[MigrationRequestIndex];
+
+						if (MigrationRequest != HighestPriorityRequest)
+						{
+							if (MigrationRequest->DestinationServerId != HighestPriorityRequest->DestinationServerId)
+							{
+								UE_LOG(LogRemoteObject, VeryVerbose, TEXT("TickObjectMigrations obj %s forwarding request for server %s"),
+									*MigrationRequests.ObjectId.ToString(),
+									*MigrationRequest->DestinationServerId.ToString());
+
+								SendRemoteObjectRequest(MigrationRequest->RequestPriority, MigrationRequest->ObjectId, HighestPriorityRequest->DestinationServerId, MigrationRequest->DestinationServerId);
+							}
+							else
+							{
+								UE_LOG(LogRemoteObject, VeryVerbose, TEXT("TickObjectMigrations obj %s denying request to server %s"),
+									*MigrationRequests.ObjectId.ToString(),
+									*MigrationRequest->DestinationServerId.ToString());
+
+								RemoteObjectDeniedTransferDelegate.Execute(MigrationRequest->ObjectId, MigrationRequest->DestinationServerId);
+							}
+						}
+					}
+					
+					// all done, delete all requests for this object
+					PendingObjectMigrationRequests.RemoveAt(MigrationRequestsIndex, EAllowShrinking::No);
+					// do not increment MigrationRequestsIndex
+				}
+				else
+				{
+					UE_LOG(LogRemoteObject, VeryVerbose, TEXT("TickObjectMigrations obj %s is locked, continuing"),
+						*MigrationRequests.ObjectId.ToString());
+
+					// object is locked, do we need to report it?
+					MigrationRequestsIndex++;
+				}
+			}
+			else
+			{
+				UE_LOG(LogRemoteObject, VeryVerbose, TEXT("TickObjectMigrations obj %s is remote, forwarding requests"),
+					*MigrationRequests.ObjectId.ToString());
+
+				// we found a request for an object we don't have, we need to forward all of these
+				// requests to whatever server we think has the object
+				FRemoteObjectStub* RemoteObjectStub = FindRemoteObjectStub(MigrationRequests.ObjectId);
+
+				FRemoteServerId CurrentResidentServerId = RemoteObjectStub ? RemoteObjectStub->ResidentServerId : MigrationRequests.ObjectId.GetServerId();
+
+				for (int32 MigrationRequestIndex = 0; MigrationRequestIndex < MigrationRequests.IndividualRequests.Num(); MigrationRequestIndex++)
+				{
+					FObjectMigrationRequest* MigrationRequest = &MigrationRequests.IndividualRequests[MigrationRequestIndex];
+				
+					if (CurrentResidentServerId != MigrationRequest->DestinationServerId)
+					{
+						SendRemoteObjectRequest(MigrationRequest->RequestPriority, MigrationRequest->ObjectId, CurrentResidentServerId, MigrationRequest->DestinationServerId);
+					}
+					else
+					{
+						RemoteObjectDeniedTransferDelegate.Execute(MigrationRequest->ObjectId, MigrationRequest->DestinationServerId);
+					}
+				}
+
+				// all done, delete all requests for this object
+				PendingObjectMigrationRequests.RemoveAt(MigrationRequestsIndex, EAllowShrinking::No);
+				// do not increment MigrationRequestsIndex
+			}
+		}
+	}
+
+	void TickRequest() final
 	{
 		using namespace UE::RemoteObject::Private;
 		using namespace UE::RemoteObject::Handle;
@@ -93,22 +328,17 @@ public:
 		// reset the list of used objects each time before we run
 		ActiveRequest->UsedObjects.SetNum(0, EAllowShrinking::No);
 
-		UE_LOG(LogRemoteObject, VeryVerbose, TEXT("\t%d required objs (%d pending requests):"), ActiveRequest->RequiredObjects.Num(), PendingObjectRequests.Num());
+		UE_LOG(LogRemoteObject, VeryVerbose, TEXT("FRemoteObjectTransferQueue: TickRequest(%s) %d required objs (%d pending requests):"), *ActiveRequest->GetRequestId().ToString(), ActiveRequest->RequiredObjects.Num(), PendingObjectRequests.Num());
 		for (int32 RequiredObjectIndex = 0; RequiredObjectIndex < ActiveRequest->RequiredObjects.Num(); RequiredObjectIndex++)
 		{
 			FRemoteObjectId RequiredObjectId = ActiveRequest->RequiredObjects[RequiredObjectIndex];
 			FRemoteObjectId CanonicalObjectId = ActiveRequest->RequiredObjectsCanonicalRoot[RequiredObjectIndex];
-			UE_LOG(LogRemoteObject, VeryVerbose, TEXT("\t\t%s [root %s] %s"), *RequiredObjectId.ToString(), *CanonicalObjectId.ToString(), Handle::IsRemote(RequiredObjectId) ? TEXT("<remote>") : TEXT("<local>"));
+			UE_LOG(LogRemoteObject, VeryVerbose, TEXT("FRemoteObjectTransferQueue: TickRequest(%s) %s [root %s] %s"), *ActiveRequest->GetRequestId().ToString(), *RequiredObjectId.ToString(), *CanonicalObjectId.ToString(), Handle::IsRemote(RequiredObjectId) ? TEXT("<remote>") : TEXT("<local>"));
 		}
 
-		// send the requests... slowly for now (to avoid straining the network subsystem)
+		// send the requests
 		for (int32 RequiredObjectIndex = 0; RequiredObjectIndex < ActiveRequest->RequiredObjects.Num(); RequiredObjectIndex++)
 		{
-			if (PendingObjectRequests.Num() >= 5)
-			{
-				break;
-			}
-
 			FRemoteObjectId RequiredObjectId = ActiveRequest->RequiredObjects[RequiredObjectIndex];
 			FRemoteObjectId RequiredCanonicalRootObjectId = ActiveRequest->RequiredObjectsCanonicalRoot[RequiredObjectIndex];
 			
@@ -116,42 +346,72 @@ public:
 
 			if (Handle::IsRemote(ObjectIdToRequest))
 			{
-				if (PendingObjectRequests.Find(ObjectIdToRequest) == INDEX_NONE)
-				{
-					// find the highest priority request that needs this object
-					FRemoteObjectRequest* HighestPriorityRequest = ActiveRequest;
+				// find the highest priority request that needs this object
+				FRemoteObjectRequest* HighestPriorityRequest = ActiveRequest;
 
-					for (int32 RequestIndex = 0; RequestIndex < GetRequestCount(); RequestIndex++)
-					{
-						FRemoteObjectRequest* ExistingRequest = GetRequestByIndex(RequestIndex);
+				for (int32 RequestIndex = 0; RequestIndex < GetRequestCount(); RequestIndex++)
+				{
+					FRemoteObjectRequest* ExistingRequest = GetRequestByIndex(RequestIndex);
 
-						bool bFoundObject =
-							(ExistingRequest->RequiredObjectsCanonicalRoot.Find(ObjectIdToRequest) != INDEX_NONE) ||
-							(ExistingRequest->RequiredObjects.Find(ObjectIdToRequest) != INDEX_NONE);
+					bool bFoundObject =
+						(ExistingRequest->RequiredObjectsCanonicalRoot.Find(ObjectIdToRequest) != INDEX_NONE) ||
+						(ExistingRequest->RequiredObjects.Find(ObjectIdToRequest) != INDEX_NONE);
 
-						if (bFoundObject)
+					if (bFoundObject)
+					{
+						if (IsHigherPriority(ExistingRequest->Priority, HighestPriorityRequest->Priority))
 						{
-							if (IsHigherPriority(ExistingRequest->Priority, HighestPriorityRequest->Priority))
-							{
-								HighestPriorityRequest = ExistingRequest;
-							}
+							HighestPriorityRequest = ExistingRequest;
 						}
 					}
+				}
+
+				// does this object id already exist in PendingObjectRequests?
+				// if it does, but the current HighestPriorityRequest is of
+				// a different priority, then we need to re-send the request
+				// to update the priority at which we require this object
+
+				FPendingObjectRequest* ExistingObjectRequest = nullptr;
+
+				for (FPendingObjectRequest& PendingObjectRequest : PendingObjectRequests)
+				{
+					if (PendingObjectRequest.ObjectId == ObjectIdToRequest)
+					{
+						ExistingObjectRequest = &PendingObjectRequest;
+						break;
+					}
+				}
+
+				const bool bUpdatingPriority = (ExistingObjectRequest != nullptr);
+
+				if (ExistingObjectRequest == nullptr)
+				{
+					ExistingObjectRequest = &PendingObjectRequests.Emplace_GetRef();
+					ExistingObjectRequest->ObjectId = ObjectIdToRequest;
+				}
+
+				if (ExistingObjectRequest->RequestPriority != HighestPriorityRequest->GetPriority())
+				{
+					// send this request
+					ExistingObjectRequest->RequestPriority = HighestPriorityRequest->GetPriority();
 
 					FRemoteObjectStub* RemoteObjectStub = FindRemoteObjectStub(ObjectIdToRequest);
 					FRemoteServerId CurrentResidentServerId = RemoteObjectStub ? RemoteObjectStub->ResidentServerId : ObjectIdToRequest.GetServerId();
 
-					// send one request
-					RequestRemoteObject(
-						HighestPriorityRequest->GetRequestId(),
+					SendRemoteObjectRequest(
 						HighestPriorityRequest->GetPriority(),
 						ObjectIdToRequest,
 						CurrentResidentServerId,
 						GetGlobalServerId());
 			
-					PendingObjectRequests.Add(ObjectIdToRequest);
-				
-					UE_LOG(LogRemoteObject, Verbose, TEXT("\tsent request for %s %s (%d pending requests)"), *RequiredObjectId.ToString(), *HighestPriorityRequest->Priority.ToString(), PendingObjectRequests.Num());
+					if (bUpdatingPriority)
+					{
+						UE_LOG(LogRemoteObject, Verbose, TEXT("FRemoteObjectTransferQueue: TickRequest(%s) sent updated priority for %s %s to server %s (%d pending requests)"), *ActiveRequest->GetRequestId().ToString(), *RequiredObjectId.ToString(), *HighestPriorityRequest->Priority.ToString(), *CurrentResidentServerId.ToString(), PendingObjectRequests.Num());
+					}
+					else
+					{
+						UE_LOG(LogRemoteObject, Verbose, TEXT("FRemoteObjectTransferQueue: TickRequest(%s) sent request for %s %s to server %s (%d pending requests)"), *ActiveRequest->GetRequestId().ToString(), *RequiredObjectId.ToString(), *HighestPriorityRequest->Priority.ToString(), *CurrentResidentServerId.ToString(), PendingObjectRequests.Num());
+					}
 				}
 			}
 			else if (!RequiredCanonicalRootObjectId.IsValid())
@@ -178,22 +438,33 @@ public:
 			}
 		}
 
-		UE_LOG(LogRemoteObject, VeryVerbose, TEXT("\tREADY: all %d required objects are local"), ActiveRequest->RequiredObjects.Num());
+		UE_LOG(LogRemoteObject, VeryVerbose, TEXT("FRemoteObjectTransferQueue: TickRequest(%s) all %d required objects are local"), *ActiveRequest->GetRequestId().ToString(), ActiveRequest->RequiredObjects.Num());
 
 		return true;
 	}
 
 	void BeginMultiServerCommit(TArray<FRemoteServerId>& OutMultiServerCommitRemoteServers) final
 	{
+		using namespace UE::RemoteObject::Serialization;
+
 		check(ActiveRequest);
 		check(!bInMultiServerCommit);
 		bInMultiServerCommit = true;
 
+		check(MultiServerCommitObjectsToReturn.Num() == 0);
+		check(MultiServerCommitSentObjects.Num() == 0);
+		check(MultiServerCommitReferencedObjects.Num() == 0);
+
 		UE_LOG(LogRemoteObject, VeryVerbose, TEXT("FRemoteObjectTransferQueue: BeginMultiServerCommit has %d used objects"), ActiveRequest->UsedObjects.Num());
 
 		// are any of our used objects borrowed?
 		for (FRemoteObjectId UsedObjectId : ActiveRequest->UsedObjects)
 		{
+			if (UsedObjectId.IsAsset())
+			{
+				continue;
+			}
+
 			UE_LOG(LogRemoteObject, VeryVerbose, TEXT("FRemoteObjectTransferQueue: used obj: %s remote: %d"), *UsedObjectId.ToString(), Handle::IsRemote(UsedObjectId));
 
 			// if we used it this run and we're about to commit, we expect that the object is local
@@ -201,13 +472,22 @@ public:
 
 			UObject* UsedObject = StaticFindObjectFastInternal(UsedObjectId);
 			check(UsedObject);
-
 			FRemoteServerId OwnerServerId = Handle::GetOwnerServerId(UsedObject);
-			if (OwnerServerId != GetGlobalServerId())
+
+			UObject* RootObject = FindCanonicalRootObjectForSerialization(UsedObject);
+			check(RootObject);
+			FRemoteServerId RootOwnerServerId = Handle::GetOwnerServerId(RootObject);
+
+			// todo: what to do if the Owner and the Root owner differ?
+			ensure(RootOwnerServerId == OwnerServerId);
+
+			if (RootOwnerServerId != GetGlobalServerId())
 			{
-				if (OutMultiServerCommitRemoteServers.Find(OwnerServerId) == INDEX_NONE)
+				MultiServerCommitObjectsToReturn.Add(RootObject);
+
+				if (OutMultiServerCommitRemoteServers.Find(RootOwnerServerId) == INDEX_NONE)
 				{
-					OutMultiServerCommitRemoteServers.Add(OwnerServerId);
+					OutMultiServerCommitRemoteServers.Add(RootOwnerServerId);
 				}
 			}
 		}
@@ -221,33 +501,7 @@ public:
 		check(MultiServerCommitReferencedObjects.Num() == 0);
 		check(bInMultiServerCommit);
 
-		TSet<UObject*> BorrowedObjectsToReturn;
-
-		// are any of our used objects borrowed?
-		for (FRemoteObjectId ObjectId : ActiveRequest->UsedObjects)
-		{
-			if (Handle::IsRemote(ObjectId))
-			{
-				UE_LOG(LogRemoteObject, VeryVerbose, TEXT("FRemoteObjectTransferQueue: ExecuteMultiServerCommit UNEXPECTED: object %s is remote"), *ObjectId.ToString());
-			}
-
-			// if we used it this run and we're about to commit, we expect that the object is local
-			check(!Handle::IsRemote(ObjectId));
-
-			UObject* Object = StaticFindObjectFastInternal(ObjectId);
-			UObject* RootObject = FindCanonicalRootObjectForSerialization(Object);
-			
-			FRemoteServerId RootOwnerServerId = Handle::GetOwnerServerId(RootObject);
-			if (RootOwnerServerId != GetGlobalServerId())
-			{
-				FRemoteObjectId RootObjectId = UE::CoreUObject::Private::FObjectHandleUtils::GetRemoteId(RootObject);
-
-				UE_LOG(LogRemoteObject, VeryVerbose, TEXT("FRemoteObjectTransferQueue: ExecuteMultiServerCommit found borrowed object %s (root %s)"), *ObjectId.ToString(), *RootObjectId.ToString());
-				BorrowedObjectsToReturn.Add(RootObject);
-			}
-		}
-
-		for (UObject* Object : BorrowedObjectsToReturn)
+		for (UObject* Object : MultiServerCommitObjectsToReturn)
 		{
 			FRemoteObjectId ObjectId = UE::CoreUObject::Private::FObjectHandleUtils::GetRemoteId(Object);
 			FRemoteServerId OwnerServerId = Handle::GetOwnerServerId(Object);
@@ -263,9 +517,8 @@ public:
 				.RemoteServerId = OwnerServerId,
 				.OwnerServerId = OwnerServerId,
 				.PhysicsServerId = OwnerServerId,
-				.RequestId = ActiveRequest->RequestId,
 				.MigrationSide = EObjectMigrationSide::Send,
-				.bForMultiServerCommit = true
+				.MultiServerCommitRequestId = ActiveRequest->RequestId,
 			};
 
 			TSet<UObject*> SentObjects;
@@ -291,6 +544,7 @@ public:
 	{
 		UE_AUTORTFM_OPEN
 		{
+			MultiServerCommitObjectsToReturn.Reset();
 			MultiServerCommitSentObjects.Reset();
 			MultiServerCommitReferencedObjects.Reset();
 		};
@@ -325,7 +579,7 @@ public:
 			FRemoteServerId PhysicsServerId = Handle::GetPhysicsServerId(SentObject);
 			FRemoteObjectId SentObjectId = UE::CoreUObject::Private::FObjectHandleUtils::GetRemoteId(SentObject);
 
-			FUObjectMigrationContext MigrationContext{ .ObjectId = SentObjectId, .RemoteServerId = OwnerServerId, .OwnerServerId = OwnerServerId, .PhysicsServerId = PhysicsServerId, .RequestId = ActiveRequest->RequestId, .MigrationSide = EObjectMigrationSide::Send };
+			FUObjectMigrationContext MigrationContext{ .ObjectId = SentObjectId, .RemoteServerId = OwnerServerId, .OwnerServerId = OwnerServerId, .PhysicsServerId = PhysicsServerId, .MigrationSide = EObjectMigrationSide::Send };
 
 			UE_LOG(LogRemoteObject, VeryVerbose, TEXT("FRemoteObjectTransferQueue: CommitMultiServerCommit postmigrate %s"), *SentObjectId.ToString());
 
@@ -337,6 +591,7 @@ public:
 
 		UE_AUTORTFM_OPEN
 		{
+			MultiServerCommitObjectsToReturn.Reset();
 			MultiServerCommitSentObjects.Reset();
 		};
 
@@ -347,7 +602,60 @@ public:
 	{
 	}
 
-	void RequestRemoteObject(FRemoteTransactionId RequestId, FRemoteWorkPriority RequestPriority, FRemoteObjectId ObjectId, FRemoteServerId LastKnownResidentServerId, FRemoteServerId DestinationServerId)
+	void RequestObjectMigration(FRemoteWorkPriority RequestPriority, FRemoteObjectId ObjectId, FRemoteServerId DestinationServerId)
+	{
+		check(DestinationServerId != GetGlobalServerId());
+
+		// enqueue this request
+		FObjectMigrationRequests* ObjectRequests = nullptr;
+
+		for (int32 RequestIndex = 0; RequestIndex < PendingObjectMigrationRequests.Num(); RequestIndex++)
+		{
+			FObjectMigrationRequests* ExistingObjectRequests = &PendingObjectMigrationRequests[RequestIndex];
+
+			if (ExistingObjectRequests->ObjectId == ObjectId)
+			{
+				ObjectRequests = ExistingObjectRequests;
+				break;
+			}
+		}
+		
+		if (ObjectRequests == nullptr)
+		{
+			ObjectRequests = &PendingObjectMigrationRequests.Emplace_GetRef();
+			ObjectRequests->ObjectId = ObjectId;
+		}
+
+		// first search the list to see if we have a request with the same DestinationServerId - if so, overwrite the priority
+		// otherwise, insert the request in priority sorted order
+		FObjectMigrationRequest* Request = nullptr;
+
+		for (FObjectMigrationRequest& ExistingRequest : ObjectRequests->IndividualRequests)
+		{
+			if (ExistingRequest.DestinationServerId == DestinationServerId)
+			{
+				Request = &ExistingRequest;
+				break;
+			}
+		}
+
+		if (!Request)
+		{
+			Request = &ObjectRequests->IndividualRequests.Emplace_GetRef();
+			Request->ObjectId = ObjectId;
+			Request->DestinationServerId = DestinationServerId;
+		}
+
+		Request->RequestPriority = RequestPriority;
+
+		// re-sort the list by priority
+		ObjectRequests->IndividualRequests.Sort([](const FObjectMigrationRequest& Lhs, const FObjectMigrationRequest& Rhs) -> bool
+			{
+				return IsHigherPriority(Lhs.RequestPriority, Rhs.RequestPriority);
+			});
+	}
+
+	void SendRemoteObjectRequest(FRemoteWorkPriority RequestPriority, FRemoteObjectId ObjectId, FRemoteServerId LastKnownResidentServerId, FRemoteServerId DestinationServerId)
 	{
 		check(!AutoRTFM::IsClosed());
 
@@ -356,14 +664,14 @@ public:
 		{
 			// We are forwarding a request for pull-migrate from the resident server id
 			// This isn't technically a migration, it is actually a request for migration...
-			RequestRemoteObjectDelegate.Execute(RequestId, RequestPriority, ObjectId, LastKnownResidentServerId, DestinationServerId);
+			RequestRemoteObjectDelegate.Execute(RequestPriority, ObjectId, LastKnownResidentServerId, DestinationServerId);
 		}
 		else
 		{
 			UE_LOG(LogRemoteObject, Log, TEXT("%hs: Restoring ObjectId %s from Database and Forwarding to %s"), __func__, *ObjectId.ToString(), *DestinationServerId.ToString());
 			// We are intercepting a request to pull-migrate since we believe the object resides in our database, migrate it "from" the database to us
 			// Inside the database call, it will fix-up these migration context parameters to make us the owner
-			FUObjectMigrationContext MigrationContext { .ObjectId = ObjectId, .RemoteServerId = DatabaseId, .OwnerServerId = DatabaseId, .PhysicsServerId = DatabaseId, .RequestId = RequestId, .MigrationSide = EObjectMigrationSide::Receive};
+			FUObjectMigrationContext MigrationContext { .ObjectId = ObjectId, .RemoteServerId = DatabaseId, .OwnerServerId = DatabaseId, .PhysicsServerId = DatabaseId, .MigrationSide = EObjectMigrationSide::Receive};
 			RestoreRemoteObjectDataDelegate.Execute(MigrationContext);
 
 #if UE_WITH_REMOTE_OBJECT_HANDLE
@@ -375,7 +683,7 @@ public:
 			// If we didn't expect it locally, send it on to the final destination
 			if (DestinationServerId != GetGlobalServerId())
 			{
-				MigrateObjectToRemoteServer(RequestId, ObjectId, DestinationServerId);
+				MigrateObjectToRemoteServer(ObjectId, DestinationServerId);
 			}
 		}
 	}
@@ -388,12 +696,6 @@ public:
 		RemoteObjectTransferDelegate.Execute(Params);
 	}
 
-	void SendDeniedRemoteObject(FRemoteTransactionId RequestId, FRemoteObjectId ObjectId, FRemoteServerId DestinationServerId)
-	{
-		check(!AutoRTFM::IsClosed());
-		RemoteObjectDeniedTransferDelegate.Execute(RequestId, ObjectId, DestinationServerId);
-	}
-
 	/**
 	* Fullfils receive request
 	* @param ObjectId Id of the object that was requested
@@ -405,10 +707,21 @@ public:
 
 		check(!AutoRTFM::IsClosed());
 
-		const int32 PendingRequestIndex = PendingObjectRequests.Find(MigrationContext.ObjectId);
-		const bool bFoundRequest = (PendingRequestIndex != INDEX_NONE);
+		int32 FoundPendingRequestIndex = INDEX_NONE;
+
+		for (int32 PendingRequestIndex = 0; PendingRequestIndex < PendingObjectRequests.Num(); PendingRequestIndex++)
+		{
+			const FPendingObjectRequest& PendingObjectRequest = PendingObjectRequests[PendingRequestIndex];
+			if (PendingObjectRequest.ObjectId == MigrationContext.ObjectId)
+			{
+				FoundPendingRequestIndex = PendingRequestIndex;
+				break;
+			}
+		}
 
-		UE_LOG(LogRemoteObject, Verbose, TEXT("     : FulfillReceiveRequest(%s) %s (was requested: %d, owner: %s)"), *MigrationContext.RequestId.ToString(), *MigrationContext.ObjectId.ToString(), bFoundRequest, *MigrationContext.OwnerServerId.ToString());
+		const bool bFoundRequest = (FoundPendingRequestIndex != INDEX_NONE);
+
+		UE_LOG(LogRemoteObject, Verbose, TEXT("FRemoteObjectTransferQueue: FulfillReceiveRequest %s (was requested: %d, owner: %s)"), *MigrationContext.ObjectId.ToString(), bFoundRequest, *MigrationContext.OwnerServerId.ToString());
 
 		UObject* Result = MigrateObjectFromObjectData(ObjectData, MigrationContext);
 		checkf(UE::CoreUObject::Private::FObjectHandleUtils::GetRemoteId(Result) == MigrationContext.ObjectId, TEXT("PeerId:%s requested object %s migration from PeerId:%s but received %s"),
@@ -417,22 +730,33 @@ public:
 
 		if (bFoundRequest)
 		{
-			PendingObjectRequests.RemoveAt(PendingRequestIndex);
+			PendingObjectRequests.RemoveAt(FoundPendingRequestIndex);
 		}
 	}
 
-	void DenyReceiveRequest(FRemoteTransactionId RequestId, FRemoteObjectId ObjectId)
+	void DenyReceiveRequest(FRemoteObjectId ObjectId)
 	{
 		check(!AutoRTFM::IsClosed());
 
-		const int32 PendingRequestIndex = PendingObjectRequests.Find(ObjectId);
-		const bool bFoundRequest = (PendingRequestIndex != INDEX_NONE);
+		int32 FoundPendingRequestIndex = INDEX_NONE;
+
+		for (int32 PendingRequestIndex = 0; PendingRequestIndex < PendingObjectRequests.Num(); PendingRequestIndex++)
+		{
+			const FPendingObjectRequest& PendingObjectRequest = PendingObjectRequests[PendingRequestIndex];
+			if (PendingObjectRequest.ObjectId == ObjectId)
+			{
+				FoundPendingRequestIndex = PendingRequestIndex;
+				break;
+			}
+		}
+
+		const bool bFoundRequest = (FoundPendingRequestIndex != INDEX_NONE);
 
-		UE_LOG(LogRemoteObject, Verbose, TEXT("     : DenyReceiveRequest(%s) %s, (was requested: %d)"), *RequestId.ToString(), *ObjectId.ToString(), bFoundRequest);
+		UE_LOG(LogRemoteObject, Verbose, TEXT("FRemoteObjectTransferQueue: DenyReceiveRequest %s, (was requested: %d)"), *ObjectId.ToString(), bFoundRequest);
 
 		if (bFoundRequest)
 		{
-			PendingObjectRequests.RemoveAt(PendingRequestIndex);
+			PendingObjectRequests.RemoveAt(FoundPendingRequestIndex);
 		}
 	}
 };
@@ -657,13 +981,7 @@ void TouchResidentObject(UObject* Object)
 	}
 }
 
-void MigrateObjectToRemoteServer(FRemoteTransactionId RequestId, UObject* Object, FRemoteServerId DestinationServerId)
-{
-	FRemoteWorkPriority RootWorkPriority = UE::RemoteExecutor::CreateRootWorkPriority();
-	MigrateObjectToRemoteServerWithExplicitPriority(RequestId, RootWorkPriority, Object, DestinationServerId);
-}
-
-void MigrateObjectToRemoteServerWithExplicitPriority(FRemoteTransactionId RequestId, FRemoteWorkPriority RequestPriority, UObject* Object, FRemoteServerId DestinationServerId)
+void TransferObjectOwnershipToRemoteServer(UObject* Object, FRemoteServerId DestinationServerId)
 {
 	using namespace UE::RemoteObject::Serialization;
 	using namespace UE::RemoteObject::Private;
@@ -672,156 +990,67 @@ void MigrateObjectToRemoteServerWithExplicitPriority(FRemoteTransactionId Reques
 
 	FRemoteObjectId ObjectId = UE::CoreUObject::Private::FObjectHandleUtils::GetRemoteId(Object);
 
-	UE_LOG(LogRemoteObject, Verbose, TEXT("     : MigrateObjectToRemoteServer(%s %s) %s is requesting obj %s (%s)"),
-		*RequestId.ToString(), 
-		*RequestPriority.ToString(), 
-		*DestinationServerId.ToString(), 
-		*ObjectId.ToString(),
-		*Object->GetName());
-
 	// find the canonical root object id to use for arbitration
 	UObject* const RootObject = FindCanonicalRootObjectForSerialization(Object);
 	const FRemoteObjectId CanonicalRootObjectId = UE::CoreUObject::Private::FObjectHandleUtils::GetRemoteId(RootObject);
 
 	// Figure out the ownership semantics
-	const bool bPushOwnershipChange = (RequestId == PushMigrationRequestId);
-	FRemoteServerId OwnerServerId = bPushOwnershipChange ? DestinationServerId : Handle::GetOwnerServerId(RootObject);
 	FRemoteServerId PhysicsServerId = Handle::GetPhysicsServerId(RootObject);
-	FUObjectMigrationContext MigrationContext{ .ObjectId = ObjectId, .RemoteServerId = DestinationServerId, .OwnerServerId = OwnerServerId, .PhysicsServerId = PhysicsServerId, .RequestId = RequestId, .MigrationSide = EObjectMigrationSide::Send };
-
-	// is the object locked?
-	const FRemoteObjectRequest* HigherPriorityRequest = nullptr;
-	const FRemoteObjectRequest* EqualPriorityRequest = nullptr;
-	const FRemoteObjectRequest* LowerPriorityRequest = nullptr;
+	FUObjectMigrationContext MigrationContext{ .ObjectId = ObjectId, .RemoteServerId = DestinationServerId, .OwnerServerId = DestinationServerId, .PhysicsServerId = PhysicsServerId, .MigrationSide = EObjectMigrationSide::Send };
 
-	UE_LOG(LogRemoteObject, Verbose, 
-		TEXT("     : MigrateObjectToRemoteServer(%s) sending obj %s [%s] (root %s, isremote %d, owner %s, dest %s)"),
-		*RequestId.ToString(),
+	UE_LOG(LogRemoteObject, Verbose, TEXT("TransferObjectOwnershipToRemoteServer %s (%s) root %s (%s) to %s"),
 		*ObjectId.ToString(),
 		*Object->GetName(),
 		*CanonicalRootObjectId.ToString(),
-		Handle::IsRemote(Object),
-		*OwnerServerId.ToString(),
-		*DestinationServerId.ToString()
-	);
-
-	for (int32 RequestIndex = 0; RequestIndex < GTransferQueue.GetRequestCount(); RequestIndex++)
-	{
-		FRemoteObjectRequest* ExistingRequest = GTransferQueue.GetRequestByIndex(RequestIndex);
-
-		for (int32 RequiredObjectIndex = 0; RequiredObjectIndex < ExistingRequest->RequiredObjectsCanonicalRoot.Num(); RequiredObjectIndex++)
-		{
-			FRemoteObjectId RequiredObjectId = ExistingRequest->RequiredObjectsCanonicalRoot[RequiredObjectIndex];
-
-			if (RequiredObjectId == CanonicalRootObjectId)
-			{
-				// do we give it up, or deny their request?
-				bool bIsHigherPriority = IsHigherPriority(ExistingRequest->Priority, RequestPriority);
-				bool bIsEqualPriority = IsEqualPriority(ExistingRequest->Priority, RequestPriority);
-				bool bIsLowerPriority = IsHigherPriority(RequestPriority, ExistingRequest->Priority);
+		*RootObject->GetName(),
+		*DestinationServerId.ToString());
 
-				if (bIsHigherPriority)
-				{
-					HigherPriorityRequest = ExistingRequest;
-					
-					// if we ever find a higher priority request, we can stop looking
-					break;
-				}
-				else if (bIsEqualPriority)
-				{
-					EqualPriorityRequest = ExistingRequest;
-				}
-				else if (bIsLowerPriority)
-				{
-					LowerPriorityRequest = ExistingRequest;
-				}
-			}
-		}
-	}
+	// send the object
+	TSet<UObject*> SentObjects;
+	TSet<UObject*> ReferencedObjects;
+	FRemoteObjectData ObjectData = SerializeObjectData(Object, SentObjects, ReferencedObjects, &MigrationContext);
 
-	bool bIsObjectLocked = false;
+	UE_LOG(LogRemoteObject, VeryVerbose, TEXT("TransferObjectOwnershipToRemoteServer %s serialization complete (%d bytes)"),
+		*ObjectId.ToString(),
+		ObjectData.GetNumBytes());
 
-	if (HigherPriorityRequest != nullptr)
-	{
-		// we keep our object locked if our priority is higher
-		UE_LOG(LogRemoteObject, Verbose, TEXT("     : MigrateObjectToRemoteServer(%s) DENYING request for our locked %s because our priority is higher: %s vs %s"), *RequestId.ToString(), *ObjectId.ToString(), *HigherPriorityRequest->Priority.ToString(), *RequestPriority.ToString());
-		bIsObjectLocked = true;
-	}
-	else if (EqualPriorityRequest)
-	{
-		if (GetGlobalServerId() < DestinationServerId)
-		{
-			UE_LOG(LogRemoteObject, Verbose, TEXT("     : MigrateObjectToRemoteServer(%s) DENYING request for our locked %s because priority %s is equal, and our server index is lower: %s vs %s"), *RequestId.ToString(), *ObjectId.ToString(), *EqualPriorityRequest->Priority.ToString(), *GetGlobalServerId().ToString(), *DestinationServerId.ToString());
-			bIsObjectLocked = true;
-		}
-		else
-		{
-			UE_LOG(LogRemoteObject, Verbose, TEXT("     : MigrateObjectToRemoteServer(%s) GRANTING request for our locked %s because priority %s is equal, and our server index is higher: %s vs %s"), *RequestId.ToString(), *ObjectId.ToString(), *EqualPriorityRequest->Priority.ToString(), *GetGlobalServerId().ToString(), *DestinationServerId.ToString());
-		}
-	}
-	else if (LowerPriorityRequest)
-	{
-		UE_LOG(LogRemoteObject, Verbose, TEXT("     : MigrateObjectToRemoteServer(%s) GRANTING request for our locked %s because our priority is lower: %s vs %s"), *RequestId.ToString(), *ObjectId.ToString(), *LowerPriorityRequest->Priority.ToString(), *RequestPriority.ToString());
-	}
-	else
+	for (UObject* RO : ReferencedObjects)
 	{
-		UE_LOG(LogRemoteObject, Verbose, TEXT("     : MigrateObjectToRemoteServer(%s) GRANTING request for %s because none of our active requests need it"), *RequestId.ToString(), *ObjectId.ToString());
+		RO->SetInternalFlags(EInternalObjectFlags::RemoteReference);
 	}
 
-	if (bIsObjectLocked)
+	for (UObject* SO : SentObjects)
 	{
-		RemoteObjectDeniedTransferDelegate.Execute(RequestId, ObjectId, DestinationServerId);
+		FRemoteObjectId SubObjectId = UE::CoreUObject::Private::FObjectHandleUtils::GetRemoteId(SO);
+		UE_LOG(LogRemoteObject, VeryVerbose, TEXT("TransferObjectOwnershipToRemoteServer %s sent subobject : %s (%s)"),
+			*ObjectId.ToString(),
+			*SubObjectId.ToString(),
+			*GetNameSafe(SO));
 	}
-	else
-	{
-		UE_LOG(LogRemoteObject, VeryVerbose, TEXT("     : MigrateObjectToRemoteServer(%s) serializing obj %s"), *RequestId.ToString(), *ObjectId.ToString());
 
-		// send the object
-		TSet<UObject*> SentObjects;
-		TSet<UObject*> ReferencedObjects;
-		FRemoteObjectData ObjectData = SerializeObjectData(Object, SentObjects, ReferencedObjects, &MigrationContext);
+	// Notify that object data has been migrated (actually will be but SendObjectData does a MoveTemp on the data)
+	OnObjectDataSentDelegate.Broadcast(ObjectData, MigrationContext);
 
-		UE_LOG(LogRemoteObject, VeryVerbose, TEXT("     : MigrateObjectToRemoteServer(%s) serialization complete obj %s (%d bytes)"), *RequestId.ToString(), *ObjectId.ToString(), ObjectData.GetNumBytes());
+	FMigrateSendParams SendRemoteObjectParams = { .MigrationContext = MigrationContext, .ObjectData = ObjectData };
+	GTransferQueue.SendRemoteObject(SendRemoteObjectParams);
 
-		for (UObject* RO : ReferencedObjects)
-		{
-			RO->SetInternalFlags(EInternalObjectFlags::RemoteReference);
-		}
-
-		for (UObject* SO : SentObjects)
-		{
-			FRemoteObjectId SubObjectId = UE::CoreUObject::Private::FObjectHandleUtils::GetRemoteId(SO);
-			UE_LOG(LogRemoteObject, VeryVerbose, TEXT("         sent subobject : %s [%s]"), *SubObjectId.ToString(), *GetNameSafe(SO));
-		}
-			
-		// Notify that object data has been migrated (actually will be but SendObjectData does a MoveTemp on the data)
-		OnObjectDataSentDelegate.Broadcast(ObjectData, MigrationContext);
+	checkf(SentObjects.Num() > 0, TEXT("PeerId:%s unable to migrate object %s to PeerId:%s"), *GetGlobalServerId().ToString(), *UE::CoreUObject::Private::FObjectHandleUtils::GetRemoteId(Object).ToString(), *DestinationServerId.ToString());
 
-		FMigrateSendParams SendRemoteObjectParams = { .MigrationContext = MigrationContext, .ObjectData = ObjectData };
-		GTransferQueue.SendRemoteObject(SendRemoteObjectParams);
+	{
+		// PostMigrate callbacks may be manipulating migrated objects and may try to migrate them back so prevent that from happening since
+		// they all will be destroyed anyway
+		FUnsafeToMigrateScope UnsafeToMigrateScope;
 
-		checkf(SentObjects.Num() > 0, TEXT("PeerId:%s unable to migrate object %s to PeerId:%s"), *GetGlobalServerId().ToString(), *UE::CoreUObject::Private::FObjectHandleUtils::GetRemoteId(Object).ToString(), *DestinationServerId.ToString());
+		// Notify that objects have been migrated to a remote server
+		OnObjectsSentDelegate.Broadcast(SentObjects, MigrationContext);
 
+		for (UObject* SentObject : SentObjects)
 		{
-			// PostMigrate callbacks may be manipulating migrated objects and may try to migrate them back so prevent that from happening since
-			// they all will be destroyed anyway
-			FUnsafeToMigrateScope UnsafeToMigrateScope;
-
-			// Notify that objects have been migrated to a remote server
-			OnObjectsSentDelegate.Broadcast(SentObjects, MigrationContext);
-
-			const bool bSentToOwnerServer = (OwnerServerId == DestinationServerId);
-			for (UObject* SentObject : SentObjects)
-			{
 #if UE_WITH_REMOTE_OBJECT_HANDLE
-				SentObject->PostMigrate(MigrationContext);
+			SentObject->PostMigrate(MigrationContext);
 #endif
-				MarkAsRemote(SentObject, DestinationServerId);
-				if (bSentToOwnerServer)
-				{
-					UE::RemoteObject::Handle::ChangeOwnerServerId(SentObject, OwnerServerId);
-				}
-			}
+			MarkAsRemote(SentObject, DestinationServerId);
+			UE::RemoteObject::Handle::ChangeOwnerServerId(SentObject, DestinationServerId);
 		}
 	}
 }
@@ -849,68 +1078,29 @@ void GetAllBorrowedObjects(TArray<FRemoteObjectId>& OutBorrowedObjectIds)
 	}
 }
 
-void MigrateObjectToRemoteServer(FRemoteTransactionId RequestId, FRemoteObjectId ObjectId, FRemoteServerId DestinationServerId)
+void MigrateObjectToRemoteServer(FRemoteObjectId ObjectId, FRemoteServerId DestinationServerId)
 {
 	FRemoteWorkPriority RootWorkPriority = UE::RemoteExecutor::CreateRootWorkPriority();
-	MigrateObjectToRemoteServerWithExplicitPriority(RequestId, RootWorkPriority, ObjectId, DestinationServerId);
+	MigrateObjectToRemoteServerWithExplicitPriority(RootWorkPriority, ObjectId, DestinationServerId);
 }
 
-void MigrateObjectToRemoteServerWithExplicitPriority(FRemoteTransactionId RequestId, FRemoteWorkPriority RequestPriority, FRemoteObjectId ObjectId, FRemoteServerId DestinationServerId)
+void MigrateObjectToRemoteServerWithExplicitPriority(FRemoteWorkPriority RequestPriority, FRemoteObjectId ObjectId, FRemoteServerId DestinationServerId)
 {
-	using namespace UE::RemoteObject::Handle;
-	using namespace UE::RemoteObject::Serialization;
-	using namespace UE::RemoteObject::Private;
-
-	UE_LOG(LogRemoteObject, Verbose, TEXT("MigrateObjectToRemoteServer(%s) obj %s to %s"), *RequestId.ToString(), *ObjectId.ToString(), *DestinationServerId.ToString());
-
-	UObject* Object = StaticFindObjectFastInternal(ObjectId);
-	if (Object && !IsRemote(Object))
-	{
-		UE_LOG(LogRemoteObject, VeryVerbose, TEXT("    : MigrateObjectToRemoteServer(%s) found obj %s"), *RequestId.ToString(), *ObjectId.ToString());
-		MigrateObjectToRemoteServerWithExplicitPriority(RequestId, RequestPriority, Object, DestinationServerId);
-	}
-	else
-	{
-		FRemoteObjectStub* RemoteObjectStub = FindRemoteObjectStub(ObjectId);
-
-		if (RemoteObjectStub)
-		{
-			UE_LOG(LogRemoteObject, Verbose, TEXT("    : MigrateObjectToRemoteServer(%s) did not find obj %s (we do have a stub)"), *RequestId.ToString(),*ObjectId.ToString());
-		}
-		else
-		{
-			UE_LOG(LogRemoteObject, Verbose, TEXT("    : MigrateObjectToRemoteServer(%s) did not find obj %s (also no stub)"), *RequestId.ToString(),*ObjectId.ToString());
-		}
-
-		FRemoteServerId CurrentResidentServerId = RemoteObjectStub ? RemoteObjectStub->ResidentServerId : ObjectId.GetServerId();
-		checkf(CurrentResidentServerId != GetGlobalServerId(), TEXT("PeerId:%s attempting to forward request for object %s to PeerId:%s but object location is unknown"), 
-			*GetGlobalServerId().ToString(), *ObjectId.ToString(), *DestinationServerId.ToString());
-
-		if (CurrentResidentServerId != DestinationServerId)
-		{
-			UE_LOG(LogRemoteObject, Verbose, TEXT("    : MigrateObjectToRemoteServer(%s) forwarding request (we believe %s owns the Object)"), *RequestId.ToString(), *CurrentResidentServerId.ToString());
-			GTransferQueue.RequestRemoteObject(RequestId, RequestPriority, ObjectId, CurrentResidentServerId, DestinationServerId);
-		}
-		else
-		{
-			UE_LOG(LogRemoteObject, Verbose, TEXT("    : MigrateObjectToRemoteServer(%s) denying request %s (we believe %s already has the Object)"), *RequestId.ToString(), *ObjectId.ToString(), *CurrentResidentServerId.ToString());
-			RemoteObjectDeniedTransferDelegate.Execute(RequestId, ObjectId, DestinationServerId);
-		}
-	}
+	GTransferQueue.RequestObjectMigration(RequestPriority, ObjectId, DestinationServerId);
 }
 
-void OnObjectDataReceived(FRemoteTransactionId RequestId, FRemoteServerId OwnerServerId, FRemoteServerId PhysicsId, FRemoteObjectId ObjectId, FRemoteServerId RemoteServerId, FRemoteObjectData& Data)
+void OnObjectDataReceived(FRemoteServerId OwnerServerId, FRemoteServerId PhysicsId, FRemoteObjectId ObjectId, FRemoteServerId RemoteServerId, FRemoteObjectData& Data)
 {
 	check(!AutoRTFM::IsClosed());
 
-	FUObjectMigrationContext MigrationContext { .ObjectId = ObjectId, .RemoteServerId = RemoteServerId, .OwnerServerId = OwnerServerId, .PhysicsServerId = PhysicsId, .RequestId = RequestId, .MigrationSide = EObjectMigrationSide::Receive };
+	FUObjectMigrationContext MigrationContext { .ObjectId = ObjectId, .RemoteServerId = RemoteServerId, .OwnerServerId = OwnerServerId, .PhysicsServerId = PhysicsId, .MigrationSide = EObjectMigrationSide::Receive };
 	GTransferQueue.FulfillReceiveRequest(Data, MigrationContext);
 }
 
-void OnObjectDataDenied(FRemoteTransactionId RequestId, FRemoteObjectId ObjectId, FRemoteServerId RemoteServerId)
+void OnObjectDataDenied(FRemoteObjectId ObjectId, FRemoteServerId RemoteServerId)
 {
 	check(!AutoRTFM::IsClosed());
-	GTransferQueue.DenyReceiveRequest(RequestId, ObjectId);
+	GTransferQueue.DenyReceiveRequest(ObjectId);
 }
 
 void PrefetchRemoteObject(FRemoteObjectId ObjectId, UObject* DestinationOuter)
@@ -1023,7 +1213,7 @@ void StoreObjectToDatabase(UObject* Object)
 
 
 			// Setup a migration context for "push migrating to the new owner: the Database Server"
-			FUObjectMigrationContext StoreToDatabaseContext{ .ObjectId = RootObjectId, .RemoteServerId = DatabaseId, .OwnerServerId = DatabaseId, .PhysicsServerId = DatabaseId, .RequestId = PushMigrationRequestId, .MigrationSide = EObjectMigrationSide::Send };
+			FUObjectMigrationContext StoreToDatabaseContext{ .ObjectId = RootObjectId, .RemoteServerId = DatabaseId, .OwnerServerId = DatabaseId, .PhysicsServerId = DatabaseId, .MigrationSide = EObjectMigrationSide::Send };
 
 			// Do the actual serialization
 			FRemoteObjectData ObjectData = SerializeObjectData(OutermostUnreachableRemoteReference, SentObjects, ReferencedObjects, &StoreToDatabaseContext);
diff --git a/Engine/Source/Runtime/CoreUObject/Private/UObject/UObjectMigrationContext.cpp b/Engine/Source/Runtime/CoreUObject/Private/UObject/UObjectMigrationContext.cpp
--- a/Engine/Source/Runtime/CoreUObject/Private/UObject/UObjectMigrationContext.cpp
+++ b/Engine/Source/Runtime/CoreUObject/Private/UObject/UObjectMigrationContext.cpp
@@ -8,7 +8,7 @@
 
 bool FUObjectMigrationContext::IsValid() const
 {
-	return ObjectId.IsValid() && RemoteServerId.IsValid() && OwnerServerId.IsValid() && RequestId.IsValid();
+	return ObjectId.IsValid() && RemoteServerId.IsValid() && OwnerServerId.IsValid();
 }
 
 EObjectMigrationRecvType FUObjectMigrationContext::GetObjectMigrationRecvType(const UObjectBase* Object) const
@@ -27,35 +27,22 @@ EObjectMigrationRecvType FUObjectMigrationContext::GetObjectMigrationRecvType(FR
 	const FRemoteServerId LocalServerId = UE::RemoteObject::GetGlobalServerId();
 	const bool bAlreadyOwns = UE::RemoteObject::Handle::IsOwned(InObjectId);
 
-	const bool bPushAssignOwnership = (RequestId == UE::RemoteObject::Transfer::PushMigrationRequestId);
-	if (bPushAssignOwnership)
-	{
-		ensureMsgf(OwnerServerId == LocalServerId,
-				   TEXT("%s attempted to push-assign ownership of %s to %s but the receiving server of that request is actually %s"),
-				   *RemoteServerId.ToString(), *ObjectId.ToString(), *OwnerServerId.ToString(), *LocalServerId.ToString());
-
-		if (bAlreadyOwns)
-		{
-			// We were pushed an object that we already owned, so the loan has been returned to us
-			return EObjectMigrationRecvType::ReturnedLoan;
-		}
+	const bool bAssignOwnership = (OwnerServerId == LocalServerId) && !bAlreadyOwns;
+	const bool bReturnLoan = (OwnerServerId == LocalServerId) && bAlreadyOwns;
 
+	if (bAssignOwnership)
+	{
 		// We were pushed an object that we do not already own, so we must assume ownership of it
 		return EObjectMigrationRecvType::AssignedOwnership;
 	}
 
-	const bool bReturnLoan = (OwnerServerId == LocalServerId);
 	if (bReturnLoan)
 	{
-		ensureMsgf(bAlreadyOwns,
-				   TEXT("Remote Server %s sent obj %s and assumed %s owns it (but it doesn't)"),
-				   *RemoteServerId.ToString(), *InObjectId.ToString(), *LocalServerId.ToString());
-
-		// A request we must have made (a non-push) has been fulfilled that returned a loaned object back to us
+		// We were given an object back that we already own
 		return EObjectMigrationRecvType::ReturnedLoan;
 	}
 
-	// We received an object we requested and ownership is not involved, so we are just borrowing it
+	// We received an object and ownership is not involved, so we are just borrowing it
 	return EObjectMigrationRecvType::Borrowed;
 #else
 	return EObjectMigrationRecvType::Invalid;
@@ -78,30 +65,17 @@ EObjectMigrationSendType FUObjectMigrationContext::GetObjectMigrationSendType(FR
 	const FRemoteServerId LocalServerId = UE::RemoteObject::GetGlobalServerId();
 	const bool bCurrentlyOwns = UE::RemoteObject::Handle::IsOwned(InObjectId);
 
-	const bool bPushAssignOwnership = (RequestId == UE::RemoteObject::Transfer::PushMigrationRequestId);
-	if (bPushAssignOwnership)
-	{
-		ensureMsgf(OwnerServerId == RemoteServerId,
-				   TEXT("Server %s attempted to push Object %s to Server %s which is not the Owner (%s is the owner)"),
-				   *LocalServerId.ToString(), *ObjectId.ToString(), *RemoteServerId.ToString(), *OwnerServerId.ToString());
+	const bool bOwnershipTransfer = (OwnerServerId == RemoteServerId) && bCurrentlyOwns;
+	const bool bReturnBorrowed = (OwnerServerId == RemoteServerId) && !bCurrentlyOwns;
 
-		if (bCurrentlyOwns)
-		{
-			// We have ownership of this object and are pushing it to the destination (reassigning the ownership)
-			return EObjectMigrationSendType::ReassignOwnership;
-		}
-
-		// Send an object we don't own back to the owning server.  The fact that it was a push is irrelevant in the outcome.
-		return EObjectMigrationSendType::ReturnBorrowed;
+	if (bOwnershipTransfer)
+	{
+		// We have ownership of this object and are pushing it to the destination (reassigning the ownership)
+		return EObjectMigrationSendType::ReassignOwnership;
 	}
 
-	const bool bReturnBorrowed = (OwnerServerId == RemoteServerId);
 	if (bReturnBorrowed)
 	{
-		ensureMsgf(!bCurrentlyOwns,
-				   TEXT("Server %s is sending a non-push migration request for Object %s to Server %s but it actually owns the Object (this should have been a push)"),
-				   *LocalServerId.ToString(), *InObjectId.ToString(), *RemoteServerId.ToString());
-
 		// Send an object we don't own back to the owning server (it requested the object back).
 		return EObjectMigrationSendType::ReturnBorrowed;
 	}
diff --git a/Engine/Source/Runtime/CoreUObject/Public/UObject/RemoteExecutor.h b/Engine/Source/Runtime/CoreUObject/Public/UObject/RemoteExecutor.h
--- a/Engine/Source/Runtime/CoreUObject/Public/UObject/RemoteExecutor.h
+++ b/Engine/Source/Runtime/CoreUObject/Public/UObject/RemoteExecutor.h
@@ -82,8 +82,9 @@ struct FRemoteWorkPriority
 
 COREUOBJECT_API FArchive& operator<<(FArchive& Ar, FRemoteWorkPriority& Priority);
 COREUOBJECT_API bool IsHigherPriority(FRemoteWorkPriority Lhs, FRemoteWorkPriority Rhs);
-COREUOBJECT_API  bool IsEqualPriority(FRemoteWorkPriority Lhs, FRemoteWorkPriority Rhs);
-COREUOBJECT_API  bool IsHigherOrEqualPriority(FRemoteWorkPriority Lhs, FRemoteWorkPriority Rhs);
+COREUOBJECT_API bool IsEqualPriority(FRemoteWorkPriority Lhs, FRemoteWorkPriority Rhs);
+COREUOBJECT_API bool IsHigherOrEqualPriority(FRemoteWorkPriority Lhs, FRemoteWorkPriority Rhs);
+COREUOBJECT_API bool operator==(FRemoteWorkPriority Lhs, FRemoteWorkPriority Rhs);
 
 class FRemoteSubsystemBase
 {
@@ -106,7 +107,8 @@ public:
 	// for subsystem implementor to override:
 	virtual const TCHAR* NameForDebug() = 0;
 	virtual void BeginRequest() = 0;
-	virtual void Tick() = 0;
+	virtual void TickSubsystem() = 0;
+	virtual void TickRequest() = 0;
 	virtual bool AreDependenciesSatisfied() const = 0;
 	virtual void BeginMultiServerCommit(TArray<FRemoteServerId>& OutMultiServerCommitRemoteServers) = 0;
 	virtual void ExecuteMultiServerCommit() = 0;
diff --git a/Engine/Source/Runtime/CoreUObject/Public/UObject/RemoteObjectTransfer.h b/Engine/Source/Runtime/CoreUObject/Public/UObject/RemoteObjectTransfer.h
--- a/Engine/Source/Runtime/CoreUObject/Public/UObject/RemoteObjectTransfer.h
+++ b/Engine/Source/Runtime/CoreUObject/Public/UObject/RemoteObjectTransfer.h
@@ -50,8 +50,6 @@ namespace UE::RemoteObject::Transfer
 		FRemoteObjectData ObjectData;
 	};
 
-	constexpr FRemoteTransactionId PushMigrationRequestId = FRemoteTransactionId(0xFFFFFFFFu);
-
 	/**
 	* (Advanced API) Adds the remote object to the current transaction's requested object list, causing
 	* the object to be migrated to this server. Does not abort the transaction if the object isn't yet local
@@ -64,30 +62,29 @@ namespace UE::RemoteObject::Transfer
 	* @param RemoteServerId Server that sent object data
 	* @param Data Remote object data. Data ownership is transferred to the remote object queue
 	*/
-	COREUOBJECT_API void OnObjectDataReceived(FRemoteTransactionId RequestId, FRemoteServerId OwnerServerId, FRemoteServerId PhysicsId, FRemoteObjectId ObjectId, FRemoteServerId RemoteServerId, FRemoteObjectData& Data);
+	COREUOBJECT_API void OnObjectDataReceived(FRemoteServerId OwnerServerId, FRemoteServerId PhysicsId, FRemoteObjectId ObjectId, FRemoteServerId RemoteServerId, FRemoteObjectData& Data);
 
 	/**
 	* Called when a remote object request was denied by a remote server
 	* @param ObjectId Remote object id
 	* @param RemoteServerId Server that owns the object data
 	*/
-	COREUOBJECT_API void OnObjectDataDenied(FRemoteTransactionId RequestId, FRemoteObjectId ObjectId, FRemoteServerId RemoteServerId);
+	COREUOBJECT_API void OnObjectDataDenied(FRemoteObjectId ObjectId, FRemoteServerId RemoteServerId);
 
 	/**
-	* Migrates an object to a remote server
+	* Migrates and transfers ownership of an object to a remote server
 	* @param Object Local object to migrate
-	* @param DestinationServerId Id of a remote server to receive the object
+	* @param DestinationServerId Id of a remote server to receive the object and assume ownership
 	*/
-	COREUOBJECT_API void MigrateObjectToRemoteServer(FRemoteTransactionId RequestId, UObject* Object, FRemoteServerId DestinationServerId);
-	COREUOBJECT_API void MigrateObjectToRemoteServerWithExplicitPriority(FRemoteTransactionId RequestId, FRemoteWorkPriority RequestPriority, UObject* Object, FRemoteServerId DestinationServerId);
+	COREUOBJECT_API void TransferObjectOwnershipToRemoteServer(UObject* Object, FRemoteServerId DestinationServerId);
 
 	/**
-	* Migrates an object to a remote server
+	* Migrates an object to a remote server without changing ownership
 	* @param ObjectId Remote object id
 	* @param DestinationServerId Id of a remote server to receive the object
 	*/
-	COREUOBJECT_API void MigrateObjectToRemoteServer(FRemoteTransactionId RequestId, FRemoteObjectId Id, FRemoteServerId DestinationServerId);
-	COREUOBJECT_API void MigrateObjectToRemoteServerWithExplicitPriority(FRemoteTransactionId RequestId, FRemoteWorkPriority RequestPriority, FRemoteObjectId Id, FRemoteServerId DestinationServerId);
+	COREUOBJECT_API void MigrateObjectToRemoteServer(FRemoteObjectId ObjectId, FRemoteServerId DestinationServerId);
+	COREUOBJECT_API void MigrateObjectToRemoteServerWithExplicitPriority(FRemoteWorkPriority RequestPriority, FRemoteObjectId Id, FRemoteServerId DestinationServerId);
 	
 	/**
 	* Migrates an object from a remote server (temp function)
@@ -142,10 +139,10 @@ namespace UE::RemoteObject::Transfer
 	extern COREUOBJECT_API TDelegate<void(const FMigrateSendParams& /*Params*/)> RemoteObjectTransferDelegate;
 
 	/** Delegate that handles an object request being denied */
-	extern COREUOBJECT_API TDelegate<void(FRemoteTransactionId /*RequestId*/, FRemoteObjectId /*ObjectId*/, FRemoteServerId /*DestinationServerId*/)> RemoteObjectDeniedTransferDelegate;
+	extern COREUOBJECT_API TDelegate<void(FRemoteObjectId /*ObjectId*/, FRemoteServerId /*DestinationServerId*/)> RemoteObjectDeniedTransferDelegate;
 
 	/** Delegate that requests remote object data from LastKnownResidentServerId to be transferred to DestinationServerId. Allows requests to be forwarded if ObjectId does not reside on LastKnownResidentServerId. */
-	extern COREUOBJECT_API TDelegate<void(FRemoteTransactionId /*RequestId*/, FRemoteWorkPriority /*RequestPriority*/, FRemoteObjectId /*ObjectId*/, FRemoteServerId /*LastKnownResidentServerId*/, FRemoteServerId /*DestinationServerId*/)> RequestRemoteObjectDelegate;
+	extern COREUOBJECT_API TDelegate<void(FRemoteWorkPriority /*RequestPriority*/, FRemoteObjectId /*ObjectId*/, FRemoteServerId /*LastKnownResidentServerId*/, FRemoteServerId /*DestinationServerId*/)> RequestRemoteObjectDelegate;
 
 	/** Delegate executed when objects have been migrated from another server. */
 	extern COREUOBJECT_API TMulticastDelegate<void(const FRemoteObjectData& /*ObjectData*/, const FUObjectMigrationContext& /*MigrationContext*/)> OnObjectDataReceivedDelegate;
diff --git a/Engine/Source/Runtime/Engine/Private/Components/PrimitiveComponent.cpp b/Engine/Source/Runtime/Engine/Private/Components/PrimitiveComponent.cpp
--- a/Engine/Source/Runtime/Engine/Private/Components/PrimitiveComponent.cpp
+++ b/Engine/Source/Runtime/Engine/Private/Components/PrimitiveComponent.cpp
@@ -1308,7 +1308,7 @@ bool UPrimitiveComponent::CanSendMigratedPhysicsState(const FUObjectMigrationCon
 
 	if (ObjMigrationSendType == EObjectMigrationSendType::ReturnBorrowed)
 	{
-		return MigrationContext->bForMultiServerCommit;
+		return MigrationContext->MultiServerCommitRequestId.IsValid();
 	}
 
 	return true;
diff --git a/Engine/Source/Runtime/Engine/Private/Player.cpp b/Engine/Source/Runtime/Engine/Private/Player.cpp
--- a/Engine/Source/Runtime/Engine/Private/Player.cpp
+++ b/Engine/Source/Runtime/Engine/Private/Player.cpp
@@ -17,6 +17,8 @@
 #include "GameFramework/CheatManager.h"
 #include "GameFramework/GameStateBase.h"
 
+#include "AutoRTFM/AutoRTFM.h"
+
 #include UE_INLINE_GENERATED_CPP_BY_NAME(Player)
 
 //////////////////////////////////////////////////////////////////////////
@@ -27,6 +29,7 @@ UPlayer::UPlayer(const FObjectInitializer& ObjectInitializer)
 {
 }
 
+UE_AUTORTFM_ALWAYS_OPEN
 FString UPlayer::ConsoleCommand(const FString& Cmd, bool bWriteToLog)
 {
 	UNetConnection* NetConn = Cast<UNetConnection>(this);
