dstm: removed an ensure

- if the root and object owner servers differ then that suggests the mutation that happened was adding a new object to the root object hierarchy. we now support this case of mixed ownership on the receiving side so we don't need this ensure

diff --git a/Engine/Source/Runtime/CoreUObject/Private/UObject/RemoteObjectTransfer.cpp b/Engine/Source/Runtime/CoreUObject/Private/UObject/RemoteObjectTransfer.cpp
--- a/Engine/Source/Runtime/CoreUObject/Private/UObject/RemoteObjectTransfer.cpp
+++ b/Engine/Source/Runtime/CoreUObject/Private/UObject/RemoteObjectTransfer.cpp
@@ -472,15 +472,11 @@ public:
 
 			UObject* UsedObject = StaticFindObjectFastInternal(UsedObjectId);
 			check(UsedObject);
-			FRemoteServerId OwnerServerId = Handle::GetOwnerServerId(UsedObject);
 
 			UObject* RootObject = FindCanonicalRootObjectForSerialization(UsedObject);
 			check(RootObject);
 			FRemoteServerId RootOwnerServerId = Handle::GetOwnerServerId(RootObject);
 
-			// todo: what to do if the Owner and the Root owner differ?
-			ensure(RootOwnerServerId == OwnerServerId);
-
 			if (RootOwnerServerId != GetGlobalServerId())
 			{
 				MultiServerCommitObjectsToReturn.Add(RootObject);
