IncrementalCook: Move EDLCookChecker into the cooker since it is not needed for other uses of SavePackage.
* Change SavePackage to report Imports, Exports and PreloadDependencies to the cooker so the cooker can populate EDLCookChecker. A future change will build upon this to save the per-package information into the CookArtifacts saved with each package.
* Related deprecation: Deprecate SAVE_CompareLinker and the implementation types for it - FSavePackageResultStruct.LinkerSave, FLinkerDiff, FLinkerDiffPackageWriter, USavePackageUtilitiesCommandlet. It was only needed for comparing SavePackage2 with legacy SavePackage, and legacy SavePackage is now removed.

#rnx
#incrementalcook

diff --git a/Engine/Source/Editor/UnrealEd/Classes/Commandlets/SavePackageUtilitiesCommandlet.h b/Engine/Source/Editor/UnrealEd/Classes/Commandlets/SavePackageUtilitiesCommandlet.h
--- a/Engine/Source/Editor/UnrealEd/Classes/Commandlets/SavePackageUtilitiesCommandlet.h
+++ b/Engine/Source/Editor/UnrealEd/Classes/Commandlets/SavePackageUtilitiesCommandlet.h
@@ -3,8 +3,6 @@
 #pragma once
 
 #include "Commandlets/Commandlet.h"
-#include "CoreMinimal.h"
-#include "UObject/ObjectMacros.h"
 
 #include "SavePackageUtilitiesCommandlet.generated.h"
 
@@ -13,19 +11,13 @@
  * It can currently compare two back to back saves of a package (or folder of packages)
  */
 UCLASS()
-class USavePackageUtilitiesCommandlet
+class
+	UE_DEPRECATED(5.6, "SavePackageUtilitiesCommandlet is no longer needed. Contact Epic if you need this functionality.")
+	USavePackageUtilitiesCommandlet
 	: public UCommandlet
 {
 	GENERATED_UCLASS_BODY()
 
 public:
-	//~ Begin UCommandlet Interface
 	virtual int32 Main(const FString& Params) override;
-	//~ End UCommandlet Interface
-
-private:
-	void InitParameters(const FString& Params);
-
-	TArray<FString> PackageNames;
-	class ITargetPlatform* TargetPlatform = nullptr;
 };
diff --git a/Engine/Source/Editor/UnrealEd/Private/Commandlets/CookCommandlet.cpp b/Engine/Source/Editor/UnrealEd/Private/Commandlets/CookCommandlet.cpp
--- a/Engine/Source/Editor/UnrealEd/Private/Commandlets/CookCommandlet.cpp
+++ b/Engine/Source/Editor/UnrealEd/Private/Commandlets/CookCommandlet.cpp
@@ -9,6 +9,7 @@
 #include "AssetRegistry/AssetRegistryModule.h"
 #include "Async/TaskGraphInterfaces.h"
 #include "CookOnTheSide/CookOnTheFlyServer.h"
+#include "Cooker/CookImportsChecker.h"
 #include "Cooker/CookProfiling.h"
 #include "CookerSettings.h"
 #include "Editor.h"
@@ -604,7 +605,7 @@ void UCookCommandlet::RunCookByTheBookCook(UCookOnTheFlyServer* CookOnTheFlyServ
 	if (bShouldVerifyEDLCookInfo)
 	{
 		bool bFullReferencesExpected = !(CookOptions & ECookByTheBookOptions::SkipHardReferences);
-		UE::SavePackageUtilities::VerifyEDLCookInfo([](UE::FLogRecord&& Record)
+		FEDLCookChecker::Verify([](UE::FLogRecord&& Record)
 			{
 #if !NO_LOGGING
 				Record.SetCategory(LogCook.GetCategoryName());
diff --git a/Engine/Source/Editor/UnrealEd/Private/Commandlets/SavePackageUtilitiesCommandlet.cpp b/Engine/Source/Editor/UnrealEd/Private/Commandlets/SavePackageUtilitiesCommandlet.cpp
--- a/Engine/Source/Editor/UnrealEd/Private/Commandlets/SavePackageUtilitiesCommandlet.cpp
+++ b/Engine/Source/Editor/UnrealEd/Private/Commandlets/SavePackageUtilitiesCommandlet.cpp
@@ -2,22 +2,7 @@
 
 #include "Commandlets/SavePackageUtilitiesCommandlet.h"
 
-#include "Editor.h"
-#include "HAL/IConsoleManager.h"
-#include "HAL/FileManager.h"
-#include "HAL/IConsoleManager.h"
-#include "Interfaces/ITargetPlatformManagerModule.h"
-#include "Misc/CoreMisc.h"
-#include "Misc/FeedbackContext.h"
-#include "Misc/PackageName.h"
-#include "Misc/Paths.h"
-#include "UObject/ArchiveCookContext.h"
-#include "UObject/LinkerDiff.h"
-#include "UObject/LinkerSave.h"
-#include "UObject/Package.h"
-#include "UObject/SavePackage.h"
-#include "UObject/UObjectGlobals.h"
-#include "UObject/UObjectHash.h"
+#include "CoreGlobals.h"
 
 USavePackageUtilitiesCommandlet::USavePackageUtilitiesCommandlet(const FObjectInitializer& ObjectInitializer)
 	: Super(ObjectInitializer)
@@ -26,94 +11,6 @@ USavePackageUtilitiesCommandlet::USavePackageUtilitiesCommandlet(const FObjectIn
 
 int32 USavePackageUtilitiesCommandlet::Main(const FString& Params)
 {
-	InitParameters(Params);
-
-	UE::Cook::ICookInfo* CookInfo = nullptr; // TODO: NeedSimpleCookInfoClass
-
-	for (const FString& PackageName : PackageNames)
-	{
-		// Load Package
-		UPackage* Package = LoadPackage(nullptr, *PackageName, LOAD_None);
-		UObject* Asset = Package->FindAssetInPackage();
-		FString Filename = FPaths::CreateTempFilename(*FPaths::ProjectSavedDir());
-
-		// CookData should only be nonzero if we are cooking.
-		TOptional<FArchiveCookContext> CookContext;
-		TOptional<FArchiveCookData> CookData;
-		if (TargetPlatform != nullptr)
-		{
-			CookContext.Emplace(Package, UE::Cook::ECookType::Unknown,
-				UE::Cook::ECookingDLC::Unknown, TargetPlatform, CookInfo);
-			CookData.Emplace(*TargetPlatform, *CookContext);
-		}
-
-		FSavePackageArgs SaveArgs;
-		SaveArgs.TopLevelFlags = RF_Public;
-		SaveArgs.SaveFlags = SAVE_CompareLinker;
-		SaveArgs.ArchiveCookData = CookData.GetPtrOrNull();
-		SaveArgs.bSlowTask = false;
-
-		// if not cooking add RF_Standalone to the top level flags
-		if (TargetPlatform == nullptr)
-		{
-			SaveArgs.TopLevelFlags |= RF_Standalone;
-		}
-
-		FSavePackageResultStruct FirstResult = GEditor->Save(Package, Asset, *Filename, SaveArgs);
-		FSavePackageResultStruct SecondResult = GEditor->Save(Package, Asset, *Filename, SaveArgs);
-	
-		if (FirstResult.LinkerSave && SecondResult.LinkerSave)
-		{
-			// Compare Linker Save info
-			FLinkerDiff LinkerDiff = FLinkerDiff::CompareLinkers(FirstResult.LinkerSave.Get(), SecondResult.LinkerSave.Get());
-			LinkerDiff.PrintDiff(*GWarn);
-		}
-
-		// Delete the temp filename
-		IFileManager::Get().Delete(*Filename);
-	}
-
+	UE_LOG(LogCore, Error, TEXT("SavePackageUtilitiesCommandlet is no longer needed. Contact Epic if you need this functionality."));
 	return 0;
 }
-
-void USavePackageUtilitiesCommandlet::InitParameters(const FString& Params)
-{
-	TArray<FString> Tokens;
-	TArray<FString> Switches;
-	ParseCommandLine(*Params, Tokens, Switches);
-
-	FString SwitchValue;
-	for (const FString& CurrentSwitch : Switches)
-	{
-		if (FParse::Value(*CurrentSwitch, TEXT("PACKAGE="), SwitchValue))
-		{
-			FString LongPackageName;
-			FPackageName::SearchForPackageOnDisk(SwitchValue, &LongPackageName, nullptr);
-			PackageNames.Add(MoveTemp(LongPackageName));
-
-		}
-		else if (FParse::Value(*CurrentSwitch, TEXT("PACKAGEFOLDER="), SwitchValue))
-		{
-			FPackageName::IteratePackagesInDirectory(SwitchValue, [this](const TCHAR* Filename)
-				{
-					PackageNames.Add(FPackageName::FilenameToLongPackageName(Filename));
-					return true;
-				});
-		}
-		else if (FParse::Value(*CurrentSwitch, TEXT("CookPlatform="), SwitchValue))
-		{
-			if (ITargetPlatformManagerModule* TPM = GetTargetPlatformManager())
-			{
-				TargetPlatform = TPM->FindTargetPlatform(SwitchValue);
-				if (TargetPlatform == nullptr)
-				{
-					// @todo Error
-				}
-
-			}
-			
-
-		}
-	}
-}
-
diff --git a/Engine/Source/Editor/UnrealEd/Private/CookOnTheFlyServer.cpp b/Engine/Source/Editor/UnrealEd/Private/CookOnTheFlyServer.cpp
--- a/Engine/Source/Editor/UnrealEd/Private/CookOnTheFlyServer.cpp
+++ b/Engine/Source/Editor/UnrealEd/Private/CookOnTheFlyServer.cpp
@@ -26,6 +26,7 @@
 #include "Cooker/CookDirector.h"
 #include "Cooker/CookGarbageCollect.h"
 #include "Cooker/CookGenerationHelper.h"
+#include "Cooker/CookImportsChecker.h"
 #include "Cooker/CookOnTheFlyServerInterface.h"
 #include "Cooker/CookPackageData.h"
 #include "Cooker/CookPackagePreloader.h"
@@ -5570,7 +5571,6 @@ public:
 	{
 		None,
 		DiffOnly,
-		LinkerDiff,
 		IterativeValidate,
 		IterativeValidatePhase1,
 		IterativeValidatePhase2,
@@ -5600,9 +5600,6 @@ public:
 			CookedPackageWriter = new FDiffPackageWriter(TUniquePtr<ICookedPackageWriter>(CookedPackageWriter),
 				InDeterminismManager);
 			break;
-		case EDiffMode::LinkerDiff:
-			CookedPackageWriter = new FLinkerDiffPackageWriter(TUniquePtr<ICookedPackageWriter>(CookedPackageWriter));
-			break;
 		case EDiffMode::IterativeValidate:
 			CookedPackageWriter = new FIterativeValidatePackageWriter(COTFS,
 				TUniquePtr<ICookedPackageWriter>(CookedPackageWriter), FIterativeValidatePackageWriter::EPhase::AllInOnePhase,
@@ -5630,7 +5627,6 @@ public:
 		{
 		case EDiffMode::None: return false;
 		case EDiffMode::DiffOnly: return true;
-		case EDiffMode::LinkerDiff: return false;
 		case EDiffMode::IterativeValidate: return true;
 		case EDiffMode::IterativeValidatePhase1: return true;
 		case EDiffMode::IterativeValidatePhase2: return true;
@@ -5653,21 +5649,15 @@ public:
 		{
 			if (DiffMode != EDiffMode::None)
 			{
-				UE_LOG(LogCook, Fatal, TEXT("-DiffOnly, -LinkerDiff, and -IterativeValidate* are mutually exclusive."));
+				UE_LOG(LogCook, Fatal, TEXT("-DiffOnly, and -IterativeValidate* are mutually exclusive."));
 			}
 		};
 
-		bool bValue;
 		if (FParse::Param(CommandLine, TEXT("DIFFONLY")))
 		{
 			EnsureMutualExclusion();
 			DiffMode = EDiffMode::DiffOnly;
 		}
-		if (FParse::Bool(CommandLine, TEXT("-LINKERDIFF="), bValue) && bValue)
-		{
-			EnsureMutualExclusion();
-			DiffMode = EDiffMode::LinkerDiff;
-		}
 		if (FParse::Param(CommandLine, TEXT("IterativeValidate")))
 		{
 			EnsureMutualExclusion();
@@ -7696,7 +7686,7 @@ void UCookOnTheFlyServer::PopulateCookedPackages(TArrayView<const ITargetPlatfor
 							COOK_STAT(++DetailedCookStats::NumPackagesIterativelySkipped);
 						}
 						// Declare the package to the EDLCookInfo verification so we don't warn about missing exports from it
-						UE::SavePackageUtilities::EDLCookInfoAddIterativelySkippedPackage(PackageName);
+						FEDLCookCheckerThreadState::Get().AddPackageWithUnknownExports(PackageName);
 					}
 					else
 					{
@@ -11824,7 +11814,9 @@ void FEDLMPCollector::ClientTickPackage(FMPCollectorClientTickPackageContext& Co
 {
 	FCbWriter Writer;
 	bool bHasData;
-	UE::SavePackageUtilities::EDLCookInfoMoveToCompactBinaryAndClear(Writer, bHasData, Context.GetPackageName());
+
+	// For simplicity, instead of sending only information related to the given Package, we send all data.
+	FEDLCookChecker::MoveToCompactBinaryAndClear(Writer, bHasData);
 	if (bHasData)
 	{
 		Context.AddMessage(Writer.Save().AsObject());
@@ -11833,7 +11825,7 @@ void FEDLMPCollector::ClientTickPackage(FMPCollectorClientTickPackageContext& Co
 
 void FEDLMPCollector::ServerReceiveMessage(FMPCollectorServerMessageContext& Context, FCbObjectView Message)
 {
-	UE::SavePackageUtilities::EDLCookInfoAppendFromCompactBinary(Message.AsFieldView());
+	FEDLCookChecker::AppendFromCompactBinary(Message.AsFieldView());
 }
 
 }
@@ -11849,7 +11841,7 @@ void UCookOnTheFlyServer::BeginCookEDLCookInfo(FBeginCookContext& BeginContext)
 	{
 		return;
 	}
-	UE::SavePackageUtilities::StartSavingEDLCookInfoForVerification();
+	FEDLCookChecker::StartSavingEDLCookInfoForVerification();
 	if (CookDirector)
 	{
 		CookDirector->Register(new UE::Cook::FEDLMPCollector());
diff --git a/Engine/Source/Editor/UnrealEd/Private/Cooker/CookGenerationHelper.cpp b/Engine/Source/Editor/UnrealEd/Private/Cooker/CookGenerationHelper.cpp
--- a/Engine/Source/Editor/UnrealEd/Private/Cooker/CookGenerationHelper.cpp
+++ b/Engine/Source/Editor/UnrealEd/Private/Cooker/CookGenerationHelper.cpp
@@ -6,6 +6,7 @@
 #include "AssetRegistry/AssetData.h"
 #include "Cooker/CookDirector.h"
 #include "Cooker/CookGarbageCollect.h"
+#include "Cooker/CookImportsChecker.h"
 #include "Cooker/CookPackageArtifacts.h"
 #include "Cooker/CookPlatformManager.h"
 #include "Cooker/CookWorkerServer.h"
@@ -2747,7 +2748,7 @@ void FCookGenerationInfo::IterativeCookValidateOrClear(FGenerationHelper& Genera
 			COOK_STAT(++DetailedCookStats::NumPackagesIterativelySkipped);
 		}
 		// Declare the package to the EDLCookInfo verification so we don't warn about missing exports from it
-		UE::SavePackageUtilities::EDLCookInfoAddIterativelySkippedPackage(PackageData->GetPackageName());
+		FEDLCookCheckerThreadState::Get().AddPackageWithUnknownExports(PackageData->GetPackageName());
 	}
 	else
 	{
diff --git a/Engine/Source/Editor/UnrealEd/Private/Cooker/CookImportsChecker.cpp b/Engine/Source/Editor/UnrealEd/Private/Cooker/CookImportsChecker.cpp
new file mode 100644
--- /dev/null
+++ b/Engine/Source/Editor/UnrealEd/Private/Cooker/CookImportsChecker.cpp
@@ -0,0 +1,804 @@
+// Copyright Epic Games, Inc. All Rights Reserved.
+
+#include "CookImportsChecker.h"
+
+#include "Algo/Sort.h"
+#include "Algo/Unique.h"
+#include "HAL/LowLevelMemTracker.h"
+#include "Logging/LogVerbosity.h"
+#include "Logging/StructuredLog.h"
+#include "Misc/CommandLine.h"
+#include "Misc/ConfigCacheIni.h"
+#include "Misc/Parse.h"
+#include "Misc/ScopeExit.h"
+#include "Misc/ScopeLock.h"
+#include "Misc/StringBuilder.h"
+#include "Serialization/CompactBinary.h"
+#include "Serialization/CompactBinarySerialization.h"
+#include "Serialization/CompactBinaryWriter.h"
+#include "Templates/TypeHash.h"
+#include "UObject/Package.h"
+#include "UObject/SavePackage.h"
+
+LLM_DEFINE_TAG(EDLCookChecker);
+
+FEDLCookChecker::FEDLNodeHash::FEDLNodeHash()
+{
+}
+
+FEDLCookChecker::FEDLNodeHash::FEDLNodeHash(const FEDLNodeHash& Other)
+{
+	bIsNode = Other.bIsNode;
+	if (bIsNode)
+	{
+		Nodes = Other.Nodes;
+	}
+	else
+	{
+		Object = Other.Object;
+	}
+	NodeID = Other.NodeID;
+	ObjectEvent = Other.ObjectEvent;
+}
+
+FEDLCookChecker::FEDLNodeHash::FEDLNodeHash(const TArray<FEDLNodeData>* InNodes, FEDLNodeID InNodeID, EObjectEvent InObjectEvent)
+	: Nodes(InNodes)
+	, NodeID(InNodeID)
+	, bIsNode(true)
+	, ObjectEvent(InObjectEvent)
+{
+}
+
+FEDLCookChecker::FEDLNodeHash::FEDLNodeHash(TObjectPtr<UObject> InObject, EObjectEvent InObjectEvent)
+	: Object(InObject)
+	, bIsNode(false)
+	, ObjectEvent(InObjectEvent)
+{
+}
+
+bool FEDLCookChecker::FEDLNodeHash::operator==(const FEDLNodeHash& Other) const
+{
+	if (ObjectEvent != Other.ObjectEvent)
+	{
+		return false;
+	}
+
+	uint32 LocalNodeID;
+	uint32 OtherNodeID;
+	TObjectPtr<const UObject> LocalObject;
+	TObjectPtr<const UObject> OtherObject;
+	FName LocalName = ObjectNameFirst(*this, LocalNodeID, LocalObject);
+	FName OtherName = ObjectNameFirst(Other, OtherNodeID, OtherObject);
+
+	do
+	{
+		if (LocalName != OtherName)
+		{
+			return false;
+		}
+		LocalName = ObjectNameNext(*this, LocalNodeID, LocalObject);
+		OtherName = ObjectNameNext(Other, OtherNodeID, OtherObject);
+	} while (!LocalName.IsNone() && !OtherName.IsNone());
+	return LocalName.IsNone() == OtherName.IsNone();
+}
+
+FEDLCookChecker::FEDLNodeHash& FEDLCookChecker::FEDLNodeHash::operator=(const FEDLNodeHash& Other)
+{
+	bIsNode = Other.bIsNode;
+	if (bIsNode)
+	{
+		Nodes = Other.Nodes;
+	}
+	else
+	{
+		Object = Other.Object;
+	}
+	NodeID = Other.NodeID;
+	ObjectEvent = Other.ObjectEvent;
+
+	return *this;
+}
+
+uint32 GetTypeHash(const FEDLCookChecker::FEDLNodeHash& A)
+{
+	uint32 Hash = 0;
+
+	uint32 LocalNodeID;
+	TObjectPtr<const UObject> LocalObject;
+	FName LocalName = FEDLCookChecker::FEDLNodeHash::ObjectNameFirst(A, LocalNodeID, LocalObject);
+	do
+	{
+		Hash = HashCombine(Hash, GetTypeHash(LocalName));
+		LocalName = FEDLCookChecker::FEDLNodeHash::ObjectNameNext(A, LocalNodeID, LocalObject);
+	} while (!LocalName.IsNone());
+
+	return (Hash << 1) | (uint32)A.ObjectEvent;
+}
+
+FName FEDLCookChecker::FEDLNodeHash::GetName() const
+{
+	if (bIsNode)
+	{
+		return (*Nodes)[NodeID].Name;
+	}
+	else
+	{
+		return Object.GetFName();
+	}
+}
+
+bool FEDLCookChecker::FEDLNodeHash::TryGetParent(FEDLCookChecker::FEDLNodeHash& Parent) const
+{
+	EObjectEvent ParentObjectEvent = EObjectEvent::Create; // For purposes of parents, which is used only to get the ObjectPath, we always use the Create version of the node as the parent
+	if (bIsNode)
+	{
+		FEDLNodeID ParentID = (*Nodes)[NodeID].ParentID;
+		if (ParentID != NodeIDInvalid)
+		{
+			Parent = FEDLNodeHash(Nodes, ParentID, ParentObjectEvent);
+			return true;
+		}
+	}
+	else
+	{
+		TObjectPtr<UObject> ParentObject = Object.GetOuter();
+		if (ParentObject)
+		{
+			Parent = FEDLNodeHash(ParentObject, ParentObjectEvent);
+			return true;
+		}
+	}
+	return false;
+}
+
+FEDLCookChecker::EObjectEvent FEDLCookChecker::FEDLNodeHash::GetObjectEvent() const
+{
+	return ObjectEvent;
+}
+
+void FEDLCookChecker::FEDLNodeHash::SetNodes(const TArray<FEDLNodeData>* InNodes)
+{
+	if (bIsNode)
+	{
+		Nodes = InNodes;
+	}
+}
+
+FName FEDLCookChecker::FEDLNodeHash::ObjectNameFirst(const FEDLNodeHash& InNode, uint32& OutNodeID, TObjectPtr<const UObject>& OutObject)
+{
+	if (InNode.bIsNode)
+	{
+		OutNodeID = InNode.NodeID;
+		return (*InNode.Nodes)[OutNodeID].Name;
+	}
+	else
+	{
+		OutObject = InNode.Object;
+		return OutObject.GetFName();
+	}
+}
+
+FName FEDLCookChecker::FEDLNodeHash::ObjectNameNext(const FEDLNodeHash& InNode, uint32& OutNodeID, TObjectPtr<const UObject>& OutObject)
+{
+	if (InNode.bIsNode)
+	{
+		OutNodeID = (*InNode.Nodes)[OutNodeID].ParentID;
+		return OutNodeID != NodeIDInvalid ? (*InNode.Nodes)[OutNodeID].Name : NAME_None;
+	}
+	else
+	{
+		OutObject = OutObject.GetOuter();
+		return OutObject ? OutObject.GetFName() : NAME_None;
+	}
+}
+
+FEDLCookChecker::FEDLNodeData::FEDLNodeData(FEDLNodeID InID, FEDLNodeID InParentID, FName InName, EObjectEvent InObjectEvent)
+	: Name(InName)
+	, ID(InID)
+	, ParentID(InParentID)
+	, ObjectEvent(InObjectEvent)
+	, bIsExport(false)
+{
+}
+
+FEDLCookChecker::FEDLNodeData::FEDLNodeData(FEDLNodeID InID, FEDLNodeID InParentID, FName InName, FEDLNodeData&& Other)
+	: Name(InName)
+	, ID(InID)
+	, ImportingPackagesSorted(MoveTemp(Other.ImportingPackagesSorted))
+	, ParentID(InParentID)
+	, ObjectEvent(Other.ObjectEvent)
+	, bIsExport(Other.bIsExport)
+{
+	// Note that Other Name and ParentID must be unmodified, since they might still be needed for GetHashCode calls from children
+	Other.ImportingPackagesSorted.Empty();
+}
+
+FEDLCookChecker::FEDLNodeHash FEDLCookChecker::FEDLNodeData::GetNodeHash(const FEDLCookChecker& Owner) const
+{
+	return FEDLNodeHash(&Owner.Nodes, ID, ObjectEvent);
+}
+
+FString FEDLCookChecker::FEDLNodeData::ToString(const FEDLCookChecker& Owner) const
+{
+	TStringBuilder<NAME_SIZE> Result;
+	switch (ObjectEvent)
+	{
+	case EObjectEvent::Create:
+		Result << TEXT("Create:");
+		break;
+	case EObjectEvent::Serialize:
+		Result << TEXT("Serialize:");
+		break;
+	default:
+		check(false);
+		break;
+	}
+	AppendPathName(Owner, Result);
+	return FString(Result);
+}
+
+void FEDLCookChecker::FEDLNodeData::AppendPathName(const FEDLCookChecker& Owner, FStringBuilderBase& Result) const
+{
+	if (ParentID != NodeIDInvalid)
+	{
+		const FEDLNodeData& ParentNode = Owner.Nodes[ParentID];
+		ParentNode.AppendPathName(Owner, Result);
+		bool bParentIsOutermost = ParentNode.ParentID == NodeIDInvalid;
+		Result << (bParentIsOutermost ? TEXT(".") : SUBOBJECT_DELIMITER);
+	}
+	Name.AppendString(Result);
+}
+
+FName FEDLCookChecker::FEDLNodeData::GetPackageName(const FEDLCookChecker& Owner) const
+{
+	if (ParentID != NodeIDInvalid)
+	{
+		// @todo ExternalPackages: We need to store ExternalPackage pointers on the Node and return that
+		return Owner.Nodes[ParentID].GetPackageName(Owner);
+	}
+	return Name;
+}
+
+void FEDLCookChecker::FEDLNodeData::Merge(FEDLCookChecker::FEDLNodeData&& Other)
+{
+	check(ObjectEvent == Other.ObjectEvent);
+	bIsExport = bIsExport || Other.bIsExport;
+
+	ImportingPackagesSorted.Append(Other.ImportingPackagesSorted);
+	Algo::Sort(ImportingPackagesSorted, FNameFastLess());
+	ImportingPackagesSorted.SetNum(Algo::Unique(ImportingPackagesSorted), EAllowShrinking::Yes);
+}
+
+FEDLCookCheckerThreadState::FEDLCookCheckerThreadState()
+{
+	Checker.SetActiveIfNeeded();
+
+	FScopeLock CookCheckerInstanceLock(&FEDLCookChecker::CookCheckerInstanceCritical);
+	FEDLCookChecker::CookCheckerInstances.Add(&Checker);
+}
+
+void FEDLCookChecker::SetActiveIfNeeded()
+{
+	bIsActive = !FParse::Param(FCommandLine::Get(), TEXT("DisableEDLCookChecker"));
+}
+
+void FEDLCookChecker::Reset()
+{
+	check(!GIsSavingPackage);
+
+	Nodes.Reset();
+	NodeHashToNodeID.Reset();
+	NodePrereqs.Reset();
+	bIsActive = false;
+}
+
+void FEDLCookChecker::AddImport(TObjectPtr<UObject> Import, UPackage* ImportingPackage)
+{
+	if (bIsActive)
+	{
+		if (!Import->GetOutermost()->HasAnyPackageFlags(PKG_CompiledIn))
+		{
+			LLM_SCOPE_BYTAG(EDLCookChecker);
+			FEDLNodeID NodeId = FindOrAddNode(FEDLNodeHash(Import, EObjectEvent::Serialize));
+			FEDLNodeData& NodeData = Nodes[NodeId];
+			FName ImportingPackageName = ImportingPackage->GetFName();
+			TArray<FName>& Sorted = NodeData.ImportingPackagesSorted;
+			int32 InsertionIndex = Algo::LowerBound(Sorted, ImportingPackageName, FNameFastLess());
+			if (InsertionIndex == Sorted.Num() || Sorted[InsertionIndex] != ImportingPackageName)
+			{
+				Sorted.Insert(ImportingPackageName, InsertionIndex);
+			}
+		}
+	}
+}
+
+void FEDLCookChecker::AddExport(UObject* Export)
+{
+	if (bIsActive)
+	{
+		LLM_SCOPE_BYTAG(EDLCookChecker);
+		FEDLNodeID SerializeID = FindOrAddNode(FEDLNodeHash(Export, EObjectEvent::Serialize));
+		Nodes[SerializeID].bIsExport = true;
+		FEDLNodeID CreateID = FindOrAddNode(FEDLNodeHash(Export, EObjectEvent::Create));
+		Nodes[CreateID].bIsExport = true;
+		AddDependency(SerializeID, CreateID); // every export must be created before it can be serialized...these arcs are implicit and not listed in any table.
+	}
+}
+
+void FEDLCookChecker::AddArc(UObject* DepObject, bool bDepIsSerialize, UObject* Export, bool bExportIsSerialize)
+{
+	if (bIsActive)
+	{
+		LLM_SCOPE_BYTAG(EDLCookChecker);
+		FEDLNodeID ExportID = FindOrAddNode(FEDLNodeHash(Export, bExportIsSerialize ? EObjectEvent::Serialize : EObjectEvent::Create));
+		FEDLNodeID DepID = FindOrAddNode(FEDLNodeHash(DepObject, bDepIsSerialize ? EObjectEvent::Serialize : EObjectEvent::Create));
+		AddDependency(ExportID, DepID);
+	}
+}
+
+void FEDLCookChecker::AddPackageWithUnknownExports(FName LongPackageName)
+{
+	LLM_SCOPE_BYTAG(EDLCookChecker);
+	if (bIsActive)
+	{
+		LLM_SCOPE_BYTAG(EDLCookChecker);
+		PackagesWithUnknownExports.Add(LongPackageName);
+	}
+}
+
+
+void FEDLCookChecker::AddDependency(FEDLNodeID SourceID, FEDLNodeID TargetID)
+{
+	NodePrereqs.Add(SourceID, TargetID);
+}
+
+void FEDLCookChecker::StartSavingEDLCookInfoForVerification()
+{
+	LLM_SCOPE_BYTAG(EDLCookChecker);
+	FScopeLock CookCheckerInstanceLock(&CookCheckerInstanceCritical);
+	for (FEDLCookChecker* Checker : CookCheckerInstances)
+	{
+		Checker->Reset();
+		Checker->SetActiveIfNeeded();
+	}
+}
+
+bool FEDLCookChecker::CheckForCyclesInner(TSet<FEDLNodeID>& Visited, TSet<FEDLNodeID>& Stack, const FEDLNodeID& Visit, FEDLNodeID& FailNode)
+{
+	bool bResult = false;
+	if (Stack.Contains(Visit))
+	{
+		FailNode = Visit;
+		bResult = true;
+	}
+	else
+	{
+		bool bWasAlreadyTested = false;
+		Visited.Add(Visit, &bWasAlreadyTested);
+		if (!bWasAlreadyTested)
+		{
+			Stack.Add(Visit);
+			for (auto It = NodePrereqs.CreateConstKeyIterator(Visit); !bResult && It; ++It)
+			{
+				bResult = CheckForCyclesInner(Visited, Stack, It.Value(), FailNode);
+			}
+			Stack.Remove(Visit);
+		}
+	}
+	UE_CLOG(bResult && Stack.Contains(FailNode), LogSavePackage, Error, TEXT("Cycle Node %s"), *Nodes[Visit].ToString(*this));
+	return bResult;
+}
+
+FEDLCookChecker::FEDLNodeID FEDLCookChecker::FindOrAddNode(const FEDLNodeHash& NodeHash)
+{
+	uint32 TypeHash = GetTypeHash(NodeHash);
+	FEDLNodeID* NodeIDPtr = NodeHashToNodeID.FindByHash(TypeHash, NodeHash);
+	if (NodeIDPtr)
+	{
+		return *NodeIDPtr;
+	}
+
+	FName Name = NodeHash.GetName();
+	FEDLNodeHash ParentHash;
+	FEDLNodeID ParentID = NodeHash.TryGetParent(ParentHash) ? FindOrAddNode(ParentHash) : NodeIDInvalid;
+	FEDLNodeID NodeID = Nodes.Num();
+	FEDLNodeData& NewNodeData = Nodes.Emplace_GetRef(NodeID, ParentID, Name, NodeHash.GetObjectEvent());
+	NodeHashToNodeID.AddByHash(TypeHash, NewNodeData.GetNodeHash(*this), NodeID);
+	return NodeID;
+}
+
+FEDLCookChecker::FEDLNodeID FEDLCookChecker::FindOrAddNode(FEDLNodeData&& NodeData, const FEDLCookChecker& OldOwnerOfNode, FEDLNodeID ParentIDInThis, bool& bNew)
+{
+	// Note that NodeData's Name and ParentID must be unmodified, since they might still be needed for GetHashCode calls from children
+
+	FEDLNodeHash NodeHash = NodeData.GetNodeHash(OldOwnerOfNode);
+	uint32 TypeHash = GetTypeHash(NodeHash);
+	FEDLNodeID* NodeIDPtr = NodeHashToNodeID.FindByHash(TypeHash, NodeHash);
+	if (NodeIDPtr)
+	{
+		bNew = false;
+		return *NodeIDPtr;
+	}
+
+	FEDLNodeID NodeID = Nodes.Num();
+	FEDLNodeData& NewNodeData = Nodes.Emplace_GetRef(NodeID, ParentIDInThis, NodeData.Name, MoveTemp(NodeData));
+	NodeHashToNodeID.AddByHash(TypeHash, NewNodeData.GetNodeHash(*this), NodeID);
+	bNew = true;
+	return NodeID;
+}
+
+FEDLCookChecker::FEDLNodeID FEDLCookChecker::FindNode(const FEDLNodeHash& NodeHash)
+{
+	const FEDLNodeID* NodeIDPtr = NodeHashToNodeID.Find(NodeHash);
+	return NodeIDPtr ? *NodeIDPtr : NodeIDInvalid;
+}
+
+void FEDLCookChecker::Merge(FEDLCookChecker&& Other)
+{
+	if (Nodes.Num() == 0)
+	{
+		Swap(Nodes, Other.Nodes);
+		Swap(NodeHashToNodeID, Other.NodeHashToNodeID);
+		Swap(NodePrereqs, Other.NodePrereqs);
+
+		// Switch the pointers in all of the swapped data to point at this instead of Other
+		for (TPair<FEDLNodeHash, FEDLNodeID>& KVPair : NodeHashToNodeID)
+		{
+			FEDLNodeHash& NodeHash = KVPair.Key;
+			NodeHash.SetNodes(&Nodes);
+		}
+	}
+	else
+	{
+		Other.NodeHashToNodeID.Empty(); // We will be invalidating the data these NodeHashes point to in the Other.Nodes loop, so empty the array now to avoid using it by accident
+
+		TArray<FEDLNodeID> RemapIDs;
+		RemapIDs.Reserve(Other.Nodes.Num());
+		for (FEDLNodeData& NodeData : Other.Nodes)
+		{
+			FEDLNodeID ParentID;
+			if (NodeData.ParentID == NodeIDInvalid)
+			{
+				ParentID = NodeIDInvalid;
+			}
+			else
+			{
+				// Parents should be earlier in the nodes list than children, since we always FindOrAdd the parent (and hence add it to the nodelist) when creating the child.
+				// Since the parent is earlier in the nodes list, we have already transferred it, and its ID in this->Nodes is therefore RemapIDs[Other.ParentID]
+				check(NodeData.ParentID < NodeData.ID);
+				ParentID = RemapIDs[NodeData.ParentID];
+			}
+
+			bool bNew;
+			FEDLNodeID NodeID = FindOrAddNode(MoveTemp(NodeData), Other, ParentID, bNew);
+			if (!bNew)
+			{
+				Nodes[NodeID].Merge(MoveTemp(NodeData));
+			}
+			RemapIDs.Add(NodeID);
+		}
+
+		for (const TPair<FEDLNodeID, FEDLNodeID>& Prereq : Other.NodePrereqs)
+		{
+			FEDLNodeID SourceID = RemapIDs[Prereq.Key];
+			FEDLNodeID TargetID = RemapIDs[Prereq.Value];
+			AddDependency(SourceID, TargetID);
+		}
+
+		Other.NodePrereqs.Empty();
+		Other.Nodes.Empty();
+	}
+
+	if (PackagesWithUnknownExports.Num() == 0)
+	{
+		Swap(PackagesWithUnknownExports, Other.PackagesWithUnknownExports);
+	}
+	else
+	{
+		PackagesWithUnknownExports.Reserve(Other.PackagesWithUnknownExports.Num());
+		for (FName PackageName : Other.PackagesWithUnknownExports)
+		{
+			PackagesWithUnknownExports.Add(PackageName);
+		}
+		Other.PackagesWithUnknownExports.Empty();
+	}
+}
+
+FEDLCookChecker FEDLCookChecker::AccumulateAndClear()
+{
+	FEDLCookChecker Accumulator;
+
+	FScopeLock CookCheckerInstanceLock(&CookCheckerInstanceCritical);
+	for (FEDLCookChecker* Checker : CookCheckerInstances)
+	{
+		if (Checker->bIsActive)
+		{
+			Accumulator.bIsActive = true;
+			Accumulator.Merge(MoveTemp(*Checker));
+			Checker->Reset();
+			Checker->bIsActive = true;
+		}
+	}
+	return Accumulator;
+}
+
+void FEDLCookChecker::Verify(const TFunction<void(UE::FLogRecord&& Record)>& MessageCallback,
+	bool bFullReferencesExpected)
+{
+	LLM_SCOPE_BYTAG(EDLCookChecker);
+
+	check(!GIsSavingPackage);
+	FEDLCookChecker Accumulator = AccumulateAndClear();
+
+	FString SeverityStr;
+	GConfig->GetString(TEXT("CookSettings"), TEXT("CookContentMissingSeverity"), SeverityStr, GEditorIni);
+	ELogVerbosity::Type MissingContentSeverity = ParseLogVerbosityFromString(SeverityStr);
+
+	if (Accumulator.bIsActive)
+	{
+		double StartTime = FPlatformTime::Seconds();
+
+		if (bFullReferencesExpected)
+		{
+			// imports to things that are not exports...
+			for (const FEDLNodeData& NodeData : Accumulator.Nodes)
+			{
+				if (NodeData.bIsExport)
+				{
+					// The node is an export; imports of it are valid
+					continue;
+				}
+
+				if (Accumulator.PackagesWithUnknownExports.Contains(NodeData.GetPackageName(Accumulator)))
+				{
+					// The node is an object in a package that exists, but for which we do not know the exports
+					// because e.g. it was iteratively skipped in the current cook. Suppress warnings about it
+					continue;
+				}
+
+				// Any imports of this non-exported node are an error; log them all if they exist
+				if (NodeData.ImportingPackagesSorted.IsEmpty())
+				{
+					continue;
+				}
+
+				const FEDLNodeData* NodeDataOfExportPackage = &NodeData;
+				while (NodeDataOfExportPackage->ParentID != NodeIDInvalid)
+				{
+					int32 ParentNodeIndex = static_cast<int32>(NodeDataOfExportPackage->ParentID);
+					check(Accumulator.Nodes.IsValidIndex(ParentNodeIndex));
+					NodeDataOfExportPackage = &Accumulator.Nodes[ParentNodeIndex];
+				}
+
+				const TCHAR* ReasonExportIsMissing = TEXT("");
+				if (NodeDataOfExportPackage->bIsExport)
+				{
+					ReasonExportIsMissing = TEXT("the object was stripped out of the target package when saved");
+				}
+				else
+				{
+					ReasonExportIsMissing = TEXT("the target package was marked NeverCook or is not cookable for the target platform");
+				}
+
+				for (FName PackageName : NodeData.ImportingPackagesSorted)
+				{
+					UE::FLogRecord Record;
+#if !NO_LOGGING
+					Record.SetCategory(LogSavePackage.GetCategoryName());
+#endif
+					Record.SetVerbosity(MissingContentSeverity);
+					Record.SetTime(UE::FLogTime::Now());
+					Record.SetFormat(TEXT("Content is missing from cook. Source package referenced an object in target package but {Reason}.")
+						TEXT("\n\tSource package: {Source}")
+						TEXT("\n\tTarget package: {Target}")
+						TEXT("\n\tReferenced object: {ReferencedObject}"));
+					{
+						FCbWriter Writer;
+						Writer.BeginObject();
+						Writer << "Reason" << ReasonExportIsMissing;
+						Writer << "Source" << WriteToUtf8String<256>(PackageName);
+						Writer << "Target" << WriteToUtf8String<256>(NodeDataOfExportPackage->Name);
+						{
+							TStringBuilder<256> ReferencedObjectStr;
+							NodeData.AppendPathName(Accumulator, ReferencedObjectStr);
+							Writer << "ReferencedObject" << ReferencedObjectStr;
+						}
+						Writer.EndObject();
+						Record.SetFields(Writer.Save().AsObject());
+					}
+					Record.SetFile(__FILE__);
+					Record.SetLine(__LINE__);
+					MessageCallback(MoveTemp(Record));
+				}
+			}
+		}
+
+		// cycles in the dep graph
+		TSet<FEDLNodeID> Visited;
+		TSet<FEDLNodeID> Stack;
+		bool bHadCycle = false;
+		for (const FEDLNodeData& NodeData : Accumulator.Nodes)
+		{
+			if (!NodeData.bIsExport)
+			{
+				continue;
+			}
+			FEDLNodeID FailNode;
+			if (Accumulator.CheckForCyclesInner(Visited, Stack, NodeData.ID, FailNode))
+			{
+				UE_LOG(LogSavePackage, Error, TEXT("----- %s contained a cycle (listed above)."), *Accumulator.Nodes[FailNode].ToString(Accumulator));
+				bHadCycle = true;
+			}
+		}
+		if (bHadCycle)
+		{
+			UE_LOG(LogSavePackage, Fatal, TEXT("EDL dep graph contained a cycle (see errors, above). This is fatal at runtime so it is fatal at cook time."));
+		}
+		UE_LOG(LogSavePackage, Display, TEXT("Took %fs to verify the EDL loading graph."), float(FPlatformTime::Seconds() - StartTime));
+	}
+}
+
+void FEDLCookChecker::MoveToCompactBinaryAndClear(FCbWriter& Writer, bool& bOutHasData)
+{
+	LLM_SCOPE_BYTAG(EDLCookChecker);
+
+	bOutHasData = false;
+	FEDLCookChecker Accumulator = AccumulateAndClear();
+	if (!Accumulator.bIsActive)
+	{
+		return;
+	}
+	if (Accumulator.Nodes.IsEmpty() && Accumulator.NodePrereqs.IsEmpty() && Accumulator.PackagesWithUnknownExports.IsEmpty())
+	{
+		return;
+	}
+	bOutHasData = true;
+
+	Accumulator.WriteToCompactBinary(Writer);
+}
+
+bool FEDLCookChecker::AppendFromCompactBinary(FCbFieldView Field)
+{
+	LLM_SCOPE_BYTAG(EDLCookChecker);
+	FEDLCookChecker Instance;
+	if (!Instance.ReadFromCompactBinary(Field))
+	{
+		return false;
+	}
+	FEDLCookChecker& CurrentChecker = FEDLCookCheckerThreadState::Get().Checker;
+	CurrentChecker.Merge(MoveTemp(Instance));
+	return true;
+}
+
+void FEDLCookChecker::WriteToCompactBinary(FCbWriter& Writer)
+{
+	Writer.BeginObject();
+	{
+		Writer.BeginArray("Nodes");
+		for (const FEDLNodeData& Node : Nodes)
+		{
+			Writer << Node.Name;
+			Writer << Node.ImportingPackagesSorted;
+			Writer << Node.ParentID;
+			Writer << static_cast<uint8>(Node.ObjectEvent);
+			Writer << Node.bIsExport;
+		}
+		Writer.EndArray();
+		Writer.BeginArray("NodePrereqs");
+		for (const TPair<FEDLNodeID, FEDLNodeID>& Pair : NodePrereqs)
+		{
+			Writer << static_cast<uint32>(Pair.Key);
+			Writer << static_cast<uint32>(Pair.Value);
+		}
+		Writer.EndArray();
+		Writer.BeginArray("PackagesWithUnknownExports");
+		for (FName PackageName : PackagesWithUnknownExports)
+		{
+			Writer << PackageName;
+		}
+		Writer.EndArray();
+	}
+	Writer.EndObject();
+}
+
+bool FEDLCookChecker::ReadFromCompactBinary(FCbFieldView Field)
+{
+	Reset();
+
+	bool bSuccess = false;
+	ON_SCOPE_EXIT
+	{
+		if (!bSuccess)
+		{
+			Reset();
+		}
+	};
+
+	FCbFieldView NodesField = Field["Nodes"];
+	const uint64 NumNodes = NodesField.AsArrayView().Num() / 5;
+	if (NumNodes > MAX_int32)
+	{
+		return false;
+	}
+	Nodes.Reserve(static_cast<int32>(NumNodes));
+	if (NodesField.HasError())
+	{
+		return false;
+	}
+	FCbFieldViewIterator NodeIter = NodesField.CreateViewIterator();
+	while (NodeIter)
+	{
+		FEDLNodeID NodeID = Nodes.Num();
+		FEDLNodeData& Node = Nodes.Emplace_GetRef();
+		Node.ID = NodeID;
+		if (!LoadFromCompactBinary(NodeIter, Node.Name)) { return false; }
+		++NodeIter;
+		if (!LoadFromCompactBinary(NodeIter, Node.ImportingPackagesSorted)) { return false; }
+		++NodeIter;
+		if (!LoadFromCompactBinary(NodeIter, Node.ParentID)) { return false; }
+		++NodeIter;
+		uint8 LocalObjectEvent;
+		if (!LoadFromCompactBinary(NodeIter, LocalObjectEvent)) { return false; }
+		if (LocalObjectEvent > static_cast<uint8>(EObjectEvent::Max)) { return false; }
+		Node.ObjectEvent = static_cast<EObjectEvent>(LocalObjectEvent);
+		++NodeIter;
+		if (!LoadFromCompactBinary(NodeIter, Node.bIsExport)) { return false; }
+		++NodeIter;
+	}
+
+	FCbFieldView PrereqsField = Field["NodePrereqs"];
+	const int64 NumNodePrereqs = PrereqsField.AsArrayView().Num() / 2;
+	if (NumNodePrereqs > MAX_int32)
+	{
+		return false;
+	}
+	NodePrereqs.Reserve(static_cast<int32>(NumNodePrereqs));
+	if (PrereqsField.HasError())
+	{
+		return false;
+	}
+	FCbFieldViewIterator PrereqsIter = PrereqsField.CreateViewIterator();
+	while (PrereqsIter)
+	{
+		uint32 Key;
+		uint32 Value;
+		if (!LoadFromCompactBinary(PrereqsIter, Key)) { return false; }
+		++PrereqsIter;
+		if (!LoadFromCompactBinary(PrereqsIter, Value)) { return false; }
+		++PrereqsIter;
+		NodePrereqs.Add(static_cast<FEDLNodeID>(Key), static_cast<FEDLNodeID>(Value));
+	}
+
+	FCbFieldView PackagesWithUnknownExportsField = Field["PackagesWithUnknownExports"];
+	const int64 NumPackagesWithUnknownExports = PackagesWithUnknownExportsField.AsArrayView().Num();
+	if (NumPackagesWithUnknownExports > MAX_int32)
+	{
+		return false;
+	}
+	PackagesWithUnknownExports.Reserve(static_cast<int32>(NumPackagesWithUnknownExports));
+	if (PackagesWithUnknownExportsField.HasError())
+	{
+		return false;
+	}
+	for (FCbFieldView PackageNameField : PackagesWithUnknownExportsField)
+	{
+		FName PackageName;
+		if (!LoadFromCompactBinary(PackageNameField, PackageName))
+		{
+			return false;
+		}
+		PackagesWithUnknownExports.Add(PackageName);
+	}
+
+	for (const FEDLNodeData& Node : Nodes)
+	{
+		NodeHashToNodeID.Add(Node.GetNodeHash(*this), Node.ID);
+	}
+	bIsActive = !Nodes.IsEmpty() || !NodePrereqs.IsEmpty() || !PackagesWithUnknownExports.IsEmpty();
+
+	bSuccess = true;
+	return true;
+}
+
+FCriticalSection FEDLCookChecker::CookCheckerInstanceCritical;
+TArray<FEDLCookChecker*> FEDLCookChecker::CookCheckerInstances;
diff --git a/Engine/Source/Editor/UnrealEd/Private/Cooker/CookImportsChecker.h b/Engine/Source/Editor/UnrealEd/Private/Cooker/CookImportsChecker.h
new file mode 100644
--- /dev/null
+++ b/Engine/Source/Editor/UnrealEd/Private/Cooker/CookImportsChecker.h
@@ -0,0 +1,201 @@
+// Copyright Epic Games, Inc. All Rights Reserved.
+
+#pragma once
+
+#include "Containers/Array.h"
+#include "Containers/Map.h"
+#include "Containers/StringFwd.h"
+#include "HAL/CriticalSection.h"
+#include "HAL/ThreadSingleton.h"
+#include "Templates/Function.h"
+#include "UObject/NameTypes.h"
+#include "UObject/Object.h"
+
+class FCbFieldView;
+class FCbWriter;
+class UPackage;
+namespace UE { class FLogRecord; }
+
+struct FEDLCookCheckerThreadState;
+
+/**
+ * Helper struct used during cooking to validate EDL dependencies
+ */
+struct FEDLCookChecker
+{
+	void SetActiveIfNeeded();
+
+	void Reset();
+
+	void AddImport(TObjectPtr<UObject> Import, UPackage* ImportingPackage);
+	void AddExport(UObject* Export);
+	void AddArc(UObject* DepObject, bool bDepIsSerialize, UObject* Export, bool bExportIsSerialize);
+	void AddPackageWithUnknownExports(FName LongPackageName);
+
+	static void StartSavingEDLCookInfoForVerification();
+	static void Verify(const TFunction<void(UE::FLogRecord&& Record)>& MessageCallback,
+		bool bFullReferencesExpected);
+	static void MoveToCompactBinaryAndClear(FCbWriter& Writer, bool& bOutHasData);
+	static bool AppendFromCompactBinary(FCbFieldView Field);
+
+private:
+	static FEDLCookChecker AccumulateAndClear();
+	void WriteToCompactBinary(FCbWriter& Writer);
+	bool ReadFromCompactBinary(FCbFieldView Field);
+
+	typedef uint32 FEDLNodeID;
+	static const FEDLNodeID NodeIDInvalid = static_cast<FEDLNodeID>(-1);
+
+	struct FEDLNodeData;
+public: // FEDLNodeHash is public only so that GetTypeHash can be defined
+	enum class EObjectEvent : uint8
+	{
+		Create,
+		Serialize,
+		Max = Serialize,
+	};
+
+	/**
+	 * Wrapper around an FEDLNodeData (or around a UObject when searching for an FEDLNodeData corresponding to the UObject)
+	 * that provides the hash-by-objectpath to lookup the FEDLNodeData for an objectpath.
+	 */
+	struct FEDLNodeHash
+	{
+		FEDLNodeHash(); // creates an uninitialized node; only use this to provide as an out parameter
+		FEDLNodeHash(const TArray<FEDLNodeData>* InNodes, FEDLNodeID InNodeID, EObjectEvent InObjectEvent);
+		FEDLNodeHash(TObjectPtr<UObject> InObject, EObjectEvent InObjectEvent);
+		FEDLNodeHash(const FEDLNodeHash& Other);
+		bool operator==(const FEDLNodeHash& Other) const;
+		FEDLNodeHash& operator=(const FEDLNodeHash& Other);
+		friend uint32 GetTypeHash(const FEDLNodeHash& A);
+
+		FName GetName() const;
+		bool TryGetParent(FEDLNodeHash& Parent) const;
+		EObjectEvent GetObjectEvent() const;
+		void SetNodes(const TArray<FEDLNodeData>* InNodes);
+
+	private:
+		static FName ObjectNameFirst(const FEDLNodeHash& InNode, uint32& OutNodeID, TObjectPtr<const UObject>& OutObject);
+		static FName ObjectNameNext(const FEDLNodeHash& InNode, uint32& OutNodeID, TObjectPtr<const UObject>& OutObject);
+
+		union
+		{
+			/**
+			 * The array of nodes from the FEDLCookChecker; this is how we lookup the node for the FEDLNodeData.
+			 * Because the FEDLNodeData are elements in an array which can resize and therefore reallocate the nodes, we cannot store the pointer to the node.
+			 * Only used if bIsNode is true.
+			 */
+			const TArray<FEDLNodeData>* Nodes;
+			/** Pointer to the Object we are looking up, if this hash was created during lookup-by-objectpath for an object */
+			TObjectPtr<const UObject> Object;
+		};
+		/** The identifier for the FEDLNodeData this hash is wrapping. Only used if bIsNode is true. */
+		FEDLNodeID NodeID;
+		/** True if this hash is wrapping an FEDLNodeData, false if it is wrapping a UObject. */
+		bool bIsNode;
+		EObjectEvent ObjectEvent;
+	};
+
+private:
+
+	/**
+	 * Node representing either the Create event or Serialize event of a UObject in the graph of runtime dependencies between UObjects.
+	 */
+	struct FEDLNodeData
+	{
+		// Note that order of the variables is important to reduce alignment waste in the size of FEDLNodeData.
+		/** Name of the UObject represented by this node; full objectpath name is obtainable by traversing parent. */
+		FName Name;
+		/**
+		 * Index of this node in the FEDLCookChecker's Nodes array. This index is used to provide a small-memory-usage
+		 * identifier for the node.
+		 */
+		FEDLNodeID ID;
+		/**
+		 * Tracks references to this node's UObjects from other packages (which is the reverse of the references
+		 * from each node that we track in NodePrereqs.)
+		 * We only need this information from each package, so we track by package name instead of node id.
+		 */
+		TArray<FName> ImportingPackagesSorted;
+		/**
+		 * ID of the node representing the UObject parent of this node's UObject. NodeIDInvalid if the UObject has no parent.
+		 * The ParentID always refers to the node for the Create event of the parent UObject.
+		 */
+		uint32 ParentID;
+		/** True if this node represents the Serialize event on the UObject, false if it represents the Create event. */
+		EObjectEvent ObjectEvent;
+		/**
+		 * True if the UObject represented by this node has been exported by a SavePackage call; used to verify that
+		 * the imports requested by packages are present somewhere in the cook.
+		 */
+		bool bIsExport;
+
+		FEDLNodeData() { /* Fields are uninitialized */ }
+		FEDLNodeData(FEDLNodeID InID, FEDLNodeID InParentID, FName InName, EObjectEvent InObjectEvent);
+		FEDLNodeData(FEDLNodeID InID, FEDLNodeID InParentID, FName InName, FEDLNodeData&& Other);
+		FEDLNodeHash GetNodeHash(const FEDLCookChecker& Owner) const;
+
+		FString ToString(const FEDLCookChecker& Owner) const;
+		void AppendPathName(const FEDLCookChecker& Owner, FStringBuilderBase& Result) const;
+		FName GetPackageName(const FEDLCookChecker& Owner) const;
+		void Merge(FEDLNodeData&& Other);
+	};
+
+	FEDLNodeID FindOrAddNode(const FEDLNodeHash& NodeLookup);
+	FEDLNodeID FindOrAddNode(FEDLNodeData&& NodeData, const FEDLCookChecker& OldOwnerOfNode, FEDLNodeID ParentIDInThis, bool& bNew);
+	FEDLNodeID FindNode(const FEDLNodeHash& NodeHash);
+	void Merge(FEDLCookChecker&& Other);
+	bool CheckForCyclesInner(TSet<FEDLNodeID>& Visited, TSet<FEDLNodeID>& Stack, const FEDLNodeID& Visit, FEDLNodeID& FailNode);
+	void AddDependency(FEDLNodeID SourceID, FEDLNodeID TargetID);
+
+	/**
+	 * All the FEDLNodeDatas that have been created for this checker. These are allocated as elements of an array rather than pointers to reduce cputime and
+	 * memory due to many small allocations, and to provide index-based identifiers. Nodes are not deleted until the checker is reset.
+	 */
+	TArray<FEDLNodeData> Nodes;
+	/** A map to lookup the node for a UObject or for the corresponding node in another thread's FEDLCookChecker. */
+	TMap<FEDLNodeHash, FEDLNodeID> NodeHashToNodeID;
+	/** The graph of dependencies between nodes. */
+	TMultiMap<FEDLNodeID, FEDLNodeID> NodePrereqs;
+	/**
+	 * Packages that were cooked iteratively and therefore have an unknown set of exports.
+	 * We suppress warnings for exports missing from these packages.
+	 */
+	TSet<FName> PackagesWithUnknownExports;
+	/** True if the EDLCookChecker should be active; it is turned off if the runtime will not be using EDL. */
+	bool bIsActive = false;
+
+	/** When cooking with concurrent saving, each thread has its own FEDLCookChecker, and these are merged after the cook is complete. */
+	static FCriticalSection CookCheckerInstanceCritical;
+	static TArray<FEDLCookChecker*> CookCheckerInstances;
+
+	friend FEDLCookCheckerThreadState;
+};
+
+/** Per-thread accessor for writing EDL dependencies to global FEDLCookChecker storage. */
+struct FEDLCookCheckerThreadState : public TThreadSingleton<FEDLCookCheckerThreadState>
+{
+	FEDLCookCheckerThreadState();
+
+	void AddImport(TObjectPtr<UObject> Import, UPackage* ImportingPackage)
+	{
+		Checker.AddImport(Import, ImportingPackage);
+	}
+	void AddExport(UObject* Export)
+	{
+		Checker.AddExport(Export);
+	}
+	void AddArc(UObject* DepObject, bool bDepIsSerialize, UObject* Export, bool bExportIsSerialize)
+	{
+		Checker.AddArc(DepObject, bDepIsSerialize, Export, bExportIsSerialize);
+	}
+	void AddPackageWithUnknownExports(FName LongPackageName)
+	{
+		Checker.AddPackageWithUnknownExports(LongPackageName);
+	}
+
+private:
+	FEDLCookChecker Checker;
+	friend TThreadSingleton<FEDLCookCheckerThreadState>;
+	friend FEDLCookChecker;
+};
\ No newline at end of file
diff --git a/Engine/Source/Editor/UnrealEd/Private/Cooker/CookRequestCluster.cpp b/Engine/Source/Editor/UnrealEd/Private/Cooker/CookRequestCluster.cpp
--- a/Engine/Source/Editor/UnrealEd/Private/Cooker/CookRequestCluster.cpp
+++ b/Engine/Source/Editor/UnrealEd/Private/Cooker/CookRequestCluster.cpp
@@ -11,6 +11,7 @@
 #include "Async/Async.h"
 #include "Cooker/CookDependency.h"
 #include "Cooker/CookGenerationHelper.h"
+#include "Cooker/CookImportsChecker.h"
 #include "Cooker/CookPackageData.h"
 #include "Cooker/CookPlatformManager.h"
 #include "Cooker/CookProfiling.h"
@@ -2203,7 +2204,7 @@ void FRequestCluster::FGraphSearch::FExploreEdgesContext::ProcessPlatformAttachm
 						*PackageData->GetInstigator(EReachability::Runtime).ToString());
 				}
 				// Declare the package to the EDLCookInfo verification so we don't warn about missing exports from it
-				UE::SavePackageUtilities::EDLCookInfoAddIterativelySkippedPackage(PackageName);
+				FEDLCookCheckerThreadState::Get().AddPackageWithUnknownExports(PackageName);
 			}
 			else
 			{
diff --git a/Engine/Source/Editor/UnrealEd/Private/Cooker/CookSavePackage.cpp b/Engine/Source/Editor/UnrealEd/Private/Cooker/CookSavePackage.cpp
--- a/Engine/Source/Editor/UnrealEd/Private/Cooker/CookSavePackage.cpp
+++ b/Engine/Source/Editor/UnrealEd/Private/Cooker/CookSavePackage.cpp
@@ -10,6 +10,7 @@
 #include "Cooker/CookDiagnostics.h"
 #include "Cooker/CookEvents.h"
 #include "Cooker/CookGenerationHelper.h"
+#include "Cooker/CookImportsChecker.h"
 #include "Cooker/CookOnTheFlyServerInterface.h"
 #include "Cooker/CookPackageArtifacts.h"
 #include "Cooker/CookPackageData.h"
@@ -467,6 +468,7 @@ void FSaveCookedPackageContext::FinishPlatform()
 			OverridePackageDependencies.GetPtrOrNull() /** const, not modified */,
 			GeneratedPackageBuildResultDependencies.GetPtrOrNull() /** non-const, function sets it to empty */);
 		RecordPlatformBuildDependencies();
+		RecordCookImportsCheckerData();
 
 		ICookedPackageWriter::FCommitPackageInfo Info;
 		Info.Attachments = GetCommitAttachments();
@@ -874,6 +876,27 @@ void FSaveCookedPackageContext::RecordPlatformBuildDependencies()
 	}
 }
 
+void FSaveCookedPackageContext::RecordCookImportsCheckerData()
+{
+	if (HasSavePackageResult())
+	{
+		FEDLCookCheckerThreadState& CookChecker = FEDLCookCheckerThreadState::Get();
+		for (UObject* Import : SavePackageResult.Imports)
+		{
+			CookChecker.AddImport(Import, Package);
+		}
+		for (UObject* Export : SavePackageResult.Exports)
+		{
+			CookChecker.AddExport(Export);
+		}
+		for (UE::SavePackageUtilities::FPreloadDependency& PreloadDependency : SavePackageResult.PreloadDependencies)
+		{
+			CookChecker.AddArc(PreloadDependency.TargetObject, PreloadDependency.bTargetIsSerialize,
+				PreloadDependency.SourceObject, PreloadDependency.bSourceIsSerialize);
+		}
+	}
+}
+
 TArray<IPackageWriter::FCommitAttachmentInfo> FSaveCookedPackageContext::GetCommitAttachments()
 {
 	TArray<IPackageWriter::FCommitAttachmentInfo> Result;
diff --git a/Engine/Source/Editor/UnrealEd/Private/Cooker/CookSavePackage.h b/Engine/Source/Editor/UnrealEd/Private/Cooker/CookSavePackage.h
--- a/Engine/Source/Editor/UnrealEd/Private/Cooker/CookSavePackage.h
+++ b/Engine/Source/Editor/UnrealEd/Private/Cooker/CookSavePackage.h
@@ -53,6 +53,7 @@ private: // Used only by UCookOnTheFlyServer, which has private access
 	void CalculateCookDependencies(FGenerationHelper* GenerationHelper, bool bGenerated,
 		const TArray<FAssetDependency>* ExtraARDependencies, FBuildResultDependenciesMap* ExtraBuildResultDependencies);
 	void RecordPlatformBuildDependencies();
+	void RecordCookImportsCheckerData();
 	bool HasSavePackageResult() const;
 
 	TArray<IPackageWriter::FCommitAttachmentInfo> GetCommitAttachments();
diff --git a/Engine/Source/Editor/UnrealEd/Private/Cooker/DiffPackageWriter.cpp b/Engine/Source/Editor/UnrealEd/Private/Cooker/DiffPackageWriter.cpp
--- a/Engine/Source/Editor/UnrealEd/Private/Cooker/DiffPackageWriter.cpp
+++ b/Engine/Source/Editor/UnrealEd/Private/Cooker/DiffPackageWriter.cpp
@@ -21,7 +21,6 @@
 #include "Serialization/LargeMemoryWriter.h"
 #include "Templates/UniquePtr.h"
 #include "Templates/UnrealTemplate.h"
-#include "UObject/LinkerDiff.h"
 #include "UObject/LinkerSave.h"
 #include "UObject/SavePackage.h"
 #include "UObject/UObjectGlobals.h"
@@ -423,85 +422,4 @@ void FDiffPackageWriter::ConditionallyDumpObjects()
 			}
 		}
 	}
-}
-
-FLinkerDiffPackageWriter::FLinkerDiffPackageWriter(TUniquePtr<ICookedPackageWriter>&& InInner)
-	: Inner(MoveTemp(InInner))
-{
-	FString DiffModeText;
-	FParse::Value(FCommandLine::Get(), TEXT("-LINKERDIFF="), DiffModeText);
-	DiffMode = EDiffMode::LDM_Consistent; // (DiffModeText == TEXT("2") || DiffModeText == TEXT("consistent"))
-}
-
-void FLinkerDiffPackageWriter::BeginPackage(const FBeginPackageInfo& Info)
-{
-	bHasStartedSecondSave = false;
-	BeginInfo = Info;
-	Inner->BeginPackage(Info);
-	SetupOtherAlgorithm();
-}
-
-void FLinkerDiffPackageWriter::UpdateSaveArguments(FSavePackageArgs& SaveArgs)
-{
-	SaveArgs.SaveFlags |= SAVE_CompareLinker;
-	Inner->UpdateSaveArguments(SaveArgs);
-}
-
-bool FLinkerDiffPackageWriter::IsAnotherSaveNeeded(FSavePackageResultStruct& PreviousResult,
-	FSavePackageArgs& SaveArgs)
-{
-	checkf(!Inner->IsAnotherSaveNeeded(PreviousResult, SaveArgs),
-		TEXT("LinkerDiffPackageWriter does not support an Inner that needs multiple saves."));
-	if (PreviousResult == ESavePackageResult::Timeout)
-	{
-		OtherResult.LinkerSave.Reset();
-		PreviousResult.LinkerSave.Reset();
-		return false;
-	}
-
-	if (!bHasStartedSecondSave)
-	{
-		bHasStartedSecondSave = true;
-		OtherResult = MoveTemp(PreviousResult);
-
-		// Resave the package with the current save algorithm.
-		SetupCurrentAlgorithm();
-
-		// The contract with the Inner is that every Begin is paired with a single commit.
-		// Send the old commit and the new begin.
-		IPackageWriter::FCommitPackageInfo CommitInfo;
-		CommitInfo.Status = PackageResultToCommitStatus(OtherResult.Result);
-		CommitInfo.PackageName = BeginInfo.PackageName;
-		CommitInfo.WriteOptions = IPackageWriter::EWriteOptions::None;
-		Inner->CommitPackage(MoveTemp(CommitInfo));
-		Inner->BeginPackage(BeginInfo);
-		return true;
-	}
-	else
-	{
-		CompareResults(PreviousResult);
-
-		OtherResult.LinkerSave.Reset();
-		PreviousResult.LinkerSave.Reset();
-
-		return false;
-	}
-}
-
-void FLinkerDiffPackageWriter::SetupOtherAlgorithm()
-{
-}
-
-void FLinkerDiffPackageWriter::SetupCurrentAlgorithm()
-{
-}
-
-void FLinkerDiffPackageWriter::CompareResults(FSavePackageResultStruct& CurrentResult)
-{
-	if (OtherResult.LinkerSave && CurrentResult.LinkerSave)
-	{
-		FLinkerDiff LinkerDiff =
-			FLinkerDiff::CompareLinkers(OtherResult.LinkerSave.Get(), CurrentResult.LinkerSave.Get());
-		LinkerDiff.PrintDiff(*GWarn);
-	}
-}
+}
\ No newline at end of file
diff --git a/Engine/Source/Editor/UnrealEd/Private/Cooker/DiffPackageWriter.h b/Engine/Source/Editor/UnrealEd/Private/Cooker/DiffPackageWriter.h
--- a/Engine/Source/Editor/UnrealEd/Private/Cooker/DiffPackageWriter.h
+++ b/Engine/Source/Editor/UnrealEd/Private/Cooker/DiffPackageWriter.h
@@ -169,144 +169,3 @@ protected:
 	bool bDumpObjects = false;
 	bool bDumpObjectsSorted = false;
 };
-
-/** A CookedPackageWriter that runs the save twice with different methods (Save1 vs Save2). */
-class FLinkerDiffPackageWriter : public ICookedPackageWriter
-{
-public:
-	FLinkerDiffPackageWriter(TUniquePtr<ICookedPackageWriter>&& InInner);
-
-	// IPackageWriter
-	virtual FCapabilities GetCapabilities() const override
-	{
-		return Inner->GetCapabilities();
-	}
-	virtual void BeginPackage(const FBeginPackageInfo& Info) override;
-	virtual void CommitPackage(FCommitPackageInfo&& Info) override
-	{
-		return Inner->CommitPackage(MoveTemp(Info));
-	}
-	virtual void WritePackageData(const FPackageInfo& Info, FLargeMemoryWriter& ExportsArchive,
-		const TArray<FFileRegion>& FileRegions) override
-	{
-		Inner->WritePackageData(Info, ExportsArchive, FileRegions);
-	}
-	virtual void WriteBulkData(const FBulkDataInfo& Info, const FIoBuffer& BulkData,
-		const TArray<FFileRegion>& FileRegions) override
-	{
-		Inner->WriteBulkData(Info, BulkData, FileRegions);
-	}
-	virtual void WriteAdditionalFile(const FAdditionalFileInfo& Info, const FIoBuffer& FileData) override
-	{
-		Inner->WriteAdditionalFile(Info, FileData);
-	}
-	virtual void WriteLinkerAdditionalData(const FLinkerAdditionalDataInfo& Info, const FIoBuffer& Data,
-		const TArray<FFileRegion>& FileRegions) override
-	{
-		Inner->WriteLinkerAdditionalData(Info, Data, FileRegions);
-	}
-	virtual void WritePackageTrailer(const FPackageTrailerInfo& Info, const FIoBuffer& Data) override
-	{
-		Inner->WritePackageTrailer(Info, Data);
-	}
-	virtual int64 GetExportsFooterSize() override
-	{
-		return Inner->GetExportsFooterSize();
-	}
-	virtual TUniquePtr<FLargeMemoryWriter> CreateLinkerArchive(FName PackageName, UObject* Asset, uint16 MultiOutputIndex) override
-	{
-		return Inner->CreateLinkerArchive(PackageName, Asset, MultiOutputIndex);
-	}
-	virtual bool IsPreSaveCompleted() const override
-	{
-		return Inner->IsPreSaveCompleted();
-	}
-
-	// ICookedPackageWriter
-	virtual FCookCapabilities GetCookCapabilities() const override
-	{
-		FCookCapabilities Result = Inner->GetCookCapabilities();
-		Result.bDiffModeSupported = false; // LinkerDiffPackageWriter can not be an inner of DiffPackageWriter
-		return Result;
-	}
-	virtual FDateTime GetPreviousCookTime() const
-	{
-		return Inner->GetPreviousCookTime();
-	}
-	virtual void Initialize(const FCookInfo& Info) override
-	{
-		Inner->Initialize(Info);
-	}
-	virtual void BeginCook(const FCookInfo& Info) override
-	{
-		Inner->BeginCook(Info);
-	}
-	virtual void EndCook(const FCookInfo& Info) override
-	{
-		Inner->EndCook(Info);
-	}
-	virtual TUniquePtr<FAssetRegistryState> LoadPreviousAssetRegistry() override
-	{
-		return Inner->LoadPreviousAssetRegistry();
-	}
-	virtual FCbObject GetOplogAttachment(FName PackageName, FUtf8StringView AttachmentKey) override
-	{
-		return Inner->GetOplogAttachment(PackageName, AttachmentKey);
-	}
-	virtual void GetOplogAttachments(TArrayView<FName> PackageNames,
-		TArrayView<FUtf8StringView> AttachmentKeys,
-		TUniqueFunction<void(FName PackageName, FUtf8StringView AttachmentKey, FCbObject&& Attachment)>&& Callback) override
-	{
-		return Inner->GetOplogAttachments(PackageNames, AttachmentKeys, MoveTemp(Callback));
-	}
-	virtual ECommitStatus GetCommitStatus(FName PackageName) override
-	{
-		return Inner->GetCommitStatus(PackageName);
-	}
-	virtual void RemoveCookedPackages(TArrayView<const FName> PackageNamesToRemove) override
-	{
-		Inner->RemoveCookedPackages(PackageNamesToRemove);
-	}
-	virtual void RemoveCookedPackages() override
-	{
-		Inner->RemoveCookedPackages();
-	}
-	virtual void UpdatePackageModificationStatus(FName PackageName, bool bIterativelyUnmodified,
-		bool& bInOutShouldIterativelySkip) override
-	{
-		Inner->UpdatePackageModificationStatus(PackageName, bIterativelyUnmodified, bInOutShouldIterativelySkip);
-	}
-	virtual EPackageWriterResult BeginCacheForCookedPlatformData(FBeginCacheForCookedPlatformDataInfo& Info) override
-	{
-		return Inner->BeginCacheForCookedPlatformData(Info);
-	}
-	virtual void UpdateSaveArguments(FSavePackageArgs& SaveArgs) override;
-	virtual bool IsAnotherSaveNeeded(FSavePackageResultStruct& PreviousResult, FSavePackageArgs& SaveArgs) override;
-	virtual TFuture<FCbObject> WriteMPCookMessageForPackage(FName PackageName) override
-	{
-		return Inner->WriteMPCookMessageForPackage(PackageName);
-	}
-	virtual bool TryReadMPCookMessageForPackage(FName PackageName, FCbObjectView Message) override
-	{
-		return Inner->TryReadMPCookMessageForPackage(PackageName, Message);
-	}
-	virtual TMap<FName, TRefCountPtr<FPackageHashes>>& GetPackageHashes() override
-	{
-		return Inner->GetPackageHashes();
-	}
-private:
-	enum class EDiffMode : uint8
-	{
-		LDM_Consistent,
-	};
-
-	void SetupOtherAlgorithm();
-	void SetupCurrentAlgorithm();
-	void CompareResults(FSavePackageResultStruct& CurrentResult);
-
-	FSavePackageResultStruct OtherResult;
-	FBeginPackageInfo BeginInfo;
-	TUniquePtr<ICookedPackageWriter> Inner;
-	EDiffMode DiffMode = EDiffMode::LDM_Consistent;
-	bool bHasStartedSecondSave = false;
-};
diff --git a/Engine/Source/Runtime/Core/Public/Serialization/PackageWriter.h b/Engine/Source/Runtime/Core/Public/Serialization/PackageWriter.h
--- a/Engine/Source/Runtime/Core/Public/Serialization/PackageWriter.h
+++ b/Engine/Source/Runtime/Core/Public/Serialization/PackageWriter.h
@@ -264,7 +264,7 @@ public:
 
 	struct FCookCapabilities
 	{
-		/** Whether this writer implements -diffonly and -linkerdiff. */
+		/** Whether this writer implements -diffonly or other difftypes. */
 		bool bDiffModeSupported = false;
 
 		/** If true, the cooker will assume no packages are written and will skip writing non-package data. */
diff --git a/Engine/Source/Runtime/CoreUObject/Private/UObject/LinkerDiff.cpp b/Engine/Source/Runtime/CoreUObject/Private/UObject/LinkerDiff.cpp
--- a/Engine/Source/Runtime/CoreUObject/Private/UObject/LinkerDiff.cpp
+++ b/Engine/Source/Runtime/CoreUObject/Private/UObject/LinkerDiff.cpp
@@ -11,12 +11,14 @@
 #include "UObject/Package.h"
 #include "Internationalization/GatherableTextData.h"
 
+PRAGMA_DISABLE_DEPRECATION_WARNINGS;
 FLinkerDiff FLinkerDiff::CompareLinkers(FLinker* LHSLinker, FLinker* RHSLinker)
 {
 	FLinkerDiff Diff;
 	Diff.Generate(LHSLinker, RHSLinker);
 	return Diff;
 }
+PRAGMA_ENABLE_DEPRECATION_WARNINGS;
 
 COREUOBJECT_API bool FLinkerDiff::HasDiffs() const
 {
diff --git a/Engine/Source/Runtime/CoreUObject/Private/UObject/Package.cpp b/Engine/Source/Runtime/CoreUObject/Private/UObject/Package.cpp
--- a/Engine/Source/Runtime/CoreUObject/Private/UObject/Package.cpp
+++ b/Engine/Source/Runtime/CoreUObject/Private/UObject/Package.cpp
@@ -85,7 +85,7 @@ UPackage::FOnPackageDirtyStateChanged UPackage::PackageDirtyStateChangedEvent;
  */
 UPackage::FOnPackageMarkedDirty UPackage::PackageMarkedDirtyEvent;
 
-PRAGMA_DISABLE_DEPRECATION_WARNINGS // Silence deprecation warnings for deprecated CookedHash member
+PRAGMA_DISABLE_DEPRECATION_WARNINGS // Silence deprecation warnings for deprecated CookedHash and LinkerSave members
 FSavePackageResultStruct::FSavePackageResultStruct()
 	: Result(ESavePackageResult::Error), TotalFileSize(0), SerializedPackageFlags(0)
 {
@@ -98,11 +98,16 @@ FSavePackageResultStruct::FSavePackageResultStruct(ESavePackageResult InResult,
 	: Result(InResult), TotalFileSize(InTotalFileSize), SerializedPackageFlags(0)
 {
 }
+FSavePackageResultStruct::FSavePackageResultStruct(ESavePackageResult InResult, int64 InTotalFileSize,
+	uint32 InSerializedPackageFlags)
+	: Result(InResult), TotalFileSize(InTotalFileSize), SerializedPackageFlags(InSerializedPackageFlags)
+{
+}
 FSavePackageResultStruct::FSavePackageResultStruct(ESavePackageResult InResult, int64 InTotalFileSize,
 	uint32 InSerializedPackageFlags, TPimplPtr<FLinkerSave> Linker)
-	: Result(InResult), TotalFileSize(InTotalFileSize), SerializedPackageFlags(InSerializedPackageFlags),
-	LinkerSave(MoveTemp(Linker))
+	: Result(InResult), TotalFileSize(InTotalFileSize), SerializedPackageFlags(InSerializedPackageFlags)
 {
+	Linker.Reset();
 }
 FSavePackageResultStruct::FSavePackageResultStruct(FSavePackageResultStruct&& Other) = default;
 FSavePackageResultStruct& FSavePackageResultStruct::operator=(FSavePackageResultStruct&& Other) = default;
diff --git a/Engine/Source/Runtime/CoreUObject/Private/UObject/SavePackage/SaveContext.cpp b/Engine/Source/Runtime/CoreUObject/Private/UObject/SavePackage/SaveContext.cpp
--- a/Engine/Source/Runtime/CoreUObject/Private/UObject/SavePackage/SaveContext.cpp
+++ b/Engine/Source/Runtime/CoreUObject/Private/UObject/SavePackage/SaveContext.cpp
@@ -221,8 +221,7 @@ FSavePackageResultStruct FSaveContext::GetFinalResult()
 	}
 
 	ESavePackageResult FinalResult = IsStubRequested() ? ESavePackageResult::GenerateStub : ESavePackageResult::Success;
-	FSavePackageResultStruct ResultData(FinalResult, TotalPackageSizeUncompressed,
-		SerializedPackageFlags, IsCompareLinker() ? MoveTemp(GetHarvestedRealm().Linker) : nullptr);
+	FSavePackageResultStruct ResultData(FinalResult, TotalPackageSizeUncompressed, SerializedPackageFlags);
 
 	ResultData.SavedAssets = MoveTemp(SavedAssets);
 	UClass* PackageClass = UPackage::StaticClass();
@@ -235,6 +234,9 @@ FSavePackageResultStruct FSaveContext::GetFinalResult()
 	}
 	TSet<FName>& SoftPackageReferenceList = GetSoftPackageReferenceList();
 	ResultData.SoftPackageReferences = SoftPackageReferenceList.Array();
+	const TSet<FName>& UntrackedPackageReferenceList = GetUntrackedSoftPackageReferenceList();
+	ResultData.UntrackedSoftPackageReferences = UntrackedPackageReferenceList.Array();
+
 #if WITH_EDITOR
 	for (const FSoftObjectPath& RuntimeDependency : ObjectSaveContext.CookRuntimeDependencies)
 	{
@@ -245,11 +247,32 @@ FSavePackageResultStruct FSaveContext::GetFinalResult()
 		}
 	}
 	ResultData.BuildResultDependencies = MoveTemp(ObjectSaveContext.BuildResultDependencies);
-#endif
 
-	const TSet<FName>& UntrackedPackageReferenceList = GetUntrackedSoftPackageReferenceList();
-	ResultData.UntrackedSoftPackageReferences = UntrackedPackageReferenceList.Array();
+	if (IsCooking())
+	{
+		// Add the imports, exports, and PreloadDependencies for the CookImportsChecker
 
+		// the package isn't actually in the export map, but that is ok, we add it as export anyway for error checking
+		ResultData.Exports.Reserve(GetExports().Num() + 1);
+		ResultData.Exports.Add(GetPackage());
+		for (const FTaggedExport& ExportData : GetExports())
+		{
+			UObject* Export = ExportData.Obj.Get();
+			check(Export);
+			ResultData.Exports.Add(Export);
+		}
+
+		ResultData.Imports.Reserve(GetImports().Num());
+		for (TObjectPtr<UObject> ImportPtr : GetImports())
+		{
+			UObject* Import = ImportPtr.Get();
+			check(Import);
+			ResultData.Imports.Add(Import);
+		}
+
+		ResultData.PreloadDependencies = MoveTemp(GetPreloadDependencies());
+	}
+#endif
 	return ResultData;
 }
 
diff --git a/Engine/Source/Runtime/CoreUObject/Private/UObject/SavePackage/SaveContext.h b/Engine/Source/Runtime/CoreUObject/Private/UObject/SavePackage/SaveContext.h
--- a/Engine/Source/Runtime/CoreUObject/Private/UObject/SavePackage/SaveContext.h
+++ b/Engine/Source/Runtime/CoreUObject/Private/UObject/SavePackage/SaveContext.h
@@ -289,6 +289,16 @@ struct FHarvestedRealm
 		return Exports;
 	}
 
+	const TArray<UE::SavePackageUtilities::FPreloadDependency>& GetPreloadDependencies() const
+	{
+		return PreloadDependencies;
+	}
+
+	TArray<UE::SavePackageUtilities::FPreloadDependency> &GetPreloadDependencies()
+	{
+		return PreloadDependencies;
+	}
+
 	const TSet<TObjectPtr<UObject>>& GetDirectImports() const
 	{
 		return DirectImports;
@@ -525,6 +535,8 @@ private:
 	TSet<FTaggedExport> Exports;
 	// Set of objects marked as import
 	TSet<TObjectPtr<UObject>> Imports;
+	// Set of EDL/ZenLoader dependencies that were written into the package for runtime loading
+	TArray<UE::SavePackageUtilities::FPreloadDependency> PreloadDependencies;
 	// Set of objects that were referenced directly from an export. Some imports are transitively added by
 	// FPackageHarvester::ProcessImport (Outer, Class, CDO, CDO subobjects, others?) for long-standing reasons
 	// (performance, loading behavior, other?)/ But some features such as allowed-import access warnings need
@@ -845,11 +857,6 @@ public:
 		return !!(SaveArgs.SaveFlags & SAVE_Concurrent);
 	}
 
-	bool IsCompareLinker() const
-	{
-		return !!(SaveArgs.SaveFlags & ESaveFlags::SAVE_CompareLinker);
-	}
-
 	bool IsIgnoringHeaderDiff() const
 	{
 		return bIgnoreHeaderDiffs;
@@ -870,16 +877,6 @@ public:
 		return (SaveArgs.SaveFlags & ESaveFlags::SAVE_RehydratePayloads) != 0;
 	}
 
-	FEDLCookCheckerThreadState* GetEDLCookChecker() const
-	{
-		return EDLCookChecker;
-	}
-
-	void SetEDLCookChecker(FEDLCookCheckerThreadState* InCookChecker)
-	{
-		EDLCookChecker = InCookChecker;
-	}
-
 	bool GetPostSaveRootRequired() const
 	{
 		return bPostSaveRootRequired;
@@ -992,6 +989,11 @@ public:
 		return GetHarvestedRealm().GetExports();
 	}
 
+	TArray<UE::SavePackageUtilities::FPreloadDependency>& GetPreloadDependencies()
+	{
+		return GetHarvestedRealm().GetPreloadDependencies();
+	}
+
 	const TSet<TObjectPtr<UObject>>& GetImports() const
 	{
 		return GetHarvestedRealm().GetImports();
@@ -1406,9 +1408,6 @@ private:
 	// Config classes shared with the old Save
 	FCanSkipEditorReferencedPackagesWhenCooking SkipEditorRefCookingSetting;
 
-	// Pointer to the EDLCookChecker associated with this context
-	FEDLCookCheckerThreadState* EDLCookChecker = nullptr;
-
 	// An object matching any GameRealmExcludedObjectMarks should be excluded from imports or exports in the game realm
 	const EObjectMark GameRealmExcludedObjectMarks;
 
diff --git a/Engine/Source/Runtime/CoreUObject/Private/UObject/SavePackage/SavePackageUtilities.cpp b/Engine/Source/Runtime/CoreUObject/Private/UObject/SavePackage/SavePackageUtilities.cpp
--- a/Engine/Source/Runtime/CoreUObject/Private/UObject/SavePackage/SavePackageUtilities.cpp
+++ b/Engine/Source/Runtime/CoreUObject/Private/UObject/SavePackage/SavePackageUtilities.cpp
@@ -1211,843 +1211,46 @@ void FObjectExportSortHelper::SortExports(FLinkerSave* Linker)
 	Algo::Sort(Linker->ExportMap, CompareObjectExports);
 }
 
-FEDLCookChecker::FEDLNodeHash::FEDLNodeHash()
-{
-}
-
-FEDLCookChecker::FEDLNodeHash::FEDLNodeHash(const FEDLNodeHash& Other)
-{
-	bIsNode = Other.bIsNode;
-	if (bIsNode)
-	{
-		Nodes = Other.Nodes;
-	}
-	else
-	{
-		Object = Other.Object;
-	}
-	NodeID = Other.NodeID;
-	ObjectEvent = Other.ObjectEvent;
-}
-
-FEDLCookChecker::FEDLNodeHash::FEDLNodeHash(const TArray<FEDLNodeData>* InNodes, FEDLNodeID InNodeID, EObjectEvent InObjectEvent)
-	: Nodes(InNodes)
-	, NodeID(InNodeID)
-	, bIsNode(true)
-	, ObjectEvent(InObjectEvent)
-{
-}
-
-FEDLCookChecker::FEDLNodeHash::FEDLNodeHash(TObjectPtr<UObject> InObject, EObjectEvent InObjectEvent)
-	: Object(InObject)
-	, bIsNode(false)
-	, ObjectEvent(InObjectEvent)
-{
-}
-
-bool FEDLCookChecker::FEDLNodeHash::operator==(const FEDLNodeHash& Other) const
-{
-	if (ObjectEvent != Other.ObjectEvent)
-	{
-		return false;
-	}
-
-	uint32 LocalNodeID;
-	uint32 OtherNodeID;
-	TObjectPtr<const UObject> LocalObject;
-	TObjectPtr<const UObject> OtherObject;
-	FName LocalName = ObjectNameFirst(*this, LocalNodeID, LocalObject);
-	FName OtherName = ObjectNameFirst(Other, OtherNodeID, OtherObject);
-
-	do
-	{
-		if (LocalName != OtherName)
-		{
-			return false;
-		}
-		LocalName = ObjectNameNext(*this, LocalNodeID, LocalObject);
-		OtherName = ObjectNameNext(Other, OtherNodeID, OtherObject);
-	} while (!LocalName.IsNone() && !OtherName.IsNone());
-	return LocalName.IsNone() == OtherName.IsNone();
-}
-
-FEDLCookChecker::FEDLNodeHash& FEDLCookChecker::FEDLNodeHash::operator=(const FEDLNodeHash& Other)
-{
-	bIsNode = Other.bIsNode;
-	if (bIsNode)
-	{
-		Nodes = Other.Nodes;
-	}
-	else
-	{
-		Object = Other.Object;
-	}
-	NodeID = Other.NodeID;
-	ObjectEvent = Other.ObjectEvent;
-
-	return *this;
-}
-
-uint32 GetTypeHash(const FEDLCookChecker::FEDLNodeHash& A)
-{
-	uint32 Hash = 0;
-
-	uint32 LocalNodeID;
-	TObjectPtr<const UObject> LocalObject;
-	FName LocalName = FEDLCookChecker::FEDLNodeHash::ObjectNameFirst(A, LocalNodeID, LocalObject);
-	do
-	{
-		Hash = HashCombine(Hash, GetTypeHash(LocalName));
-		LocalName = FEDLCookChecker::FEDLNodeHash::ObjectNameNext(A, LocalNodeID, LocalObject);
-	} while (!LocalName.IsNone());
-
-	return (Hash << 1) | (uint32)A.ObjectEvent;
-}
-
-FName FEDLCookChecker::FEDLNodeHash::GetName() const
-{
-	if (bIsNode)
-	{
-		return (*Nodes)[NodeID].Name;
-	}
-	else
-	{
-		return Object.GetFName();
-	}
-}
-
-bool FEDLCookChecker::FEDLNodeHash::TryGetParent(FEDLCookChecker::FEDLNodeHash& Parent) const
-{
-	EObjectEvent ParentObjectEvent = EObjectEvent::Create; // For purposes of parents, which is used only to get the ObjectPath, we always use the Create version of the node as the parent
-	if (bIsNode)
-	{
-		FEDLNodeID ParentID = (*Nodes)[NodeID].ParentID;
-		if (ParentID != NodeIDInvalid)
-		{
-			Parent = FEDLNodeHash(Nodes, ParentID, ParentObjectEvent);
-			return true;
-		}
-	}
-	else
-	{
-		TObjectPtr<UObject> ParentObject = Object.GetOuter();
-		if (ParentObject)
-		{
-			Parent = FEDLNodeHash(ParentObject, ParentObjectEvent);
-			return true;
-		}
-	}
-	return false;
-}
-
-FEDLCookChecker::EObjectEvent FEDLCookChecker::FEDLNodeHash::GetObjectEvent() const
-{
-	return ObjectEvent;
-}
-
-void FEDLCookChecker::FEDLNodeHash::SetNodes(const TArray<FEDLNodeData>* InNodes)
-{
-	if (bIsNode)
-	{
-		Nodes = InNodes;
-	}
-}
-
-FName FEDLCookChecker::FEDLNodeHash::ObjectNameFirst(const FEDLNodeHash& InNode, uint32& OutNodeID,  TObjectPtr<const UObject>& OutObject)
-{
-	if (InNode.bIsNode)
-	{
-		OutNodeID = InNode.NodeID;
-		return (*InNode.Nodes)[OutNodeID].Name;
-	}
-	else
-	{
-		OutObject = InNode.Object;
-		return OutObject.GetFName();
-	}
-}
-
-FName FEDLCookChecker::FEDLNodeHash::ObjectNameNext(const FEDLNodeHash& InNode, uint32& OutNodeID, TObjectPtr<const UObject>& OutObject)
-{
-	if (InNode.bIsNode)
-	{
-		OutNodeID = (*InNode.Nodes)[OutNodeID].ParentID;
-		return OutNodeID != NodeIDInvalid ? (*InNode.Nodes)[OutNodeID].Name : NAME_None;
-	}
-	else
-	{
-		OutObject = OutObject.GetOuter();
-		return OutObject ? OutObject.GetFName() : NAME_None;
-	}
-}
-
-FEDLCookChecker::FEDLNodeData::FEDLNodeData(FEDLNodeID InID, FEDLNodeID InParentID, FName InName, EObjectEvent InObjectEvent)
-	: Name(InName)
-	, ID(InID)
-	, ParentID(InParentID)
-	, ObjectEvent(InObjectEvent)
-	, bIsExport(false)
-{
-}
-
-FEDLCookChecker::FEDLNodeData::FEDLNodeData(FEDLNodeID InID, FEDLNodeID InParentID, FName InName, FEDLNodeData&& Other)
-	: Name(InName)
-	, ID(InID)
-	, ImportingPackagesSorted(MoveTemp(Other.ImportingPackagesSorted))
-	, ParentID(InParentID)
-	, ObjectEvent(Other.ObjectEvent)
-	, bIsExport(Other.bIsExport)
-{
-	// Note that Other Name and ParentID must be unmodified, since they might still be needed for GetHashCode calls from children
-	Other.ImportingPackagesSorted.Empty();
-}
-
-FEDLCookChecker::FEDLNodeHash FEDLCookChecker::FEDLNodeData::GetNodeHash(const FEDLCookChecker& Owner) const
-{
-	return FEDLNodeHash(&Owner.Nodes, ID, ObjectEvent);
-}
-
-FString FEDLCookChecker::FEDLNodeData::ToString(const FEDLCookChecker& Owner) const
-{
-	TStringBuilder<NAME_SIZE> Result;
-	switch (ObjectEvent)
-	{
-	case EObjectEvent::Create:
-		Result << TEXT("Create:");
-		break;
-	case EObjectEvent::Serialize:
-		Result << TEXT("Serialize:");
-		break;
-	default:
-		check(false);
-		break;
-	}
-	AppendPathName(Owner, Result);
-	return FString(Result);
-}
-
-void FEDLCookChecker::FEDLNodeData::AppendPathName(const FEDLCookChecker& Owner, FStringBuilderBase& Result) const
-{
-	if (ParentID != NodeIDInvalid)
-	{
-		const FEDLNodeData& ParentNode = Owner.Nodes[ParentID];
-		ParentNode.AppendPathName(Owner, Result);
-		bool bParentIsOutermost = ParentNode.ParentID == NodeIDInvalid;
-		Result << (bParentIsOutermost ? TEXT(".") : SUBOBJECT_DELIMITER);
-	}
-	Name.AppendString(Result);
-}
-
-FName FEDLCookChecker::FEDLNodeData::GetPackageName(const FEDLCookChecker& Owner) const
-{
-	if (ParentID != NodeIDInvalid)
-	{
-		// @todo ExternalPackages: We need to store ExternalPackage pointers on the Node and return that
-		return Owner.Nodes[ParentID].GetPackageName(Owner);
-	}
-	return Name;
-}
-
-void FEDLCookChecker::FEDLNodeData::Merge(FEDLCookChecker::FEDLNodeData&& Other)
-{
-	check(ObjectEvent == Other.ObjectEvent);
-	bIsExport = bIsExport || Other.bIsExport;
-
-	ImportingPackagesSorted.Append(Other.ImportingPackagesSorted);
-	Algo::Sort(ImportingPackagesSorted, FNameFastLess());
-	ImportingPackagesSorted.SetNum(Algo::Unique(ImportingPackagesSorted), EAllowShrinking::Yes);
-}
-
-FEDLCookCheckerThreadState::FEDLCookCheckerThreadState()
-{
-	Checker.SetActiveIfNeeded();
-
-	FScopeLock CookCheckerInstanceLock(&FEDLCookChecker::CookCheckerInstanceCritical);
-	FEDLCookChecker::CookCheckerInstances.Add(&Checker);
-}
-
-void FEDLCookChecker::SetActiveIfNeeded()
-{
-	bIsActive = !FParse::Param(FCommandLine::Get(), TEXT("DisableEDLCookChecker"));
-}
-
-void FEDLCookChecker::Reset()
-{
-	check(!GIsSavingPackage);
-
-	Nodes.Reset();
-	NodeHashToNodeID.Reset();
-	NodePrereqs.Reset();
-	bIsActive = false;
-}
-
-LLM_DEFINE_TAG(EDLCookChecker);
-
-void FEDLCookChecker::AddImport(TObjectPtr<UObject> Import, UPackage* ImportingPackage)
-{
-	if (bIsActive)
-	{
-		if (!Import->GetOutermost()->HasAnyPackageFlags(PKG_CompiledIn))
-		{
-			LLM_SCOPE_BYTAG(EDLCookChecker);
-			FEDLNodeID NodeId = FindOrAddNode(FEDLNodeHash(Import, EObjectEvent::Serialize));
-			FEDLNodeData& NodeData = Nodes[NodeId];
-			FName ImportingPackageName = ImportingPackage->GetFName();
-			TArray<FName>& Sorted = NodeData.ImportingPackagesSorted;
-			int32 InsertionIndex = Algo::LowerBound(Sorted, ImportingPackageName, FNameFastLess());
-			if (InsertionIndex == Sorted.Num() || Sorted[InsertionIndex] != ImportingPackageName)
-			{
-				Sorted.Insert(ImportingPackageName, InsertionIndex);
-			}
-		}
-	}
-}
-void FEDLCookChecker::AddExport(UObject* Export)
-{
-	if (bIsActive)
-	{
-		LLM_SCOPE_BYTAG(EDLCookChecker);
-		FEDLNodeID SerializeID = FindOrAddNode(FEDLNodeHash(Export, EObjectEvent::Serialize));
-		Nodes[SerializeID].bIsExport = true;
-		FEDLNodeID CreateID = FindOrAddNode(FEDLNodeHash(Export, EObjectEvent::Create));
-		Nodes[CreateID].bIsExport = true;
-		AddDependency(SerializeID, CreateID); // every export must be created before it can be serialized...these arcs are implicit and not listed in any table.
-	}
-}
-
-void FEDLCookChecker::AddArc(UObject* DepObject, bool bDepIsSerialize, UObject* Export, bool bExportIsSerialize)
-{
-	if (bIsActive)
-	{
-		LLM_SCOPE_BYTAG(EDLCookChecker);
-		FEDLNodeID ExportID = FindOrAddNode(FEDLNodeHash(Export, bExportIsSerialize ? EObjectEvent::Serialize : EObjectEvent::Create));
-		FEDLNodeID DepID = FindOrAddNode(FEDLNodeHash(DepObject, bDepIsSerialize ? EObjectEvent::Serialize : EObjectEvent::Create));
-		AddDependency(ExportID, DepID);
-	}
-}
-
-void FEDLCookChecker::AddPackageWithUnknownExports(FName LongPackageName)
-{
-	if (bIsActive)
-	{
-		LLM_SCOPE_BYTAG(EDLCookChecker);
-		PackagesWithUnknownExports.Add(LongPackageName);
-	}
-}
-
-
-void FEDLCookChecker::AddDependency(FEDLNodeID SourceID, FEDLNodeID TargetID)
-{
-	NodePrereqs.Add(SourceID, TargetID);
-}
-
-void FEDLCookChecker::StartSavingEDLCookInfoForVerification()
-{
-	LLM_SCOPE_BYTAG(EDLCookChecker);
-	FScopeLock CookCheckerInstanceLock(&CookCheckerInstanceCritical);
-	for (FEDLCookChecker* Checker : CookCheckerInstances)
-	{
-		Checker->Reset();
-		Checker->SetActiveIfNeeded();
-	}
-}
-
-bool FEDLCookChecker::CheckForCyclesInner(TSet<FEDLNodeID>& Visited, TSet<FEDLNodeID>& Stack, const FEDLNodeID& Visit, FEDLNodeID& FailNode)
-{
-	bool bResult = false;
-	if (Stack.Contains(Visit))
-	{
-		FailNode = Visit;
-		bResult = true;
-	}
-	else
-	{
-		bool bWasAlreadyTested = false;
-		Visited.Add(Visit, &bWasAlreadyTested);
-		if (!bWasAlreadyTested)
-		{
-			Stack.Add(Visit);
-			for (auto It = NodePrereqs.CreateConstKeyIterator(Visit); !bResult && It; ++It)
-			{
-				bResult = CheckForCyclesInner(Visited, Stack, It.Value(), FailNode);
-			}
-			Stack.Remove(Visit);
-		}
-	}
-	UE_CLOG(bResult && Stack.Contains(FailNode), LogSavePackage, Error, TEXT("Cycle Node %s"), *Nodes[Visit].ToString(*this));
-	return bResult;
-}
-
-FEDLCookChecker::FEDLNodeID FEDLCookChecker::FindOrAddNode(const FEDLNodeHash& NodeHash)
-{
-	uint32 TypeHash = GetTypeHash(NodeHash);
-	FEDLNodeID* NodeIDPtr = NodeHashToNodeID.FindByHash(TypeHash, NodeHash);
-	if (NodeIDPtr)
-	{
-		return *NodeIDPtr;
-	}
 
-	FName Name = NodeHash.GetName();
-	FEDLNodeHash ParentHash;
-	FEDLNodeID ParentID = NodeHash.TryGetParent(ParentHash) ? FindOrAddNode(ParentHash) : NodeIDInvalid;
-	FEDLNodeID NodeID = Nodes.Num();
-	FEDLNodeData& NewNodeData = Nodes.Emplace_GetRef(NodeID, ParentID, Name, NodeHash.GetObjectEvent());
-	NodeHashToNodeID.AddByHash(TypeHash, NewNodeData.GetNodeHash(*this), NodeID);
-	return NodeID;
-}
-
-FEDLCookChecker::FEDLNodeID FEDLCookChecker::FindOrAddNode(FEDLNodeData&& NodeData, const FEDLCookChecker& OldOwnerOfNode, FEDLNodeID ParentIDInThis, bool& bNew)
-{
-	// Note that NodeData's Name and ParentID must be unmodified, since they might still be needed for GetHashCode calls from children
-
-	FEDLNodeHash NodeHash = NodeData.GetNodeHash(OldOwnerOfNode);
-	uint32 TypeHash = GetTypeHash(NodeHash);
-	FEDLNodeID* NodeIDPtr = NodeHashToNodeID.FindByHash(TypeHash, NodeHash);
-	if (NodeIDPtr)
-	{
-		bNew = false;
-		return *NodeIDPtr;
-	}
-
-	FEDLNodeID NodeID = Nodes.Num();
-	FEDLNodeData& NewNodeData = Nodes.Emplace_GetRef(NodeID, ParentIDInThis, NodeData.Name, MoveTemp(NodeData));
-	NodeHashToNodeID.AddByHash(TypeHash, NewNodeData.GetNodeHash(*this), NodeID);
-	bNew = true;
-	return NodeID;
-}
-
-FEDLCookChecker::FEDLNodeID FEDLCookChecker::FindNode(const FEDLNodeHash& NodeHash)
-{
-	const FEDLNodeID* NodeIDPtr = NodeHashToNodeID.Find(NodeHash);
-	return NodeIDPtr ? *NodeIDPtr : NodeIDInvalid;
-}
-
-void FEDLCookChecker::Merge(FEDLCookChecker&& Other)
-{
-	if (Nodes.Num() == 0)
-	{
-		Swap(Nodes, Other.Nodes);
-		Swap(NodeHashToNodeID, Other.NodeHashToNodeID);
-		Swap(NodePrereqs, Other.NodePrereqs);
-
-		// Switch the pointers in all of the swapped data to point at this instead of Other
-		for (TPair<FEDLNodeHash, FEDLNodeID>& KVPair : NodeHashToNodeID)
-		{
-			FEDLNodeHash& NodeHash = KVPair.Key;
-			NodeHash.SetNodes(&Nodes);
-		}
-	}
-	else
-	{
-		Other.NodeHashToNodeID.Empty(); // We will be invalidating the data these NodeHashes point to in the Other.Nodes loop, so empty the array now to avoid using it by accident
-
-		TArray<FEDLNodeID> RemapIDs;
-		RemapIDs.Reserve(Other.Nodes.Num());
-		for (FEDLNodeData& NodeData : Other.Nodes)
-		{
-			FEDLNodeID ParentID;
-			if (NodeData.ParentID == NodeIDInvalid)
-			{
-				ParentID = NodeIDInvalid;
-			}
-			else
-			{
-				// Parents should be earlier in the nodes list than children, since we always FindOrAdd the parent (and hence add it to the nodelist) when creating the child.
-				// Since the parent is earlier in the nodes list, we have already transferred it, and its ID in this->Nodes is therefore RemapIDs[Other.ParentID]
-				check(NodeData.ParentID < NodeData.ID);
-				ParentID = RemapIDs[NodeData.ParentID];
-			}
-
-			bool bNew;
-			FEDLNodeID NodeID = FindOrAddNode(MoveTemp(NodeData), Other, ParentID, bNew);
-			if (!bNew)
-			{
-				Nodes[NodeID].Merge(MoveTemp(NodeData));
-			}
-			RemapIDs.Add(NodeID);
-		}
-
-		for (const TPair<FEDLNodeID, FEDLNodeID>& Prereq : Other.NodePrereqs)
-		{
-			FEDLNodeID SourceID = RemapIDs[Prereq.Key];
-			FEDLNodeID TargetID = RemapIDs[Prereq.Value];
-			AddDependency(SourceID, TargetID);
-		}
-
-		Other.NodePrereqs.Empty();
-		Other.Nodes.Empty();
-	}
-
-	if (PackagesWithUnknownExports.Num() == 0)
-	{
-		Swap(PackagesWithUnknownExports, Other.PackagesWithUnknownExports);
-	}
-	else
-	{
-		PackagesWithUnknownExports.Reserve(Other.PackagesWithUnknownExports.Num());
-		for (FName PackageName : Other.PackagesWithUnknownExports)
-		{
-			PackagesWithUnknownExports.Add(PackageName);
-		}
-		Other.PackagesWithUnknownExports.Empty();
-	}
-}
-
-FEDLCookChecker FEDLCookChecker::AccumulateAndClear()
-{
-	FEDLCookChecker Accumulator;
-
-	FScopeLock CookCheckerInstanceLock(&CookCheckerInstanceCritical);
-	for (FEDLCookChecker* Checker : CookCheckerInstances)
-	{
-		if (Checker->bIsActive)
-		{
-			Accumulator.bIsActive = true;
-			Accumulator.Merge(MoveTemp(*Checker));
-			Checker->Reset();
-			Checker->bIsActive = true;
-		}
-	}
-	return Accumulator;
-}
-
-void FEDLCookChecker::Verify(const UE::SavePackageUtilities::FEDLLogRecordCallback& MessageCallback,
-	bool bFullReferencesExpected)
-{
-	check(!GIsSavingPackage);
-	FEDLCookChecker Accumulator = AccumulateAndClear();
-
-	FString SeverityStr;
-	GConfig->GetString(TEXT("CookSettings"), TEXT("CookContentMissingSeverity"), SeverityStr, GEditorIni);
-	ELogVerbosity::Type MissingContentSeverity = ParseLogVerbosityFromString(SeverityStr);
-
-	if (Accumulator.bIsActive)
-	{
-		double StartTime = FPlatformTime::Seconds();
-			
-		if (bFullReferencesExpected)
-		{
-			// imports to things that are not exports...
-			for (const FEDLNodeData& NodeData : Accumulator.Nodes)
-			{
-				if (NodeData.bIsExport)
-				{
-					// The node is an export; imports of it are valid
-					continue;
-				}
-
-				if (Accumulator.PackagesWithUnknownExports.Contains(NodeData.GetPackageName(Accumulator)))
-				{
-					// The node is an object in a package that exists, but for which we do not know the exports
-					// because e.g. it was iteratively skipped in the current cook. Suppress warnings about it
-					continue;
-				}
-
-				// Any imports of this non-exported node are an error; log them all if they exist
-				if (NodeData.ImportingPackagesSorted.IsEmpty())
-				{
-					continue;
-				}
-
-				const FEDLNodeData* NodeDataOfExportPackage = &NodeData;
-				while (NodeDataOfExportPackage->ParentID != NodeIDInvalid)
-				{
-					int32 ParentNodeIndex = static_cast<int32>(NodeDataOfExportPackage->ParentID);
-					check(Accumulator.Nodes.IsValidIndex(ParentNodeIndex));
-					NodeDataOfExportPackage = &Accumulator.Nodes[ParentNodeIndex];
-				}
-
-				const TCHAR* ReasonExportIsMissing = TEXT("");
-				if (NodeDataOfExportPackage->bIsExport)
-				{
-					ReasonExportIsMissing = TEXT("the object was stripped out of the target package when saved");
-				}
-				else
-				{
-					ReasonExportIsMissing = TEXT("the target package was marked NeverCook or is not cookable for the target platform");
-				}
-
-				for (FName PackageName : NodeData.ImportingPackagesSorted)
-				{
-					UE::FLogRecord Record;
-#if !NO_LOGGING
-					Record.SetCategory(LogSavePackage.GetCategoryName());
-#endif
-					Record.SetVerbosity(MissingContentSeverity);
-					Record.SetTime(UE::FLogTime::Now());
-					Record.SetFormat(TEXT("Content is missing from cook. Source package referenced an object in target package but {Reason}.")
-						TEXT("\n\tSource package: {Source}")
-						TEXT("\n\tTarget package: {Target}")
-						TEXT("\n\tReferenced object: {ReferencedObject}"));
-					{
-						FCbWriter Writer;
-						Writer.BeginObject();
-						Writer << "Reason" << ReasonExportIsMissing;
-						Writer << "Source" << WriteToUtf8String<256>(PackageName);
-						Writer << "Target" << WriteToUtf8String<256>(NodeDataOfExportPackage->Name);
-						{
-							TStringBuilder<256> ReferencedObjectStr;
-							NodeData.AppendPathName(Accumulator, ReferencedObjectStr);
-							Writer << "ReferencedObject" << ReferencedObjectStr;
-						}
-						Writer.EndObject();
-						Record.SetFields(Writer.Save().AsObject());
-					}
-					Record.SetFile(__FILE__);
-					Record.SetLine(__LINE__);
-					MessageCallback(MoveTemp(Record));
-				}
-			}
-		}
-
-		// cycles in the dep graph
-		TSet<FEDLNodeID> Visited;
-		TSet<FEDLNodeID> Stack;
-		bool bHadCycle = false;
-		for (const FEDLNodeData& NodeData : Accumulator.Nodes)
-		{
-			if (!NodeData.bIsExport)
-			{
-				continue;
-			}
-			FEDLNodeID FailNode;
-			if (Accumulator.CheckForCyclesInner(Visited, Stack, NodeData.ID, FailNode))
-			{
-				UE_LOG(LogSavePackage, Error, TEXT("----- %s contained a cycle (listed above)."), *Accumulator.Nodes[FailNode].ToString(Accumulator));
-				bHadCycle = true;
-			}
-		}
-		if (bHadCycle)
-		{
-			UE_LOG(LogSavePackage, Fatal, TEXT("EDL dep graph contained a cycle (see errors, above). This is fatal at runtime so it is fatal at cook time."));
-		}
-		UE_LOG(LogSavePackage, Display, TEXT("Took %fs to verify the EDL loading graph."), float(FPlatformTime::Seconds() - StartTime));
-	}
-}
-
-void FEDLCookChecker::MoveToCompactBinaryAndClear(FCbWriter& Writer, bool& bOutHasData)
-{
-	bOutHasData = false;
-	FEDLCookChecker Accumulator = AccumulateAndClear();
-	if (!Accumulator.bIsActive)
-	{
-		return;
-	}
-	if (Accumulator.Nodes.IsEmpty() && Accumulator.NodePrereqs.IsEmpty() && Accumulator.PackagesWithUnknownExports.IsEmpty())
-	{
-		return;
-	}
-	bOutHasData = true;
-
-	Accumulator.WriteToCompactBinary(Writer);
-}
-
-bool FEDLCookChecker::AppendFromCompactBinary(FCbFieldView Field)
-{
-	FEDLCookChecker Instance;
-	if (!Instance.ReadFromCompactBinary(Field))
-	{
-		return false;
-	}
-	FEDLCookChecker& CurrentChecker = FEDLCookCheckerThreadState::Get().Checker;
-	CurrentChecker.Merge(MoveTemp(Instance));
-	return true;
-}
-
-void FEDLCookChecker::WriteToCompactBinary(FCbWriter& Writer)
-{
-	Writer.BeginObject();
-	{
-		Writer.BeginArray("Nodes");
-		for (const FEDLNodeData& Node : Nodes)
-		{
-			Writer << Node.Name;
-			Writer << Node.ImportingPackagesSorted;
-			Writer << Node.ParentID;
-			Writer << static_cast<uint8>(Node.ObjectEvent);
-			Writer << Node.bIsExport;
-		}
-		Writer.EndArray();
-		Writer.BeginArray("NodePrereqs");
-		for (const TPair<FEDLNodeID, FEDLNodeID>& Pair : NodePrereqs)
-		{
-			Writer << static_cast<uint32>(Pair.Key);
-			Writer << static_cast<uint32>(Pair.Value);
-		}
-		Writer.EndArray();
-		Writer.BeginArray("PackagesWithUnknownExports");
-		for (FName PackageName : PackagesWithUnknownExports)
-		{
-			Writer << PackageName;
-		}
-		Writer.EndArray();
-	}
-	Writer.EndObject();
-}
-
-
-bool FEDLCookChecker::ReadFromCompactBinary(FCbFieldView Field)
-{
-	Reset();
-
-	bool bSuccess = false;
-	ON_SCOPE_EXIT
-	{
-		if (!bSuccess)
-		{
-			Reset();
-		}
-	};
-
-	FCbFieldView NodesField = Field["Nodes"];
-	const uint64 NumNodes = NodesField.AsArrayView().Num() / 5;
-	if (NumNodes > MAX_int32)
-	{
-		return false;
-	}
-	Nodes.Reserve(static_cast<int32>(NumNodes));
-	if (NodesField.HasError())
-	{
-		return false;
-	}
-	FCbFieldViewIterator NodeIter = NodesField.CreateViewIterator();
-	while (NodeIter)
-	{
-		FEDLNodeID NodeID = Nodes.Num();
-		FEDLNodeData& Node = Nodes.Emplace_GetRef();
-		Node.ID = NodeID;
-		if (!LoadFromCompactBinary(NodeIter, Node.Name)) { return false; }
-		++NodeIter;
-		if (!LoadFromCompactBinary(NodeIter, Node.ImportingPackagesSorted)) { return false; }
-		++NodeIter;
-		if (!LoadFromCompactBinary(NodeIter, Node.ParentID)) { return false; }
-		++NodeIter;
-		uint8 LocalObjectEvent;
-		if (!LoadFromCompactBinary(NodeIter, LocalObjectEvent)) { return false; }
-		if (LocalObjectEvent > static_cast<uint8>(EObjectEvent::Max)) { return false; }
-		Node.ObjectEvent = static_cast<EObjectEvent>(LocalObjectEvent);
-		++NodeIter;
-		if (!LoadFromCompactBinary(NodeIter, Node.bIsExport)) { return false; }
-		++NodeIter;
-	}
-
-	FCbFieldView PrereqsField = Field["NodePrereqs"];
-	const int64 NumNodePrereqs = PrereqsField.AsArrayView().Num() / 2;
-	if (NumNodePrereqs > MAX_int32)
-	{
-		return false;
-	}
-	NodePrereqs.Reserve(static_cast<int32>(NumNodePrereqs));
-	if (PrereqsField.HasError())
-	{
-		return false;
-	}
-	FCbFieldViewIterator PrereqsIter = PrereqsField.CreateViewIterator();
-	while (PrereqsIter)
-	{
-		uint32 Key;
-		uint32 Value;
-		if (!LoadFromCompactBinary(PrereqsIter, Key)) { return false; }
-		++PrereqsIter;
-		if (!LoadFromCompactBinary(PrereqsIter, Value)) { return false; }
-		++PrereqsIter;
-		NodePrereqs.Add(static_cast<FEDLNodeID>(Key), static_cast<FEDLNodeID>(Value));
-	}
-
-	FCbFieldView PackagesWithUnknownExportsField = Field["PackagesWithUnknownExports"];
-	const int64 NumPackagesWithUnknownExports = PackagesWithUnknownExportsField.AsArrayView().Num();
-	if (NumPackagesWithUnknownExports > MAX_int32)
-	{
-		return false;
-	}
-	PackagesWithUnknownExports.Reserve(static_cast<int32>(NumPackagesWithUnknownExports));
-	if (PackagesWithUnknownExportsField.HasError())
-	{
-		return false;
-	}
-	for (FCbFieldView PackageNameField : PackagesWithUnknownExportsField)
-	{
-		FName PackageName;
-		if (!LoadFromCompactBinary(PackageNameField, PackageName))
-		{
-			return false;
-		}
-		PackagesWithUnknownExports.Add(PackageName);
-	}
-
-	for (const FEDLNodeData& Node : Nodes)
-	{
-		NodeHashToNodeID.Add(Node.GetNodeHash(*this), Node.ID);
-	}
-	bIsActive = !Nodes.IsEmpty() || !NodePrereqs.IsEmpty() || !PackagesWithUnknownExports.IsEmpty();
-
-	bSuccess = true;
-	return true;
-}
-
-FCriticalSection FEDLCookChecker::CookCheckerInstanceCritical;
-TArray<FEDLCookChecker*> FEDLCookChecker::CookCheckerInstances;
 
 namespace UE::SavePackageUtilities
 {
 
 void StartSavingEDLCookInfoForVerification()
 {
-	LLM_SCOPE_BYTAG(EDLCookChecker);
-	FEDLCookChecker::StartSavingEDLCookInfoForVerification();
 }
 
 void VerifyEDLCookInfo(bool bFullReferencesExpected)
 {
-	VerifyEDLCookInfo([](UE::FLogRecord&& Record)
-		{
-#if !NO_LOGGING
-			DispatchDynamicLogRecord(Record);
-#endif
-		}, bFullReferencesExpected);
 }
 
 void VerifyEDLCookInfo(const UE::SavePackageUtilities::FEDLMessageCallback& MessageCallback,
 	bool bFullReferencesExpected)
 {
-	VerifyEDLCookInfo([&MessageCallback](FLogRecord&& Record)
-		{
-			TStringBuilder<256> Message;
-			Record.FormatMessageTo(Message);
-			MessageCallback(Record.GetVerbosity(), Message.ToView());
-		}, bFullReferencesExpected);
 }
 
 void VerifyEDLCookInfo(const UE::SavePackageUtilities::FEDLLogRecordCallback& MessageCallback,
 	bool bFullReferencesExpected)
 {
-	LLM_SCOPE_BYTAG(EDLCookChecker);
-	FEDLCookChecker::Verify(MessageCallback, bFullReferencesExpected);
 }
 
 void EDLCookInfoAddIterativelySkippedPackage(FName LongPackageName)
 {
-	LLM_SCOPE_BYTAG(EDLCookChecker);
-	FEDLCookCheckerThreadState::Get().AddPackageWithUnknownExports(LongPackageName);
 }
 
 void EDLCookInfoMoveToCompactBinaryAndClear(FCbWriter& Writer, bool& bOutHasData)
 {
-	LLM_SCOPE_BYTAG(EDLCookChecker);
-	FEDLCookChecker::MoveToCompactBinaryAndClear(Writer, bOutHasData);
+	bOutHasData = false;
 }
 
 void EDLCookInfoMoveToCompactBinaryAndClear(FCbWriter& Writer, bool& bOutHasData, FName PackageName)
 {
-	// For simplicity, instead of sending only information related to the given Package, we send all data.
-	LLM_SCOPE_BYTAG(EDLCookChecker);
-	FEDLCookChecker::MoveToCompactBinaryAndClear(Writer, bOutHasData);
+	bOutHasData = false;
 }
 
 bool EDLCookInfoAppendFromCompactBinary(FCbFieldView Field)
 {
-	LLM_SCOPE_BYTAG(EDLCookChecker);
-	return FEDLCookChecker::AppendFromCompactBinary(Field);
+	return false;
 }
 
 }
diff --git a/Engine/Source/Runtime/CoreUObject/Private/UObject/SavePackage/SavePackageUtilities.h b/Engine/Source/Runtime/CoreUObject/Private/UObject/SavePackage/SavePackageUtilities.h
--- a/Engine/Source/Runtime/CoreUObject/Private/UObject/SavePackage/SavePackageUtilities.h
+++ b/Engine/Source/Runtime/CoreUObject/Private/UObject/SavePackage/SavePackageUtilities.h
@@ -6,7 +6,6 @@
 #include "Containers/Map.h"
 #include "Containers/Set.h"
 #include "Containers/StringFwd.h"
-#include "HAL/ThreadSingleton.h"
 #include "Memory/MemoryView.h"
 #include "Misc/FileHelper.h"
 #include "ProfilingDebugging/CookStats.h"
@@ -146,183 +145,6 @@ struct FObjectExportSortHelper
 	static void SortExports(FLinkerSave* Linker);
 };
 
-struct FEDLCookCheckerThreadState;
-
-/**
- * Helper struct used during cooking to validate EDL dependencies
- */
-struct FEDLCookChecker
-{
-	void SetActiveIfNeeded();
-
-	void Reset();
-
-	void AddImport(TObjectPtr<UObject> Import, UPackage* ImportingPackage);
-	void AddExport(UObject* Export);
-	void AddArc(UObject* DepObject, bool bDepIsSerialize, UObject* Export, bool bExportIsSerialize);
-	void AddPackageWithUnknownExports(FName LongPackageName);
-
-	static void StartSavingEDLCookInfoForVerification();
-	static void Verify(const UE::SavePackageUtilities::FEDLLogRecordCallback& MessageCallback,
-		bool bFullReferencesExpected);
-	static void MoveToCompactBinaryAndClear(FCbWriter& Writer, bool& bOutHasData);
-	static bool AppendFromCompactBinary(FCbFieldView Field);
-
-private:
-	static FEDLCookChecker AccumulateAndClear();
-	void WriteToCompactBinary(FCbWriter& Writer);
-	bool ReadFromCompactBinary(FCbFieldView Field);
-
-	typedef uint32 FEDLNodeID;
-	static const FEDLNodeID NodeIDInvalid = static_cast<FEDLNodeID>(-1);
-
-	struct FEDLNodeData;
-public: // FEDLNodeHash is public only so that GetTypeHash can be defined
-	enum class EObjectEvent : uint8
-	{
-		Create,
-		Serialize,
-		Max = Serialize,
-	};
-
-	/**
-	 * Wrapper aroundan FEDLNodeData (or around a UObject when searching for an FEDLNodeData corresponding to the UObject)
-	 * that provides the hash-by-objectpath to lookup the FEDLNodeData for an objectpath.
-	 */
-	struct FEDLNodeHash
-	{
-		FEDLNodeHash(); // creates an uninitialized node; only use this to provide as an out parameter
-		FEDLNodeHash(const TArray<FEDLNodeData>* InNodes, FEDLNodeID InNodeID, EObjectEvent InObjectEvent);
-		FEDLNodeHash(TObjectPtr<UObject> InObject, EObjectEvent InObjectEvent);
-		FEDLNodeHash(const FEDLNodeHash& Other);
-		bool operator==(const FEDLNodeHash& Other) const;
-		FEDLNodeHash& operator=(const FEDLNodeHash& Other);
-		friend uint32 GetTypeHash(const FEDLNodeHash& A);
-
-		FName GetName() const;
-		bool TryGetParent(FEDLNodeHash& Parent) const;
-		EObjectEvent GetObjectEvent() const;
-		void SetNodes(const TArray<FEDLNodeData>* InNodes);
-
-	private:
-		static FName ObjectNameFirst(const FEDLNodeHash& InNode, uint32& OutNodeID, TObjectPtr<const UObject>& OutObject);
-		static FName ObjectNameNext(const FEDLNodeHash& InNode, uint32& OutNodeID, TObjectPtr<const UObject>& OutObject);
-			
-		union
-		{
-			/**
-			 * The array of nodes from the FEDLCookChecker; this is how we lookup the node for the FEDLNodeData.
-			 * Because the FEDLNodeData are elements in an array which can resize and therefore reallocate the nodes, we cannot store the pointer to the node.
-			 * Only used if bIsNode is true.
-			 */
-			const TArray<FEDLNodeData>* Nodes;
-			/** Pointer to the Object we are looking up, if this hash was created during lookup-by-objectpath for an object */
-			TObjectPtr<const UObject> Object;
-		};
-		/** The identifier for the FEDLNodeData this hash is wrapping. Only used if bIsNode is true. */
-		FEDLNodeID NodeID;
-		/** True if this hash is wrapping an FEDLNodeData, false if it is wrapping a UObject. */
-		bool bIsNode;
-		EObjectEvent ObjectEvent;
-	};
-
-private:
-
-	/**
-	 * Node representing either the Create event or Serialize event of a UObject in the graph of runtime dependencies between UObjects.
-	 */
-	struct FEDLNodeData
-	{
-		// Note that order of the variables is important to reduce alignment waste in the size of FEDLNodeData.
-		/** Name of the UObject represented by this node; full objectpath name is obtainable by combining parent data with the name. */
-		FName Name;
-		/** Index of this node in the FEDLCookChecker's Nodes array. This index is used to provide a small-memory-usage identifier for the node. */
-		FEDLNodeID ID;
-		/**
-		 * Tracks references to this node's UObjects from other packages (which is the reverse of the references from each node that we track in NodePrereqs.)
-		 * We only need this information from each package, so we track by package name instead of node id.
-		 */
-		TArray<FName> ImportingPackagesSorted;
-		/**
-		 * ID of the node representing the UObject parent of this node's UObject. NodeIDInvalid if the UObject has no parent.
-		 * The ParentID always refers to the node for the Create event of the parent UObject.
-		 */
-		uint32 ParentID;
-		/** True if this node represents the Serialize event on the UObject, false if it represents the Create event. */
-		EObjectEvent ObjectEvent;
-		/** True if the UObject represented by this node has been exported by a SavePackage call; used to verify that the imports requested by packages are present somewhere in the cook. */
-		bool bIsExport;
-
-		FEDLNodeData() { /* Fields are uninitialized */ }
-		FEDLNodeData(FEDLNodeID InID, FEDLNodeID InParentID, FName InName, EObjectEvent InObjectEvent);
-		FEDLNodeData(FEDLNodeID InID, FEDLNodeID InParentID, FName InName, FEDLNodeData&& Other);
-		FEDLNodeHash GetNodeHash(const FEDLCookChecker& Owner) const;
-
-		FString ToString(const FEDLCookChecker& Owner) const;
-		void AppendPathName(const FEDLCookChecker& Owner, FStringBuilderBase& Result) const;
-		FName GetPackageName(const FEDLCookChecker& Owner) const;
-		void Merge(FEDLNodeData&& Other);
-	};
-
-	FEDLNodeID FindOrAddNode(const FEDLNodeHash& NodeLookup);
-	FEDLNodeID FindOrAddNode(FEDLNodeData&& NodeData, const FEDLCookChecker& OldOwnerOfNode, FEDLNodeID ParentIDInThis, bool& bNew);
-	FEDLNodeID FindNode(const FEDLNodeHash& NodeHash);
-	void Merge(FEDLCookChecker&& Other);
-	bool CheckForCyclesInner(TSet<FEDLNodeID>& Visited, TSet<FEDLNodeID>& Stack, const FEDLNodeID& Visit, FEDLNodeID& FailNode);
-	void AddDependency(FEDLNodeID SourceID, FEDLNodeID TargetID);
-
-	/**
-	 * All the FEDLNodeDatas that have been created for this checker. These are allocated as elements of an array rather than pointers to reduce cputime and
-	 * memory due to many small allocations, and to provide index-based identifiers. Nodes are not deleted until the checker is reset.
-	 */
-	TArray<FEDLNodeData> Nodes;
-	/** A map to lookup the node for a UObject or for the corresponding node in another thread's FEDLCookChecker. */
-	TMap<FEDLNodeHash, FEDLNodeID> NodeHashToNodeID;
-	/** The graph of dependencies between nodes. */
-	TMultiMap<FEDLNodeID, FEDLNodeID> NodePrereqs;
-	/**
-	 * Packages that were cooked iteratively and therefore have an unknown set of exports.
-	 * We suppress warnings for exports missing from these packages.
-	 */
-	TSet<FName> PackagesWithUnknownExports;
-	/** True if the EDLCookChecker should be active; it is turned off if the runtime will not be using EDL. */
-	bool bIsActive = false;
-
-	/** When cooking with concurrent saving, each thread has its own FEDLCookChecker, and these are merged after the cook is complete. */
-	static FCriticalSection CookCheckerInstanceCritical;
-	static TArray<FEDLCookChecker*> CookCheckerInstances;
-
-	friend FEDLCookCheckerThreadState;
-};
-
-/** Per-thread accessor for writing EDL dependencies to global FEDLCookChecker storage. */
-struct FEDLCookCheckerThreadState : public TThreadSingleton<FEDLCookCheckerThreadState>
-{
-	FEDLCookCheckerThreadState();
-
-	void AddImport(TObjectPtr<UObject> Import, UPackage* ImportingPackage)
-	{
-		Checker.AddImport(Import, ImportingPackage);
-	}
-	void AddExport(UObject* Export)
-	{
-		Checker.AddExport(Export);
-	}
-	void AddArc(UObject* DepObject, bool bDepIsSerialize, UObject* Export, bool bExportIsSerialize)
-	{
-		Checker.AddArc(DepObject, bDepIsSerialize, Export, bExportIsSerialize);
-	}
-	void AddPackageWithUnknownExports(FName LongPackageName)
-	{
-		Checker.AddPackageWithUnknownExports(LongPackageName);
-	}
-
-private:
-	FEDLCookChecker Checker;
-	friend TThreadSingleton<FEDLCookCheckerThreadState>;
-	friend FEDLCookChecker;
-};
-
 // Utility functions used by both UPackage::Save and/or UPackage::Save2
 namespace UE::SavePackageUtilities
 {
diff --git a/Engine/Source/Runtime/CoreUObject/Private/UObject/SavePackage2.cpp b/Engine/Source/Runtime/CoreUObject/Private/UObject/SavePackage2.cpp
--- a/Engine/Source/Runtime/CoreUObject/Private/UObject/SavePackage2.cpp
+++ b/Engine/Source/Runtime/CoreUObject/Private/UObject/SavePackage2.cpp
@@ -754,19 +754,6 @@ ESavePackageResult ValidateExports(FSaveContext& SaveContext)
 	}
 #endif
 
-	if (SaveContext.IsCooking())
-	{
-		// Add the exports for the cook checker
-		if (FEDLCookCheckerThreadState* EDLCookChecker = SaveContext.GetEDLCookChecker())
-		{
-			// the package isn't actually in the export map, but that is ok, we add it as export anyway for error checking
-			EDLCookChecker->AddExport(SaveContext.GetPackage());
-			for (const FTaggedExport& Export : SaveContext.GetExports())
-			{
-				EDLCookChecker->AddExport(Export.Obj);
-			}
-		}
-	}
 	return ReturnSuccessOrCancel();
 }
 
@@ -1067,21 +1054,6 @@ ESavePackageResult ValidateImports(FSaveContext& SaveContext)
 	}
 	PRAGMA_ENABLE_DEPRECATION_WARNINGS
 
-	// Cooking checks
-	// Currently do not use the edl checker for the optional context
-	if (SaveContext.IsCooking() && SaveContext.GetCurrentHarvestingRealm() != ESaveRealm::Optional)
-	{
-		// Now that imports are validated add them to the cook checker if available
-		if (FEDLCookCheckerThreadState* EDLCookChecker = SaveContext.GetEDLCookChecker())
-		{
-			for (UObject* Import : SaveContext.GetImports())
-			{
-				check(Import);
-				EDLCookChecker->AddImport(Import, Package);
-			}
-		}
-	}
-
 	return ReturnSuccessOrCancel();
 }
 
@@ -1665,6 +1637,8 @@ ESavePackageResult BuildLinker(FSaveContext& SaveContext)
 void SavePreloadDependencies(FStructuredArchive::FRecord& StructuredArchiveRoot, FSaveContext& SaveContext)
 {
 	FLinkerSave* Linker = SaveContext.GetLinker();
+	TArray<UE::SavePackageUtilities::FPreloadDependency>& PreloadDependencies = SaveContext.GetPreloadDependencies();
+
 	auto IncludeObjectAsDependency = [Linker, &SaveContext](int32 CallSite, TSet<FPackageIndex>& AddTo, TObjectPtr<UObject> ToTest, UObject* ForObj, bool bMandatory, bool bOnlyIfInLinkerTable)
 	{
 		// Skip transient, editor only, and excluded client/server objects
@@ -1911,9 +1885,7 @@ void SavePreloadDependencies(FStructuredArchive::FRecord& StructuredArchiveRoot,
 				IncludeIndexAsDependency(CreateBeforeCreateDependencies, Export.OuterIndex);
 				IncludeIndexAsDependency(CreateBeforeCreateDependencies, Export.SuperIndex);
 			}
-			// Currently do not validate the optional context with the edl checker
-			FEDLCookCheckerThreadState* EDLCookChecker = SaveContext.GetCurrentHarvestingRealm() != ESaveRealm::Optional ? SaveContext.GetEDLCookChecker() : nullptr;
-			auto AddArcForDepChecking = [&Linker, &Export, EDLCookChecker](bool bExportIsSerialize, FPackageIndex Dep, bool bDepIsSerialize)
+			auto AddArcForDepChecking = [&Linker, &Export, &PreloadDependencies](bool bExportIsSerialize, FPackageIndex Dep, bool bDepIsSerialize)
 			{
 				check(Export.Object);
 				check(!Dep.IsNull());
@@ -1948,10 +1920,12 @@ void SavePreloadDependencies(FStructuredArchive::FRecord& StructuredArchiveRoot,
 				}
 
 				Linker->DepListForErrorChecking.Add(Dep);
-				if (EDLCookChecker && DepObject)
-				{
-					EDLCookChecker->AddArc(DepObject, bDepIsSerialize, Export.Object, bExportIsSerialize);
-				}
+				PreloadDependencies.Add({
+					.SourceObject = Export.Object,
+					.TargetObject = DepObject,
+					.bSourceIsSerialize = bExportIsSerialize,
+					.bTargetIsSerialize = bDepIsSerialize
+					});
 			};
 
 			for (FPackageIndex Index : SerializationBeforeSerializationDependencies)
@@ -3634,8 +3608,6 @@ ESavePackageResult SaveHarvestedRealms(FSaveContext& SaveContext, ESaveRealm Har
  */
 ESavePackageResult InnerSave(FSaveContext& SaveContext)
 {
-	SaveContext.SetEDLCookChecker(&FEDLCookCheckerThreadState::Get());
-
 	// Create slow task dialog if needed
 	const int32 TotalSaveSteps = 3;
 	FScopedSlowTask SlowTask(TotalSaveSteps, FText(), SaveContext.IsUsingSlowTask());
diff --git a/Engine/Source/Runtime/CoreUObject/Public/UObject/LinkerDiff.h b/Engine/Source/Runtime/CoreUObject/Public/UObject/LinkerDiff.h
--- a/Engine/Source/Runtime/CoreUObject/Public/UObject/LinkerDiff.h
+++ b/Engine/Source/Runtime/CoreUObject/Public/UObject/LinkerDiff.h
@@ -13,7 +13,9 @@ class FName;
 class FOutputDevice;
 struct FLinkerExportObject;
 
-class FLinkerDiff
+class
+UE_DEPRECATED(5.6, "Functionality is no longer needed, now that legacy Save is removed and Save2 is standard. Contact Epic if you need this functionality.")
+FLinkerDiff
 {
 public:
 	static COREUOBJECT_API FLinkerDiff CompareLinkers(FLinker* LHSLinker, FLinker* RHSLinker);
diff --git a/Engine/Source/Runtime/CoreUObject/Public/UObject/ObjectMacros.h b/Engine/Source/Runtime/CoreUObject/Public/UObject/ObjectMacros.h
--- a/Engine/Source/Runtime/CoreUObject/Public/UObject/ObjectMacros.h
+++ b/Engine/Source/Runtime/CoreUObject/Public/UObject/ObjectMacros.h
@@ -97,7 +97,7 @@ enum ESaveFlags
 //	SAVE_Unused						= 0x00000200,
 //	SAVE_Unused						= 0x00000400,
 //	SAVE_Unused						= 0x00000800,
-	SAVE_CompareLinker				= 0x00001000,	///< Return the linker save to compare against another
+	SAVE_CompareLinker UE_DEPRECATED(5.6, "Comparing Linkers is no longer used. Contact Epic if you need this functionality.") = 0x00001000,
 	SAVE_BulkDataByReference		= 0x00002000,	///< When saving to a different file than the package's LoadedPath, point bulkdata in the new file to be loaded from the original file.
 	SAVE_Unversioned_Properties		= 0x00004000,	///< Properties are saved without property name information, and are saved/loaded in the order of the current binary.
 													///This is only reasonable to use with full cooked builds for distribution, or with a domain that selects the payload matching a given version externally.
diff --git a/Engine/Source/Runtime/CoreUObject/Public/UObject/Package.h b/Engine/Source/Runtime/CoreUObject/Public/UObject/Package.h
--- a/Engine/Source/Runtime/CoreUObject/Public/UObject/Package.h
+++ b/Engine/Source/Runtime/CoreUObject/Public/UObject/Package.h
@@ -103,6 +103,23 @@ inline bool IsSuccessful(ESavePackageResult Result)
 		Result == ESavePackageResult::ReplaceCompletely;
 }
 
+namespace UE::SavePackageUtilities
+{
+
+/**
+ * A dependency in a runtime cook package from the load phase (create or serialize) of one object to the load phase of
+ * another object.
+ * Experimental, may be changed without deprecation.
+ */
+struct FPreloadDependency
+{
+	UObject* SourceObject;
+	UObject* TargetObject;
+	bool bSourceIsSerialize;
+	bool bTargetIsSerialize;
+};
+
+}
 /**
 * Struct returned from save package, contains the enum as well as extra data about what was written
 */
@@ -120,7 +137,7 @@ struct FSavePackageResultStruct
 	/** Serialized package flags */
 	uint32 SerializedPackageFlags;
 
-	/** Linker for linker comparison after save. */
+	UE_DEPRECATED(5.6, "Returning the LinkerSave for comparison is no longer used. Contact Epic if you need this functionality.")
 	TPimplPtr<FLinkerSave> LinkerSave;
 
 	TArray<FAssetData> SavedAssets;
@@ -133,6 +150,20 @@ struct FSavePackageResultStruct
 	UE_DEPRECATED(5.6, "Use BuildResultDependencies instead..")
 	TArray<UE::Cook::FCookDependency> CookDependencies;
 	UE::Cook::FBuildResultDependenciesMap BuildResultDependencies;
+
+	/** Only populated during cook saves. */
+	TArray<UObject*> Imports;
+	/** Only populated during cook saves. */
+	TArray<UObject*> Exports;
+	/**
+	 * Only populated during cook saves. The list of object dependencies from exports in this package to
+	 * other objects, either imports or exports, and also specifying which stage of the source object and which stage
+	 * of the target object have the dependency. These dependencies are also recorded in the save package and are used
+	 * for runtime loading of cooked packages. This output variable allows the cooker to run validation on that graph
+	 * of dependencies.
+	 * Experimental, may be changed without deprecation.
+	 */
+	UE_INTERNAL TArray<UE::SavePackageUtilities::FPreloadDependency> PreloadDependencies;
 #endif
 
 	/** Constructors, it will implicitly construct from the result enum */
@@ -140,7 +171,10 @@ struct FSavePackageResultStruct
 	COREUOBJECT_API FSavePackageResultStruct(ESavePackageResult InResult);
 	COREUOBJECT_API FSavePackageResultStruct(ESavePackageResult InResult, int64 InTotalFileSize);
 	COREUOBJECT_API FSavePackageResultStruct(ESavePackageResult InResult, int64 InTotalFileSize,
-		uint32 InSerializedPackageFlags, TPimplPtr<FLinkerSave> Linker = nullptr);
+		uint32 InSerializedPackageFlags);
+	UE_DEPRECATED(5.6, "Returning the LinkerSave for comparison is no longer used. Contact Epic if you need this functionality.")
+	COREUOBJECT_API FSavePackageResultStruct(ESavePackageResult InResult, int64 InTotalFileSize,
+		uint32 InSerializedPackageFlags, TPimplPtr<FLinkerSave> Linker);
 	COREUOBJECT_API FSavePackageResultStruct(FSavePackageResultStruct&& Other);
 	COREUOBJECT_API FSavePackageResultStruct& operator=(FSavePackageResultStruct&& Other);
 	COREUOBJECT_API ~FSavePackageResultStruct();
diff --git a/Engine/Source/Runtime/CoreUObject/Public/UObject/SavePackage.h b/Engine/Source/Runtime/CoreUObject/Public/UObject/SavePackage.h
--- a/Engine/Source/Runtime/CoreUObject/Public/UObject/SavePackage.h
+++ b/Engine/Source/Runtime/CoreUObject/Public/UObject/SavePackage.h
@@ -319,17 +319,25 @@ namespace UE::SavePackageUtilities
 	COREUOBJECT_API void ResetCookStats();
 	COREUOBJECT_API int32 GetNumPackagesSaved();
 
+	UE_DEPRECATED(5.6, "Functionality has moved into private cooker implementation; contact Epic if you need this functionality.")
 	COREUOBJECT_API void StartSavingEDLCookInfoForVerification();
 	using FEDLMessageCallback = TFunction<void(ELogVerbosity::Type, FStringView)>;
 	using FEDLLogRecordCallback = TFunction<void(UE::FLogRecord&& Record)>;
+	UE_DEPRECATED(5.6, "Functionality has moved into private cooker implementation; contact Epic if you need this functionality.")
 	COREUOBJECT_API void VerifyEDLCookInfo(bool bFullReferencesExpected = true);
-	UE_DEPRECATED(5.5, "Use version that takes FEDLLogRecordCallback")
+	UE_DEPRECATED(5.5, "Functionality has moved into private cooker implementation; contact Epic if you need this functionality.")
 	COREUOBJECT_API void VerifyEDLCookInfo(const FEDLMessageCallback& MessageCallback, bool bFullReferencesExpected = true);
+	UE_DEPRECATED(5.6, "Functionality has moved into private cooker implementation; contact Epic if you need this functionality.")
 	COREUOBJECT_API void VerifyEDLCookInfo(const FEDLLogRecordCallback& MessageCallback, bool bFullReferencesExpected = true);
+	UE_DEPRECATED(5.6, "Functionality has moved into private cooker implementation; contact Epic if you need this functionality.")
 	COREUOBJECT_API void EDLCookInfoAddIterativelySkippedPackage(FName LongPackageName);
+	UE_DEPRECATED(5.6, "Functionality has moved into private cooker implementation; contact Epic if you need this functionality.")
 	COREUOBJECT_API void EDLCookInfoMoveToCompactBinaryAndClear(FCbWriter& Writer, bool& bOutHasData);
+	UE_DEPRECATED(5.6, "Functionality has moved into private cooker implementation; contact Epic if you need this functionality.")
 	COREUOBJECT_API void EDLCookInfoMoveToCompactBinaryAndClear(FCbWriter& Writer, bool& bOutHasData, FName PackageName);
+	UE_DEPRECATED(5.6, "Functionality has moved into private cooker implementation; contact Epic if you need this functionality.")
 	COREUOBJECT_API bool EDLCookInfoAppendFromCompactBinary(FCbFieldView Field);
+
 	UE_DEPRECATED(5.5, "No longer used; skiponlyeditoronly is used instead and tracks editoronly references via savepackage results.")
 	inline bool CanSkipEditorReferencedPackagesWhenCooking() { return false; }
 
