dstm:

- fix so that we also look at whether the root object is an asset before we consider it for multi-server commit
- fix up the logging
- fix a bug where we could add an object to the UsedObjects list more than once

diff --git a/Engine/Source/Runtime/CoreUObject/Private/UObject/RemoteObjectTransfer.cpp b/Engine/Source/Runtime/CoreUObject/Private/UObject/RemoteObjectTransfer.cpp
--- a/Engine/Source/Runtime/CoreUObject/Private/UObject/RemoteObjectTransfer.cpp
+++ b/Engine/Source/Runtime/CoreUObject/Private/UObject/RemoteObjectTransfer.cpp
@@ -465,20 +465,37 @@ public:
 				continue;
 			}
 
-			UE_LOG(LogRemoteObject, VeryVerbose, TEXT("FRemoteObjectTransferQueue: used obj: %s remote: %d"), *UsedObjectId.ToString(), Handle::IsRemote(UsedObjectId));
-
 			// if we used it this run and we're about to commit, we expect that the object is local
 			check(!Handle::IsRemote(UsedObjectId));
 
 			UObject* UsedObject = StaticFindObjectFastInternal(UsedObjectId);
 			check(UsedObject);
+			FRemoteServerId OwnerServerId = Handle::GetOwnerServerId(UsedObject);
 
 			UObject* RootObject = FindCanonicalRootObjectForSerialization(UsedObject);
 			check(RootObject);
+			FRemoteObjectId RootObjectId = FRemoteObjectId(RootObject);
 			FRemoteServerId RootOwnerServerId = Handle::GetOwnerServerId(RootObject);
 
+			if (RootObjectId.IsAsset())
+			{
+				// if the used object is not an asset, but the root of it is, then we
+				// expect that the used object owner is the local server
+				check(OwnerServerId == GetGlobalServerId());
+
+				continue;
+			}
+
 			if (RootOwnerServerId != GetGlobalServerId())
 			{
+				UE_LOG(LogRemoteObject, VeryVerbose, TEXT("FRemoteObjectTransferQueue: borrowed obj: %s (%s) owner server: %s (root: %s (s) root owner: %s)"),
+					*UsedObjectId.ToString(),
+					*UsedObject->GetName(),
+					*OwnerServerId.ToString(),
+					*RootObjectId.ToString(),
+					*RootObject->GetName(),
+					*RootOwnerServerId.ToString());
+
 				MultiServerCommitObjectsToReturn.Add(RootObject);
 
 				if (OutMultiServerCommitRemoteServers.Find(RootOwnerServerId) == INDEX_NONE)
@@ -885,6 +902,8 @@ void MigrateObjectFromRemoteServer(FRemoteObjectId ObjectId, FRemoteServerId Cur
 		if (!GTransferQueue.bInMultiServerCommit && GTransferQueue.ActiveRequest->UsedObjects.Find(ObjectId) == INDEX_NONE)
 		{
 			GTransferQueue.ActiveRequest->UsedObjects.Add(ObjectId);
+
+			OnObjectTouchedDelegate.Broadcast(GTransferQueue.ActiveRequest->RequestId, ObjectId);
 		}
 
 		// add it to the active request's list of required objects
@@ -912,10 +931,6 @@ void MigrateObjectFromRemoteServer(FRemoteObjectId ObjectId, FRemoteServerId Cur
 
 			GTransferQueue.ActiveRequest->RequiredObjects.Add(ObjectId);
 			GTransferQueue.ActiveRequest->RequiredObjectsCanonicalRoot.Add(CanonicalRootObjectId);
-
-			GTransferQueue.ActiveRequest->UsedObjects.Add(ObjectId);
-
-			OnObjectTouchedDelegate.Broadcast(GTransferQueue.ActiveRequest->RequestId, ObjectId);
 		}
 		
 		if (UE_LOG_ACTIVE(LogRemoteObject, VeryVerbose))
@@ -958,6 +973,8 @@ void TouchResidentObject(UObject* Object)
 			if (GTransferQueue.ActiveRequest->UsedObjects.Find(ObjectId) == INDEX_NONE)
 			{
 				GTransferQueue.ActiveRequest->UsedObjects.Add(ObjectId);
+
+				OnObjectTouchedDelegate.Broadcast(GTransferQueue.ActiveRequest->RequestId, ObjectId);
 			}
 
 			// add it to the active request's list of required objects
@@ -970,8 +987,6 @@ void TouchResidentObject(UObject* Object)
 				// called when object pointers get resolved
 				FRemoteObjectId CanonicalRootObjectId;
 				GTransferQueue.ActiveRequest->RequiredObjectsCanonicalRoot.Add(CanonicalRootObjectId);
-
-				OnObjectTouchedDelegate.Broadcast(GTransferQueue.ActiveRequest->RequestId, ObjectId);
 			}
 		};
 	}
@@ -1117,6 +1132,8 @@ void PrefetchRemoteObject(FRemoteObjectId ObjectId, UObject* DestinationOuter)
 		if (!GTransferQueue.bInMultiServerCommit && GTransferQueue.ActiveRequest->UsedObjects.Find(ObjectId) == INDEX_NONE)
 		{
 			GTransferQueue.ActiveRequest->UsedObjects.Add(ObjectId);
+
+			OnObjectTouchedDelegate.Broadcast(GTransferQueue.ActiveRequest->RequestId, ObjectId);
 		}
 
 		// it's possible a previous prefetch or strong fetch already added this to the RequiredObjects list
@@ -1127,8 +1144,6 @@ void PrefetchRemoteObject(FRemoteObjectId ObjectId, UObject* DestinationOuter)
 			GTransferQueue.ActiveRequest->RequiredObjects.Add(ObjectId);
 			GTransferQueue.ActiveRequest->RequiredObjectsCanonicalRoot.Add(ObjectId);
 			UE_LOG(LogRemoteObject, VeryVerbose, TEXT("PrefetchRemoteObject(%s) : obj %s"), *GTransferQueue.ActiveRequest->GetRequestId().ToString(), *ObjectId.ToString());
-
-			OnObjectTouchedDelegate.Broadcast(GTransferQueue.ActiveRequest->RequestId, ObjectId);
 		}
 	};
 }
