dstm:

- inside FWeakObjectPtr::TryResolveRemoteObject, added a call to TouchResidentObject since we don't hit any of the ResolveObject calls if the object is local and already resolved. we need to ensure we do call Touch to report that we've used the object (used by the multi-server commit system)

diff --git a/Engine/Source/Runtime/CoreUObject/Private/UObject/WeakObjectPtr.cpp b/Engine/Source/Runtime/CoreUObject/Private/UObject/WeakObjectPtr.cpp
--- a/Engine/Source/Runtime/CoreUObject/Private/UObject/WeakObjectPtr.cpp
+++ b/Engine/Source/Runtime/CoreUObject/Private/UObject/WeakObjectPtr.cpp
@@ -241,6 +241,10 @@ FUObjectItem* FWeakObjectPtr::TryResolveRemoteObject() const
 			// Object memory is still on this server but it's marked as remote so we need to resolve it
 			ResolvedObject = ResolveObject(ResolvedObject, ERemoteReferenceType::Weak);
 		}
+		else
+		{
+			TouchResidentObject(ResolvedObject);
+		}
 	}
 	else if (FRemoteObjectStub* Stub = FindRemoteObjectStub(ObjectRemoteId))
 	{
