Async loading propogates NotInstalled errors from IoDispatcher and IoPackageStore.  This gives reliable errors for when IAD packages are not installed locally (as distinct from missing).

diff --git a/Engine/Plugins/Runtime/Metasound/Source/MetasoundEngine/Private/MetasoundAssetSubsystem.cpp b/Engine/Plugins/Runtime/Metasound/Source/MetasoundEngine/Private/MetasoundAssetSubsystem.cpp
--- a/Engine/Plugins/Runtime/Metasound/Source/MetasoundEngine/Private/MetasoundAssetSubsystem.cpp
+++ b/Engine/Plugins/Runtime/Metasound/Source/MetasoundEngine/Private/MetasoundAssetSubsystem.cpp
@@ -618,16 +618,11 @@ namespace Metasound::Engine
 				break;
 			}
 
-			case EAsyncLoadingResult::Failed:
+			default:
 			{
 				UE_LOG(LogMetaSound, Warning, TEXT("MetaSoundAssetManager request to aync load asset '%s' failed"), *PackageLoadedArgs.PackageName.ToString());
 				break;
 			}
-
-			default:
-			{
-				checkNoEntry();
-			}
 		}
 	}
 
diff --git a/Engine/Source/Runtime/Core/Private/IO/PackageStore.cpp b/Engine/Source/Runtime/Core/Private/IO/PackageStore.cpp
--- a/Engine/Source/Runtime/Core/Private/IO/PackageStore.cpp
+++ b/Engine/Source/Runtime/Core/Private/IO/PackageStore.cpp
@@ -193,11 +193,16 @@ EPackageStoreEntryStatus FPackageStore::GetPackageStoreEntry(FPackageId PackageI
 	for (const FBackendAndPriority& Backend : Backends)
 	{
 		EPackageStoreEntryStatus Status = Backend.Value->GetPackageStoreEntry(PackageId, PackageName, OutPackageStoreEntry);
-		if (Status >= EPackageStoreEntryStatus::Pending)
+		const bool bContinueSearch = 
+			Status == EPackageStoreEntryStatus::None || 
+			Status == EPackageStoreEntryStatus::Missing;
+
+		if (!bContinueSearch)
 		{
 			return Status;
 		}
 	}
+
 	return EPackageStoreEntryStatus::Missing;
 }
 
diff --git a/Engine/Source/Runtime/Core/Public/IO/PackageStore.h b/Engine/Source/Runtime/Core/Public/IO/PackageStore.h
--- a/Engine/Source/Runtime/Core/Public/IO/PackageStore.h
+++ b/Engine/Source/Runtime/Core/Public/IO/PackageStore.h
@@ -26,6 +26,7 @@ enum class EPackageStoreEntryStatus
 {
 	None,
 	Missing,
+	NotInstalled,
 	Pending,
 	Ok,
 };
diff --git a/Engine/Source/Runtime/CoreUObject/Private/Serialization/AsyncLoading2.cpp b/Engine/Source/Runtime/CoreUObject/Private/Serialization/AsyncLoading2.cpp
--- a/Engine/Source/Runtime/CoreUObject/Private/Serialization/AsyncLoading2.cpp
+++ b/Engine/Source/Runtime/CoreUObject/Private/Serialization/AsyncLoading2.cpp
@@ -482,6 +482,13 @@ static FAutoConsoleVariableRef CVarPreferCookedPackagesInEditor(
 	ECVF_Default
 );
 
+static bool GFailLoadOnNotInstalledImport = true;
+static FAutoConsoleVariableRef CVarGFailLoadOnNotInstalledImport(
+	TEXT("s.FailLoadOnNotInstalledImport"),
+	GFailLoadOnNotInstalledImport,
+	TEXT("Fail package load if an imported package is not installed"),
+	ECVF_Default
+);
 
 
 #if WITH_EDITOR
@@ -3561,7 +3568,12 @@ struct FAsyncPackage2
 	/** Returns true if loading has failed */
 	FORCEINLINE bool HasLoadFailed() const
 	{
-		return bLoadHasFailed;
+		return LoadStatus != EAsyncLoadingResult::Succeeded;
+	}
+
+	FORCEINLINE EAsyncLoadingResult::Type GetLoadStatus() const
+	{
+		return LoadStatus;
 	}
 
 	/** Adds new request ID to the existing package */
@@ -3791,7 +3803,7 @@ private:
 	FAllDependenciesState		AllDependenciesFullyLoadedState;
 
 	/** True if our load has failed */
-	bool						bLoadHasFailed = false;
+	EAsyncLoadingResult::Type	LoadStatus = EAsyncLoadingResult::Succeeded;
 	/** True if this package was created by this async package */
 	bool						bCreatedLinkerRoot = false;
 	/** List of all request handles. */
@@ -4029,11 +4041,16 @@ private:
 		TArray<TUniquePtr<FLoadPackageAsyncProgressDelegate>, TInlineAllocator<2>> ProgressCallbacks;
 		TArray<int32, TInlineAllocator<2>> RequestIDs;
 
-		static FCompletedPackageRequest FromMissingPackage(
+		static FCompletedPackageRequest FromUnreslovedPackage(
 			const FAsyncPackageDesc2& Desc,
+			EAsyncLoadingResult::Type Result,
 			TUniquePtr<FLoadPackageAsyncDelegate>&& CompletionCallback)
 		{
-			FCompletedPackageRequest Res = {Desc.UPackageName, EAsyncLoadingResult::Failed};
+			FCompletedPackageRequest Res = 
+			{
+				.PackageName = Desc.UPackageName,
+				.Result = Result
+			};
 			Res.CompletionCallbacks.Add(MoveTemp(CompletionCallback));
 			Res.RequestIDs.Add(Desc.RequestID);
 			return Res;
@@ -4044,12 +4061,12 @@ private:
 		{
 			FCompletedPackageRequest Request
 			{
-				Package->Desc.UPackageName,
-				Package->HasLoadFailed() ? EAsyncLoadingResult::Failed : EAsyncLoadingResult::Succeeded,
-				Package->LinkerRoot,
-				Package,
-				MoveTemp(Package->CompletionCallbacks),
-				MoveTemp(Package->ProgressCallbacks)
+				.PackageName = Package->Desc.UPackageName,
+				.Result = Package->GetLoadStatus(),
+				.UPackage = Package->LinkerRoot,
+				.AsyncPackage = Package,
+				.CompletionCallbacks = MoveTemp(Package->CompletionCallbacks),
+				.ProgressCallbacks = MoveTemp(Package->ProgressCallbacks)
 			};
 
 			Request.RequestIDs.Append(Package->RequestIDs);
@@ -4283,7 +4300,7 @@ public:
 	void UpdatePackagePriorityRecursive(FAsyncLoadingThreadState2& ThreadState, FAsyncPackage2* Package, int32 NewPriority);
 
 	FAsyncPackage2* FindOrInsertPackage(FAsyncLoadingThreadState2& ThreadState, FAsyncPackageDesc2& InDesc, bool& bInserted, FAsyncPackage2* ImportedByPackage, TUniquePtr<FLoadPackageAsyncDelegate>&& PackageLoadedDelegate = TUniquePtr<FLoadPackageAsyncDelegate>(), TUniquePtr<FLoadPackageAsyncProgressDelegate>&& PackageProgressDelegate = TUniquePtr<FLoadPackageAsyncProgressDelegate>());
-	void QueueMissingPackage(FAsyncLoadingThreadState2& ThreadState, FAsyncPackageDesc2& PackageDesc, TUniquePtr<FLoadPackageAsyncDelegate>&& LoadPackageAsyncDelegate, TUniquePtr<FLoadPackageAsyncProgressDelegate>&& LoadPackageAsyncProgressDelegate);
+	void QueueUnresolvedPackage(FAsyncLoadingThreadState2& ThreadState, EPackageStoreEntryStatus PackageStatus, FAsyncPackageDesc2& PackageDesc, TUniquePtr<FLoadPackageAsyncDelegate>&& LoadPackageAsyncDelegate, TUniquePtr<FLoadPackageAsyncProgressDelegate>&& LoadPackageAsyncProgressDelegate);
 
 	/**
 	* [ASYNC* THREAD] Loads all packages
@@ -5395,7 +5412,7 @@ bool FAsyncLoadingThread2::CreateAsyncPackagesFromQueue(FAsyncLoadingThreadState
 			InitializeAsyncPackageFromPackageStore(ThreadState, &IoBatch, PendingPackage, FakePackageEntry);
 			// Simulate StartLoading() getting back a failed IoRequest and let it go through all package states
 			PendingPackage->AsyncPackageLoadingState = EAsyncPackageLoadingState2::WaitingForIo;
-			PendingPackage->bLoadHasFailed = true;
+			PendingPackage->LoadStatus = EAsyncLoadingResult::FailedMissing;
 			PendingPackage->GetPackageNode(EEventLoadNode2::Package_ProcessSummary).ReleaseBarrier(&ThreadState);
 			// Remove from PendingPackages
 			It.RemoveCurrent();
@@ -5501,9 +5518,9 @@ bool FAsyncLoadingThread2::CreateAsyncPackagesFromQueue(FAsyncLoadingThreadState
 			}
 
 			FAsyncPackageDesc2 PackageDesc = FAsyncPackageDesc2::FromPackageRequest(Request, UPackageName, PackageIdToLoad, PackageLoader);
-			if (PackageStatus == EPackageStoreEntryStatus::Missing)
+			if (PackageStatus == EPackageStoreEntryStatus::Missing || PackageStatus == EPackageStoreEntryStatus::NotInstalled)
 			{
-				QueueMissingPackage(ThreadState, PackageDesc, MoveTemp(Request.PackageLoadedDelegate), MoveTemp(Request.PackageProgressDelegate));
+				QueueUnresolvedPackage(ThreadState, PackageStatus, PackageDesc, MoveTemp(Request.PackageLoadedDelegate), MoveTemp(Request.PackageProgressDelegate));
 			}
 			else
 			{
@@ -6224,6 +6241,8 @@ void FAsyncPackage2::ImportPackagesRecursiveInner(FAsyncLoadingThreadState2& Thr
 			bIsFullyLoaded = false;
 		}
 #endif
+		const bool bIsNotInstalled = GFailLoadOnNotInstalledImport && ImportedPackageStatus == EPackageStoreEntryStatus::NotInstalled;
+		const bool bIsMissing = (ImportedPackageStatus == EPackageStoreEntryStatus::Missing) || (!GFailLoadOnNotInstalledImport && ImportedPackageStatus == EPackageStoreEntryStatus::NotInstalled);
 		if (bIsFullyLoaded)
 		{
 			ImportedPackage = AsyncLoadingThread.FindAsyncPackage(ImportedPackageId);
@@ -6233,7 +6252,30 @@ void FAsyncPackage2::ImportPackagesRecursiveInner(FAsyncLoadingThreadState2& Thr
 			}
 			bInserted = false;
 		}
-		else if (ImportedPackageStatus == EPackageStoreEntryStatus::Missing)
+		else if (bIsNotInstalled)
+		{
+			LoadStatus = EAsyncLoadingResult::FailedNotInstalled;
+			ImportedPackageRef.SetHasFailed();
+
+			UE_ASYNC_PACKAGE_CLOG(!ImportedPackageUPackageName.IsNone(), Log, Desc, TEXT("ImportPackages: SkipPackage"),
+				TEXT("Skipping not installed imported package %s (0x%s)"), *ImportedPackageNameToLoad.ToString(), *LexToString(ImportedPackageId));
+			if (!ImportedPackageUPackageName.IsNone())
+			{
+				FName PackagePathToLoadName = Desc.PackagePathToLoad.GetPackageFName();
+				ExecuteOnGameThread(TEXT("ExplanationForNotInstalledPackage"), [ImportedPackageNameToLoad, PackagePathToLoadName]()
+				{
+					FString ImportedPackageNameToLoadString = ImportedPackageNameToLoad.ToString();
+					FString PackagePathToLoadString = PackagePathToLoadName.ToString();
+					FMessageLog("LoadErrors").Warning(FText::FormatNamed(NSLOCTEXT("Core", "AsyncLoading_SkippedPackage_NotInstalled", "While trying to load package {Package}, a dependent package {DependentPackage} was unavailable because it was not installed."),
+						TEXT("Package"), FText::FromString(PackagePathToLoadString),
+						TEXT("DependentPackage"), FText::FromString(ImportedPackageNameToLoadString)));
+					return false;
+				});
+			}
+
+			continue;
+		}
+		else if (bIsMissing)
 		{
 			if (!ImportedPackageRef.HasPackage()) // If we found a package it's not actually missing but we can't load it anyway
 			{
@@ -6447,6 +6489,22 @@ void FAsyncPackage2::DetachLinker()
 }
 #endif
 
+namespace UE::AsyncPackage2
+{
+	EAsyncLoadingResult::Type AsyncLoadingResultFromIoError(EIoErrorCode Error)
+	{
+		switch (Error)
+		{
+		case EIoErrorCode::Ok:
+			return EAsyncLoadingResult::Succeeded;
+		case EIoErrorCode::NotInstalled:
+			return EAsyncLoadingResult::FailedNotInstalled;
+		}
+
+		return EAsyncLoadingResult::Failed;
+	}
+}
+
 void FAsyncPackage2::StartLoading(FAsyncLoadingThreadState2& ThreadState, FIoBatch& IoBatch)
 {
 	TRACE_CPUPROFILER_EVENT_SCOPE(StartLoading);
@@ -6488,7 +6546,7 @@ void FAsyncPackage2::StartLoading(FAsyncLoadingThreadState2& ThreadState, FIoBat
 				{
 					UE_ASYNC_PACKAGE_LOG(Warning, Desc, TEXT("StartBundleIoRequests: FailedRead"),
 						TEXT("Failed reading optional chunk for package: %s"), *Result.Status().ToString());
-					bLoadHasFailed = true;
+					LoadStatus = UE::AsyncPackage2::AsyncLoadingResultFromIoError(Result.Status().GetErrorCode());
 				}
 				int32 LocalPendingIoRequestsCounter = AsyncLoadingThread.PendingIoRequestsCounter.DecrementExchange() - 1;
 				TRACE_COUNTER_SET(AsyncLoadingPendingIoRequests, LocalPendingIoRequestsCounter);
@@ -6526,7 +6584,7 @@ void FAsyncPackage2::StartLoading(FAsyncLoadingThreadState2& ThreadState, FIoBat
 			{
 				UE_ASYNC_PACKAGE_LOG(Warning, Desc, TEXT("StartBundleIoRequests: FailedRead"),
 					TEXT("Failed reading chunk for package: %s"), *Result.Status().ToString());
-				bLoadHasFailed = true;
+				LoadStatus = UE::AsyncPackage2::AsyncLoadingResultFromIoError(Result.Status().GetErrorCode());
 			}
 			int32 LocalPendingIoRequestsCounter = AsyncLoadingThread.PendingIoRequestsCounter.DecrementExchange() - 1;
 			TRACE_COUNTER_SET(AsyncLoadingPendingIoRequests, LocalPendingIoRequestsCounter);
@@ -6588,9 +6646,9 @@ EEventLoadNodeExecutionResult FAsyncPackage2::ProcessLinkerLoadPackageSummary(FA
 	check(LinkerResult != FLinkerLoad::LINKER_TimedOut); // TODO: Add support for timeouts here
 	if (LinkerResult == FLinkerLoad::LINKER_Failed)
 	{
-		bLoadHasFailed = true;
+		LoadStatus = EAsyncLoadingResult::FailedLinker;
 	}
-	check(LinkerLoadState->Linker->HasFinishedInitialization() || bLoadHasFailed);
+	check(LinkerLoadState->Linker->HasFinishedInitialization() || LinkerResult == FLinkerLoad::LINKER_Failed);
 
 	LinkerLoadState->LinkerLoadHeaderData.ImportMap.SetNum(LinkerLoadState->Linker->ImportMap.Num());
 	TArray<FName, TInlineAllocator<128>> ImportedPackageNames;
@@ -6724,7 +6782,7 @@ EEventLoadNodeExecutionResult FAsyncPackage2::ProcessLinkerLoadPackageSummary(FA
 	if (Desc.bCanBeImported)
 	{
 		FLoadedPackageRef& PackageRef = AsyncLoadingThread.GlobalImportStore.FindPackageRefChecked(Desc.UPackageId, Desc.UPackageName);
-		if (!bLoadHasFailed)
+		if (!HasLoadFailed())
 		{
 			PackageRef.PreInsertPublicExports({ LinkerLoadState->LinkerLoadHeaderData.ExportMap });
 
@@ -6962,7 +7020,7 @@ bool FAsyncPackage2::CreateLinkerLoadExports(FAsyncLoadingThreadState2& ThreadSt
 
 EEventLoadNodeExecutionResult FAsyncPackage2::ExecutePostLoadLinkerLoadPackageExports(FAsyncLoadingThreadState2& ThreadState)
 {
-	if (!bLoadHasFailed)
+	if (!HasLoadFailed())
 	{
 		SCOPED_LOADTIMER(PostLoadObjectsTime);
 		TRACE_LOADTIME_POSTLOAD_SCOPE;
@@ -7232,7 +7290,7 @@ EEventLoadNodeExecutionResult FAsyncPackage2::Event_ProcessPackageSummary(FAsync
 		return Package->ProcessLinkerLoadPackageSummary(ThreadState);
 	}
 #endif
-	if (Package->bLoadHasFailed)
+	if (Package->HasLoadFailed())
 	{
 		if (Package->Desc.bCanBeImported)
 		{
@@ -7409,7 +7467,7 @@ EEventLoadNodeExecutionResult FAsyncPackage2::Event_ProcessExportBundle(FAsyncLo
 
 	check(InExportBundleIndex < Package->Data.TotalExportBundleCount);
 	
-	if (!Package->bLoadHasFailed)
+	if (!Package->HasLoadFailed())
 	{
 		bool bIsOptionalSegment = false;
 #if WITH_EDITOR
@@ -8381,7 +8439,7 @@ EEventLoadNodeExecutionResult FAsyncPackage2::Event_ExportsDone(FAsyncLoadingThr
 
 	Package->AsyncPackageLoadingState = EAsyncPackageLoadingState2::ExportsDone;
 
-	if (!Package->bLoadHasFailed && Package->Desc.bCanBeImported)
+	if (!Package->HasLoadFailed() && Package->Desc.bCanBeImported)
 	{
 		FLoadedPackageRef& PackageRef = Package->AsyncLoadingThread.GlobalImportStore.FindPackageRefChecked(Package->Desc.UPackageId, Package->Desc.UPackageName);
 		const bool bSnapshotExportCount = true;
@@ -8729,7 +8787,7 @@ EEventLoadNodeExecutionResult FAsyncPackage2::Event_PostLoadExportBundle(FAsyncL
 
 	EEventLoadNodeExecutionResult LoadingState = EEventLoadNodeExecutionResult::Complete;
 
-	if (!Package->bLoadHasFailed)
+	if (!Package->HasLoadFailed())
 	{
 		// Begin async loading, simulates BeginLoad
 		Package->BeginAsyncLoad();
@@ -8862,7 +8920,7 @@ EEventLoadNodeExecutionResult FAsyncPackage2::Event_DeferredPostLoadExportBundle
 	check(InExportBundleIndex < Package->Data.TotalExportBundleCount);
 	EEventLoadNodeExecutionResult LoadingState = EEventLoadNodeExecutionResult::Complete;
 
-	if (!Package->bLoadHasFailed)
+	if (!Package->HasLoadFailed())
 	{
 		SCOPED_LOADTIMER(PostLoadDeferredObjectsTime);
 		TRACE_LOADTIME_POSTLOAD_SCOPE;
@@ -9274,7 +9332,7 @@ EAsyncPackageState::Type FAsyncLoadingThread2::ProcessLoadedPackagesFromGameThre
 			if (Package->AsyncPackageLoadingState == EAsyncPackageLoadingState2::CreateClusters)
 			{
 				SCOPE_CYCLE_COUNTER(STAT_FAsyncPackage_CreateClustersGameThread);
-				if (Package->bLoadHasFailed || !CanCreateObjectClusters())
+				if (Package->HasLoadFailed() || !CanCreateObjectClusters())
 				{
 					Package->AsyncPackageLoadingState = EAsyncPackageLoadingState2::Complete;
 				}
@@ -9311,7 +9369,7 @@ EAsyncPackageState::Type FAsyncLoadingThread2::ProcessLoadedPackagesFromGameThre
 				}
 
 #if WITH_EDITOR
-				if (!Package->bLoadHasFailed)
+				if (!Package->HasLoadFailed())
 				{
 					// In the editor we need to find any assets and packages and add them to list for later callback
 					EditorCompletedUPackages.Add(Package->LinkerRoot);
@@ -10688,7 +10746,7 @@ EAsyncPackageState::Type FAsyncPackage2::ProcessExternalReads(FAsyncLoadingThrea
 
 EAsyncPackageState::Type FAsyncPackage2::PostLoadInstances(FAsyncLoadingThreadState2& ThreadState)
 {
-	if (bLoadHasFailed)
+	if (HasLoadFailed())
 	{
 		return EAsyncPackageState::Complete;
 	}
@@ -10728,7 +10786,7 @@ EAsyncPackageState::Type FAsyncPackage2::CreateClusters(FAsyncLoadingThreadState
 
 void FAsyncPackage2::FinishUPackage()
 {
-	if (LinkerRoot && !bLoadHasFailed)
+	if (LinkerRoot && !HasLoadFailed())
 	{
 		// Mark package as having been fully loaded and update load time.
 		LinkerRoot->MarkAsFullyLoaded();
@@ -10788,7 +10846,7 @@ void FAsyncLoadingThread2::ConditionalProcessEditorCallbacks()
 void FAsyncPackage2::AddCompletionCallback(TUniquePtr<FLoadPackageAsyncDelegate>&& Callback)
 {
 	// This is to ensure that there is no one trying to subscribe to a already loaded package
-	//check(!bLoadHasFinished && !bLoadHasFailed);
+	//check(!bLoadHasFinished && !HasLoadFailed());
 	CompletionCallbacks.Emplace(MoveTemp(Callback));
 }
 
@@ -10873,7 +10931,7 @@ int32 FAsyncLoadingThread2::LoadPackage(const FPackagePath& InPackagePath, FName
 	return LoadPackageInternal(InPackagePath, InCustomName, MoveTemp(CompletionDelegate), TUniquePtr<FLoadPackageAsyncProgressDelegate>(), InPackageFlags, InPIEInstanceID, InPackagePriority, InInstancingContext, InLoadFlags);
 }
 
-void FAsyncLoadingThread2::QueueMissingPackage(FAsyncLoadingThreadState2& ThreadState, FAsyncPackageDesc2& PackageDesc, TUniquePtr<FLoadPackageAsyncDelegate>&& PackageLoadedDelegate, TUniquePtr<FLoadPackageAsyncProgressDelegate>&& PackageProgressDelegate)
+void FAsyncLoadingThread2::QueueUnresolvedPackage(FAsyncLoadingThreadState2& ThreadState, EPackageStoreEntryStatus PackageStatus, FAsyncPackageDesc2& PackageDesc, TUniquePtr<FLoadPackageAsyncDelegate>&& PackageLoadedDelegate, TUniquePtr<FLoadPackageAsyncProgressDelegate>&& PackageProgressDelegate)
 {
 	const FName FailedPackageName = PackageDesc.UPackageName;
 
@@ -10892,9 +10950,18 @@ void FAsyncLoadingThread2::QueueMissingPackage(FAsyncLoadingThreadState2& Thread
 		bIssueWarning &= !bIsScriptPackage;
 	}
 
-	UE_CLOG(bIssueWarning, LogStreaming, Warning,
-		TEXT("LoadPackage: SkipPackage: %s (0x%s) - The package to load does not exist on disk or in the loader"),
-		*FailedPackageName.ToString(), *LexToString(PackageDesc.PackageIdToLoad));
+	if (PackageStatus == EPackageStoreEntryStatus::NotInstalled)
+	{
+		UE_CLOG(bIssueWarning, LogStreaming, Warning,
+			TEXT("LoadPackage: SkipPackage: %s (0x%s) - The package to load does not exist on disk or in the loader but may be installed on demand"),
+			*FailedPackageName.ToString(), *LexToString(PackageDesc.PackageIdToLoad));
+	}
+	else
+	{
+		UE_CLOG(bIssueWarning, LogStreaming, Warning,
+			TEXT("LoadPackage: SkipPackage: %s (0x%s) - The package to load does not exist on disk or in the loader"),
+			*FailedPackageName.ToString(), *LexToString(PackageDesc.PackageIdToLoad));
+	}
 
 	if (PackageProgressDelegate.IsValid())
 	{
@@ -10912,9 +10979,15 @@ void FAsyncLoadingThread2::QueueMissingPackage(FAsyncLoadingThreadState2& Thread
 
 	if (PackageLoadedDelegate.IsValid())
 	{
+		EAsyncLoadingResult::Type Result = EAsyncLoadingResult::FailedMissing;
+		if (PackageStatus == EPackageStoreEntryStatus::NotInstalled)
+		{
+			Result = EAsyncLoadingResult::FailedNotInstalled;
+		}
+
 		FScopeLock _(&FailedPackageRequestsCritical);
 		FailedPackageRequests.Add(
-			FCompletedPackageRequest::FromMissingPackage(PackageDesc, MoveTemp(PackageLoadedDelegate)));
+			FCompletedPackageRequest::FromUnreslovedPackage(PackageDesc, Result, MoveTemp(PackageLoadedDelegate)));
 	}
 	else
 	{
diff --git a/Engine/Source/Runtime/CoreUObject/Public/UObject/UObjectGlobals.h b/Engine/Source/Runtime/CoreUObject/Public/UObject/UObjectGlobals.h
--- a/Engine/Source/Runtime/CoreUObject/Public/UObject/UObjectGlobals.h
+++ b/Engine/Source/Runtime/CoreUObject/Public/UObject/UObjectGlobals.h
@@ -657,10 +657,16 @@ COREUOBJECT_API UPackage* LoadPackage(UPackage* InOuter, const FPackagePath& InP
 /** Async package and object loading result */
 namespace EAsyncLoadingResult
 {
-	enum Type
+	enum Type : uint8
 	{
 		/** Package or object failed to load */
 		Failed,
+		/** Package or object failed to load because it is missing */
+		FailedMissing,
+		/** Package or object failed to load due to a linker error */
+		FailedLinker,
+		/** Package or object failed to load because its data is not local */
+		FailedNotInstalled,
 		/** Package or object loaded successfully */
 		Succeeded,
 		/** Async loading was canceled */
diff --git a/Engine/Source/Runtime/Engine/Private/Experimental/StreamableManagerError.cpp b/Engine/Source/Runtime/Engine/Private/Experimental/StreamableManagerError.cpp
--- a/Engine/Source/Runtime/Engine/Private/Experimental/StreamableManagerError.cpp
+++ b/Engine/Source/Runtime/Engine/Private/Experimental/StreamableManagerError.cpp
@@ -12,11 +12,13 @@ UE::UnifiedError::FError GetStreamableError(EAsyncLoadingResult::Type Result)
 	switch (Result)
 	{
 	case EAsyncLoadingResult::Failed:
+	case EAsyncLoadingResult::FailedMissing:
+	case EAsyncLoadingResult::FailedLinker:
 		// We could supply an error string if async loading bubbled one up...
 		// Possibly GetExplanationForUnavailablePackage?
 		return AsyncLoadFailed::MakeError();
-	/*case EAsyncLoadingResult::FailedNotInstalled:
-		return AsyncLoadNotInstalled::MakeError();*/
+	case EAsyncLoadingResult::FailedNotInstalled:
+		return AsyncLoadNotInstalled::MakeError();
 	case EAsyncLoadingResult::Canceled:
 		return AsyncLoadCancelled::MakeError();
 	}
diff --git a/Engine/Source/Runtime/Experimental/IoStore/OnDemand/Private/OnDemandPackageStoreBackend.cpp b/Engine/Source/Runtime/Experimental/IoStore/OnDemand/Private/OnDemandPackageStoreBackend.cpp
--- a/Engine/Source/Runtime/Experimental/IoStore/OnDemand/Private/OnDemandPackageStoreBackend.cpp
+++ b/Engine/Source/Runtime/Experimental/IoStore/OnDemand/Private/OnDemandPackageStoreBackend.cpp
@@ -123,19 +123,12 @@ EPackageStoreEntryStatus FOnDemandPackageStoreBackend::GetPackageStoreEntry(
 		OutPackageStoreEntry.ShaderMapHashes =
 			MakeArrayView(FilePackageStoreEntry->ShaderMapHashes.Data(), FilePackageStoreEntry->ShaderMapHashes.Num());
 
-		// TODO: this depends on some package loading changes that aren't ready to go in yet.
-		// We can't return missing in this case because we have code paths that need to be able
-		// to find NotInstalled packages.
-#if 0
 		if (Entry->bIsInstalled)
 		{
 			return EPackageStoreEntryStatus::Ok;
 		}
 
 		return EPackageStoreEntryStatus::NotInstalled;
-#else
-		return EPackageStoreEntryStatus::Ok;
-#endif
 	}
 
 	return EPackageStoreEntryStatus::Missing;
