Fixes to remote object resetting archives to properly handle missing / remote subobjects and properties.
Remote objects marked as garbage will now always be reconstructed when they're being migrated back (at least temporarily)
Cleaned up and removed unused code in FArchiveRemoteObjectWriter

diff --git a/Engine/Source/Runtime/CoreUObject/Private/UObject/RemoteObjectSerialization.cpp b/Engine/Source/Runtime/CoreUObject/Private/UObject/RemoteObjectSerialization.cpp
--- a/Engine/Source/Runtime/CoreUObject/Private/UObject/RemoteObjectSerialization.cpp
+++ b/Engine/Source/Runtime/CoreUObject/Private/UObject/RemoteObjectSerialization.cpp
@@ -34,6 +34,7 @@ static FAutoConsoleVariableRef CVarResetBorrowedObjects(
 
 using FNameIndexType = FPackedRemoteObjectPathName::FNameIndexType;
 
+DEFINE_LOG_CATEGORY_STATIC(LogRemoteSerialization, Log, All);
 
 FArchive& operator<<(FArchive& Ar, FRemoteObjectBytes& Chunk)
 {
@@ -188,8 +189,8 @@ protected:
 
 		const int64 TotalSize = NameTableSize + ObjectData.GetNumBytes();
 
-		UE_LOG(LogRemoteObject, Log, TEXT("%s Object Data stats for %s %s (Object Data toal: %d, total: %lld):"), *GetArchiveName(), *FRemoteObjectId(RootObject).ToString(), *GetFullNameSafe(RootObject), ObjectData.GetNumBytes(), TotalSize);
-		UE_LOG(LogRemoteObject, Log, TEXT("  Name Table total size: %lld (Names: %d, RemoteIds: %d, Paths: %d)"), NameTableSize, ObjectData.Tables.Names.Num(), ObjectData.Tables.RemoteIds.Num(), ObjectData.PathNames.Num());
+		UE_LOG(LogRemoteSerialization, Log, TEXT("%s Object Data stats for %s %s (Object Data toal: %d, total: %lld):"), *GetArchiveName(), *FRemoteObjectId(RootObject).ToString(), *GetFullNameSafe(RootObject), ObjectData.GetNumBytes(), TotalSize);
+		UE_LOG(LogRemoteSerialization, Log, TEXT("  Name Table total size: %lld (Names: %d, RemoteIds: %d, Paths: %d)"), NameTableSize, ObjectData.Tables.Names.Num(), ObjectData.Tables.RemoteIds.Num(), ObjectData.PathNames.Num());
 
 		TArray<FString> SortedPathNames;
 		for (const FPackedRemoteObjectPathName& PathName : ObjectData.PathNames)
@@ -199,7 +200,7 @@ protected:
 		SortedPathNames.Sort();
 		for (const FString& PathName : SortedPathNames)
 		{
-			UE_LOG(LogRemoteObject, Log, TEXT("    %s"), *PathName);
+			UE_LOG(LogRemoteSerialization, Log, TEXT("    %s"), *PathName);
 		}
 
 		for (const TPair<FString, TMap<FProperty*, FMigratedPropertyStats>>& ObjectStatsPair : ObjectPropertyStats)
@@ -214,13 +215,13 @@ protected:
 				Total += Stats.Size;
 			}
 
-			UE_LOG(LogRemoteObject, Log, TEXT("  Data serialized for %s (Total: %lld):"), Obj.Len() ? *Obj : TEXT("Native Serialize"), Total);
+			UE_LOG(LogRemoteSerialization, Log, TEXT("  Data serialized for %s (Total: %lld):"), Obj.Len() ? *Obj : TEXT("Native Serialize"), Total);
 
 			for (const TPair<FProperty*, FMigratedPropertyStats>& PropertyPair : PropStats)
 			{
 				FProperty* Prop = PropertyPair.Key;
 				const FMigratedPropertyStats& Stats = PropertyPair.Value;
-				UE_LOG(LogRemoteObject, Log, TEXT("    %s: size: %lld, count: %lld"), Prop ? *GetFullNameSafe(Prop) : TEXT("Native Serialize"), Stats.Size, Stats.Count);
+				UE_LOG(LogRemoteSerialization, Log, TEXT("    %s: size: %lld, count: %lld"), Prop ? *GetFullNameSafe(Prop) : TEXT("Native Serialize"), Stats.Size, Stats.Count);
 			}
 		}
 	}
@@ -337,7 +338,6 @@ protected:
 	TMap<FName, FNameIndexType> NameMap;
 	TMap<UObject*, FNameIndexType> PathNameMap;
 	TMap<FRemoteObjectId, FNameIndexType> RemoteIdMap;
-	TMap<UObject*, UObject*>* ObjectToArchetypeMap = nullptr;
 	TSet<UObject*>* ReferencedObjectsSet = nullptr;
 
 public:
@@ -373,14 +373,7 @@ public:
 		return Info;
 	}
 
-protected:
-
-	virtual FRemoteObjectReferenceInfo GetReferenceInfo(const FObjectPtr& ObjPtr) const
-	{
-		return GetReferenceInfo(ObjPtr, RootObject);
-	}
-
-	virtual void PopulateObjectHeader(UObject* Object, FRemoteObjectHeader& OutHeader) const
+	virtual bool PopulateObjectHeader(UObject* Object, FRemoteObjectHeader& OutHeader) const
 	{
 		using namespace UE::RemoteObject;
 		using namespace UE::RemoteObject::Private;
@@ -393,60 +386,68 @@ protected:
 		OutHeader.Outer = Outer;
 		OutHeader.Archetype = Object->GetArchetype();
 		OutHeader.InternalFlags = (int32)(Object->GetInternalFlags() & EInternalObjectFlags::Garbage);
+
+		return true;
 	}
 
-private:
+protected:
+
+	virtual FRemoteObjectReferenceInfo GetReferenceInfo(const FObjectPtr& ObjPtr) const
+	{
+		return GetReferenceInfo(ObjPtr, RootObject);
+	}
 
-	void WriteObjectReference(const FObjectPtr& ObjPtr)
+	virtual void WriteObjectReference(const FRemoteObjectReferenceInfo& RefInfo)
 	{
 		using namespace UE::CoreUObject::Private;
 
-		FRemoteObjectReferenceInfo Info = GetReferenceInfo(ObjPtr);
-
-		if (Info.Object)
+		*this << RefInfo.Type;
+		// Always serialize unique id as objects may not exist on the other server and then we may need to pull them from this server
+		FNameIndexType IdIndex = 0;
+		if (RefInfo.Type != ERemoteReferenceType::None)
 		{
-			if (UObject** ReplacementObject = (ObjectToArchetypeMap ? ObjectToArchetypeMap->Find(Info.Object) : nullptr))
+			IdIndex = AddRemoteIdToIdMap(RefInfo.Id);
+			*this << IdIndex;
+
+			if (RefInfo.Type == ERemoteReferenceType::PathName)
 			{
-				Info.Object = *ReplacementObject;
-				Info.Id = FRemoteObjectId(Info.Object);
+				FNameIndexType PathNameIndex = AddPathNameToNameMap(RefInfo.Object);
+				*this << PathNameIndex;
 			}
+		}		
+	}
+
+	void WriteObjectPtr(const FObjectPtr& ObjPtr)
+	{
+		using namespace UE::CoreUObject::Private;
+
+		FRemoteObjectReferenceInfo Info = GetReferenceInfo(ObjPtr);
+		WriteObjectReference(Info);
+
+		if (ObjPtr && !ObjPtr.IsRemote())
+		{
+			UObject* Object = ObjPtr.Get();
 
 			// Anything can be marked as a remote reference, even assets in which case we rely on this flag to be set so that GC calls
 			// StoreObjectToDatabase for any remotely referenced asset (and only for remotely referenced assets) that's about to be GC'd
-			if (!Info.Object->HasAnyInternalFlags(EInternalObjectFlags::RemoteReference))
+			if (!Object->HasAnyInternalFlags(EInternalObjectFlags::RemoteReference))
 			{
 				if (ReferencedObjectsSet)
 				{
 					UE_AUTORTFM_OPEN
 					{
-						ReferencedObjectsSet->Add(Info.Object);
+						ReferencedObjectsSet->Add(Object);
 					};
 				}
 			}
-		}
-
-		*this << Info.Type;
-		// Always serialize unique id as objects may not exist on the other server and then we may need to pull them from this server
-		FNameIndexType IdIndex = 0;
-		if (Info.Type != ERemoteReferenceType::None)
-		{
-			IdIndex = AddRemoteIdToIdMap(Info.Id);
-			*this << IdIndex;
-		}
 
-		if (Info.Type == ERemoteReferenceType::IdOnly)
-		{
-			// Add subobjects of the root object to the list of objects to serialize
-			if (Info.bIsSubobject)
+			if (Info.Type == ERemoteReferenceType::IdOnly && Info.bIsSubobject)
 			{
-				ObjectsToSerialize.Add(Info.Object);
+				// Add subobjects of the root object to the list of objects to serialize
+				ensureMsgf(Object->IsIn(RootObject), TEXT("We're about to serialize a subobject %s which is not a subobject of a root object %s"), *Object->GetPathName(), *RootObject->GetPathName());
+				ObjectsToSerialize.Add(Object);
 			}
 		}
-		else if (Info.Type == ERemoteReferenceType::PathName)
-		{
-			FNameIndexType PathNameIndex = AddPathNameToNameMap(Info.Object);
-			*this << PathNameIndex;
-		}
 	}
 
 public:
@@ -462,18 +463,6 @@ public:
 		ReferencedObjectsSet = OutReferencedObjectsSet;
 	}
 
-	virtual UObject* GetArchetypeFromLoader(const UObject* Obj) override
-	{
-		if (ObjectToArchetypeMap)
-		{
-			if (UObject** FoundArchetype = ObjectToArchetypeMap->Find(Obj))
-			{
-				return *FoundArchetype;
-			}
-		}
-		return nullptr;
-	}
-
 	FNameIndexType AddNameToNameMap(const FName& Name)
 	{
 		FNameIndexType NameIndex = TNumericLimits<FNameIndexType>::Max();
@@ -533,7 +522,7 @@ public:
 
 	virtual FArchive& operator<<(UObject*& Obj) override
 	{
-		WriteObjectReference(FObjectPtr(Obj));
+		WriteObjectPtr(FObjectPtr(Obj));
 		return *this;
 	}
 
@@ -547,7 +536,7 @@ public:
 	virtual FArchive& operator<<(FLazyObjectPtr& Value) override { return FArchiveUObject::SerializeLazyObjectPtr(*this, Value); }
 	virtual FArchive& operator<<(FObjectPtr& Value) override 
 	{ 
-		WriteObjectReference(Value);
+		WriteObjectPtr(Value);
 		return *this;
 	}
 	virtual FArchive& operator<<(FSoftObjectPtr& Value) override { return FArchiveUObject::SerializeSoftObjectPtr(*this, Value); }
@@ -566,7 +555,7 @@ public:
 		{
 			Ptr = Value.Get(/** bEvenIfGarbage */ true);
 		}
-		WriteObjectReference(Ptr);
+		WriteObjectPtr(Ptr);
 
 		return *this;
 	}
@@ -576,10 +565,9 @@ public:
 		return ObjectsToSerialize;
 	}
 
-	void SerializeRemoteObject(UObject* Object, FRemoteObjectHeader& OutHeader)
+	FORCENOINLINE void SerializeRemoteObject(UObject* Object, const FRemoteObjectHeader& Header)
 	{
-		PopulateObjectHeader(Object, OutHeader);
-		AddNameToNameMap(OutHeader.Name);
+		AddNameToNameMap(Header.Name);
 		Object->Serialize(*this);
 	}
 };
@@ -620,16 +608,17 @@ using FRemoteObjectReaderScope = TRemoteObjectArchiveScope<FMemoryReader>;
 * Helper archive that serializes the difference between archetypes and their instances
 * This is achieved using delta serialization but the data we serialize against is coming from instances of the archetypes
 * Effectively this is the opposite how delta serialization normally works which serializes instances of archetypes agains the archetypes
+* In other words this archive is used to serialize archetypes, not their instances.
 */
 class FArchetypeDeltaWriter : public FArchiveRemoteObjectWriter
 {
-	TMap<UObject*, UObject*>& SubstituteArchetypeMap;
+	TMap<UObject*, UObject*>& ArchetypeToInstanceMap;
 
 public:
 
-	FArchetypeDeltaWriter(UObject* InRootObject, FRemoteObjectData& OutObjectData, TMap<UObject*, UObject*>& InSubstituteArchetypeMap)
+	FArchetypeDeltaWriter(UObject* InRootObject, FRemoteObjectData& OutObjectData, TMap<UObject*, UObject*>& InArchetypeToInstanceMap)
 		: FArchiveRemoteObjectWriter(InRootObject, OutObjectData, /*MigrationContext*/ nullptr, TEXT("RemoteArchetypeDeltaWriter"))
-		, SubstituteArchetypeMap(InSubstituteArchetypeMap)
+		, ArchetypeToInstanceMap(InArchetypeToInstanceMap)
 	{
 	}
 
@@ -637,42 +626,53 @@ public:
 	{
 		// Since this archive serializes the archetype we want the archetype of that archetype to be its instance (effectively reversing the object -> archetype relationship)
 		// this way we will delta serialize the difference between the archetype and its instance
-		if (UObject** FoundArchetype = SubstituteArchetypeMap.Find(Obj))
+		if (UObject** FoundArchetypeInstance = ArchetypeToInstanceMap.Find(Obj))
 		{
-			return *FoundArchetype;
+			return *FoundArchetypeInstance;
 		}
+		UE_LOG(LogRemoteSerialization, Warning, TEXT("FArchetypeDeltaWriter::ArchetypeToInstanceMap does not contain an archetype mapping for %s"), *GetPathNameSafe(Obj));
 		return nullptr;
 	}
 
 protected:
 
-	virtual FRemoteObjectReferenceInfo GetReferenceInfo(const FObjectPtr& ObjPtr) const override
+	virtual void WriteObjectReference(const FRemoteObjectReferenceInfo& RefInfo) override
 	{
-		FRemoteObjectReferenceInfo Info = FArchiveRemoteObjectWriter::GetReferenceInfo(ObjPtr, RootObject);
-		if (Info.Object)
+		using namespace UE::CoreUObject::Private;
+
+		FRemoteObjectReferenceInfo ReplacementInfo(RefInfo);
+
+		// If we're serializing a reference to an archetype try to replace it with a reference to its instance
+		// This way the produced delta between the archetype and its instance will be correctly pointing to the instances of default subobjects
+		// we can then deserialize over the archetype instance.
+		if (UObject** Instance = ArchetypeToInstanceMap.Find(RefInfo.Object))
 		{
-			// In addition to substituting the archetype of the serialized objects (which are archetypes themselves) we want to substitue references to these objects as well (to pretend we're serializing their instances)
-			// this way when we deserialize the serialized data on top of the instances all references to instances of subobjects point to the right objects (instances, not archetypes)
-			if (UObject** SubstituteObject = SubstituteArchetypeMap.Find(Info.Object))
-			{
-				Info.Object = *SubstituteObject;
-				Info.Id = FRemoteObjectId(Info.Object);
-			}
+			ReplacementInfo.Object = *Instance;
+			ReplacementInfo.Id = FRemoteObjectId(*Instance);
 		}
-		return Info;
+
+		FArchiveRemoteObjectWriter::WriteObjectReference(ReplacementInfo);
 	}
 
-	virtual void PopulateObjectHeader(UObject* Object, FRemoteObjectHeader& OutHeader) const override
+public:
+
+	virtual bool PopulateObjectHeader(UObject* Object, FRemoteObjectHeader& OutHeader) const override
 	{
-		// As per the comment in GetReferenceInfo - when serializing archetype data pretend we're actually serializing instances of the archetype so substitute
-		// the archetype object with its instance when serializing the object header
-		UObject* HeaderObject = Object;
-		if (UObject** ReplacementObject = SubstituteArchetypeMap.Find(Object))
+		// As per comment in WriteObjectReference - when serializing archetype data pretend we're actually serializing references to instances of the archetype 
+		// so substitute the archetype object with its instance when serializing the object header
+		if (UObject** Instance = ArchetypeToInstanceMap.Find(Object))
 		{
-			HeaderObject = *ReplacementObject;
+			Object = *Instance;
 		}
 
-		FArchiveRemoteObjectWriter::PopulateObjectHeader(HeaderObject, OutHeader);
+		// If for some reason we failed to substitute the archetype with its instance or for some other reason Object is an asset, skip it
+		bool bCanSaveObject = !FRemoteObjectId(Object).IsAsset();
+		if (bCanSaveObject)
+		{
+			bCanSaveObject = FArchiveRemoteObjectWriter::PopulateObjectHeader(Object, OutHeader);
+		}
+		
+		return bCanSaveObject;
 	}
 };
 
@@ -740,8 +740,11 @@ class FArchiveRemoteObjectReader : public TArchiveRemoteObjectBase<FMemoryReader
 				// If the serialized reference was not found in memory or if the reference was serialized as id-only 
 				// keep it as an unresolved ObjectPtr and store a pointer to it so that we can try to resolve it after all objects have been deserialize
 #if UE_WITH_REMOTE_OBJECT_HANDLE
-				FObjectHandle Handle = FObjectHandle::FromIdNoResolve(ObjId);
-				Value = FObjectPtr(Handle);
+				if (!EnumHasAnyFlags(DeserializeFlags, ERemoteObjectSerializationFlags::PreserveRemoteReferences) || !Value.IsRemote())
+				{
+					FObjectHandle Handle = FObjectHandle::FromIdNoResolve(ObjId);
+					Value = FObjectPtr(Handle);
+				}
 #endif // UE_WITH_REMOTE_OBJECT_HANDLE
 			}
 		}
@@ -927,7 +930,7 @@ UObject* ConstructRemoteObject(const FRemoteObjectHeader& Header, FRemoteObjectC
 	else
 	{
 		Outer = StaticFindObjectFastInternal(Header.Outer.GetRemoteId());
-		UE_CLOG(!Outer, LogRemoteObject, Fatal, TEXT("Failed to resolve an Outer when constructing remote object"))
+		UE_CLOG(!Outer, LogRemoteSerialization, Fatal, TEXT("Failed to resolve an Outer when constructing remote object"))
 	}
 #endif
 	FName Name = Header.Name;
@@ -938,7 +941,7 @@ UObject* ConstructRemoteObject(const FRemoteObjectHeader& Header, FRemoteObjectC
 	{
 		if (FRemoteObjectId(Object) != Header.RemoteId)
 		{
-			UE_LOG(LogRemoteObject, Warning, TEXT("Received remote object %s with identical pathname (%s) as a local object %s. Remote object will be renamed."), *Header.RemoteId.ToString(), *Object->GetPathName(), *FRemoteObjectId(Object).ToString());
+			UE_LOG(LogRemoteSerialization, Warning, TEXT("Received remote object %s with identical pathname (%s) as a local object %s. Remote object will be renamed."), *Header.RemoteId.ToString(), *Object->GetPathName(), *FRemoteObjectId(Object).ToString());
 			Name = FName();
 			Object = nullptr;
 		}
@@ -954,9 +957,9 @@ UObject* ConstructRemoteObject(const FRemoteObjectHeader& Header, FRemoteObjectC
 	}
 
 	// If not or the object is marked as remote (which means we brought it back before it was GC'd) (re)construct it
-	// Unless we explcitly want to re-use existing objects and skip re-construction to avoid side-effects
+	// Unless we explcitly want to re-use existing (valid / not marked as garbage) objects and skip re-construction to avoid side-effects
 	// Note that even then an object may not exist on this server (it could've been constructed on a different server when its owner was migrated or was simply GC'd)
-	const bool bSkipConstruction = Object && (!IsRemote(Object) || EnumHasAnyFlags(DeserializeFlags, ERemoteObjectSerializationFlags::UseExistingObjects));
+	const bool bSkipConstruction = IsValid(Object) && (!IsRemote(Object) || EnumHasAnyFlags(DeserializeFlags, ERemoteObjectSerializationFlags::UseExistingObjects));
 	if (!bSkipConstruction)
 	{
 		FStaticConstructObjectParameters Params(Class);
@@ -1039,24 +1042,41 @@ void SerializeObjectDataInternal(FArchiveRemoteObjectWriter& Ar, UObject* Reques
 	int32 SerializedObjectIndex = 0;
 	bool bSerializedRequestedObject = false;
 
+	TSet<UObject*> ProcessedObjects;
 	do
 	{
 		for (; SerializedObjectIndex < Ar.GetObjectsToSerialize().Num(); ++SerializedObjectIndex)
 		{
 			UObject* ObjectToSerialize = Ar.GetObjectsToSerialize()[SerializedObjectIndex];
-			if (!OutObjects.Contains(ObjectToSerialize))
+			if (!ProcessedObjects.Contains(ObjectToSerialize))
 			{
-				FRemoteObjectWriterScope Scope(Ar, ObjectToSerialize);
-				OutObjects.Add(ObjectToSerialize);
-				FRemoteObjectHeader& Header = ObjectHeaders.Emplace_GetRef();
-				Ar.SerializeRemoteObject(ObjectToSerialize, Header);
-				Header.NextOffset = Ar.Tell();
+				ProcessedObjects.Add(ObjectToSerialize);
+
+				FRemoteObjectWriterScope Scope(Ar, ObjectToSerialize);				
+				FRemoteObjectHeader Header;
+				if (Ar.PopulateObjectHeader(ObjectToSerialize, Header))
+				{					
+					ObjectHeaders.Add(Header);
+					OutObjects.Add(ObjectToSerialize);
+
+					Ar.SerializeRemoteObject(ObjectToSerialize, ObjectHeaders.Last());
+
+					ObjectHeaders.Last().NextOffset = Ar.Tell();
+				}
+				else
+				{
+					UE_LOG(LogRemoteSerialization, Warning, TEXT("Unable to serialize object (asset: %s) %s"), 
+						FRemoteObjectId(ObjectToSerialize).IsAsset() ? TEXT("yes") : TEXT("no"),
+						*ObjectToSerialize->GetPathName());
+				}
 			}
 		}
 
 		bSerializedRequestedObject = OutObjects.Contains(RequestedObject);
 		if (!bSerializedRequestedObject)
 		{
+			checkf(!ProcessedObjects.Contains(RequestedObject), TEXT("%s couldn't be serialized"), *RequestedObject->GetPathName());
+
 			// InObject was a default subobject (see GRemoteObjectsMigrateFullHierarchy) but when we serialized its parent 
 			// it turned out that the parent had no direct reference to InObject in which case we need to manually add InObject to ObjectsToSerialize list
 			Ar.GetObjectsToSerialize().Add(RequestedObject);
@@ -1148,7 +1168,11 @@ void ResetRemoteObject(UObject* InObject)
 		const FRemoteServerId LocalServerId = UE::RemoteObject::GetGlobalServerId();
 		for (UObject* DeserializedObject : DeserializedObjects)
 		{
-			UE::RemoteObject::Handle::ChangeOwnerServerId(DeserializedObject, LocalServerId);
+			FRemoteObjectId RemoteId = UE::CoreUObject::Private::FObjectHandleUtils::GetRemoteId(DeserializedObject);
+			if (ensureMsgf(!RemoteId.IsAsset(), TEXT("Attempted to Reset Asset %s (we must have received something we thought was an Asset). Check why this is happening."), *GetNameSafe(DeserializedObject)))
+			{
+				UE::RemoteObject::Handle::ChangeOwnerServerId(DeserializedObject, LocalServerId);
+			}
 		}
 	}
 	delete ArchetypeDelta;
diff --git a/Engine/Source/Runtime/CoreUObject/Private/UObject/RemoteObjectTransfer.cpp b/Engine/Source/Runtime/CoreUObject/Private/UObject/RemoteObjectTransfer.cpp
--- a/Engine/Source/Runtime/CoreUObject/Private/UObject/RemoteObjectTransfer.cpp
+++ b/Engine/Source/Runtime/CoreUObject/Private/UObject/RemoteObjectTransfer.cpp
@@ -22,7 +22,7 @@
 namespace UE::RemoteObject::Transfer
 {
 const FRemoteServerId DatabaseId(ERemoteServerIdConstants::Database);
-static TAutoConsoleVariable<bool> CVarRemoveMissingMigratedComponents(TEXT("DSTM.RemoveMissingMigratedComponents"), false, TEXT("When we receive less SubObjects (Components) than we expect, remove the extras that we have locally (assume they were removed on purpose)"));
+static TAutoConsoleVariable<bool> CVarRemoveMissingMigratedComponents(TEXT("DSTM.RemoveMissingMigratedComponents"), true, TEXT("When we receive less SubObjects (Components) than we expect, remove the extras that we have locally (assume they were removed on purpose)"));
 
 UObject* MigrateObjectFromObjectData(FRemoteObjectData& ObjectData, const FUObjectMigrationContext& MigrationContext);
 
@@ -486,7 +486,7 @@ public:
 				continue;
 			}
 
-			if (RootOwnerServerId != GetGlobalServerId())
+			if (RootOwnerServerId != GetGlobalServerId() && !RootOwnerServerId.IsAsset())
 			{
 				UE_LOG(LogRemoteObject, VeryVerbose, TEXT("FRemoteObjectTransferQueue: borrowed obj: %s (%s) owner server: %s (root: %s (s) root owner: %s)"),
 					*UsedObjectId.ToString(),
