Moving MinFlagBitIndex out of EInternalObjectFlags as it breaks VS debugging of enum flag values

diff --git a/Engine/Source/Runtime/CoreUObject/Public/UObject/ObjectMacros.h b/Engine/Source/Runtime/CoreUObject/Public/UObject/ObjectMacros.h
--- a/Engine/Source/Runtime/CoreUObject/Public/UObject/ObjectMacros.h
+++ b/Engine/Source/Runtime/CoreUObject/Public/UObject/ObjectMacros.h
@@ -620,7 +620,7 @@ enum class EInternalObjectFlags : int32
 {
 	None = 0,
 
-	MinFlagBitIndex = 14,
+	// NOTE that the minimum flag bit index is currently 14. See EInternalObjectFlags_MinFlagBitIndex and UE_ENABLE_FUOBJECT_ITEM_PACKING in UObjectArray.h
 
 	ReachabilityFlag0 = 1 << 14, ///< One of the flags used by Garbage Collector to determine UObject's reachability state
 	ReachabilityFlag1 = 1 << 15, ///< One of the flags used by Garbage Collector to determine UObject's reachability state
@@ -649,6 +649,7 @@ enum class EInternalObjectFlags : int32
 ENUM_CLASS_FLAGS(EInternalObjectFlags);
 
 //~ Make sure these macros are up to date!
+#define EInternalObjectFlags_MinFlagBitIndex 14
 #define EInternalObjectFlags_GarbageCollectionKeepFlags (EInternalObjectFlags::Native | EInternalObjectFlags::Async | EInternalObjectFlags::AsyncLoadingPhase1 | EInternalObjectFlags::AsyncLoadingPhase2 | EInternalObjectFlags::LoaderImport | EInternalObjectFlags::RefCounted)
 #define EInternalObjectFlags_AllFlags (EInternalObjectFlags::ReachabilityFlag0 | EInternalObjectFlags::ReachabilityFlag1 | EInternalObjectFlags::ReachabilityFlag2 | EInternalObjectFlags::AutoRTFMConstructionAborted | EInternalObjectFlags::Remote | EInternalObjectFlags::RemoteReference | EInternalObjectFlags::LoaderImport | EInternalObjectFlags::Garbage | EInternalObjectFlags::ReachableInCluster | EInternalObjectFlags::ClusterRoot | EInternalObjectFlags::Native | EInternalObjectFlags::RefCounted | EInternalObjectFlags::Async | EInternalObjectFlags::AsyncLoadingPhase1 | EInternalObjectFlags::AsyncLoadingPhase2 | EInternalObjectFlags::RootSet | EInternalObjectFlags::PendingConstruction | (EInternalObjectFlags::Unreachable))
 #define EInternalObjectFlags_RootFlags (EInternalObjectFlags::RootSet | EInternalObjectFlags_GarbageCollectionKeepFlags)
diff --git a/Engine/Source/Runtime/CoreUObject/Public/UObject/UObjectArray.h b/Engine/Source/Runtime/CoreUObject/Public/UObject/UObjectArray.h
--- a/Engine/Source/Runtime/CoreUObject/Public/UObject/UObjectArray.h
+++ b/Engine/Source/Runtime/CoreUObject/Public/UObject/UObjectArray.h
@@ -60,12 +60,12 @@ struct FUObjectItem
 private:
 	// Currently we assume UObjects are aligned by 8 bytes, that gives us 3 lower bits as zeros that we can discard.
 	// This will give us total 45 bits in a pointer that we pack into a int32 and the remaining 13 bits we pack with Flags
-	// EInternalObjectFlags::MinFlagBitIndex at the time of writing this is 14 and we have only 1 bit left in the EInternalObjectFlags for future use
+	// EInternalObjectFlags_MinFlagBitIndex at the time of writing this is 14 and we have only 1 bit left in the EInternalObjectFlags for future use
 	// We can increase UObject alignment to 16 bytes to get one more bit and reduce the overall addressable virtual memory range to get more bits if necessary
 	constexpr static int32 UObjectAlignment = 8;
 	constexpr static int32 UObjectPtrTrailingZeroes = FMath::CountTrailingZerosConstExpr(UObjectAlignment);
-	static_assert(int(EInternalObjectFlags::MinFlagBitIndex) >= 48 - 32 - UObjectPtrTrailingZeroes, "We need at least 13 bits to pack higher bits of a UObject pointer into Flags");
-	constexpr static int32 FlagsMask = 0xFFFFFFFF << int(EInternalObjectFlags::MinFlagBitIndex);
+	static_assert(int(EInternalObjectFlags_MinFlagBitIndex) >= 48 - 32 - UObjectPtrTrailingZeroes, "We need at least 13 bits to pack higher bits of a UObject pointer into Flags");
+	constexpr static int32 FlagsMask = 0xFFFFFFFF << int(EInternalObjectFlags_MinFlagBitIndex);
 	constexpr static int32 PtrMask = ~FlagsMask;
 
 	// Stores lower 32 bits of UObject pointer shifted by 3 to the left as all our allocations are at leas 8 bytes aligned and lower 3 bits will always be 0
