Removed $ check from FPackageName::GetVersePath.  Objects like $SolarisSignatureFunctionOuter derive from UVerseClass which now has its own logic for forming Verse paths.  If someone were to call this function with a UVerseClass top level object path (which we presently don't do) we want to log an error to know about it.

diff --git a/Engine/Source/Runtime/CoreUObject/Private/Misc/PackageName.cpp b/Engine/Source/Runtime/CoreUObject/Private/Misc/PackageName.cpp
--- a/Engine/Source/Runtime/CoreUObject/Private/Misc/PackageName.cpp
+++ b/Engine/Source/Runtime/CoreUObject/Private/Misc/PackageName.cpp
@@ -2460,12 +2460,6 @@ UE::Core::FVersePath FPackageName::GetVersePath(const FTopLevelAssetPath& AssetP
 
 	FString VerseModule = FPaths::Combine(PluginVersePath, RelativePath, AssetNameBuilder);
 
-	// Hack to reject names containing "$" - currently used for non-user facing vobject names in Verse, e.g. $SolarisSignatureFunctionOuter
-	if (VerseModule.Contains(TEXT("$")))
-	{
-		return {};
-	}
-
 	UE::Core::FVersePath Result;
 #if NO_LOGGING
 	UE::Core::FVersePath::TryMake(Result, MoveTemp(VerseModule));
