Updated the ancient ClassConfigName hack in UClass to avoid stomping memory when FName is 4 bytes

diff --git a/Engine/Source/Runtime/CoreUObject/Private/UObject/Class.cpp b/Engine/Source/Runtime/CoreUObject/Private/UObject/Class.cpp
--- a/Engine/Source/Runtime/CoreUObject/Private/UObject/Class.cpp
+++ b/Engine/Source/Runtime/CoreUObject/Private/UObject/Class.cpp
@@ -5061,11 +5061,10 @@ void UClass::DeferredRegister(UClass *UClassStaticClass, const TCHAR* PackageNam
 
 	// Get stashed registration info.
 
-	// PVS-Studio justifiably complains about this cast, but we expect this to work because we 'know' that 
-	// we're coming from the UClass constructor that is used when 'statically linked'. V580 disables 
-	// a warning that indicates this is an 'odd explicit type casting'.
-	const TCHAR* InClassConfigName = *(TCHAR**)&ClassConfigName; //-V580 //-V641
-	ClassConfigName = InClassConfigName;
+	// We expect this to work because we 'know' that we're coming from the UClass constructor
+	// that is used when 'statically linked'.
+	ClassConfigName = (const TCHAR*)SparseClassData;
+	SparseClassData = nullptr;
 
 	// Propagate inherited flags.
 	UClass* SuperClass = GetSuperClass();
@@ -6484,11 +6483,9 @@ PRAGMA_ENABLE_DEPRECATION_WARNINGS
 
 	SetCppTypeInfoStatic(&DefaultCppClassTypeInfoStatic);
 
-	// We store the pointer to the ConfigName in an FName temporarily, this cast is intentional
-	// as we expect the mis-typed data to get picked up in UClass::DeferredRegister. PVS-Studio
-	// complains about this operation, but AFAIK it is safe (and we've been doing it a long time)
-	// so the warning has been disabled for now:
-	*(const TCHAR**)&ClassConfigName = InConfigName; //-V580
+	// We store the pointer to the ConfigName in a void* temporarily, this cast is intentional
+	// as we expect the mis-typed data to get picked up in UClass::DeferredRegister.
+	SparseClassData = const_cast<TCHAR*>(InConfigName);
 }
 
 UClass::~UClass() = default;
