Removed the ancient ClassConfigName hack in UClass to avoid stomping memory when FName is 4 bytes

diff --git a/Engine/Source/Runtime/CoreUObject/Private/UObject/Class.cpp b/Engine/Source/Runtime/CoreUObject/Private/UObject/Class.cpp
--- a/Engine/Source/Runtime/CoreUObject/Private/UObject/Class.cpp
+++ b/Engine/Source/Runtime/CoreUObject/Private/UObject/Class.cpp
@@ -5059,14 +5059,6 @@ void UClass::DeferredRegister(UClass *UClassStaticClass, const TCHAR* PackageNam
 #endif
 	);
 
-	// Get stashed registration info.
-
-	// PVS-Studio justifiably complains about this cast, but we expect this to work because we 'know' that 
-	// we're coming from the UClass constructor that is used when 'statically linked'. V580 disables 
-	// a warning that indicates this is an 'odd explicit type casting'.
-	const TCHAR* InClassConfigName = *(TCHAR**)&ClassConfigName; //-V580 //-V641
-	ClassConfigName = InClassConfigName;
-
 	// Propagate inherited flags.
 	UClass* SuperClass = GetSuperClass();
 	if (SuperClass != NULL)
@@ -6472,7 +6464,7 @@ UClass::UClass
 ,	ClassGeneratedBy		( nullptr )
 ,	PropertiesPendingDestruction( nullptr )
 #endif
-,	ClassConfigName			()
+,	ClassConfigName			(InConfigName)
 ,	NetFields				()
 PRAGMA_DISABLE_DEPRECATION_WARNINGS
 ,	ClassDefaultObject		( nullptr )
@@ -6483,12 +6475,6 @@ PRAGMA_ENABLE_DEPRECATION_WARNINGS
 	// If you add properties here, please update the other constructors and PurgeClass()
 
 	SetCppTypeInfoStatic(&DefaultCppClassTypeInfoStatic);
-
-	// We store the pointer to the ConfigName in an FName temporarily, this cast is intentional
-	// as we expect the mis-typed data to get picked up in UClass::DeferredRegister. PVS-Studio
-	// complains about this operation, but AFAIK it is safe (and we've been doing it a long time)
-	// so the warning has been disabled for now:
-	*(const TCHAR**)&ClassConfigName = InConfigName; //-V580
 }
 
 UClass::~UClass() = default;
