[ObjectHandleTracking]
- Improve serialization performance under TSAN by removing unnecessary barrier

diff --git a/Engine/Source/Runtime/CoreUObject/Private/UObject/ObjectHandleTracking.cpp b/Engine/Source/Runtime/CoreUObject/Private/UObject/ObjectHandleTracking.cpp
--- a/Engine/Source/Runtime/CoreUObject/Private/UObject/ObjectHandleTracking.cpp
+++ b/Engine/Source/Runtime/CoreUObject/Private/UObject/ObjectHandleTracking.cpp
@@ -158,7 +158,7 @@ namespace UE::CoreUObject
 				FObjectHandleTrackingCallbackId Result;
 				UE_AUTORTFM_OPEN
 				{
-					HandleReadCallbackQuantity.fetch_add(1, std::memory_order_release);
+					HandleReadCallbackQuantity.fetch_add(1, std::memory_order_relaxed);
 					FLightweightWriteScopeLock _(HandleLock);
 					NextHandleId++;
 					ReadHandleCallbacks.Add({ NextHandleId, Func });
@@ -171,7 +171,7 @@ namespace UE::CoreUObject
 			{
 				UE_AUTORTFM_OPEN
 				{
-					HandleReadCallbackQuantity.fetch_sub(1, std::memory_order_release);
+					HandleReadCallbackQuantity.fetch_sub(1, std::memory_order_relaxed);
 					FLightweightWriteScopeLock _(HandleLock);
 					for (int32 i = ReadHandleCallbacks.Num() - 1; i >= 0; --i)
 					{
diff --git a/Engine/Source/Runtime/CoreUObject/Public/UObject/ObjectHandleTracking.h b/Engine/Source/Runtime/CoreUObject/Public/UObject/ObjectHandleTracking.h
--- a/Engine/Source/Runtime/CoreUObject/Public/UObject/ObjectHandleTracking.h
+++ b/Engine/Source/Runtime/CoreUObject/Public/UObject/ObjectHandleTracking.h
@@ -108,14 +108,14 @@ namespace UE::CoreUObject::Private
 	COREUOBJECT_API void OnHandleReadInternal(const UObject* Object);
 	inline void OnHandleRead(TArrayView<const UObject* const> Objects)
 	{
-		if (HandleReadCallbackQuantity.load(std::memory_order_acquire) > 0)
+		if (HandleReadCallbackQuantity.load(std::memory_order_relaxed) > 0)
 		{
 			OnHandleReadInternal(Objects);
 		}
 	}
 	inline void OnHandleRead(const UObject* Object)
 	{
-		if (HandleReadCallbackQuantity.load(std::memory_order_acquire) > 0)
+		if (HandleReadCallbackQuantity.load(std::memory_order_relaxed) > 0)
 		{
 			OnHandleReadInternal(Object);
 		}
