Async Loading:
- Added the possibility for taking ownership of the progress callback delegate, useful when chaining load requests without having to allocate a new delegate with captures on the async loading thread.
- Added the possibility to filter out progress callback invocations, to minimize the number of calls when only interested in a subset.

diff --git a/Engine/Source/Runtime/CoreUObject/Private/Serialization/AsyncLoading2.cpp b/Engine/Source/Runtime/CoreUObject/Private/Serialization/AsyncLoading2.cpp
--- a/Engine/Source/Runtime/CoreUObject/Private/Serialization/AsyncLoading2.cpp
+++ b/Engine/Source/Runtime/CoreUObject/Private/Serialization/AsyncLoading2.cpp
@@ -4095,16 +4095,23 @@ private:
 			{
 				TRACE_CPUPROFILER_EVENT_SCOPE(PackageProgressCallbacks_Completion);
 
-				const FLoadPackageAsyncProgressParams Params
-				{
-					.PackageName = PackageName,
-					.LoadedPackage = UPackage,
-					.ProgressType = Result == EAsyncLoadingResult::Succeeded ? EAsyncLoadingProgress::FullyLoaded : EAsyncLoadingProgress::Failed
-				};
-				
+				const EAsyncLoadingProgress ProgressType = Result == EAsyncLoadingResult::Succeeded ? EAsyncLoadingProgress::FullyLoaded : EAsyncLoadingProgress::Failed;
+
 				for (TUniquePtr<FLoadPackageAsyncProgressDelegate>& ProgressCallback : ProgressCallbacks)
 				{
-					ProgressCallback->ExecuteIfBound(Params);
+					FLoadPackageAsyncProgressParams Params
+					{
+						.PackageName = PackageName,
+						.LoadedPackage = UPackage,
+						.ProgressType = ProgressType
+					};
+					
+					ProgressCallback->Invoke(Params);
+
+					if (Params.bReleaseDelegate)
+					{
+						verify(ProgressCallback.Release());
+					}
 				}
 				ProgressCallbacks.Empty();
 			}
@@ -7394,7 +7401,7 @@ void FAsyncPackage2::CallProgressCallbacks(EAsyncLoadingProgress ProgressType)
 	{
 		TRACE_CPUPROFILER_EVENT_SCOPE(FAsyncPackage2::CallProgressCallbacks);
 
-		const FLoadPackageAsyncProgressParams Params
+		FLoadPackageAsyncProgressParams Params
 		{
 			.PackageName = Desc.UPackageName,
 			.LoadedPackage = GetLinkerRoot(),
@@ -7403,7 +7410,7 @@ void FAsyncPackage2::CallProgressCallbacks(EAsyncLoadingProgress ProgressType)
 
 		for (TUniquePtr<FLoadPackageAsyncProgressDelegate>& ProgressCallback : ProgressCallbacks)
 		{
-			ProgressCallback->ExecuteIfBound(Params);
+			ProgressCallback->Invoke(Params);
 		}
 	}
 }
@@ -10967,14 +10974,19 @@ void FAsyncLoadingThread2::QueueUnresolvedPackage(FAsyncLoadingThreadState2& Thr
 	{
 		TRACE_CPUPROFILER_EVENT_SCOPE(PackageProgressCallback_Failed);
 
-		const FLoadPackageAsyncProgressParams Params
+		FLoadPackageAsyncProgressParams Params
 		{
 			.PackageName = FailedPackageName,
 			.LoadedPackage = nullptr,
 			.ProgressType = EAsyncLoadingProgress::Failed
 		};
 
-		PackageProgressDelegate->ExecuteIfBound(Params);
+		PackageProgressDelegate->Invoke(Params);
+
+		if (Params.bReleaseDelegate)
+		{
+			verify(PackageProgressDelegate.Release());
+		}
 	}
 
 	if (PackageLoadedDelegate.IsValid())
diff --git a/Engine/Source/Runtime/CoreUObject/Public/UObject/UObjectGlobals.h b/Engine/Source/Runtime/CoreUObject/Public/UObject/UObjectGlobals.h
--- a/Engine/Source/Runtime/CoreUObject/Public/UObject/UObjectGlobals.h
+++ b/Engine/Source/Runtime/CoreUObject/Public/UObject/UObjectGlobals.h
@@ -708,18 +708,50 @@ DECLARE_DELEGATE_ThreeParams(FLoadPackageAsyncDelegate, const FName& /*PackageNa
 struct FLoadPackageAsyncProgressParams
 {
 	/* Name of the package. */
-	FName PackageName { NAME_None };
+	const FName PackageName { NAME_None };
 	/* Pointer to UPackage being loaded, can be nullptr depending on async loading progress. */
-	UPackage* LoadedPackage { nullptr };
+	const UPackage* LoadedPackage { nullptr };
 	/* Progress of async loading. */
-	EAsyncLoadingProgress ProgressType { EAsyncLoadingProgress::Failed };
+	const EAsyncLoadingProgress ProgressType { EAsyncLoadingProgress::Failed };
+	/* Whether the current delegate should be released or destroyed when not required anymore. */
+	bool bReleaseDelegate { false };
 };
 
 /**
  * Thread-safe delegate called on progress of async package loading.
- * @param	Params        Struct containing the parameters for the callback.
  */
-using FLoadPackageAsyncProgressDelegate = TTSDelegate<void(const FLoadPackageAsyncProgressParams& Params)>;
+class FLoadPackageAsyncProgressDelegate
+{
+public:
+	using DelegateType = TTSDelegate<void(FLoadPackageAsyncProgressParams& Params)>;
+
+	FLoadPackageAsyncProgressDelegate(DelegateType&& InDelegate, uint32 InMask = MAX_uint32)
+		: Delegate(InDelegate)
+		, Mask(InMask)
+	{}
+
+	constexpr static inline uint32 BuildMask(EAsyncLoadingProgress Progress)
+	{
+		return 1 << (uint32)Progress;
+	}
+
+	constexpr static inline uint32 BuildMask(auto... Types)
+	{
+		return (0 | ... | BuildMask(Types));
+	}
+
+	void Invoke(FLoadPackageAsyncProgressParams& Params)
+	{
+		if (Mask & BuildMask(Params.ProgressType))
+		{
+			Delegate.ExecuteIfBound(Params);
+		}
+	}
+
+private:
+	TTSDelegate<void(FLoadPackageAsyncProgressParams& Params)> Delegate;
+	const uint32 Mask { MAX_uint32 };
+};
 
 /**
  * Optional parameters passed to the LoadPackageAsync function.
diff --git a/Engine/Source/Runtime/Engine/Private/Tests/Loading/UnattendedAsyncLoadingTests.cpp b/Engine/Source/Runtime/Engine/Private/Tests/Loading/UnattendedAsyncLoadingTests.cpp
--- a/Engine/Source/Runtime/Engine/Private/Tests/Loading/UnattendedAsyncLoadingTests.cpp
+++ b/Engine/Source/Runtime/Engine/Private/Tests/Loading/UnattendedAsyncLoadingTests.cpp
@@ -113,8 +113,8 @@ namespace UnattendedLoadTestImpl
 				LoadPackageAsync(PackageName.ToString(),
 					{
 						.ProgressDelegate = MakeUnique<FLoadPackageAsyncProgressDelegate>(
-							FLoadPackageAsyncProgressDelegate::CreateLambda(
-								[this, ExecutionCheck = FExecutionCheck()](const FLoadPackageAsyncProgressParams& Params) mutable
+							FLoadPackageAsyncProgressDelegate::DelegateType::CreateLambda(
+								[this, ExecutionCheck = FExecutionCheck()](FLoadPackageAsyncProgressParams& Params) mutable
 								{
 									// If we already triggered the next load from this callback, don't do anything else
 									if (ExecutionCheck.HasBeenExecuted())
