Physics Island Merging + Support for Proximity Check on Kinematic Object:
1. Each object checks for its neighboring objects and mark their physics servers as merged.
2. Each server broadcasts merging information across other servers.
3. Each server merges all received merging information and update all physics objects using them.
Also add a cvar dstm.MigratePhysicsObjects as a switch and some logging options

diff --git a/Engine/Source/Runtime/CoreUObject/Private/UObject/RemoteObject.cpp b/Engine/Source/Runtime/CoreUObject/Private/UObject/RemoteObject.cpp
--- a/Engine/Source/Runtime/CoreUObject/Private/UObject/RemoteObject.cpp
+++ b/Engine/Source/Runtime/CoreUObject/Private/UObject/RemoteObject.cpp
@@ -119,6 +119,22 @@ public:
 	{
 		return AssetPaths.Find(ObjectId);
 	}
+	void UpdateAllPhysicsServerId(const TMap<uint32, uint32>& PhysicsServerMergingMap)
+	{
+#if UE_WITH_REMOTE_OBJECT_HANDLE
+		for (TPair<FRemoteObjectId, UE::RemoteObject::Handle::FRemoteObjectStub*>& Elem : RemoteObjects)
+		{
+			if (Elem.Value)
+			{
+				const uint64 OriginalPhysicsId = Elem.Value->PhysicsOverrideServerId.GetIdNumber();
+				if (const uint32* MergedPhysicsId = PhysicsServerMergingMap.Find(OriginalPhysicsId))
+				{
+					Elem.Value->PhysicsOverrideServerId = FRemoteServerId(*MergedPhysicsId);
+				}
+			}
+		}
+#endif
+	}
 };
 FRemoteObjectMaps* ObjectMaps = nullptr;
 
@@ -423,6 +439,16 @@ void ClearAllPhysicsServerId()
 #endif
 }
 
+void UpdateAllPhysicsServerId(const TMap<uint32, uint32>& PhysicsServerMergingMap)
+{
+#if UE_WITH_REMOTE_OBJECT_HANDLE
+	if (Private::ObjectMaps)
+	{
+		Private::ObjectMaps->UpdateAllPhysicsServerId(PhysicsServerMergingMap);
+	}
+#endif
+}
+
 UObject* ResolveObject(const FRemoteObjectStub* Stub, ERemoteReferenceType RefType /*= ERemoteReferenceType::Strong*/)
 {
 	using namespace UE::RemoteObject::Transfer;
diff --git a/Engine/Source/Runtime/CoreUObject/Public/UObject/RemoteObject.h b/Engine/Source/Runtime/CoreUObject/Public/UObject/RemoteObject.h
--- a/Engine/Source/Runtime/CoreUObject/Public/UObject/RemoteObject.h
+++ b/Engine/Source/Runtime/CoreUObject/Public/UObject/RemoteObject.h
@@ -133,4 +133,6 @@ namespace UE::RemoteObject::Handle
 	COREUOBJECT_API void ChangePhysicsServerId(const UObject* Object, FRemoteServerId NewPhysicsServerId);
 
 	COREUOBJECT_API void ClearAllPhysicsServerId();
+
+	COREUOBJECT_API void UpdateAllPhysicsServerId(const TMap<uint32, uint32>& PhysicsServerMergingMap);
 }
