Prevent a crash when instanciating a prefab containing a deleted one in the editor.
Added a Unitttest to test that specific case

#rnx

diff --git a/Engine/Source/Runtime/CoreUObject/Private/UObject/OverridableManager.cpp b/Engine/Source/Runtime/CoreUObject/Private/UObject/OverridableManager.cpp
--- a/Engine/Source/Runtime/CoreUObject/Private/UObject/OverridableManager.cpp
+++ b/Engine/Source/Runtime/CoreUObject/Private/UObject/OverridableManager.cpp
@@ -154,7 +154,7 @@ void FOverridableManager::RestoreOverrideState(TNotNull<const UObject*> OldObjec
 EOverriddenState FOverridableManager::GetOverriddenState(TNotNull<UObject*> Object)
 {
 #if WITH_EDITORONLY_DATA
-	if(const FOverriddenPropertySet* OverriddenProperties = GetOverriddenProperties(Object))
+	if (const FOverriddenPropertySet* OverriddenProperties = GetOverriddenProperties(Object))
 	{
 		// Transient should not affect the overridden state
 		if (Object->HasAnyFlags(RF_Transient))
@@ -318,7 +318,7 @@ EOverriddenPropertyOperation FOverridableManager::GetOverriddenPropertyOperation
 void FOverridableManager::ClearOverrides(TNotNull<UObject*> Object)
 {
 #if WITH_EDITORONLY_DATA
-	if(FOverriddenPropertySet* ThisObjectOverriddenProperties = OverriddenObjectAnnotations.Find(Object))
+	if (FOverriddenPropertySet* ThisObjectOverriddenProperties = OverriddenObjectAnnotations.Find(Object))
 	{
 		ThisObjectOverriddenProperties->Reset();
 	}
@@ -398,7 +398,7 @@ void FOverridableManager::HandleObjectsReInstantiated(const TMap<UObject*, UObje
 	const TMap<const UObjectBase *, FOverriddenPropertyAnnotation>& AnnotationMap = OverriddenObjectAnnotations.GetAnnotationMap();
 	for (const auto& Pair : AnnotationMap)
 	{
-		if( FOverriddenPropertySet* OverriddenProperties = Pair.Value.OverriddenProperties.Get())
+		if (FOverriddenPropertySet* OverriddenProperties = Pair.Value.OverriddenProperties.Get())
 		{
 			OverriddenProperties->HandleObjectsReInstantiated(OldToNewInstanceMap);
 		}
@@ -411,12 +411,12 @@ void FOverridableManager::AddReferencedObjects(FReferenceCollector& Collector)
 #if WITH_EDITORONLY_DATA
 	// This isn't great but this object implements this FGCObject interface just to so 
 	// that the replace references archives go through this object
-	if (GIsReinstancing)
+	if (!GIsGarbageCollecting)
 	{
 		const TMap<const UObjectBase *, FOverriddenPropertyAnnotation>& AnnotationMap = OverriddenObjectAnnotations.GetAnnotationMap();
 		for (const auto& Pair : AnnotationMap)
 		{
-			if( FOverriddenPropertySet* OverriddenProperties = Pair.Value.OverriddenProperties.Get())
+			if (FOverriddenPropertySet* OverriddenProperties = Pair.Value.OverriddenProperties.Get())
 			{
 				OverriddenProperties->AddReferencedObjects(Collector);
 			}
@@ -431,10 +431,10 @@ void FOverridableManager::HandleDeadObjectReferences(const TSet<UClass*>& DeadCl
 	const TMap<const UObjectBase *, FOverriddenPropertyAnnotation>& AnnotationMap = OverriddenObjectAnnotations.GetAnnotationMap();
 	for (const auto& Pair : AnnotationMap)
 	{
-		if( FOverriddenPropertySet* OverriddenProperties = Pair.Value.OverriddenProperties.Get())
+		if (FOverriddenPropertySet* OverriddenProperties = Pair.Value.OverriddenProperties.Get())
 		{
 			OverriddenProperties->HandleDeadObjectReferences(ActiveInstances, TemplateInstances);
 		}
 	}
 #endif // WITH_EDITORONLY_DATA
-}
\ No newline at end of file
+}
diff --git a/Engine/Source/Runtime/CoreUObject/Private/UObject/OverriddenPropertySet.cpp b/Engine/Source/Runtime/CoreUObject/Private/UObject/OverriddenPropertySet.cpp
--- a/Engine/Source/Runtime/CoreUObject/Private/UObject/OverriddenPropertySet.cpp
+++ b/Engine/Source/Runtime/CoreUObject/Private/UObject/OverriddenPropertySet.cpp
@@ -403,9 +403,9 @@ void FOverriddenPropertyNodeID::HandleObjectsReInstantiated(const TMap<UObject*,
 	}
 }
 
-void FOverriddenPropertyNodeID::AddReferencedObjects(FReferenceCollector& Collector)
+void FOverriddenPropertyNodeID::AddReferencedObjects(FReferenceCollector& Collector, UObject* Owner)
 {
-	Collector.AddReferencedObject(Object);
+	Collector.AddReferencedObject(Object, Owner);
 }
 
 void FOverriddenPropertyNodeID::HandleDeadObjectReferences(const TSet<UObject*>& ActiveInstances, const TSet<UObject*>& TemplateInstances)
@@ -1623,11 +1623,11 @@ void FOverriddenPropertySet::AddReferencedObjects(FReferenceCollector& Collector
 {
 	for (FOverriddenPropertyNode& Node : OverriddenPropertyNodes)
 	{
-		Node.NodeID.AddReferencedObjects(Collector);
+		Node.NodeID.AddReferencedObjects(Collector, Owner);
 		for (auto& Pair : Node.SubPropertyNodeKeys)
 		{
-			Pair.Key.AddReferencedObjects(Collector);
-			Pair.Value.AddReferencedObjects(Collector);
+			Pair.Key.AddReferencedObjects(Collector, Owner);
+			Pair.Value.AddReferencedObjects(Collector, Owner);
 		}
 	}
 }
diff --git a/Engine/Source/Runtime/CoreUObject/Public/UObject/OverriddenPropertySet.h b/Engine/Source/Runtime/CoreUObject/Public/UObject/OverriddenPropertySet.h
--- a/Engine/Source/Runtime/CoreUObject/Public/UObject/OverriddenPropertySet.h
+++ b/Engine/Source/Runtime/CoreUObject/Public/UObject/OverriddenPropertySet.h
@@ -244,7 +244,7 @@ struct FOverriddenPropertyNodeID
 
 	void HandleObjectsReInstantiated(const TMap<UObject*, UObject*>& Map);
 	void HandleDeadObjectReferences(const TSet<UObject*>& ActiveInstances, const TSet<UObject*>& TemplateInstances);
-	void AddReferencedObjects(FReferenceCollector& Collector);
+	void AddReferencedObjects(FReferenceCollector& Collector, UObject* Owner);
 
 	friend uint32 GetTypeHash(const FOverriddenPropertyNodeID& NodeID)
 	{
