Change FFieldPath to properly handle blueprint class redirectors when resolving structs by name

diff --git a/Engine/Source/Runtime/CoreUObject/Private/UObject/FieldPath.cpp b/Engine/Source/Runtime/CoreUObject/Private/UObject/FieldPath.cpp
--- a/Engine/Source/Runtime/CoreUObject/Private/UObject/FieldPath.cpp
+++ b/Engine/Source/Runtime/CoreUObject/Private/UObject/FieldPath.cpp
@@ -12,6 +12,7 @@
 #include "UObject/ReleaseObjectVersion.h"
 #include "UObject/LinkerLoad.h"
 #include "UObject/UObjectThreadContext.h"
+#include "UObject/ObjectRedirector.h"
 #include "UObject/PropertyHelper.h"
 #include "Algo/Reverse.h"
 
@@ -242,6 +243,17 @@ UStruct* FFieldPath::TryToResolveOwnerFromStruct(UStruct* InCurrentStruct /*= nu
 		LastOuter = Outer;
 	}
 	UStruct* Owner = Cast<UStruct>(LastOuter);
+
+	if (!Owner && LastOuter)
+	{
+		// The outer resolve uses names, so follow a redirector chain if needed
+		while (UObjectRedirector* Redirector = Cast<UObjectRedirector>(LastOuter))
+		{
+			LastOuter = Redirector->DestinationObject;
+			Owner = Cast<UStruct>(LastOuter);
+		}
+	}
+
 	if (Owner)
 	{
 		ResolvedOwner = Owner;
