[ScriptCore]
- Add safety check since this is not thread-safe

diff --git a/Engine/Source/Runtime/CoreUObject/Private/UObject/ScriptCore.cpp b/Engine/Source/Runtime/CoreUObject/Private/UObject/ScriptCore.cpp
--- a/Engine/Source/Runtime/CoreUObject/Private/UObject/ScriptCore.cpp
+++ b/Engine/Source/Runtime/CoreUObject/Private/UObject/ScriptCore.cpp
@@ -376,6 +376,8 @@ extern CORE_API void (*GPrintScriptCallStackFn)();
 FEditorScriptExecutionGuard::FEditorScriptExecutionGuard()
 	: bOldGAllowScriptExecutionInEditor(GAllowActorScriptExecutionInEditor)
 {
+	check(IsInGameThread());
+
 	GAllowActorScriptExecutionInEditor = true;
 
 	if( GIsEditor && !FApp::IsGame() )
