Linker Instancing Context
* Fix linker automatic populate of the instancing context map to also include remappings from the instanced package map function, if set.
* Optimization: automatically populate the linker instancing context package map when resolved from the instanced package map function, to avoid calling the map function on already resolved mappings.

diff --git a/Engine/Source/Runtime/CoreUObject/Private/UObject/LinkerInstancingContext.cpp b/Engine/Source/Runtime/CoreUObject/Private/UObject/LinkerInstancingContext.cpp
--- a/Engine/Source/Runtime/CoreUObject/Private/UObject/LinkerInstancingContext.cpp
+++ b/Engine/Source/Runtime/CoreUObject/Private/UObject/LinkerInstancingContext.cpp
@@ -75,6 +75,15 @@ public:
 			Instanced = *InstancedPtr;
 			return true;
 		}
+		if (InstancedPackageMapFunc)
+		{
+			const FName MappedResult = InstancedPackageMapFunc(Original);
+			if (MappedResult != Original)
+			{
+				Instanced = MappedResult;
+				return true;
+			}
+		}
 		return false;
 	}
 
@@ -164,16 +173,14 @@ public:
 		FName RemappedPackageName;
 		{
 			FLinkerInstancingContextReadLock ScopeLock(Lock);
-			RemappedPackageName = InstancedPackageMap.RemapPackage(PackageName);
+			
+			bool bWasFound;
+			RemappedPackageName = InstancedPackageMap.RemapPackage(PackageName, &bWasFound);
 
-			if ((RemappedPackageName == PackageName) && InstancedPackageMapFunc)
+			if (!bWasFound && InstancedPackageMapFunc)
 			{
 				RemappedPackageName = InstancedPackageMapFunc(PackageName);
-
-				if (RemappedPackageName != PackageName)
-				{
-					bAddPackageMapping = true; 
-				}
+				bAddPackageMapping = true; 
 			}
 		}
 
@@ -590,5 +597,4 @@ bool FLinkerInstancingContextTests::RunTest(const FString& Parameters)
 	}
 
 	return true;
-}
-
+}
\ No newline at end of file
diff --git a/Engine/Source/Runtime/CoreUObject/Public/UObject/LinkerInstancingContext.h b/Engine/Source/Runtime/CoreUObject/Public/UObject/LinkerInstancingContext.h
--- a/Engine/Source/Runtime/CoreUObject/Public/UObject/LinkerInstancingContext.h
+++ b/Engine/Source/Runtime/CoreUObject/Public/UObject/LinkerInstancingContext.h
@@ -36,15 +36,23 @@ public:
 	}
 
 	/** Remap the package name from the import table to its instanced counterpart, otherwise return the name unmodified. */
-	FName RemapPackage(const FName& PackageName) const
+	FName RemapPackage(const FName& PackageName, bool* bWasFoundPtr = nullptr) const
 	{
 		if (const FName* RemappedName = InstancedPackageMapping.Find(PackageName))
 		{
+			if (bWasFoundPtr)
+			{
+				*bWasFoundPtr = true;
+			}
 			return *RemappedName;
 		}
+		if (bWasFoundPtr)
+		{
+			*bWasFoundPtr = false;
+		}
 		return PackageName;
 	}
-
+	
 	/** Add a mapping from a package name to a new package name. There should be no separators (. or :) in these strings. */
 	COREUOBJECT_API void AddPackageMapping(FName Original, FName Instanced);
 
