[Core] Added Resolved to FTemplateString

diff --git a/Engine/Source/Runtime/CoreUObject/Public/UObject/NoExportTypes.h b/Engine/Source/Runtime/CoreUObject/Public/UObject/NoExportTypes.h
--- a/Engine/Source/Runtime/CoreUObject/Public/UObject/NoExportTypes.h
+++ b/Engine/Source/Runtime/CoreUObject/Public/UObject/NoExportTypes.h
@@ -3394,6 +3394,12 @@ struct FTemplateString
 	*/
 	UPROPERTY(EditAnywhere, BlueprintReadWrite, Category = Template)
 	FString Template;
+	
+	/**
+	* The resolved (localizable) string.
+	*/
+	UPROPERTY(VisibleAnywhere, BlueprintReadWrite, Category = Template)
+	FText Resolved;
 };
 
 USTRUCT(noexport, BlueprintType)
diff --git a/Engine/Source/Runtime/CoreUObject/Public/UObject/TemplateString.h b/Engine/Source/Runtime/CoreUObject/Public/UObject/TemplateString.h
--- a/Engine/Source/Runtime/CoreUObject/Public/UObject/TemplateString.h
+++ b/Engine/Source/Runtime/CoreUObject/Public/UObject/TemplateString.h
@@ -6,7 +6,7 @@
 #include "Containers/UnrealString.h"
 #include "CoreMinimal.h"
 
-/** Structure for templated strings that are displayed in the editor with a allowed args. */
+/** Structure for templated strings that are displayed in the editor with an optional set of valid args. */
 struct FTemplateString
 {
 	/**
@@ -14,6 +14,12 @@ struct FTemplateString
 	*/
 	FString Template;
 
+	/**
+	* The (localizable) resolved text, used to persist the result.
+	* This is especially useful when resolution isn't always available, or is static (editor vs. runtime, for example)
+	*/
+	FText Resolved;
+
 	/** Returns validity based on brace matching, and if provided, arg presence in ValidArgs. */
 	bool IsValid(const TArray<FString>& InValidArgs = {}) const;
 };
