Adjust plugin asset leak reporting to strip standalone flag on all identified leaking package and not only the one we will search for.

#rnx

diff --git a/Engine/Source/Runtime/CoreUObject/Private/Misc/CoreUObjectPluginManager.cpp b/Engine/Source/Runtime/CoreUObject/Private/Misc/CoreUObjectPluginManager.cpp
--- a/Engine/Source/Runtime/CoreUObject/Private/Misc/CoreUObjectPluginManager.cpp
+++ b/Engine/Source/Runtime/CoreUObject/Private/Misc/CoreUObjectPluginManager.cpp
@@ -233,25 +233,25 @@ namespace UE::CoreUObject::Private
 			int32 i = 0;
 			for (UPackage* Package : LeakedPackages)
 			{
-				if (i++ < GLeakedAssetTrace_MaxReportCount)
-				{
-					PackagesToSearchFor.Add(Package);
+				ObjectsInPackage.Reset();
 
-					ObjectsInPackage.Reset();
+				// To improve the reporting for stale references we clear out the RF_Standalone flag on
+				// every object in the packages we will check. The flags will be restored after the Find.
+				GetObjectsWithPackage(Package, ObjectsInPackage, false);
 
-					// To improve the reporting for stale references we clear out the RF_Standalone flag on
-					// every object in the packages we will check. The flags will be restored after the Find.
-					GetObjectsWithPackage(Package, ObjectsInPackage, false);
-
-					for (UObject* Object : ObjectsInPackage)
+				for (UObject* Object : ObjectsInPackage)
+				{
+					if (Object->HasAnyFlags(RF_Standalone))
 					{
-						if (Object->HasAnyFlags(RF_Standalone))
-						{
-							Object->ClearFlags(RF_Standalone);
-							ObjectsThatHadStandaloneCleared.Add(Object);
-						}
+						Object->ClearFlags(RF_Standalone);
+						ObjectsThatHadStandaloneCleared.Add(Object);
 					}
 				}
+
+				if (i++ < GLeakedAssetTrace_MaxReportCount)
+				{
+					PackagesToSearchFor.Add(Package);
+				}
 			}
 
 			UE_LOG(PluginHandlerLog, Display, TEXT("Searching for references to %d leaked packages (%d omitted for speed) from plugin %s"), LeakedPackages.Num(), OmittedCount, *PluginName);
