minor fix in  SaveContext API to properly return array by reference

#rnx

diff --git a/Engine/Source/Runtime/CoreUObject/Private/UObject/SavePackage/SaveContext.h b/Engine/Source/Runtime/CoreUObject/Private/UObject/SavePackage/SaveContext.h
--- a/Engine/Source/Runtime/CoreUObject/Private/UObject/SavePackage/SaveContext.h
+++ b/Engine/Source/Runtime/CoreUObject/Private/UObject/SavePackage/SaveContext.h
@@ -1276,7 +1276,7 @@ public:
 		return SaveArgs.SavePackageContext ? SaveArgs.SavePackageContext->GetExternalImportValidations().Num() > 0 : false;
 	}
 
-	const TArray<TFunction<FSavePackageSettings::ExternalImportValidationFunc>> GetExternalImportValidations() const
+	const TArray<TFunction<FSavePackageSettings::ExternalImportValidationFunc>>& GetExternalImportValidations() const
 	{
 		check(SaveArgs.SavePackageContext);
 		return SaveArgs.SavePackageContext->GetExternalImportValidations();
@@ -1287,7 +1287,7 @@ public:
 		return SaveArgs.SavePackageContext ? SaveArgs.SavePackageContext->GetExternalExportValidations().Num() > 0 : false;
 	}
 
-	const TArray<TFunction<FSavePackageSettings::ExternalExportValidationFunc>> GetExternalExportValidations() const
+	const TArray<TFunction<FSavePackageSettings::ExternalExportValidationFunc>>& GetExternalExportValidations() const
 	{
 		check(SaveArgs.SavePackageContext);
 		return SaveArgs.SavePackageContext->GetExternalExportValidations();
