Introducing immutable class default object used by remote object serialization to provide stable defaults to serialize against when migrating objects
Use ro.UseImmutableArchetypes cvar to disable or enable.

diff --git a/Engine/Source/Runtime/AssetRegistry/Private/AssetRegistry.cpp b/Engine/Source/Runtime/AssetRegistry/Private/AssetRegistry.cpp
--- a/Engine/Source/Runtime/AssetRegistry/Private/AssetRegistry.cpp
+++ b/Engine/Source/Runtime/AssetRegistry/Private/AssetRegistry.cpp
@@ -908,7 +908,7 @@ UAssetRegistryImpl::UAssetRegistryImpl(const FObjectInitializer& ObjectInitializ
 
 	UE::AssetRegistry::Impl::FInitializeContext Context{ *this };
 
-	if (HasAnyFlags(RF_ClassDefaultObject))
+	if (HasAnyFlags(RF_ClassDefaultObject) && !HasAnyFlags(RF_ImmutableDefaultObject))
 	{
 		check(UE::AssetRegistry::Private::IAssetRegistrySingleton::Singleton == nullptr && IAssetRegistryInterface::Default == nullptr);
 		UE::AssetRegistry::Private::IAssetRegistrySingleton::Singleton = this;
@@ -1843,7 +1843,7 @@ void UAssetRegistryImpl::FinishDestroy()
 		UE::AssetDependencyGatherer::Private::FRegisteredAssetDependencyGatherer::OnAssetDependencyGathererRegistered.RemoveAll(this);
 #endif // WITH_EDITOR
 
-		if (HasAnyFlags(RF_ClassDefaultObject))
+		if (HasAnyFlags(RF_ClassDefaultObject) && !HasAnyFlags(RF_ImmutableDefaultObject))
 		{
 			check(UE::AssetRegistry::Private::IAssetRegistrySingleton::Singleton == this && IAssetRegistryInterface::Default == &GAssetRegistryInterface);
 			UE::AssetRegistry::Private::IAssetRegistrySingleton::Singleton = nullptr;
diff --git a/Engine/Source/Runtime/CoreUObject/Private/UObject/Class.cpp b/Engine/Source/Runtime/CoreUObject/Private/UObject/Class.cpp
--- a/Engine/Source/Runtime/CoreUObject/Private/UObject/Class.cpp
+++ b/Engine/Source/Runtime/CoreUObject/Private/UObject/Class.cpp
@@ -4833,6 +4833,10 @@ void UClass::AddReferencedObjects(UObject* InThis, FReferenceCollector& Collecto
 		This->CallAddReferencedObjects(This->ClassDefaultObject, Collector);
 	}
 
+#if UE_WITH_REMOTE_OBJECT_HANDLE
+	Collector.AddStableReference(&This->ImmutableDefaultObject);
+#endif
+
 	Collector.AddStableReference( &This->SparseClassDataStruct );
 
 	// Add sparse class data
@@ -5553,6 +5557,68 @@ void UClass::InternalCreateDefaultObjectWrapper() const
 	const_cast<UClass*>(this)->CreateDefaultObject();
 }
 
+#if UE_WITH_REMOTE_OBJECT_HANDLE
+void UClass::InternalCreateImmutableDefaultObject() const
+{
+	UObject* DefaultObject = GetDefaultObject();
+	checkf(DefaultObject, TEXT("Unable to create immutable default object because the default object couldn't be created for class %s"), *GetPathName());
+	checkf(ImmutableDefaultObject == nullptr, TEXT("Attempting to create another instance of mutable default object for %s"), *GetPathName());
+
+	UObject* ImmutableDefault = nullptr;
+	const bool bIsClassNative = HasAnyClassFlags(CLASS_Native | CLASS_Intrinsic);
+	{
+		checkf(DefaultObject->GetClass() != DefaultObject->GetOuter(), TEXT("Attempting to create an immutable CDO on top of the original CDO"));
+		FStaticConstructObjectParameters ConstructParams(DefaultObject->GetClass());
+		ConstructParams.Outer = DefaultObject->GetClass();
+		ConstructParams.Name = DefaultObject->GetFName();
+		ConstructParams.SetFlags = RF_ClassDefaultObject | RF_ImmutableDefaultObject | RF_ArchetypeObject | RF_Transient;
+		ImmutableDefault = StaticConstructObject_Internal(ConstructParams);
+	}
+	if (!HasAnyClassFlags(CLASS_Native | CLASS_Intrinsic))
+	{
+		// Non-native CDOs need to be duplicated to copy all of their non-native properties
+		check(DefaultObject->GetClass() != DefaultObject->GetOuter());
+		FObjectDuplicationParameters DuplicationParameters(DefaultObject, DefaultObject->GetClass());
+		DuplicationParameters.DestName = DefaultObject->GetFName();
+		DuplicationParameters.ApplyFlags = RF_ArchetypeObject | RF_Transient;
+		DuplicationParameters.DuplicationSeed.Add(DefaultObject, ImmutableDefault);
+		ImmutableDefault = StaticDuplicateObjectEx(DuplicationParameters);
+	}
+	checkf(ImmutableDefault, TEXT("Unable to create immutable default object for %s"), *GetPathName());
+	ImmutableDefault->PostCDOContruct();
+	FPermanentObjectPoolExtents PermanentPool;
+	if (GUObjectArray.IsDisregardForGC(this) || PermanentPool.Contains(this))
+	{
+		ImmutableDefault->AddToRoot();
+	}
+	const_cast<UClass*>(this)->ImmutableDefaultObject = ImmutableDefault;
+}
+#endif
+
+void UClass::SetDefaultObject(UObject* InClassDefaultObject)
+{
+PRAGMA_DISABLE_DEPRECATION_WARNINGS
+#if UE_WITH_REMOTE_OBJECT_HANDLE
+	if (InClassDefaultObject != ClassDefaultObject && ImmutableDefaultObject)
+	{
+		ImmutableDefaultObject->ClearFlags(RF_ClassDefaultObject | RF_ImmutableDefaultObject | RF_ArchetypeObject);
+		ImmutableDefaultObject->RemoveFromRoot();
+		ImmutableDefaultObject->Rename(*MakeUniqueObjectName(GetTransientPackage(), this).ToString(), GetTransientPackage(), REN_DoNotDirty | REN_NonTransactional | REN_DontCreateRedirectors | REN_ForceNoResetLoaders);
+		ImmutableDefaultObject = nullptr;
+	}
+#endif
+	ClassDefaultObject = InClassDefaultObject;
+PRAGMA_ENABLE_DEPRECATION_WARNINGS
+}
+
+void UClass::PostLoadDefaultObject(UObject* Object) 
+{ 
+	Object->PostLoad();
+#if UE_WITH_REMOTE_OBJECT_HANDLE
+	// Construct the immutable version of the CDO immediately after the CDO has been loaded
+	GetImmutableDefaultObject();
+#endif
+}
 /**
 * Helper function for determining if the given class is compatible with structured archive serialization
 */
diff --git a/Engine/Source/Runtime/CoreUObject/Private/UObject/Obj.cpp b/Engine/Source/Runtime/CoreUObject/Private/UObject/Obj.cpp
--- a/Engine/Source/Runtime/CoreUObject/Private/UObject/Obj.cpp
+++ b/Engine/Source/Runtime/CoreUObject/Private/UObject/Obj.cpp
@@ -1727,7 +1727,10 @@ void UObject::Serialize(FStructuredArchive::FRecord Record)
 				!UnderlyingArchive.IsLoadingFromCookedPackage() &&
 				(UnderlyingArchive.UseUnversionedPropertySerialization() || UnderlyingArchive.IsTransacting()))
 			{
-				FOverridableManager::Get().SerializeOverriddenProperties(this, Record);
+				if (FOverridableManager* OverridableManager = FOverridableManager::TryGet())
+				{
+					OverridableManager->SerializeOverriddenProperties(this, Record);
+				}
 #if WITH_EDITORONLY_DATA
 				UE::FInitializedPropertyValueState(this).Serialize(Record);
 #endif
diff --git a/Engine/Source/Runtime/CoreUObject/Private/UObject/RemoteObjectSerialization.cpp b/Engine/Source/Runtime/CoreUObject/Private/UObject/RemoteObjectSerialization.cpp
--- a/Engine/Source/Runtime/CoreUObject/Private/UObject/RemoteObjectSerialization.cpp
+++ b/Engine/Source/Runtime/CoreUObject/Private/UObject/RemoteObjectSerialization.cpp
@@ -32,6 +32,14 @@ static FAutoConsoleVariableRef CVarResetBorrowedObjects(
 	GRemoteObjectsMigrateFullHierarchy,
 	TEXT("Whether remote objects that were borrowed should be reset upon returning to their owner server instead of being reconstructed"));
 
+int32 GUseImmutableArchetypes = 1;
+static FAutoConsoleVariableRef CVarUseImmutableArchetypes(
+	TEXT("ro.UseImmutableArchetypes"),
+	GUseImmutableArchetypes,
+	TEXT("Whether to use immutable archetypes when serializing remote object data"));
+
+const UObject* FindImmutableArchetype(const UObject* InObj);
+
 using FNameIndexType = FPackedRemoteObjectPathName::FNameIndexType;
 
 DEFINE_LOG_CATEGORY_STATIC(LogRemoteSerialization, Log, All);
@@ -111,6 +119,21 @@ void SaveObjectToDisk(const UE::RemoteObject::Transfer::FMigrateSendParams& Para
 namespace UE::RemoteObject::Serialization
 {
 
+UObject* FindArchetype(const UObject* InObj)
+{
+	if (GUseImmutableArchetypes)
+	{
+		UClass* ObjectClass = InObj->GetClass();
+		// No need to get the immutable CDO for BP classes as BP classes are assets themselves and although their CDOs can still technically be modified at runtime
+		// they never are because they can be GC'd and reset to their original state when they're reloaded so they're not a persistent storage like the native CDOs
+		if (!ObjectClass->HasAnyClassFlags(CLASS_CompiledFromBlueprint))
+		{
+			return const_cast<UObject*>(FindImmutableArchetype(InObj));
+		}
+	}
+	return InObj->GetArchetype();
+}
+
 const FRemoteObjectConstructionParams* FRemoteObjectConstructionOverrides::Find(FName InName, UObject* InOuter) const
 {
 	const FRemoteObjectId OuterId(InOuter);
@@ -463,6 +486,11 @@ public:
 		ReferencedObjectsSet = OutReferencedObjectsSet;
 	}
 
+	virtual UObject* GetArchetypeFromLoader(const UObject* Obj) override
+	{
+		return FindArchetype(Obj);
+	}
+
 	FNameIndexType AddNameToNameMap(const FName& Name)
 	{
 		FNameIndexType NameIndex = TNumericLimits<FNameIndexType>::Max();
@@ -612,11 +640,11 @@ using FRemoteObjectReaderScope = TRemoteObjectArchiveScope<FMemoryReader>;
 */
 class FArchetypeDeltaWriter : public FArchiveRemoteObjectWriter
 {
-	TMap<UObject*, UObject*>& ArchetypeToInstanceMap;
+	TMap<const UObject*, UObject*>& ArchetypeToInstanceMap;
 
 public:
 
-	FArchetypeDeltaWriter(UObject* InRootObject, FRemoteObjectData& OutObjectData, TMap<UObject*, UObject*>& InArchetypeToInstanceMap)
+	FArchetypeDeltaWriter(UObject* InRootObject, FRemoteObjectData& OutObjectData, TMap<const UObject*, UObject*>& InArchetypeToInstanceMap)
 		: FArchiveRemoteObjectWriter(InRootObject, OutObjectData, /*MigrationContext*/ nullptr, TEXT("RemoteArchetypeDeltaWriter"))
 		, ArchetypeToInstanceMap(InArchetypeToInstanceMap)
 	{
@@ -631,7 +659,7 @@ public:
 			return *FoundArchetypeInstance;
 		}
 		UE_LOG(LogRemoteSerialization, Warning, TEXT("FArchetypeDeltaWriter::ArchetypeToInstanceMap does not contain an archetype mapping for %s"), *GetPathNameSafe(Obj));
-		return nullptr;
+		return FArchiveRemoteObjectWriter::GetArchetypeFromLoader(Obj);
 	}
 
 protected:
@@ -869,17 +897,17 @@ public:
 class FArchiveSubObjectGatherer : public FArchiveUObject
 {
 	UObject* RootObject = nullptr;
-	TMap<UObject*, UObject*>& ArchetypeToObjectMap;
+	TMap<const UObject*, UObject*>& ArchetypeToObjectMap;
 	TArray<UObject*>& ObjectsToSerialize;
 
 public:
 
-	FArchiveSubObjectGatherer(UObject* InRootObject, TMap<UObject*, UObject*>& InArchetypeToObjectMap, TArray<UObject*>& InObjectsToSerialize)
+	FArchiveSubObjectGatherer(UObject* InRootObject, TMap<const UObject*, UObject*>& InArchetypeToObjectMap, TArray<UObject*>& InObjectsToSerialize)
 		: RootObject(InRootObject)
 		, ArchetypeToObjectMap(InArchetypeToObjectMap)
 		, ObjectsToSerialize(InObjectsToSerialize)
 	{
-		ArchetypeToObjectMap.Add(RootObject->GetArchetype(), RootObject);
+		ArchetypeToObjectMap.Add(FindArchetype(RootObject), RootObject);
 		ArIsObjectReferenceCollector = true;
 		SetIsPersistent(false);
 		SetPortFlags(PPF_AvoidRemoteObjectMigration);
@@ -890,7 +918,7 @@ public:
 		FRemoteObjectReferenceInfo Info = FArchiveRemoteObjectWriter::GetReferenceInfo(Value, RootObject);
 		if (Info.bIsSubobject)
 		{
-			ArchetypeToObjectMap.Add(Info.Object->GetArchetype(), Info.Object);
+			ArchetypeToObjectMap.Add(FindArchetype(Info.Object), Info.Object);
 			ObjectsToSerialize.Add(Info.Object);
 		}
 		return *this;
@@ -1122,7 +1150,7 @@ void ResetRemoteObject(UObject* InObject)
 		Object = FindCanonicalRootObjectForSerialization(Object);
 	}
 
-	TMap<UObject*, UObject*> ReverseArchetypeToObjectMap;
+	TMap<const UObject*, UObject*> ReverseArchetypeToObjectMap;
 
 	// Serialize Object to gather its subobjects and map the subobject archetypes to their respective instances
 	// This map will also be used to replace archetype object pathnames in the serialized archetype data to their instances' pathnames
@@ -1148,7 +1176,7 @@ void ResetRemoteObject(UObject* InObject)
 		// Serialize the Object Archetype against the Object (and its subobjects)
 		// This will produce a delta between the archetypes and their instances which will then be used to restore the instances' state to the archetypes'
 		FRemoteObjectId RequestedObjectId(Object);
-		UObject* Archetype = Object->GetArchetype();
+		UObject* Archetype = const_cast<UObject*>(FindArchetype(Object));
 		FArchetypeDeltaWriter Ar(Archetype, *ArchetypeDelta, ReverseArchetypeToObjectMap);
 		TSet<UObject*> ArchetypeSerializedObjects;
 		SerializeObjectDataInternal(Ar, Archetype, RequestedObjectId, ArchetypeSerializedObjects);
diff --git a/Engine/Source/Runtime/CoreUObject/Private/UObject/UObjectArchetype.cpp b/Engine/Source/Runtime/CoreUObject/Private/UObject/UObjectArchetype.cpp
--- a/Engine/Source/Runtime/CoreUObject/Private/UObject/UObjectArchetype.cpp
+++ b/Engine/Source/Runtime/CoreUObject/Private/UObject/UObjectArchetype.cpp
@@ -82,9 +82,10 @@ FEditorCacheArchetypeManager& FEditorCacheArchetypeManager::Get()
 
 UObject* GetArchetypeImpl(const UObject* InObject, const FObjectArchetypeHelper::IObjectArchetypePolicy* Policy);
 
-UObject* GetArchetypeFromRequiredInfoImpl(const UClass* Class, const UObject* Outer, FName Name, EObjectFlags ObjectFlags, bool bUseUpToDateClass, const FObjectArchetypeHelper::IObjectArchetypePolicy* Policy)
+template <bool bUseImmutableArchetype = false, bool bUseArchetypeCache = UE_CACHE_ARCHETYPE>
+FORCENOINLINE UObject* FindArchetypeFromRequiredInfoImpl(const UClass* Class, const UObject* Outer, FName Name, EObjectFlags ObjectFlags, bool bUseUpToDateClass, const FObjectArchetypeHelper::IObjectArchetypePolicy* Policy)
 {
-	UObject* Result = NULL;
+	UObject* Result = nullptr;
 	const bool bIsCDO = !!(ObjectFlags & RF_ClassDefaultObject);
 	if (bIsCDO)
 	{
@@ -99,29 +100,33 @@ UObject* GetArchetypeFromRequiredInfoImpl(const UClass* Class, const UObject* Ou
 			FScopedUObjectHashTablesLock HashTablesLock;
 
 			UObject* ArchetypeToSearch = nullptr;
-#if UE_CACHE_ARCHETYPE
-			ArchetypeToSearch = GetArchetypeImpl(Outer, Policy);
-#if UE_VERIFY_CACHED_ARCHETYPE
+			// Archetype cache is currently not supported with immutable archetypes (we'd need two caches - one for mutable and the other for immutable archetypes)
+			if (bUseArchetypeCache)
 			{
-				UObject* VerifyArchetype = GetArchetypeFromRequiredInfoImpl(Outer->GetClass(), Outer->GetOuter(), Outer->GetFName(), Outer->GetFlags(), bUseUpToDateClass, Policy);
-				checkf(ArchetypeToSearch == VerifyArchetype, TEXT("Cached archetype mismatch, expected: %s, cached: %s"), *GetFullNameSafe(VerifyArchetype), *GetFullNameSafe(ArchetypeToSearch));
-			}
+				ArchetypeToSearch = GetArchetypeImpl(Outer, Policy);
+#if UE_VERIFY_CACHED_ARCHETYPE
+				{
+					UObject* VerifyArchetype = FindArchetypeFromRequiredInfoImpl<bUseImmutableArchetype, bUseArchetypeCache>(Outer->GetClass(), Outer->GetOuter(), Outer->GetFName(), Outer->GetFlags(), bUseUpToDateClass, Policy);
+					checkf(ArchetypeToSearch == VerifyArchetype, TEXT("Cached archetype mismatch, expected: %s, cached: %s"), *GetFullNameSafe(VerifyArchetype), *GetFullNameSafe(ArchetypeToSearch));
+				}
 #endif // UE_VERIFY_CACHED_ARCHETYPE
-#else
-		
-#if WITH_EDITOR
-			ArchetypeToSearch = Policy ? Policy->GetArchetype(Outer) : nullptr;
-			if (!ArchetypeToSearch)
-			{
-				ArchetypeToSearch = FEditorCacheArchetypeManager::Get().GetCachedArchetype(Outer);
 			}
+			else
+			{
+#if WITH_EDITOR
+				ArchetypeToSearch = Policy ? Policy->GetArchetype(Outer) : nullptr;
+				if (!ArchetypeToSearch)
+				{
+					ArchetypeToSearch = FEditorCacheArchetypeManager::Get().GetCachedArchetype(Outer);
+				}
 #endif
 
-			if (!ArchetypeToSearch)
-			{
-				ArchetypeToSearch = GetArchetypeFromRequiredInfoImpl(Outer->GetClass(), Outer->GetOuter(), Outer->GetFName(), Outer->GetFlags(), bUseUpToDateClass, Policy);
+				if (!ArchetypeToSearch)
+				{
+					ArchetypeToSearch = FindArchetypeFromRequiredInfoImpl<bUseImmutableArchetype, bUseArchetypeCache>(Outer->GetClass(), Outer->GetOuter(), Outer->GetFName(), Outer->GetFlags(), bUseUpToDateClass, Policy);
+				}
 			}
-#endif // UE_CACHE_ARCHETYPE
+
 			UObject* MyArchetype = static_cast<UObject*>(FindObjectWithOuter(ArchetypeToSearch, Class, Name));
 			if (MyArchetype)
 			{
@@ -166,7 +171,16 @@ UObject* GetArchetypeFromRequiredInfoImpl(const UClass* Class, const UObject* Ou
 		if (!Result)
 		{
 			// nothing found, I am not a CDO, so this is just the class CDO
-			Result = bUseUpToDateClass ? Class->GetAuthoritativeClass()->GetDefaultObject() : Class->GetDefaultObject();
+#if UE_WITH_REMOTE_OBJECT_HANDLE
+			if (bUseImmutableArchetype)
+			{
+				Result = const_cast<UObject*>(Class->GetImmutableDefaultObject());
+			}
+			else
+#endif
+			{
+				Result = bUseUpToDateClass ? Class->GetAuthoritativeClass()->GetDefaultObject() : Class->GetDefaultObject();
+			}
 		}
 	}
 
@@ -186,7 +200,7 @@ void CacheArchetypeForObject(UObject* Object, UObject* Archetype)
 #if UE_CACHE_ARCHETYPE
 #if UE_VERIFY_CACHED_ARCHETYPE
 	bool bUseUpToDateClass = false;
-	UObject* VerifyArchetype = GetArchetypeFromRequiredInfoImpl(Object->GetClass(), Object->GetOuter(), Object->GetFName(), Object->GetFlags(), bUseUpToDateClass);
+	UObject* VerifyArchetype = FindArchetypeFromRequiredInfoImpl(Object->GetClass(), Object->GetOuter(), Object->GetFName(), Object->GetFlags(), bUseUpToDateClass);
 	checkf(Archetype == VerifyArchetype, TEXT("Cached archetype mismatch, expected: %s, cached: %s"), *GetFullNameSafe(VerifyArchetype), *GetFullNameSafe(Archetype));
 #endif
 	int32 ArchetypeIndex = GUObjectArray.ObjectToIndex(Archetype);
@@ -197,7 +211,7 @@ void CacheArchetypeForObject(UObject* Object, UObject* Archetype)
 UObject* UObject::GetArchetypeFromRequiredInfo(const UClass* Class, const UObject* Outer, FName Name, EObjectFlags ObjectFlags)
 {
 	bool bUseUpToDateClass = false;
-	return GetArchetypeFromRequiredInfoImpl(Class, Outer, Name, ObjectFlags, bUseUpToDateClass, nullptr);
+	return FindArchetypeFromRequiredInfoImpl(Class, Outer, Name, ObjectFlags, bUseUpToDateClass, nullptr);
 }
 
 //DECLARE_FLOAT_ACCUMULATOR_STAT(TEXT("UObject::GetArchetype"), STAT_FArchiveRealtimeGC_GetArchetype, STATGROUP_GC);
@@ -227,7 +241,7 @@ UObject* GetArchetypeImpl(const UObject* InObject, const FObjectArchetypeHelper:
 #if UE_CACHE_ARCHETYPE
 	if (!bEnableArchetypeCache)
 	{
-		return GetArchetypeFromRequiredInfoImpl(InObject->GetClass(), InObject->GetOuter(), InObject->GetFName(), InObject->GetFlags(), bUseUpToDateClass, Policy);
+		return FindArchetypeFromRequiredInfoImpl(InObject->GetClass(), InObject->GetOuter(), InObject->GetFName(), InObject->GetFlags(), bUseUpToDateClass, Policy);
 	}
 
 	UObject* Archetype = nullptr;
@@ -236,7 +250,7 @@ UObject* GetArchetypeImpl(const UObject* InObject, const FObjectArchetypeHelper:
 	int32 SerialNumber = ArchetypeIndex == INDEX_NONE ? INDEX_NONE : GUObjectArray.GetSerialNumber(ArchetypeIndex);
 	if ((ArchetypeIndex == INDEX_NONE) || (SerialNumber != Annoatation.SerialNumber))
 	{
-		Archetype = GetArchetypeFromRequiredInfoImpl(InObject->GetClass(), InObject->GetOuter(), InObject->GetFName(), InObject->GetFlags(), bUseUpToDateClass, Policy);
+		Archetype = FindArchetypeFromRequiredInfoImpl(InObject->GetClass(), InObject->GetOuter(), InObject->GetFName(), InObject->GetFlags(), bUseUpToDateClass, Policy);
 		// If the Outer is pending load we can't cache the archetype as it may be inacurate
 		if (Archetype && !(InObject->GetOuter() && InObject->GetOuter()->HasAnyFlags(RF_NeedLoad)))
 		{
@@ -250,7 +264,7 @@ UObject* GetArchetypeImpl(const UObject* InObject, const FObjectArchetypeHelper:
 		check(ArchetypeItem != nullptr);
 		Archetype = static_cast<UObject*>(ArchetypeItem->GetObject());
 #if UE_VERIFY_CACHED_ARCHETYPE
-		UObject* ExpectedArchetype = GetArchetypeFromRequiredInfoImpl(InObject->GetClass(), InObject->GetOuter(), InObject->GetFName(), InObject->GetFlags(), bUseUpToDateClass, Policy);
+		UObject* ExpectedArchetype = FindArchetypeFromRequiredInfoImpl(InObject->GetClass(), InObject->GetOuter(), InObject->GetFName(), InObject->GetFlags(), bUseUpToDateClass, Policy);
 		if (ExpectedArchetype != Archetype)
 		{
 			UE_LOG(LogClass, Fatal, TEXT("Cached archetype mismatch, expected: %s, cached: %s"), *GetFullNameSafe(ExpectedArchetype), *GetFullNameSafe(Archetype));
@@ -262,7 +276,7 @@ UObject* GetArchetypeImpl(const UObject* InObject, const FObjectArchetypeHelper:
 
 	return Archetype;
 #else
-	return GetArchetypeFromRequiredInfoImpl(InObject->GetClass(), InObject->GetOuter(), InObject->GetFName(), InObject->GetFlags(), bUseUpToDateClass, Policy);
+	return FindArchetypeFromRequiredInfoImpl(InObject->GetClass(), InObject->GetOuter(), InObject->GetFName(), InObject->GetFlags(), bUseUpToDateClass, Policy);
 #endif // UE_CACHE_ARCHETYPE
 }
 
@@ -283,3 +297,24 @@ void CleanupCachedArchetypes()
 	ArchetypeAnnotation.RemoveAllAnnotations();
 #endif
 }
+
+#ifndef UE_WITH_IMMUTABLEARCHETYPE_DEBUGGING
+#define UE_WITH_IMMUTABLEARCHETYPE_DEBUGGING 0
+#endif
+
+const UObject* FindImmutableArchetype(const UObject* InObj)
+{
+	const bool bUseUpToDateClass = true;
+	const UObject* ImmutableArchetype = FindArchetypeFromRequiredInfoImpl</*bUseImmutableArchetype=*/true, /*bUseArchetypeCache=*/false>(InObj->GetClass(), InObj->GetOuter(), InObj->GetFName(), InObj->GetFlags(), bUseUpToDateClass, nullptr);
+#if UE_WITH_IMMUTABLEARCHETYPE_DEBUGGING
+	const UObject* Archetype = InObj->GetArchetype();
+	checkf(ImmutableArchetype, TEXT("Unable to find immutable archetype for %s"), *InObj->GetFullName());
+	checkf(ImmutableArchetype->HasAnyFlags(RF_ArchetypeObject), TEXT("Immutable archetype is not an archetype"));
+	checkf(ImmutableArchetype->GetFName() == Archetype->GetFName(), TEXT("Immutable archetype name mismatch: expected: %s, got: %s"), *Archetype->GetName(), *ImmutableArchetype->GetName());
+	checkf(ImmutableArchetype->GetClass() == Archetype->GetClass(), TEXT("Immutable archetype class mismatch: expected: %s, got: %s"), *Archetype->GetClass()->GetPathName(), *ImmutableArchetype->GetClass()->GetPathName());
+	checkf(ImmutableArchetype->HasAnyFlags(RF_ClassDefaultObject) == Archetype->HasAnyFlags(RF_ClassDefaultObject), TEXT("Immutable archetype flags mismatch"));
+	// Immutable CDOs have different outers than the original CDOs so we can't check if their outer class matches the original CDO
+	checkf(ImmutableArchetype->HasAnyFlags(RF_ClassDefaultObject) || ImmutableArchetype->GetOuter()->GetClass() == Archetype->GetOuter()->GetClass(), TEXT("Immutable archetype outer class mismatch: expected: %s, got: %s"), *Archetype->GetOuter()->GetClass()->GetPathName(), *ImmutableArchetype->GetOuter()->GetClass()->GetPathName());
+#endif // UE_WITH_IMMUTABLEARCHETYPE_DEBUGGING
+	return ImmutableArchetype;
+}
\ No newline at end of file
diff --git a/Engine/Source/Runtime/CoreUObject/Private/UObject/UObjectBase.cpp b/Engine/Source/Runtime/CoreUObject/Private/UObject/UObjectBase.cpp
--- a/Engine/Source/Runtime/CoreUObject/Private/UObject/UObjectBase.cpp
+++ b/Engine/Source/Runtime/CoreUObject/Private/UObject/UObjectBase.cpp
@@ -881,6 +881,9 @@ static void UObjectLoadAllCompiledInDefaultProperties(TArray<UClass*>& OutAllNew
 			{
 				UE_LOG(LogUObjectBootstrap, Verbose, TEXT("GetDefaultObject Begin %s %s"), *Class->GetOutermost()->GetName(), *Class->GetName());
 				Class->GetDefaultObject();
+#if UE_WITH_REMOTE_OBJECT_HANDLE
+				Class->GetImmutableDefaultObject();
+#endif
 				UE_LOG(LogUObjectBootstrap, Verbose, TEXT("GetDefaultObject End %s %s"), *Class->GetOutermost()->GetName(), *Class->GetName());
 			}
 		}
@@ -890,6 +893,9 @@ static void UObjectLoadAllCompiledInDefaultProperties(TArray<UClass*>& OutAllNew
 			{
 				UE_LOG(LogUObjectBootstrap, Verbose, TEXT("GetDefaultObject Begin %s %s"), *Class->GetOutermost()->GetName(), *Class->GetName());
 				Class->GetDefaultObject();
+#if UE_WITH_REMOTE_OBJECT_HANDLE
+				Class->GetImmutableDefaultObject();
+#endif
 				UE_LOG(LogUObjectBootstrap, Verbose, TEXT("GetDefaultObject End %s %s"), *Class->GetOutermost()->GetName(), *Class->GetName());
 			}
 		}
@@ -899,6 +905,9 @@ static void UObjectLoadAllCompiledInDefaultProperties(TArray<UClass*>& OutAllNew
 			{
 				UE_LOG(LogUObjectBootstrap, Verbose, TEXT("GetDefaultObject Begin %s %s"), *Class->GetOutermost()->GetName(), *Class->GetName());
 				Class->GetDefaultObject();
+#if UE_WITH_REMOTE_OBJECT_HANDLE
+				Class->GetImmutableDefaultObject();
+#endif
 				UE_LOG(LogUObjectBootstrap, Verbose, TEXT("GetDefaultObject End %s %s"), *Class->GetOutermost()->GetName(), *Class->GetName());
 			}
 		}
diff --git a/Engine/Source/Runtime/CoreUObject/Private/UObject/UObjectGlobals.cpp b/Engine/Source/Runtime/CoreUObject/Private/UObject/UObjectGlobals.cpp
--- a/Engine/Source/Runtime/CoreUObject/Private/UObject/UObjectGlobals.cpp
+++ b/Engine/Source/Runtime/CoreUObject/Private/UObject/UObjectGlobals.cpp
@@ -3553,7 +3553,7 @@ UObject* StaticAllocateObject
 	checkSlow(InOuter != INVALID_OBJECT); // not legal
 	check(InClass && InClass->ClassWithin && InClass->ClassConstructor);
 
-	const bool bCreatingCDO = (InFlags & RF_ClassDefaultObject) != 0;
+	const bool bCreatingCDO = (InFlags & RF_ClassDefaultObject) != 0 && (InFlags & RF_ImmutableDefaultObject) == 0;
 	const bool bCreatingArchetype = (InFlags & RF_ArchetypeObject) != 0;
 #if WITH_EDITOR
 	if (GIsEditor)
@@ -3665,7 +3665,7 @@ UObject* StaticAllocateObject
 		LinkerIndex = Obj->GetLinkerIndex();
 		InternalSetFlags |= (Obj->GetInternalFlags() & (EInternalObjectFlags::Native | EInternalObjectFlags::RootSet | EInternalObjectFlags::LoaderImport));
 
-		if ( bCreatingCDO )
+		if (bCreatingCDO || !!(InFlags & RF_ImmutableDefaultObject))
 		{
 			check(Obj->HasAllFlags(RF_ClassDefaultObject));
 			Obj->SetFlags(InFlags);
@@ -4205,7 +4205,7 @@ void FObjectInitializer::PostConstructInit()
 	}
 
 	SCOPE_CYCLE_COUNTER(STAT_PostConstructInitializeProperties);
-	const bool bIsCDO = Obj->HasAnyFlags(RF_ClassDefaultObject);
+	const bool bIsCDO = Obj->HasAnyFlags(RF_ClassDefaultObject) && !Obj->HasAnyFlags(RF_ImmutableDefaultObject);
 	UClass* Class = Obj->GetClass();
 	UClass* SuperClass = Class->GetSuperClass();
 
diff --git a/Engine/Source/Runtime/CoreUObject/Public/UObject/Class.h b/Engine/Source/Runtime/CoreUObject/Public/UObject/Class.h
--- a/Engine/Source/Runtime/CoreUObject/Public/UObject/Class.h
+++ b/Engine/Source/Runtime/CoreUObject/Public/UObject/Class.h
@@ -3075,6 +3075,10 @@ public:
 	TObjectPtr<UObject> ClassDefaultObject;
 
 protected:
+#if UE_WITH_REMOTE_OBJECT_HANDLE
+	TObjectPtr<UObject> ImmutableDefaultObject;
+#endif
+
 	/** This is where we store the data that is only changed per class instead of per instance */
 	void* SparseClassData;
 
@@ -3487,15 +3491,26 @@ public:
 	}
 
 	/**
-	 * Sets the default object of the class
-	 * @param	InClassDefaultObject new class default object or null
+	 * Get the immutable default object from the class
+	 * @param bCreateIfNeeded if true (default) then the CDO is created if it is null
+	 * @return the immutable CDO for this class
 	 */
-	void SetDefaultObject(UObject* InClassDefaultObject)
+#if UE_WITH_REMOTE_OBJECT_HANDLE
+	const UObject* GetImmutableDefaultObject(bool bCreateIfNeeded = true) const
 	{
-		PRAGMA_DISABLE_DEPRECATION_WARNINGS
-		ClassDefaultObject = InClassDefaultObject;
-		PRAGMA_ENABLE_DEPRECATION_WARNINGS
+		if (ImmutableDefaultObject == nullptr && bCreateIfNeeded)
+		{
+			InternalCreateImmutableDefaultObject();
+		}
+		return ImmutableDefaultObject;
 	}
+#endif
+
+	/**
+	 * Sets the default object of the class
+	 * @param	InClassDefaultObject new class default object or null
+	 */
+	COREUOBJECT_API void SetDefaultObject(UObject* InClassDefaultObject);
 
 	/**
 	 * Checks if the class default object is valid (not null and its address is a properly aligned pointer)
@@ -3685,7 +3700,7 @@ public:
 	/** Wraps the PostLoad() call for the class default object.
 	 * @param Object the default object to call PostLoad() on
 	 */
-	virtual void PostLoadDefaultObject(UObject* Object) { Object->PostLoad(); }
+	COREUOBJECT_API virtual void PostLoadDefaultObject(UObject* Object);
 
 	/** 
 	 * Purges out the properties of this class in preparation for it to be regenerated
@@ -3827,6 +3842,9 @@ private:
 	 */
 	COREUOBJECT_API void AssembleReferenceTokenStreamInternal(bool bForce = false);
 	COREUOBJECT_API void InternalCreateDefaultObjectWrapper() const;
+#if UE_WITH_REMOTE_OBJECT_HANDLE
+	COREUOBJECT_API void InternalCreateImmutableDefaultObject() const;
+#endif
 protected:
 	/**
 	 * Get the default object from the class, creating it if missing, if requested or under a few other circumstances
diff --git a/Engine/Source/Runtime/CoreUObject/Public/UObject/ObjectMacros.h b/Engine/Source/Runtime/CoreUObject/Public/UObject/ObjectMacros.h
--- a/Engine/Source/Runtime/CoreUObject/Public/UObject/ObjectMacros.h
+++ b/Engine/Source/Runtime/CoreUObject/Public/UObject/ObjectMacros.h
@@ -577,7 +577,7 @@ enum EObjectFlags
 	RF_DuplicateTransient		=0x00800000,	///< Object should not be included in any type of duplication (copy/paste, binary duplication, etc.)
 	RF_StrongRefOnFrame			=0x01000000,	///< References to this object from persistent function frame are handled as strong ones.
 	RF_NonPIEDuplicateTransient	=0x02000000,	///< Object should not be included for duplication unless it's being duplicated for a PIE session
-	// RF_Dynamic				=0x04000000,	///< Was removed along with bp nativization
+	RF_ImmutableDefaultObject	=0x04000000,	///< Hidden version of the class default object that cannot and will never be mutated
 	RF_WillBeLoaded				=0x08000000,	///< This object was constructed during load and will be loaded shortly
 	RF_HasExternalPackage		=0x10000000,	///< This object has an external package assigned and should look it up when getting the outermost package
 	// RF_Unused				=0x20000000,
diff --git a/Engine/Source/Runtime/Engine/Private/LevelStreaming.cpp b/Engine/Source/Runtime/Engine/Private/LevelStreaming.cpp
--- a/Engine/Source/Runtime/Engine/Private/LevelStreaming.cpp
+++ b/Engine/Source/Runtime/Engine/Private/LevelStreaming.cpp
@@ -420,7 +420,7 @@ ULevelStreaming::ULevelStreaming(const FObjectInitializer& ObjectInitializer)
 	bGarbageCollectionClusteringEnabled = true;
 
 #if ENABLE_ACTOR_PRIMITIVE_COLOR_HANDLER
-	if (HasAnyFlags(RF_ClassDefaultObject) && ExactCast<ULevelStreaming>(this))
+	if (HasAnyFlags(RF_ClassDefaultObject) && !HasAnyFlags(RF_ImmutableDefaultObject) && ExactCast<ULevelStreaming>(this))
 	{
 		FActorPrimitiveColorHandler::Get().RegisterPrimitiveColorHandler(TEXT("LevelColor"), LOCTEXT("LevelColor", "Level Color"), [](const UPrimitiveComponent* InPrimitiveComponent)
 		{
diff --git a/Engine/Source/Runtime/Engine/Private/WorldPartition/WorldPartitionRuntimeCellTransformerISM.cpp b/Engine/Source/Runtime/Engine/Private/WorldPartition/WorldPartitionRuntimeCellTransformerISM.cpp
--- a/Engine/Source/Runtime/Engine/Private/WorldPartition/WorldPartitionRuntimeCellTransformerISM.cpp
+++ b/Engine/Source/Runtime/Engine/Private/WorldPartition/WorldPartitionRuntimeCellTransformerISM.cpp
@@ -30,7 +30,7 @@ UWorldPartitionRuntimeCellTransformerISM::UWorldPartitionRuntimeCellTransformerI
 #endif
 
 #if ENABLE_ACTOR_PRIMITIVE_COLOR_HANDLER
-	if (HasAnyFlags(RF_ClassDefaultObject) && ExactCast<UWorldPartitionRuntimeCellTransformerISM>(this))
+	if (HasAnyFlags(RF_ClassDefaultObject) && !HasAnyFlags(RF_ImmutableDefaultObject) && ExactCast<UWorldPartitionRuntimeCellTransformerISM>(this))
 	{
 		FActorPrimitiveColorHandler::Get().RegisterPrimitiveColorHandler(TEXT("CellTransformerISM"), LOCTEXT("CellTransformerISM", "Cell Transformer ISM"), false, [](const UPrimitiveComponent* InPrimitiveComponent)
 		{
