Minor improvments to remote object debugging.

Added log message with the server id of the current process (useful when analyzing logs from multiple servers)
Improved the ensure assuming the existance of remote object stub to include the object pathname, remote id and address.

diff --git a/Engine/Source/Runtime/CoreUObject/Private/UObject/RemoteObject.cpp b/Engine/Source/Runtime/CoreUObject/Private/UObject/RemoteObject.cpp
--- a/Engine/Source/Runtime/CoreUObject/Private/UObject/RemoteObject.cpp
+++ b/Engine/Source/Runtime/CoreUObject/Private/UObject/RemoteObject.cpp
@@ -153,6 +153,8 @@ void InitServerId()
 		RemoteServerId = FRemoteServerId(ServerId);
 	}
 	checkf(RemoteServerId.IsValid(), TEXT("Remote ServerId is not valid"));
+
+	UE_LOG(LogRemoteObject, Display, TEXT("Remote ServerId: %s"), *RemoteServerId.ToString());
 }
 
 void InitRemoteObjects()
@@ -386,7 +388,7 @@ void ChangeOwnerServerId(const UObject* Object, FRemoteServerId NewOwnerServerId
 	UE::RemoteObject::Handle::FRemoteObjectStub* RemoteStub = FindRemoteObjectStub(FObjectHandleUtils::GetRemoteId(Object));
 
 	// The remote stub is always expected to be found for this object.
-	if (ensure(RemoteStub))
+	if (ensureMsgf(RemoteStub, TEXT("Missing stub for %s (%s / 0x%016llx)"), *GetPathNameSafe(Object), *FRemoteObjectId(Object).ToString(), (int64)(PTRINT)Object))
 	{
 		RemoteStub->OwningServerId = NewOwnerServerId;
 	}
