Make FOverridableManager thread-safe, so we can use it in the editor with async loading thread enabled.

diff --git a/Engine/Source/Runtime/CoreUObject/Private/UObject/OverridableManager.cpp b/Engine/Source/Runtime/CoreUObject/Private/UObject/OverridableManager.cpp
--- a/Engine/Source/Runtime/CoreUObject/Private/UObject/OverridableManager.cpp
+++ b/Engine/Source/Runtime/CoreUObject/Private/UObject/OverridableManager.cpp
@@ -22,7 +22,7 @@ void FOverridableManager::Create()
 bool FOverridableManager::IsEnabled(TNotNull<const UObject*> Object)
 {
 #if WITH_EDITORONLY_DATA
-	return OverriddenObjectAnnotations.IsEnabled(Object);
+	return !GetAnnotation(Object).IsDefault();
 #else
 	return false;
 #endif // WITH_EDITORONLY_DATA
@@ -31,14 +31,14 @@ bool FOverridableManager::IsEnabled(TNotNull<const UObject*> Object)
 void FOverridableManager::Enable(TNotNull<UObject*> Object)
 {
 #if WITH_EDITORONLY_DATA
-	OverriddenObjectAnnotations.FindOrAdd(Object);
+	FindOrAdd(Object);
 #endif // WITH_EDITORONLY_DATA
 }
 
 void FOverridableManager::Disable(TNotNull<UObject*> Object, bool bPropagateToSubObject/*= false*/)
 {
 #if WITH_EDITORONLY_DATA
-	OverriddenObjectAnnotations.RemoveAnnotation(Object);
+	RemoveAnnotation(Object);
 
 	if (bPropagateToSubObject)
 	{
@@ -68,7 +68,7 @@ void FOverridableManager::DisableInstancedSubObject(TNotNull<const UObject*> Obj
 void FOverridableManager::InheritEnabledFrom(TNotNull<UObject*> Object, const UObject* DefaultData)
 {
 #if WITH_EDITORONLY_DATA
-	if (!OverriddenObjectAnnotations.IsEnabled(Object))
+	if (!IsEnabled(Object))
 	{
 		const UObject* Outer = Object->GetOuter();
 		if ((Outer && IsEnabled(Outer)) || (DefaultData && IsEnabled(DefaultData)))
@@ -82,7 +82,7 @@ void FOverridableManager::InheritEnabledFrom(TNotNull<UObject*> Object, const UO
 bool FOverridableManager::NeedSubObjectTemplateInstantiation(TNotNull<const UObject*> Object)
 {
 #if WITH_EDITORONLY_DATA
-	if (const FOverriddenPropertySet* OverriddenProperties = OverriddenObjectAnnotations.Find(Object))
+	if (const FOverriddenPropertySet* OverriddenProperties = Find(Object))
 	{
 		return OverriddenProperties->bNeedsSubobjectTemplateInstantiation;
 	}
@@ -95,7 +95,7 @@ bool FOverridableManager::NeedSubObjectTemplateInstantiation(TNotNull<const UObj
 FOverriddenPropertySet* FOverridableManager::GetOverriddenPropertiesInternal(TNotNull<UObject*> Object)
 {
 #if WITH_EDITORONLY_DATA
-	return OverriddenObjectAnnotations.Find(Object);
+	return Find(Object);
 #else
 	return nullptr;
 #endif // WITH_EDITORONLY_DATA
@@ -104,7 +104,7 @@ FOverriddenPropertySet* FOverridableManager::GetOverriddenPropertiesInternal(TNo
 const FOverriddenPropertySet* FOverridableManager::GetOverriddenPropertiesInternal(TNotNull<const UObject*> Object)
 {
 #if WITH_EDITORONLY_DATA
-	return OverriddenObjectAnnotations.Find(Object);
+	return Find(Object);
 #else
 	return nullptr;
 #endif // WITH_EDITORONLY_DATA
@@ -121,7 +121,7 @@ FOverriddenPropertySet* FOverridableManager::RestoreOverrideOperation(TNotNull<U
 	}
 
 #if WITH_EDITORONLY_DATA
-	FOverriddenPropertySet& ObjectOverriddenProperties = OverriddenObjectAnnotations.FindOrAdd(Object);
+	FOverriddenPropertySet& ObjectOverriddenProperties = FindOrAdd(Object);
 	ObjectOverriddenProperties.Reset();
 	ObjectOverriddenProperties.SetOverriddenPropertyOperation(Operation, /*CurrentPropertyChain*/nullptr, /*Property*/nullptr);
 	ObjectOverriddenProperties.bNeedsSubobjectTemplateInstantiation = bNeedsSubobjectTemplateInstantiation;
@@ -224,7 +224,7 @@ EOverriddenState FOverridableManager::GetOverriddenState(TNotNull<UObject*> Obje
 void FOverridableManager::OverrideObject(TNotNull<UObject*> Object)
 {
 #if WITH_EDITORONLY_DATA
-	if (FOverriddenPropertySet* ThisObjectOverriddenProperties = OverriddenObjectAnnotations.Find(Object))
+	if (FOverriddenPropertySet* ThisObjectOverriddenProperties = Find(Object))
 	{
 		// Passing no property node means we are overriding the object itself
 		ThisObjectOverriddenProperties->OverrideProperty(FPropertyVisitorPath::InvalidIterator(), Object);
@@ -262,7 +262,7 @@ void FOverridableManager::PropagateOverrideToInstancedSubObjects(TNotNull<UObjec
 void FOverridableManager::OverrideProperty(TNotNull<UObject*> Object, const FPropertyVisitorPath& PropertyPath)
 {
 #if WITH_EDITORONLY_DATA
-	if (FOverriddenPropertySet* ThisObjectOverriddenProperties = OverriddenObjectAnnotations.Find(Object))
+	if (FOverriddenPropertySet* ThisObjectOverriddenProperties = Find(Object))
 	{
 		ThisObjectOverriddenProperties->OverrideProperty(PropertyPath.GetRootIterator(), Object);
 	}
@@ -272,7 +272,7 @@ void FOverridableManager::OverrideProperty(TNotNull<UObject*> Object, const FPro
 bool FOverridableManager::ClearOverriddenProperty(TNotNull<UObject*> Object, FPropertyVisitorPath::Iterator PropertyIterator)
 {
 #if WITH_EDITORONLY_DATA
-	if (FOverriddenPropertySet* ThisObjectOverriddenProperties = OverriddenObjectAnnotations.Find(Object))
+	if (FOverriddenPropertySet* ThisObjectOverriddenProperties = Find(Object))
 	{
 		return ThisObjectOverriddenProperties->ClearOverriddenProperty(PropertyIterator);
 	}
@@ -297,7 +297,7 @@ void FOverridableManager::PostOverrideProperty(TNotNull<UObject*> Object, const
 void FOverridableManager::NotifyPropertyChange(const EPropertyNotificationType Notification, TNotNull<UObject*> Object, FPropertyVisitorPath::Iterator PropertyIterator, const EPropertyChangeType::Type ChangeType)
 {
 #if WITH_EDITORONLY_DATA
-	if (FOverriddenPropertySet* ThisObjectOverriddenProperties = OverriddenObjectAnnotations.Find(Object))
+	if (FOverriddenPropertySet* ThisObjectOverriddenProperties = Find(Object))
 	{
 		ThisObjectOverriddenProperties->NotifyPropertyChange(Notification, PropertyIterator, ChangeType, Object);
 	}
@@ -307,7 +307,7 @@ void FOverridableManager::NotifyPropertyChange(const EPropertyNotificationType N
 EOverriddenPropertyOperation FOverridableManager::GetOverriddenPropertyOperation(TNotNull<UObject*> Object, FPropertyVisitorPath::Iterator PropertyIterator, bool* bOutInheritedOperation)
 {
 #if WITH_EDITORONLY_DATA
-	if (const FOverriddenPropertySet* ThisObjectOverriddenProperties = OverriddenObjectAnnotations.Find(Object))
+	if (const FOverriddenPropertySet* ThisObjectOverriddenProperties = Find(Object))
 	{
 		return ThisObjectOverriddenProperties->GetOverriddenPropertyOperation(PropertyIterator, bOutInheritedOperation);
 	}
@@ -318,7 +318,7 @@ EOverriddenPropertyOperation FOverridableManager::GetOverriddenPropertyOperation
 void FOverridableManager::ClearOverrides(TNotNull<UObject*> Object)
 {
 #if WITH_EDITORONLY_DATA
-	if (FOverriddenPropertySet* ThisObjectOverriddenProperties = OverriddenObjectAnnotations.Find(Object))
+	if (FOverriddenPropertySet* ThisObjectOverriddenProperties = Find(Object))
 	{
 		ThisObjectOverriddenProperties->Reset();
 	}
@@ -395,8 +395,8 @@ FOverridableManager::FOverridableManager()
 void FOverridableManager::HandleObjectsReInstantiated(const TMap<UObject*, UObject*>& OldToNewInstanceMap)
 {
 #if WITH_EDITORONLY_DATA
-	const TMap<const UObjectBase *, FOverriddenPropertyAnnotation>& AnnotationMap = OverriddenObjectAnnotations.GetAnnotationMap();
-	for (const auto& Pair : AnnotationMap)
+	UE::TScopeLock AnnotationMapLock(AnnotationMapCritical);
+	for (const auto& Pair : GetAnnotationMap())
 	{
 		if (FOverriddenPropertySet* OverriddenProperties = Pair.Value.OverriddenProperties.Get())
 		{
@@ -413,8 +413,8 @@ void FOverridableManager::AddReferencedObjects(FReferenceCollector& Collector)
 	// that the replace references archives go through this object
 	if (!GIsGarbageCollecting)
 	{
-		const TMap<const UObjectBase *, FOverriddenPropertyAnnotation>& AnnotationMap = OverriddenObjectAnnotations.GetAnnotationMap();
-		for (const auto& Pair : AnnotationMap)
+		UE::TScopeLock AnnotationMapLock(AnnotationMapCritical);
+		for (const auto& Pair : GetAnnotationMap())
 		{
 			if (FOverriddenPropertySet* OverriddenProperties = Pair.Value.OverriddenProperties.Get())
 			{
@@ -428,8 +428,8 @@ void FOverridableManager::AddReferencedObjects(FReferenceCollector& Collector)
 void FOverridableManager::HandleDeadObjectReferences(const TSet<UClass*>& DeadClasses, const TSet<UObject*>& ActiveInstances, const TSet<UObject*>& TemplateInstances)
 {
 #if WITH_EDITORONLY_DATA
-	const TMap<const UObjectBase *, FOverriddenPropertyAnnotation>& AnnotationMap = OverriddenObjectAnnotations.GetAnnotationMap();
-	for (const auto& Pair : AnnotationMap)
+	UE::TScopeLock AnnotationMapLock(AnnotationMapCritical);
+	for (const auto& Pair : GetAnnotationMap())
 	{
 		if (FOverriddenPropertySet* OverriddenProperties = Pair.Value.OverriddenProperties.Get())
 		{
@@ -438,3 +438,30 @@ void FOverridableManager::HandleDeadObjectReferences(const TSet<UClass*>& DeadCl
 	}
 #endif // WITH_EDITORONLY_DATA
 }
+
+#if WITH_EDITORONLY_DATA
+FOverriddenPropertySet* FOverridableManager::Find(TNotNull<const UObject*> Object)
+{
+	FOverriddenPropertyAnnotation Annotation = GetAnnotation(Object);
+	return Annotation.OverriddenProperties.Get();
+}
+
+FOverriddenPropertySet& FOverridableManager::FindChecked(TNotNull<const UObject*> Object)
+{
+	FOverriddenPropertyAnnotation Annotation = GetAnnotation(Object);
+	checkf(!Annotation.IsDefault(), TEXT("Caller is expecting the object to have overridable serialization enabled"));
+	return *Annotation.OverriddenProperties.Get();
+}
+
+FOverriddenPropertySet& FOverridableManager::FindOrAdd(TNotNull<UObject*> Object)
+{
+	FOverriddenPropertyAnnotation Annotation = GetAnnotation(Object);
+	if (Annotation.IsDefault())
+	{
+		Annotation.OverriddenProperties = MakeShared<FOverriddenPropertySet>(Object);
+		AddAnnotation(Object, Annotation);
+	}
+
+	return *Annotation.OverriddenProperties;
+}
+#endif // WITH_EDITORONLY_DATA
\ No newline at end of file
diff --git a/Engine/Source/Runtime/CoreUObject/Public/UObject/OverridableManager.h b/Engine/Source/Runtime/CoreUObject/Public/UObject/OverridableManager.h
--- a/Engine/Source/Runtime/CoreUObject/Public/UObject/OverridableManager.h
+++ b/Engine/Source/Runtime/CoreUObject/Public/UObject/OverridableManager.h
@@ -32,43 +32,6 @@ struct FOverriddenPropertyAnnotation
 	TSharedPtr<FOverriddenPropertySet> OverriddenProperties;
 };
 
-/**
- * Global container of overridden object annotations
- */
-class FOverriddenPropertyAnnotations : public FUObjectAnnotationSparse<FOverriddenPropertyAnnotation, true/*bAutoRemove*/>
-{
-public:
-	bool IsEnabled(TNotNull<const UObject*> Object)
-	{
-		return !GetAnnotation(Object).IsDefault();
-	}
-
-	FOverriddenPropertySet* Find(TNotNull<const UObject*> Object)
-	{
-		FOverriddenPropertyAnnotation Annotation = GetAnnotation(Object);
-		return Annotation.OverriddenProperties.Get();
-	}
-
-	FOverriddenPropertySet& FindChecked(TNotNull<const UObject*> Object)
-	{
-		FOverriddenPropertyAnnotation Annotation = GetAnnotation(Object);
-		checkf(!Annotation.IsDefault(), TEXT("Caller is expecting the object to have overridable serialization enabled"));
-		return *Annotation.OverriddenProperties.Get();
-	}
-
-	FOverriddenPropertySet& FindOrAdd(TNotNull<UObject*> Object)
-	{
-		FOverriddenPropertyAnnotation Annotation = GetAnnotation(Object);
-		if (Annotation.IsDefault())
-		{
-			Annotation.OverriddenProperties = MakeShared<FOverriddenPropertySet>(Object);
-			AddAnnotation(Object, Annotation);
-		}
-
-		return *Annotation.OverriddenProperties;
-	}
-};
-
 #endif // WITH_EDITORONLY_DATA
 
 UENUM()
@@ -89,6 +52,9 @@ struct FOverrideManagerCapabilityInterface;
  *************************************************************************************
  */
 class FOverridableManager : public FGCObject
+#if WITH_EDITORONLY_DATA
+	, protected FUObjectAnnotationSparse<FOverriddenPropertyAnnotation, true/*bAutoRemove*/>
+#endif // WITH_EDITORONLY_DATA
 {
 public:
 	/**
@@ -344,7 +310,9 @@ protected:
 	COREUOBJECT_API static FOverridableManager* OverridableManager;
 
 #if WITH_EDITORONLY_DATA
-	FOverriddenPropertyAnnotations OverriddenObjectAnnotations;
+	FOverriddenPropertySet* Find(TNotNull<const UObject*> Object);
+	FOverriddenPropertySet& FindChecked(TNotNull<const UObject*> Object);
+	FOverriddenPropertySet& FindOrAdd(TNotNull<UObject*> Object);
 #else
 	FUObjectAnnotationSparseBool NeedsSubobjectTemplateInstantiation;
 #endif // WITH_EDITORONLY_DATA
diff --git a/Engine/Source/Runtime/CoreUObject/Public/UObject/UObjectAnnotation.h b/Engine/Source/Runtime/CoreUObject/Public/UObject/UObjectAnnotation.h
--- a/Engine/Source/Runtime/CoreUObject/Public/UObject/UObjectAnnotation.h
+++ b/Engine/Source/Runtime/CoreUObject/Public/UObject/UObjectAnnotation.h
@@ -334,6 +334,7 @@ private:
 #endif
 	}
 
+protected:
 	/**
 	 * Guards the annotation map, cache key and cache value.
 	 */
