Removing spurious ChangeOwnerServerId during Resetting of Borrowed Objects

#rnx

diff --git a/Engine/Source/Runtime/CoreUObject/Private/UObject/RemoteObjectSerialization.cpp b/Engine/Source/Runtime/CoreUObject/Private/UObject/RemoteObjectSerialization.cpp
--- a/Engine/Source/Runtime/CoreUObject/Private/UObject/RemoteObjectSerialization.cpp
+++ b/Engine/Source/Runtime/CoreUObject/Private/UObject/RemoteObjectSerialization.cpp
@@ -1192,16 +1192,6 @@ void ResetRemoteObject(UObject* InObject)
 		// Additionally we don't want to recursively re-enter this function so let the deserialization process know we're already resetting migrated object(s)
 		ERemoteObjectSerializationFlags DeserializationFlags = ERemoteObjectSerializationFlags::PreserveRemoteReferences | ERemoteObjectSerializationFlags::UseExistingObjects | ERemoteObjectSerializationFlags::Resetting;
 		DeserializeObjectData(*ArchetypeDelta, /*MigrationContext*/ nullptr, DeserializedIds, DeserializedObjects, DeserializationFlags);
-
-		const FRemoteServerId LocalServerId = UE::RemoteObject::GetGlobalServerId();
-		for (UObject* DeserializedObject : DeserializedObjects)
-		{
-			FRemoteObjectId RemoteId = UE::CoreUObject::Private::FObjectHandleUtils::GetRemoteId(DeserializedObject);
-			if (ensureMsgf(!RemoteId.IsAsset(), TEXT("Attempted to Reset Asset %s (we must have received something we thought was an Asset). Check why this is happening."), *GetNameSafe(DeserializedObject)))
-			{
-				UE::RemoteObject::Handle::ChangeOwnerServerId(DeserializedObject, LocalServerId);
-			}
-		}
 	}
 	delete ArchetypeDelta;
 }
