Small fixes for ensuring RemoteObjects are always updated to the right ResidentServerId.  Also, defaulting to an Invalid GlobalServerId so we can differentiate between Server processes and Proxy/Client processes.

#rnx

diff --git a/Engine/Source/Runtime/CoreUObject/Private/UObject/RemoteObject.cpp b/Engine/Source/Runtime/CoreUObject/Private/UObject/RemoteObject.cpp
--- a/Engine/Source/Runtime/CoreUObject/Private/UObject/RemoteObject.cpp
+++ b/Engine/Source/Runtime/CoreUObject/Private/UObject/RemoteObject.cpp
@@ -86,11 +86,12 @@ public:
 			FRemoteServerId ObjectServerId = Stub->Id.GetServerId();
 			Stub->OwningServerId = ObjectServerId.IsValid() ? ObjectServerId : GetGlobalServerId();
 		}
-		Stub->ResidentServerId = DestinationServerId;
+
 		FUObjectItem* ObjectItem = GUObjectArray.ObjectToObjectItem(Object);
 		checkf(ObjectItem, TEXT("Attempting to get a serial number for an object that does not exist in the global UObject array (it's possible GUObjectArray is not initialized yet, ObjectIndex=%d)"), GUObjectArray.ObjectToIndex(Object));
 		Stub->SerialNumber = ObjectItem->GetSerialNumber();
 		Stub->Name = Object->GetFName();
+		Stub->ResidentServerId = DestinationServerId;
 
 		return Stub;
 	}
@@ -145,14 +146,14 @@ void InitServerId()
 	{
 		if (!FParse::Value(FCommandLine::Get(), TEXT("LocalPeerId="), ServerId))
 		{
-			RemoteServerId = FRemoteServerId(ERemoteServerIdConstants::FirstValid);
+			RemoteServerId = FRemoteServerId(ERemoteServerIdConstants::Invalid);
 		}
 	}
 	if (!ServerId.IsEmpty())
 	{
 		RemoteServerId = FRemoteServerId(ServerId);
+		checkf(RemoteServerId.IsValid(), TEXT("Remote ServerId is not valid"));
 	}
-	checkf(RemoteServerId.IsValid(), TEXT("Remote ServerId is not valid"));
 
 	UE_LOG(LogRemoteObject, Display, TEXT("Remote ServerId: %s"), *RemoteServerId.ToString());
 }
@@ -222,10 +223,12 @@ void MarkAsRemote(UObject* Object, FRemoteServerId DestinationServerId)
 
 void MarkAsLocal(UObject* Object)
 {
+	ensureMsgf(!Object->HasAnyFlags(EObjectFlags::RF_ClassDefaultObject | EObjectFlags::RF_ArchetypeObject), TEXT("We're about to set an ArchetypeObject %s as remote reference"), *GetNameSafe(Object));
+
 	FUObjectItem* ObjectItem = GUObjectArray.ObjectToObjectItem(Object);
 	ObjectItem->ClearFlags(EInternalObjectFlags::Remote);
 	ObjectItem->SetFlags(EInternalObjectFlags::RemoteReference);
-	ensureMsgf(!Object->HasAnyFlags(EObjectFlags::RF_ClassDefaultObject | EObjectFlags::RF_ArchetypeObject), TEXT("We're about to set an ArchetypeObject %s as remote reference"), *GetNameSafe(Object));
+	ObjectMaps->FindOrAddRemoteObjectStub(Object, GetGlobalServerId());
 }
 
 void StoreAssetPath(UObject* Object)
