[CoreUObject]
* Changed operators in PersistentObjectPtr to not be hidden friend unless they have to. This fixes compile errors with hidden friends complaining about using unions in inlined hidden friends

diff --git a/Engine/Source/Runtime/CoreUObject/Public/UObject/PersistentObjectPtr.h b/Engine/Source/Runtime/CoreUObject/Public/UObject/PersistentObjectPtr.h
--- a/Engine/Source/Runtime/CoreUObject/Public/UObject/PersistentObjectPtr.h
+++ b/Engine/Source/Runtime/CoreUObject/Public/UObject/PersistentObjectPtr.h
@@ -177,36 +177,38 @@ struct TPersistentObjectPtr
 	}
 
 	/** Compare pointers for equality. Only Serial Number matters for the base implementation */
-	FORCEINLINE friend bool operator==(const TPersistentObjectPtr& Lhs, const TPersistentObjectPtr& Rhs)
+	FORCEINLINE bool operator==(const TPersistentObjectPtr& Rhs) const
 	{
-		return Lhs.ObjectID == Rhs.ObjectID;
+		return ObjectID == Rhs.ObjectID;
 	}
 
-	FORCEINLINE friend bool operator==(const TPersistentObjectPtr& Lhs, TYPE_OF_NULLPTR)
+	FORCEINLINE bool operator==(TYPE_OF_NULLPTR) const
 	{
-		return !Lhs.IsValid();
+		return !IsValid();
 	}
 
-	FORCEINLINE friend bool operator==(TYPE_OF_NULLPTR, const TPersistentObjectPtr& Rhs)
+	/** Compare pointers for inequality. Only Serial Number matters for the base implementation */
+	FORCEINLINE bool operator!=(const TPersistentObjectPtr& Rhs) const
 	{
-		return !Rhs.IsValid();
+		return ObjectID != Rhs.ObjectID;
 	}
 
-	/** Compare pointers for inequality. Only Serial Number matters for the base implementation */
-	FORCEINLINE friend bool operator!=(const TPersistentObjectPtr& Lhs, const TPersistentObjectPtr& Rhs)
+	FORCEINLINE bool operator!=(TYPE_OF_NULLPTR) const
 	{
-		return Lhs.ObjectID != Rhs.ObjectID;
+		return IsValid();
 	}
 
-	FORCEINLINE friend bool operator!=(const TPersistentObjectPtr& Lhs, TYPE_OF_NULLPTR)
+#if !PLATFORM_COMPILER_HAS_GENERATED_COMPARISON_OPERATORS
+	FORCEINLINE friend bool operator==(TYPE_OF_NULLPTR, const TPersistentObjectPtr& Rhs)
 	{
-		return Lhs.IsValid();
+		return !Rhs.IsValid();
 	}
 
 	FORCEINLINE friend bool operator!=(TYPE_OF_NULLPTR, const TPersistentObjectPtr& Rhs)
 	{
 		return Rhs.IsValid();
 	}
+#endif
 
 	/**  
 	 * Test if this does not point to a live UObject, but may in the future
