[Core]
* Fixed AssetLog hidden friend actually being hidden friend

diff --git a/Engine/Source/Runtime/CoreUObject/Private/UObject/UObjectGlobals.cpp b/Engine/Source/Runtime/CoreUObject/Private/UObject/UObjectGlobals.cpp
--- a/Engine/Source/Runtime/CoreUObject/Private/UObject/UObjectGlobals.cpp
+++ b/Engine/Source/Runtime/CoreUObject/Private/UObject/UObjectGlobals.cpp
@@ -6035,8 +6035,10 @@ COREUOBJECT_API UFunction* FindDelegateSignature(FName DelegateSignatureName)
 	return nullptr;
 }
 
-void UE::SerializeForLog(FCbWriter& Writer, const FAssetLog& AssetLog)
+void UE::FAssetLog::SerializeForLog(FCbWriter& Writer) const
 {
+	const FAssetLog& AssetLog = *this;
+
 	FString ObjectPath;
 	FString LocalPath;
 
diff --git a/Engine/Source/Runtime/CoreUObject/Public/UObject/UObjectGlobals.h b/Engine/Source/Runtime/CoreUObject/Public/UObject/UObjectGlobals.h
--- a/Engine/Source/Runtime/CoreUObject/Public/UObject/UObjectGlobals.h
+++ b/Engine/Source/Runtime/CoreUObject/Public/UObject/UObjectGlobals.h
@@ -3418,9 +3418,10 @@ public:
 	inline explicit FAssetLog(const FPackagePath& InPath UE_LIFETIMEBOUND) : PackagePath(&InPath) {}
 	inline explicit FAssetLog(const UObject* InObject UE_LIFETIMEBOUND) : Object(InObject) {}
 
-	COREUOBJECT_API friend void SerializeForLog(FCbWriter& Writer, const FAssetLog& AssetLog);
+	friend void SerializeForLog(FCbWriter& Writer, const FAssetLog& AssetLog) { AssetLog.SerializeForLog(Writer); }
 
 private:
+	COREUOBJECT_API void SerializeForLog(FCbWriter& Writer) const;
 	const TCHAR* Path = nullptr;
 	const FPackagePath* PackagePath = nullptr;
 	const UObject* Object = nullptr;
