Moved functions with default lambda functions to be explicitly named. This works around headerunit/c++modules issue and also produce less noise on the symbol files

diff --git a/Engine/Source/Runtime/CoreUObject/Public/UObject/ObjectPtr.h b/Engine/Source/Runtime/CoreUObject/Public/UObject/ObjectPtr.h
--- a/Engine/Source/Runtime/CoreUObject/Public/UObject/ObjectPtr.h
+++ b/Engine/Source/Runtime/CoreUObject/Public/UObject/ObjectPtr.h
@@ -500,7 +500,7 @@ namespace ObjectPtr_Private
 template <typename T>
 struct TObjectPtr
 {
-#ifndef PLATFORM_COMPILER_IWYU
+#if !defined(PLATFORM_COMPILER_IWYU) && !defined(UE_HEADER_UNITS)
 	// TObjectPtr should only be used on types T that are EITHER:
 	// - incomplete (ie: forward declared and we have not seen their definition yet)
 	// - complete and derived from UObject
@@ -1099,6 +1099,8 @@ ToRawPtrTArrayUnsafe(ArrayType&& Array)
 }
 PRAGMA_ENABLE_DEPRECATION_WARNINGS
 
+inline constexpr auto TContainerElementTypeCompatibility_DefaultOperatorFunc = [](auto& InIt) -> decltype(auto) { return *InIt; };
+
 template <typename T>
 struct TContainerElementTypeCompatibility<TObjectPtr<T>>
 {
@@ -1107,7 +1109,7 @@ struct TContainerElementTypeCompatibility<TObjectPtr<T>>
 
 	template <typename IterBeginType, typename IterEndType, typename OperatorType = std::remove_reference_t<decltype(*std::declval<IterBeginType>())>& (*)(IterBeginType&)>
 	UE_OBJPTR_DEPRECATED(5.0, "Reinterpretation between ranges of one type to another type is deprecated.")
-	static void ReinterpretRange(IterBeginType Iter, IterEndType IterEnd, OperatorType Operator = [](IterBeginType& InIt) -> decltype(auto) { return *InIt; })
+	static void ReinterpretRange(IterBeginType Iter, IterEndType IterEnd, OperatorType Operator = TContainerElementTypeCompatibility_DefaultOperatorFunc)
 	{
 #if UE_WITH_OBJECT_HANDLE_LATE_RESOLVE || UE_WITH_OBJECT_HANDLE_TRACKING
 		while (Iter != IterEnd)
@@ -1121,7 +1123,7 @@ struct TContainerElementTypeCompatibility<TObjectPtr<T>>
 
 	template <typename IterBeginType, typename IterEndType, typename SizeType, typename OperatorType = std::remove_reference_t<decltype(*std::declval<IterBeginType>())>& (*)(IterBeginType&)>
 	UE_OBJPTR_DEPRECATED(5.0, "Reinterpretation between ranges of one type to another type is deprecated.")
-	static void ReinterpretRangeContiguous(IterBeginType Iter, IterEndType IterEnd, SizeType Size, OperatorType Operator = [](IterBeginType& InIt) -> decltype(auto) { return *InIt; })
+	static void ReinterpretRangeContiguous(IterBeginType Iter, IterEndType IterEnd, SizeType Size, OperatorType Operator = TContainerElementTypeCompatibility_DefaultOperatorFunc)
 	{
 #if UE_WITH_OBJECT_HANDLE_LATE_RESOLVE || UE_WITH_OBJECT_HANDLE_TRACKING
 		const TObjectPtr<T>* Begin = &*Iter;
@@ -1150,7 +1152,7 @@ struct TContainerElementTypeCompatibility<const TObjectPtr<T>>
 
 	template <typename IterBeginType, typename IterEndType, typename OperatorType = const std::remove_cv_t<std::remove_reference_t<decltype(*std::declval<IterBeginType>())>>&(*)(IterBeginType&)>
 	UE_OBJPTR_DEPRECATED(5.0, "Reinterpretation between ranges of one type to another type is deprecated.")
-	static void ReinterpretRange(IterBeginType Iter, IterEndType IterEnd, OperatorType Operator = [](IterBeginType& InIt) -> decltype(auto) { return *InIt; })
+	static void ReinterpretRange(IterBeginType Iter, IterEndType IterEnd, OperatorType Operator = TContainerElementTypeCompatibility_DefaultOperatorFunc)
 	{
 #if UE_WITH_OBJECT_HANDLE_LATE_RESOLVE || UE_WITH_OBJECT_HANDLE_TRACKING
 		while (Iter != IterEnd)
@@ -1164,7 +1166,7 @@ struct TContainerElementTypeCompatibility<const TObjectPtr<T>>
 
 	template <typename IterBeginType, typename IterEndType, typename SizeType, typename OperatorType = std::remove_reference_t<decltype(*std::declval<IterBeginType>())>& (*)(IterBeginType&)>
 	UE_OBJPTR_DEPRECATED(5.0, "Reinterpretation between ranges of one type to another type is deprecated.")
-	static void ReinterpretRangeContiguous(IterBeginType Iter, IterEndType IterEnd, SizeType Size, OperatorType Operator = [](IterBeginType& InIt) -> decltype(auto) { return *InIt; })
+	static void ReinterpretRangeContiguous(IterBeginType Iter, IterEndType IterEnd, SizeType Size, OperatorType Operator = TContainerElementTypeCompatibility_DefaultOperatorFunc)
 	{
 #if UE_WITH_OBJECT_HANDLE_LATE_RESOLVE || UE_WITH_OBJECT_HANDLE_TRACKING
 		const TObjectPtr<T>* Begin = &*Iter;
diff --git a/Engine/Source/Runtime/CoreUObject/Public/UObject/SoftObjectPath.h b/Engine/Source/Runtime/CoreUObject/Public/UObject/SoftObjectPath.h
--- a/Engine/Source/Runtime/CoreUObject/Public/UObject/SoftObjectPath.h
+++ b/Engine/Source/Runtime/CoreUObject/Public/UObject/SoftObjectPath.h
@@ -40,6 +40,8 @@ struct FUObjectSerializeContext;
 /** Delegate called on completion of async loading a soft object. The UObject will be null if the load failed */
 DECLARE_DELEGATE_TwoParams(FLoadSoftObjectPathAsyncDelegate, const FSoftObjectPath&, UObject*);
 
+constexpr inline auto FSoftObjectPath_DefaultPreFixupForPIEFunction = [](int32, FSoftObjectPath&) {};
+
 /**
  * A struct that contains a string reference to an object, either a package, a top level asset or a subobject.
  * This can be used to make soft references to assets that are loaded on demand.
@@ -408,10 +410,10 @@ struct FSoftObjectPath
 	COREUOBJECT_API void PostLoadPath(FArchive* InArchive) const;
 
 	/** Fixes up this SoftObjectPath to add the PIE prefix depending on what is currently active, returns true if it was modified. The overload that takes an explicit PIE instance is preferred, if it's available. */
-	COREUOBJECT_API bool FixupForPIE(TFunctionRef<void(int32, FSoftObjectPath&)> InPreFixupForPIECustomFunction = [](int32, FSoftObjectPath&) {});
+	COREUOBJECT_API bool FixupForPIE(TFunctionRef<void(int32, FSoftObjectPath&)> InPreFixupForPIECustomFunction = FSoftObjectPath_DefaultPreFixupForPIEFunction);
 
 	/** Fixes up this SoftObjectPath to add the PIE prefix for the given PIEInstance index, returns true if it was modified */
-	COREUOBJECT_API bool FixupForPIE(int32 PIEInstance, TFunctionRef<void(int32, FSoftObjectPath&)> InPreFixupForPIECustomFunction = [](int32, FSoftObjectPath&) {});
+	COREUOBJECT_API bool FixupForPIE(int32 PIEInstance, TFunctionRef<void(int32, FSoftObjectPath&)> InPreFixupForPIECustomFunction = FSoftObjectPath_DefaultPreFixupForPIEFunction);
 
 	/** Fixes soft object path for CoreRedirects to handle renamed native objects, returns true if it was modified */
 	COREUOBJECT_API bool FixupCoreRedirects();
diff --git a/Engine/Source/Runtime/Engine/Classes/Engine/TextureRenderTarget.h b/Engine/Source/Runtime/Engine/Classes/Engine/TextureRenderTarget.h
--- a/Engine/Source/Runtime/Engine/Classes/Engine/TextureRenderTarget.h
+++ b/Engine/Source/Runtime/Engine/Classes/Engine/TextureRenderTarget.h
@@ -103,7 +103,8 @@ class UTextureRenderTarget : public UTexture
 	 * @return true if conversion from render to texture was successful
 	 */
 	using FOnTextureChangingDelegate = TFunctionRef<void(UTexture* /*InTexture*/)>;
-	ENGINE_API bool UpdateTexture(UTexture* InTexture, EConstructTextureFlags InFlags = CTF_Default, const TArray<uint8>* InAlphaOverride = nullptr, FOnTextureChangingDelegate InOnTextureChangingDelegate = [](UTexture*){}, FText* OutErrorMessage = nullptr);
+	constexpr static auto OnTextureChangingDefault = [](UTexture*){};
+	ENGINE_API bool UpdateTexture(UTexture* InTexture, EConstructTextureFlags InFlags = CTF_Default, const TArray<uint8>* InAlphaOverride = nullptr, FOnTextureChangingDelegate InOnTextureChangingDelegate = OnTextureChangingDefault, FText* OutErrorMessage = nullptr);
 
 	/**
 	 * Utility for creating a new UTexture from a UTextureRenderTarget. The method can fail if conversion is impossible
diff --git a/Engine/Source/Runtime/Engine/Public/ContentStreaming.h b/Engine/Source/Runtime/Engine/Public/ContentStreaming.h
--- a/Engine/Source/Runtime/Engine/Public/ContentStreaming.h
+++ b/Engine/Source/Runtime/Engine/Public/ContentStreaming.h
@@ -477,11 +477,13 @@ struct IRenderAssetStreamingManager : public IStreamingManager
 	/** Return all bounds related to the ref object */
 	virtual void GetObjectReferenceBounds(const UObject* RefObject, TArray<FBox>& AssetBoxes) = 0;
 
+	constexpr static auto ShouldChooseDefault = [](const UPrimitiveComponent*) { return true; };
+
 	/** Return all components referencing the asset */
 	virtual void GetAssetComponents(
 		const UStreamableRenderAsset* RenderAsset,
 		TArray<const UPrimitiveComponent*>& OutComps,
-		TFunction<bool(const UPrimitiveComponent*)> ShouldChoose = [](const UPrimitiveComponent*) { return true; }) = 0;
+		TFunction<bool(const UPrimitiveComponent*)> ShouldChoose = ShouldChooseDefault) = 0;
 
 	//BEGIN: APIs for backward compatibility
 	ENGINE_API void UpdateIndividualTexture(UTexture2D* Texture);
