Fix compilation error for Replication System Tests Win64

diff --git a/Engine/Source/Runtime/CoreUObject/Private/UObject/RemoteObject.cpp b/Engine/Source/Runtime/CoreUObject/Private/UObject/RemoteObject.cpp
--- a/Engine/Source/Runtime/CoreUObject/Private/UObject/RemoteObject.cpp
+++ b/Engine/Source/Runtime/CoreUObject/Private/UObject/RemoteObject.cpp
@@ -127,7 +127,7 @@ public:
 		{
 			if (Elem.Value)
 			{
-				const uint64 OriginalPhysicsId = Elem.Value->PhysicsOverrideServerId.GetIdNumber();
+				const uint32 OriginalPhysicsId = Elem.Value->PhysicsOverrideServerId.GetIdNumber();
 				if (const uint32* MergedPhysicsId = PhysicsServerMergingMap.Find(OriginalPhysicsId))
 				{
 					Elem.Value->PhysicsOverrideServerId = FRemoteServerId(*MergedPhysicsId);
