[CoreUObjectPlugin]

If we are leaking packages and rename the packages so asny attempt to reload the package won't load the garbage package and will instead get the package from disk.

#rnx

diff --git a/Engine/Source/Runtime/CoreUObject/Private/Misc/CoreUObjectPluginManager.cpp b/Engine/Source/Runtime/CoreUObject/Private/Misc/CoreUObjectPluginManager.cpp
--- a/Engine/Source/Runtime/CoreUObject/Private/Misc/CoreUObjectPluginManager.cpp
+++ b/Engine/Source/Runtime/CoreUObject/Private/Misc/CoreUObjectPluginManager.cpp
@@ -142,7 +142,7 @@ namespace UE::CoreUObject::Private
 	}
 
 	// Check if any assets from the plugin mount point have leaked, and if so trace them.
-	// Then rename them to allow new copies of them to be loaded. 
+	// Then rename the assets and mark them as garbage to allow new copies of them to be loaded. 
 	void HandlePossibleAssetLeaks(const FString& PluginName)
 	{
 		TRACE_CPUPROFILER_EVENT_SCOPE(HandlePossibleAssetLeaks);
@@ -275,7 +275,7 @@ namespace UE::CoreUObject::Private
 
 			Package->MarkAsGarbage();
 
-			if ((!GIsEditor && !UObject::IsGarbageEliminationEnabled())
+			if (!UObject::IsGarbageEliminationEnabled()
 #if WITH_LOW_LEVEL_TESTS
 				|| GRenameLeakedPackages
 #endif
