Fixing code responsible for resetting remote objects finding the incorrect archetype for BP class instances' subobjects

diff --git a/Engine/Source/Runtime/CoreUObject/Private/UObject/RemoteObjectSerialization.cpp b/Engine/Source/Runtime/CoreUObject/Private/UObject/RemoteObjectSerialization.cpp
--- a/Engine/Source/Runtime/CoreUObject/Private/UObject/RemoteObjectSerialization.cpp
+++ b/Engine/Source/Runtime/CoreUObject/Private/UObject/RemoteObjectSerialization.cpp
@@ -121,12 +121,19 @@ namespace UE::RemoteObject::Serialization
 
 UObject* FindArchetype(const UObject* InObj)
 {
+	using namespace UE::CoreUObject::Private;
+
 	if (GUseImmutableArchetypes)
 	{
-		UClass* ObjectClass = InObj->GetClass();
-		// No need to get the immutable CDO for BP classes as BP classes are assets themselves and although their CDOs can still technically be modified at runtime
+		bool bNativeObject = true;
+		// No need to get the immutable CDO for BP class instance or their subobjects as BP classes are assets themselves and although their CDOs can still technically be modified at runtime
 		// they never are because they can be GC'd and reset to their original state when they're reloaded so they're not a persistent storage like the native CDOs
-		if (!ObjectClass->HasAnyClassFlags(CLASS_CompiledFromBlueprint))
+		for (const UObject* OuterIt = InObj; OuterIt && bNativeObject; OuterIt = FObjectHandleUtils::GetNonAccessTrackedOuterNoResolve(OuterIt))
+		{
+			UClass* ObjectClass = OuterIt->GetClass();
+			bNativeObject = !ObjectClass->HasAnyClassFlags(CLASS_CompiledFromBlueprint);
+		}
+		if (bNativeObject)
 		{
 			return const_cast<UObject*>(FindImmutableArchetype(InObj));
 		}
