Previously, we have cleaned up the World-to-Transform calculations so they no longer happen in PostMigrate.  There is a case where SceneComponent::OnRegister will perform a AttachToParent (since it's thought this only happens through initial spawning code).  Since that's no longer the case (OnRegister now happens during Migration), we needed a way to tell SceneComponent::OnRegister that we're migrating and thus ignore the call to AttachToParent (we've already deserialized the correct state).  The final result is this allows us to migrate Attached Actors which are Remote (we no longer dereference them outside of a transaction).  The actual usefulness of this will only become apparent when the Proxy GUIDs match the NetGUIDs, because there's no way to actually have BasedMovement without that apparently.

#rnx

diff --git a/Engine/Source/Runtime/CoreUObject/Internal/UObject/UObjectMigrationContext.h b/Engine/Source/Runtime/CoreUObject/Internal/UObject/UObjectMigrationContext.h
--- a/Engine/Source/Runtime/CoreUObject/Internal/UObject/UObjectMigrationContext.h
+++ b/Engine/Source/Runtime/CoreUObject/Internal/UObject/UObjectMigrationContext.h
@@ -52,7 +52,7 @@ enum class EObjectMigrationSendType : uint8
 COREUOBJECT_API const TCHAR* ToString(EObjectMigrationSendType Value);
 
 /** Structure that holds context for what we're intending to do during migration */
-struct COREUOBJECT_API FUObjectMigrationContext
+struct COREUOBJECT_API FUObjectMigrationContext final
 {
 	/** The Object that is being migrated, it may not be the top-level UObject */
 	FRemoteObjectId ObjectId;
@@ -67,7 +67,7 @@ struct COREUOBJECT_API FUObjectMigrationContext
 	FRemoteServerId PhysicsServerId;
 
 	/** Which side of the migration does this context apply to? Sending or Receiving? */
-	EObjectMigrationSide MigrationSide;
+	EObjectMigrationSide MigrationSide = EObjectMigrationSide::Receive;
 
 	/** If valid, the request id of the multi-server commit */
 	FRemoteTransactionId MultiServerCommitRequestId;
@@ -86,4 +86,45 @@ struct COREUOBJECT_API FUObjectMigrationContext
 	*/
 	EObjectMigrationSendType GetObjectMigrationSendType(const UObjectBase* Object) const;
 	EObjectMigrationSendType GetObjectMigrationSendType(FRemoteObjectId ObjectId) const;
+
+	/** Gets the current Migration Context if it exists (only exists during actual object migrations) */
+	static const FUObjectMigrationContext* const GetCurrentMigrationContext()
+	{
+		checkf(IsInGameThread(), TEXT("%hs: We expect migrations to only occur on the GameThread."), __func__);
+		return CurrentMigrationContext;
+	}
+
+	
+private: // Interface used for FScopedObjectMigrationContext
+	friend struct FScopedObjectMigrationContext;
+
+	inline static const FUObjectMigrationContext* CurrentMigrationContext = nullptr;
+	static void SetCurrentMigrationContext(const FUObjectMigrationContext* NewMigrationContext)
+	{
+		checkf(IsInGameThread(), TEXT("%hs: We expect migrations to only occur on the GameThread."), __func__);
+		CurrentMigrationContext = NewMigrationContext;
+	}
 };
+
+/**
+ * Helper class to be used whenever we create an FUObjectMigrationContext to ensure all calls within the scope
+ * have their FUObjectMigrationContext::GetCurrentMigrationContext set to the desired value.  It will correctly pop
+ * and restore the previous value on destruction, obeying the stack of migrations (such as needing to create a new
+ * MigrationContext when we're loading from the Database).
+ */
+struct FScopedObjectMigrationContext
+{
+	FScopedObjectMigrationContext(const FUObjectMigrationContext& NewMigrationContext)
+		: PreviousMigrationContext(FUObjectMigrationContext::CurrentMigrationContext)
+	{
+		FUObjectMigrationContext::SetCurrentMigrationContext(&NewMigrationContext);
+	}
+
+	~FScopedObjectMigrationContext()
+	{
+		FUObjectMigrationContext::SetCurrentMigrationContext(PreviousMigrationContext);
+	}
+
+private:
+	const FUObjectMigrationContext* PreviousMigrationContext;
+};
\ No newline at end of file
diff --git a/Engine/Source/Runtime/CoreUObject/Private/UObject/RemoteObjectTransfer.cpp b/Engine/Source/Runtime/CoreUObject/Private/UObject/RemoteObjectTransfer.cpp
--- a/Engine/Source/Runtime/CoreUObject/Private/UObject/RemoteObjectTransfer.cpp
+++ b/Engine/Source/Runtime/CoreUObject/Private/UObject/RemoteObjectTransfer.cpp
@@ -200,7 +200,14 @@ public:
 
 					FRemoteServerId OwnerServerId = Handle::GetOwnerServerId(RootObject);
 					FRemoteServerId PhysicsServerId = Handle::GetPhysicsServerId(RootObject);
-					FUObjectMigrationContext MigrationContext{ .ObjectId = HighestPriorityRequest->ObjectId, .RemoteServerId = HighestPriorityRequest->DestinationServerId, .OwnerServerId = OwnerServerId, .PhysicsServerId = PhysicsServerId, .MigrationSide = EObjectMigrationSide::Send };
+					FUObjectMigrationContext MigrationContext {
+						.ObjectId = HighestPriorityRequest->ObjectId,
+						.RemoteServerId = HighestPriorityRequest->DestinationServerId,
+						.OwnerServerId = OwnerServerId,
+						.PhysicsServerId = PhysicsServerId,
+						.MigrationSide = EObjectMigrationSide::Send
+					};
+					FScopedObjectMigrationContext ScopedMigrationContext(MigrationContext);
 
 					TSet<UObject*> SentObjects;
 					TSet<UObject*> ReferencedObjects;
@@ -524,8 +531,7 @@ public:
 			// NOTE: we need to run SerializeObjectData in the CLOSED because it turns out some of the
 			// serialization code ends up mutating the object, so we need to be able to undo those
 			// mutations if we need to abort past this point
-			FUObjectMigrationContext MigrationContext
-			{
+			FUObjectMigrationContext MigrationContext {
 				.ObjectId = ObjectId,
 				.RemoteServerId = OwnerServerId,
 				.OwnerServerId = OwnerServerId,
@@ -533,6 +539,7 @@ public:
 				.MigrationSide = EObjectMigrationSide::Send,
 				.MultiServerCommitRequestId = ActiveRequest->RequestId,
 			};
+			FScopedObjectMigrationContext ScopedMigrationContext(MigrationContext);
 
 			TSet<UObject*> SentObjects;
 			FRemoteObjectData ObjectData = SerializeObjectData(Object, SentObjects, MultiServerCommitReferencedObjects, &MigrationContext);
@@ -592,7 +599,14 @@ public:
 			FRemoteServerId PhysicsServerId = Handle::GetPhysicsServerId(SentObject);
 			FRemoteObjectId SentObjectId = UE::CoreUObject::Private::FObjectHandleUtils::GetRemoteId(SentObject);
 
-			FUObjectMigrationContext MigrationContext{ .ObjectId = SentObjectId, .RemoteServerId = OwnerServerId, .OwnerServerId = OwnerServerId, .PhysicsServerId = PhysicsServerId, .MigrationSide = EObjectMigrationSide::Send };
+			FUObjectMigrationContext MigrationContext {
+				.ObjectId = SentObjectId,
+				.RemoteServerId = OwnerServerId,
+				.OwnerServerId = OwnerServerId,
+				.PhysicsServerId = PhysicsServerId,
+				.MigrationSide = EObjectMigrationSide::Send
+			};
+			FScopedObjectMigrationContext ScopedMigrationContext(MigrationContext);
 
 			UE_LOG(LogRemoteObject, VeryVerbose, TEXT("FRemoteObjectTransferQueue: CommitMultiServerCommit postmigrate %s"), *SentObjectId.ToString());
 
@@ -684,7 +698,14 @@ public:
 			UE_LOG(LogRemoteObject, Log, TEXT("%hs: Restoring ObjectId %s from Database and Forwarding to %s"), __func__, *ObjectId.ToString(), *DestinationServerId.ToString());
 			// We are intercepting a request to pull-migrate since we believe the object resides in our database, migrate it "from" the database to us
 			// Inside the database call, it will fix-up these migration context parameters to make us the owner
-			FUObjectMigrationContext MigrationContext { .ObjectId = ObjectId, .RemoteServerId = DatabaseId, .OwnerServerId = DatabaseId, .PhysicsServerId = DatabaseId, .MigrationSide = EObjectMigrationSide::Receive};
+			FUObjectMigrationContext MigrationContext {
+				.ObjectId = ObjectId,
+				.RemoteServerId = DatabaseId,
+				.OwnerServerId = DatabaseId,
+				.PhysicsServerId = DatabaseId,
+				.MigrationSide = EObjectMigrationSide::Receive
+			};
+			FScopedObjectMigrationContext ScopedMigrationContext(MigrationContext);
 			RestoreRemoteObjectDataDelegate.Execute(MigrationContext);
 
 #if UE_WITH_REMOTE_OBJECT_HANDLE
@@ -1007,8 +1028,15 @@ void TransferObjectOwnershipToRemoteServer(UObject* Object, FRemoteServerId Dest
 
 	// Figure out the ownership semantics
 	FRemoteServerId PhysicsServerId = Handle::GetPhysicsServerId(RootObject);
-	FUObjectMigrationContext MigrationContext{ .ObjectId = ObjectId, .RemoteServerId = DestinationServerId, .OwnerServerId = DestinationServerId, .PhysicsServerId = PhysicsServerId, .MigrationSide = EObjectMigrationSide::Send };
-
+	FUObjectMigrationContext MigrationContext {
+		.ObjectId = ObjectId,
+		.RemoteServerId = DestinationServerId,
+		.OwnerServerId = DestinationServerId,
+		.PhysicsServerId = PhysicsServerId,
+		.MigrationSide = EObjectMigrationSide::Send
+	};
+    FScopedObjectMigrationContext ScopedMigrationContext(MigrationContext);
+	
 	UE_LOG(LogRemoteObject, Verbose, TEXT("TransferObjectOwnershipToRemoteServer %s (%s) root %s (%s) to %s"),
 		*ObjectId.ToString(),
 		*Object->GetName(),
@@ -1104,7 +1132,15 @@ void OnObjectDataReceived(FRemoteServerId OwnerServerId, FRemoteServerId Physics
 {
 	check(!AutoRTFM::IsClosed());
 
-	FUObjectMigrationContext MigrationContext { .ObjectId = ObjectId, .RemoteServerId = RemoteServerId, .OwnerServerId = OwnerServerId, .PhysicsServerId = PhysicsId, .MigrationSide = EObjectMigrationSide::Receive };
+	FUObjectMigrationContext MigrationContext {
+		.ObjectId = ObjectId,
+		.RemoteServerId = RemoteServerId,
+		.OwnerServerId = OwnerServerId,
+		.PhysicsServerId = PhysicsId,
+		.MigrationSide = EObjectMigrationSide::Receive
+	};
+	FScopedObjectMigrationContext ScopedMigrationContext(MigrationContext);
+	
 	GTransferQueue.FulfillReceiveRequest(Data, MigrationContext);
 }
 
@@ -1224,7 +1260,14 @@ void StoreObjectToDatabase(UObject* Object)
 
 
 			// Setup a migration context for "push migrating to the new owner: the Database Server"
-			FUObjectMigrationContext StoreToDatabaseContext{ .ObjectId = RootObjectId, .RemoteServerId = DatabaseId, .OwnerServerId = DatabaseId, .PhysicsServerId = DatabaseId, .MigrationSide = EObjectMigrationSide::Send };
+			FUObjectMigrationContext StoreToDatabaseContext {
+				.ObjectId = RootObjectId,
+				.RemoteServerId = DatabaseId,
+				.OwnerServerId = DatabaseId,
+				.PhysicsServerId = DatabaseId,
+				.MigrationSide = EObjectMigrationSide::Send
+			};
+			FScopedObjectMigrationContext ScopedMigrationContext(StoreToDatabaseContext);
 
 			// Do the actual serialization
 			FRemoteObjectData ObjectData = SerializeObjectData(OutermostUnreachableRemoteReference, SentObjects, ReferencedObjects, &StoreToDatabaseContext);
diff --git a/Engine/Source/Runtime/Engine/Private/Components/SceneComponent.cpp b/Engine/Source/Runtime/Engine/Private/Components/SceneComponent.cpp
--- a/Engine/Source/Runtime/Engine/Private/Components/SceneComponent.cpp
+++ b/Engine/Source/Runtime/Engine/Private/Components/SceneComponent.cpp
@@ -39,6 +39,10 @@
 #include "Misc/DataValidation.h"
 #include "Misc/EnumRange.h"
 
+#if UE_WITH_REMOTE_OBJECT_HANDLE
+#include "UObject/UObjectMigrationContext.h"
+#endif
+
 #define LOCTEXT_NAMESPACE "SceneComponent"
 
 namespace SceneComponentStatics
@@ -59,6 +63,11 @@ namespace SceneComponentCVars
 		TEXT("false: fix up any children that are missing the parent, true: fixes up all children except non-replicated root components (default)"),
 		ECVF_Default);
 
+#if UE_WITH_REMOTE_OBJECT_HANDLE
+	bool bAllowRemoteAttachedChildren = true;
+	static FAutoConsoleVariableRef AllowRemoteAttachedChildren(TEXT("DSTM.AllowRemoteAttachedChildren"), bAllowRemoteAttachedChildren, TEXT("Do we allow AttachChildren of a SceneComponent to exist only Remotely?"));
+#endif
+
 #if !(UE_BUILD_SHIPPING || UE_BUILD_TEST)
 	bool bWarnSceneComponentStaticMobilityMove = true;
 	static FAutoConsoleVariableRef CVarWarnSceneComponentStaticMobilityMove(
@@ -806,7 +815,17 @@ void USceneComponent::OnPreRegister()
 void USceneComponent::OnRegister()
 {
 	TRACE_CPUPROFILER_EVENT_SCOPE(USceneComponent::OnRegister)
+
+#if UE_WITH_REMOTE_OBJECT_HANDLE
+	// Only AttachToParent if we are not migrating
+	if (!FUObjectMigrationContext::GetCurrentMigrationContext())
+	{
+		AttachToParent();
+	}
+#else
 	AttachToParent();
+#endif
+
 	// Cache the level collection that contains the level in which this component is registered for fast access in IsVisible().
 	const UWorld* const World = GetWorld();
 	if (World)
@@ -3694,7 +3713,7 @@ void USceneComponent::Serialize(FArchive& Ar)
 
 #if UE_WITH_REMOTE_OBJECT_HANDLE
 	TArray<TObjectPtr<USceneComponent>> RestoreChildren;
-	if (Ar.IsMigratingRemoteObjects() && Ar.IsSaving())
+	if (Ar.IsMigratingRemoteObjects() && Ar.IsSaving() && !SceneComponentCVars::bAllowRemoteAttachedChildren)
 	{
 		for (int Index = AttachChildren.Num() - 1 ; Index >= 0 ; --Index)
 		{
@@ -3737,7 +3756,7 @@ void USceneComponent::Serialize(FArchive& Ar)
 		Ar << Bounds;
 
 		AttachChildren.Append(RestoreChildren);
-		if (Ar.IsLoading())
+		if (Ar.IsLoading() && !SceneComponentCVars::bAllowRemoteAttachedChildren)
 		{
 			for (TObjectPtr<USceneComponent>& Child : AttachChildren)
 			{
