IsOwned(Id) will now also return true for assets

diff --git a/Engine/Source/Runtime/CoreUObject/Private/UObject/RemoteObject.cpp b/Engine/Source/Runtime/CoreUObject/Private/UObject/RemoteObject.cpp
--- a/Engine/Source/Runtime/CoreUObject/Private/UObject/RemoteObject.cpp
+++ b/Engine/Source/Runtime/CoreUObject/Private/UObject/RemoteObject.cpp
@@ -344,21 +344,21 @@ bool IsOwned(FRemoteObjectId ObjectId)
 	using namespace UE::CoreUObject::Private;
 	using namespace UE::RemoteObject::Private;
 
-	bool Result = false;
+	bool bResult = true;
 #if UE_WITH_REMOTE_OBJECT_HANDLE
 	UE::RemoteObject::Handle::FRemoteObjectStub* RemoteStub = FindRemoteObjectStub(ObjectId);
 	if (RemoteStub)
 	{
-		Result = (RemoteStub->OwningServerId == GetGlobalServerId());
+		bResult = (RemoteStub->OwningServerId == GetGlobalServerId() || RemoteStub->OwningServerId.IsAsset());
 	}
 	else
 	{
 		const FRemoteServerId ServerId = ObjectId.GetServerId();
 		// Invalid server Id means local native objects which were created before the local server had a chance to have an id assigned
-		return !ServerId.IsValid() || ServerId == UE::RemoteObject::GetGlobalServerId();
+		bResult = (!ServerId.IsValid() || ServerId.IsAsset() || ServerId == UE::RemoteObject::GetGlobalServerId());
 	}
 #endif
-	return Result;
+	return bResult;
 }
 
 FRemoteServerId GetOwnerServerId(const UObject* Object)
