Exposing StructContainsLooseProperties to the pubic API.

diff --git a/Engine/Source/Runtime/CoreUObject/Internal/UObject/InstanceDataObjectUtils.h b/Engine/Source/Runtime/CoreUObject/Internal/UObject/InstanceDataObjectUtils.h
--- a/Engine/Source/Runtime/CoreUObject/Internal/UObject/InstanceDataObjectUtils.h
+++ b/Engine/Source/Runtime/CoreUObject/Internal/UObject/InstanceDataObjectUtils.h
@@ -36,7 +36,7 @@ UE_API bool IsInstanceDataObjectImpersonationEnabledOnSave();
 /** Query if InstanceDataObjects should be constructed using an archetype chain. */
 UE_API bool IsInstanceDataObjectArchetypeChainEnabled();
 
-bool StructContainsLooseProperties(const UStruct* Struct);
+UE_API bool StructContainsLooseProperties(const UStruct* Struct);
 UE_API bool StructIsInstanceDataObjectStruct(const UStruct* Struct);
 
 UE_API bool IsPropertyLoose(const FProperty* Property);
