Fixing a bug in the way localization.EnablePackageRemapping is implemented. It was previously refusing to set up the localization packages in the editor.

#rnx

diff --git a/Engine/Source/Runtime/CoreUObject/Private/Internationalization/PackageLocalizationManager.cpp b/Engine/Source/Runtime/CoreUObject/Private/Internationalization/PackageLocalizationManager.cpp
--- a/Engine/Source/Runtime/CoreUObject/Private/Internationalization/PackageLocalizationManager.cpp
+++ b/Engine/Source/Runtime/CoreUObject/Private/Internationalization/PackageLocalizationManager.cpp
@@ -22,7 +22,7 @@ namespace PackageLocalizationInternal
 
 	bool ShouldSkipPackageRemapping()
 	{
-		return GIsEditor || !GEnableLocalizationPackageRemapping;
+		return !GEnableLocalizationPackageRemapping;
 	}
 }
 
@@ -103,7 +103,7 @@ void FPackageLocalizationManager::InitializeFromCache(const TSharedRef<IPackageL
 	// Only preemptively attempt to conditionally update the cache outside of the editor where such things
 	// will happen almost immediately in a localized game, where as in the editor it's a bunch of work that
 	// likely won't be used until using some localization menus in the editor.
-	if (PackageLocalizationInternal::ShouldSkipPackageRemapping())
+	if (!GIsEditor && !PackageLocalizationInternal::ShouldSkipPackageRemapping())
 	{
 		ActiveCache->ConditionalUpdateCache();
 	}
