[Feature] added OverrrideEventHelper which is now the prefered way to perform operations on the overridable serialization system. This will invoke property change notifications which can be picked up by listeners like the IDO system and handled accordingly. This change also required various bug fixes and code hardening to support responding to the change notifications being sent.

diff --git a/Engine/Source/Runtime/CoreUObject/Private/UObject/Obj.cpp b/Engine/Source/Runtime/CoreUObject/Private/UObject/Obj.cpp
--- a/Engine/Source/Runtime/CoreUObject/Private/UObject/Obj.cpp
+++ b/Engine/Source/Runtime/CoreUObject/Private/UObject/Obj.cpp
@@ -520,7 +520,8 @@ void UObject::PreEditChange( FEditPropertyChain& PropertyAboutToChange )
 	}
 
 	// forward the notification to the FProperty* version of PreEditChange
-	PreEditChange(PropertyAboutToChange.GetActiveNode()->GetValue());
+	FProperty* PropertyBeingEdited = PropertyAboutToChange.GetActiveNode() ? PropertyAboutToChange.GetActiveNode()->GetValue() : nullptr;
+	PreEditChange(PropertyBeingEdited);
 
 	FCoreUObjectDelegates::OnPreObjectPropertyChanged.Broadcast(this, PropertyAboutToChange);
 
@@ -538,7 +539,8 @@ void UObject::PreEditChange( FEditPropertyChain& PropertyAboutToChange )
 
 void UObject::PostEditChangeChainProperty(FPropertyChangedChainEvent& PropertyChangedEvent)
 {
-	FPropertyChangedEvent PropertyEvent(PropertyChangedEvent.PropertyChain.GetActiveNode()->GetValue(), PropertyChangedEvent.ChangeType);
+	FProperty* PropertyBeingEdited = PropertyChangedEvent.PropertyChain.GetActiveNode() ? PropertyChangedEvent.PropertyChain.GetActiveNode()->GetValue() : nullptr;
+	FPropertyChangedEvent PropertyEvent(PropertyBeingEdited, PropertyChangedEvent.ChangeType);
 
 	// Set up array index per object map so that GetArrayIndex returns a valid result
 	TArray<TMap<FString, int32>> ArrayIndexForProperty;
diff --git a/Engine/Source/Runtime/CoreUObject/Private/UObject/OverriddenPropertySet.cpp b/Engine/Source/Runtime/CoreUObject/Private/UObject/OverriddenPropertySet.cpp
--- a/Engine/Source/Runtime/CoreUObject/Private/UObject/OverriddenPropertySet.cpp
+++ b/Engine/Source/Runtime/CoreUObject/Private/UObject/OverriddenPropertySet.cpp
@@ -596,6 +596,12 @@ EOverriddenPropertyOperation FOverriddenPropertySet::GetOverriddenPropertyOperat
 bool FOverriddenPropertySet::ClearOverriddenProperty(FOverriddenPropertyNode& ParentPropertyNode, FPropertyVisitorPath::Iterator PropertyIterator, const void* Data)
 {
 	FOverridableManager& OverridableManager = FOverridableManager::Get();
+	if (!PropertyIterator)
+	{
+		// if no property iterator is provided, clear all overrides
+		OverridableManager.ClearOverrides(Owner);
+		return true;
+	}
 
 	bool bClearedOverrides = false;
 	const void* SubValuePtr = Data;
@@ -802,6 +808,15 @@ void FOverriddenPropertySet::NotifyPropertyChange(FOverriddenPropertyNode* Paren
 {
 	checkf(IsValid(Owner), TEXT("Expecting a valid overridable owner"));
 
+	if (ChangeType == EPropertyChangeType::ResetToDefault)
+	{
+		if (Notification == EPropertyNotificationType::PostEdit)
+		{
+			this->ClearOverriddenProperty(*ParentPropertyNode, PropertyIterator, Data);
+		}
+		return;
+	}
+
 	FOverridableManager& OverridableManager = FOverridableManager::Get();
 	if (!PropertyIterator)
 	{
@@ -860,6 +875,52 @@ void FOverriddenPropertySet::NotifyPropertyChange(FOverriddenPropertyNode* Paren
 		}
 	};
 
+	// Because the PreEdit API doesn't provide Index info, we need to snapshot changed containers during pre-edit
+	// so we can intuit which element was removed during PostEdit. This is a map-like structure that stores the latest snapshot for each contianer property
+	static struct
+	{
+		uint8* Find(const FProperty* ContainerProperty) const
+		{
+			for (const TPair<const FProperty*, uint8*>& Element : Data)
+			{
+				if (Element.Key == ContainerProperty)
+				{
+					return Element.Value;
+				}
+			}
+			return nullptr;
+		}
+
+		void Free(const FProperty* ContainerProperty)
+		{
+			for (int32 I = 0; I < Data.Num(); ++I)
+			{
+				if (Data[I].Key == ContainerProperty)
+				{
+					ContainerProperty->DestroyValue(Data[I].Value);
+					FMemory::Free(Data[I].Value);
+					Data.RemoveAtSwap(I);
+					return;
+				}
+			}
+			checkf(false, TEXT("Expecting a matching property to the allocated memory"));
+		}
+
+		uint8* FindOrAdd(const FProperty* ContainerProperty)
+		{
+			if (uint8* Found = Find(ContainerProperty))
+			{
+				return Found;
+			}
+			int32 I = Data.Add({ ContainerProperty, (uint8*)FMemory::Malloc(ContainerProperty->GetSize(), ContainerProperty->GetMinAlignment()) });
+			ContainerProperty->InitializeValue(Data[I].Value);
+			return Data[I].Value;
+		}
+
+		// There's not many elements so we're using an array of TPairs for cache friendliness. A TMap would work fine here as well.
+		TArray<TPair<const FProperty*, uint8*>> Data;
+	} SavedPreEditContainers;
+
 	FPropertyVisitorPath::Iterator NextPropertyIterator = PropertyIterator+1;
 	if (const FArrayProperty* ArrayProperty = CastField<FArrayProperty>(Property))
 	{
@@ -874,18 +935,19 @@ void FOverriddenPropertySet::NotifyPropertyChange(FOverriddenPropertyNode* Paren
 			{
 				checkf(ArrayProperty->Inner, TEXT("Expecting an inner type for Arrays"));
 
-				static TArray<uint8> SavedPreEditSubObjects;
-				FScriptArrayHelper PreEditSubObjectsArrayHelper(ArrayProperty, &SavedPreEditSubObjects);
 				if (Notification == EPropertyNotificationType::PreEdit)
 				{
-					PreEditSubObjectsArrayHelper.EmptyAndAddValues(ArrayHelper.Num());
-					for(int32 i = 0; i < ArrayHelper.Num(); i++)
+					FScriptArrayHelper PreEditArrayHelper(ArrayProperty, SavedPreEditContainers.FindOrAdd(ArrayProperty));
+					PreEditArrayHelper.EmptyAndAddValues(ArrayHelper.Num());
+					for (int32 i = 0; i < ArrayHelper.Num(); i++)
 					{
-						InnerObjectProperty->SetObjectPropertyValue(PreEditSubObjectsArrayHelper.GetElementPtr(i), InnerObjectProperty->GetObjectPropertyValue(ArrayHelper.GetElementPtr(i)));
+						InnerObjectProperty->SetObjectPropertyValue(PreEditArrayHelper.GetElementPtr(i), InnerObjectProperty->GetObjectPropertyValue(ArrayHelper.GetElementPtr(i)));
 					}
 					return;
 				}
 
+				FScriptArrayHelper PreEditArrayHelper(ArrayProperty, SavedPreEditContainers.Find(ArrayProperty));
+
 				auto ArrayReplace = [&]
 				{
 					if (SubPropertyNode)
@@ -934,8 +996,8 @@ void FOverriddenPropertySet::NotifyPropertyChange(FOverriddenPropertyNode* Paren
 
 				auto ArrayRemoveImpl = [&]()
 				{
-					checkf(PreEditSubObjectsArrayHelper.IsValidIndex(ArrayIndex), TEXT("ArrayRemove change type expected to have an valid index"));
-					if (UObject* RemovedSubObject = InnerObjectProperty->GetObjectPropertyValue(PreEditSubObjectsArrayHelper.GetElementPtr(ArrayIndex)))
+					checkf(PreEditArrayHelper.IsValidIndex(ArrayIndex), TEXT("ArrayRemove change type expected to have an valid index"));
+					if (UObject* RemovedSubObject = InnerObjectProperty->GetObjectPropertyValue(PreEditArrayHelper.GetElementPtr(ArrayIndex)))
 					{
 						if(SubPropertyNode)
 						{
@@ -1045,7 +1107,7 @@ void FOverriddenPropertySet::NotifyPropertyChange(FOverriddenPropertyNode* Paren
 				{
 					checkf(ArrayIndex == INDEX_NONE, TEXT("ArrayClear change type should not have associated indexes"));
 
-					for (int i = 0; i < PreEditSubObjectsArrayHelper.Num(); ++i)
+					for (int i = 0; i < PreEditArrayHelper.Num(); ++i)
 					{
 						ArrayIndex = i;
 						ArrayRemoveImpl();
@@ -1095,30 +1157,9 @@ void FOverriddenPropertySet::NotifyPropertyChange(FOverriddenPropertyNode* Paren
 		int32 InternalMapIndex = LogicalMapIndex != INDEX_NONE ? MapHelper.FindInternalIndex(LogicalMapIndex) : INDEX_NONE;
 		if (!NextPropertyIterator)
 		{
-			static const FProperty* SavedProp = nullptr;
-			static uint8* SavedPreEditMap = nullptr;
-
-			auto FreePreEditMap = []()
-			{
-				if (SavedPreEditMap)
-				{
-					checkf(SavedProp, TEXT("Expecting a matching property to the allocated memory"));
-					SavedProp->DestroyValue(SavedPreEditMap);
-					FMemory::Free(SavedPreEditMap);
-					SavedPreEditMap = nullptr;
-					SavedProp = nullptr;
-				}
-			};
-
 			if (Notification == EPropertyNotificationType::PreEdit)
 			{
-				FreePreEditMap();
-
-				SavedPreEditMap = (uint8*)FMemory::Malloc(MapProperty->GetSize(), MapProperty->GetMinAlignment());
-				MapProperty->InitializeValue(SavedPreEditMap);
-				SavedProp = MapProperty;
-
-				FScriptMapHelper PreEditMapHelper(MapProperty, SavedPreEditMap);
+				FScriptMapHelper PreEditMapHelper(MapProperty, SavedPreEditContainers.FindOrAdd(MapProperty));
 				PreEditMapHelper.EmptyValues();
 				for (FScriptMapHelper::FIterator It(MapHelper); It; ++It)
 				{
@@ -1127,14 +1168,15 @@ void FOverriddenPropertySet::NotifyPropertyChange(FOverriddenPropertyNode* Paren
 				return;
 			}
 
-			checkf(SavedProp == MapProperty, TEXT("Expecting the same property as the pre edit flow"));
+			uint8* SavedPreEditMap = SavedPreEditContainers.Find(MapProperty);
+			checkf(SavedPreEditMap, TEXT("Expecting the same property as the pre edit flow"));
 			FScriptMapHelper PreEditMapHelper(MapProperty, SavedPreEditMap);
 			// The logical should map directly to the pre edit map internal index as we skipped all of the invalid entries
 			int32 InternalPreEditMapIndex = LogicalMapIndex;
 
 			ON_SCOPE_EXIT
 			{
-				FreePreEditMap();
+				SavedPreEditContainers.Free(MapProperty);
 			};
 
 			auto MapReplace = [&]()
@@ -1384,11 +1426,22 @@ UObject* FOverriddenPropertySet::TryGetInstancedSubObjectValue(const FObjectProp
 		return nullptr;
 	}
 
+	// subobject pointers in IDOs point to the instance subobjects. For this purpose we need to redirect them to IDO subobjects
 	UObject* SubObject = FromProperty->GetObjectPropertyValue(ValuePtr);
+	const UObject* ExpectedOuter = Owner;
+	TFunction<UObject* (UObject*)> RedirectMethod = [](UObject* Obj) {return Obj; };
+#if WITH_EDITORONLY_DATA
+	if (const UObject* Instance = UE::FPropertyBagRepository::Get().FindInstanceForDataObject(Owner))
+	{
+		RedirectMethod = UE::ResolveInstanceDataObject;
+		ExpectedOuter = Instance;
+	}
+#endif
+
 	if (FromProperty->HasAnyPropertyFlags(CPF_PersistentInstance)
-		|| (FromProperty->IsA<FObjectProperty>() && SubObject && SubObject->IsIn(Owner)))
+		|| (FromProperty->IsA<FObjectProperty>() && SubObject && SubObject->IsIn(ExpectedOuter)))
 	{
-		return SubObject;
+		return RedirectMethod(SubObject);
 	}
 
 	return nullptr;
diff --git a/Engine/Source/Runtime/CoreUObject/Private/UObject/OverrideEventHelper.cpp b/Engine/Source/Runtime/CoreUObject/Private/UObject/OverrideEventHelper.cpp
new file mode 100644
--- /dev/null
+++ b/Engine/Source/Runtime/CoreUObject/Private/UObject/OverrideEventHelper.cpp
@@ -0,0 +1,147 @@
+#include "UObject/OverrideEventHelper.h"
+#include "UObject/OverridableManager.h"
+
+
+namespace UE
+{
+	// helper class that constructs a change event using FPropertyVisitorPaths
+	class ProxyChangeEvent : public FEditPropertyChain, public FPropertyChangedChainEvent
+	{
+	public:
+		ProxyChangeEvent(TNotNull<UObject*> Object, const FPropertyVisitorPath& PropertyPath, EPropertyChangeType::Type ChangeType)
+			: FPropertyChangedChainEvent(*this, CreateChangeEvent(Object, PropertyPath, ChangeType))
+		{
+			ArrayIndicesPerObject.AddDefaulted(ObjectIteratorIndex + 1);
+			for (const FPropertyVisitorInfo& Info : PropertyPath.GetPath())
+			{
+				FEditPropertyChain::AddTail(const_cast<FProperty*>(Info.Property));
+				ArrayIndicesPerObject[ObjectIteratorIndex].Add(Info.Property->GetName(), Info.Index);
+			}
+			FPropertyChangedChainEvent::SetArrayIndexPerObject(ProxyChangeEvent::ArrayIndicesPerObject);
+		}
+
+		ProxyChangeEvent(TNotNull<UObject*> Object, EPropertyChangeType::Type ChangeType)
+			: FPropertyChangedChainEvent(*this, CreateChangeEvent(Object, {}, ChangeType))
+		{}
+
+		ProxyChangeEvent(TNotNull<UObject*> Object, const FPropertyChangedEvent& PropertyEvent, const FEditPropertyChain& PropertyChain, EPropertyChangeType::Type ChangeType)
+			: FPropertyChangedChainEvent(*this, const_cast<FPropertyChangedEvent&>(PropertyEvent))
+		{
+			FPropertyChangedChainEvent::ChangeType = ChangeType;
+			// The copy constructor of PropertyChain was deleted so copy it manually here.
+			for (FEditPropertyChain::TDoubleLinkedListNode* Itr = PropertyChain.GetHead(); Itr && Itr->GetValue(); Itr = Itr->GetNextNode())
+			{
+				FEditPropertyChain::AddTail(Itr->GetValue());
+				if (PropertyChain.GetActiveMemberNode() == Itr)
+				{
+					FEditPropertyChain::ActiveMemberPropertyNode = FEditPropertyChain::GetTail();
+				}
+				if (PropertyChain.GetActiveNode() == Itr)
+				{
+					FEditPropertyChain::ActivePropertyNode = FEditPropertyChain::GetTail();
+				}
+			}
+		}
+
+	private:
+		static FPropertyChangedEvent CreateChangeEvent(TNotNull<UObject*> Object, const FPropertyVisitorPath& PropertyPath, EPropertyChangeType::Type ChangeType)
+		{
+			FProperty* ChangedProperty = PropertyPath.Num() > 0 ? const_cast<FProperty*>(PropertyPath.Top().Property) : nullptr;
+			FPropertyChangedEvent Result(ChangedProperty, ChangeType, { Object });
+			Result.ObjectIteratorIndex = 0;
+			return Result;
+		}
+
+		TArray<TMap<FString, int32>> ArrayIndicesPerObject;
+	};
+
+	void SendOverrideObjectEvent(TNotNull<UObject*> Object)
+	{
+#if UE_EDITOR
+		// an event with no provided property path will override the entire object
+		ProxyChangeEvent ChangeEvent(Object, EPropertyChangeType::Unspecified);
+		Object->PreEditChange(ChangeEvent);
+		Object->PostEditChangeChainProperty(ChangeEvent);
+#endif
+	}
+
+	void SendClearOverridesEvent(TNotNull<UObject*> Object)
+	{
+#if UE_EDITOR
+		if (FOverridableManager::Get().IsEnabled(Object))
+		{
+			// a ResetToDefault event with no provided property path will clear all of the object's overrides
+			ProxyChangeEvent ChangeEvent(Object, EPropertyChangeType::ResetToDefault);
+			Object->PreEditChange(ChangeEvent);
+			Object->PostEditChangeChainProperty(ChangeEvent);
+		}
+#endif
+	}
+
+	void SendOverridePropertyEvent(TNotNull<UObject*> Object, const FPropertyVisitorPath& PropertyPath, EPropertyChangeType::Type ChangeType)
+	{
+#if UE_EDITOR
+		ProxyChangeEvent ChangeEvent(Object, PropertyPath, ChangeType);
+		Object->PreEditChange(ChangeEvent);
+		Object->PostEditChangeChainProperty(ChangeEvent);
+#endif
+	}
+
+	void SendOverridePropertyEvent(TNotNull<UObject*> Object, const FPropertyChangedEvent& PropertyEvent, const FEditPropertyChain& PropertyChain)
+	{
+#if UE_EDITOR
+		ProxyChangeEvent ChangeEvent(Object, PropertyEvent, PropertyChain, PropertyEvent.ChangeType);
+		Object->PreEditChange(ChangeEvent);
+		Object->PostEditChangeChainProperty(ChangeEvent);
+#endif
+	}
+
+	void SendClearOverriddenPropertyEvent(TNotNull<UObject*> Object, const FPropertyVisitorPath& PropertyPath)
+	{
+#if UE_EDITOR
+		ProxyChangeEvent ChangeEvent(Object, PropertyPath, EPropertyChangeType::ResetToDefault);
+		Object->PreEditChange(ChangeEvent);
+		Object->PostEditChangeChainProperty(ChangeEvent);
+#endif
+	}
+
+	void SendClearOverriddenPropertyEvent(TNotNull<UObject*> Object, const FPropertyChangedEvent& PropertyEvent, const FEditPropertyChain& PropertyChain)
+	{
+#if UE_EDITOR
+		ProxyChangeEvent ChangeEvent(Object, PropertyEvent, PropertyChain, EPropertyChangeType::ResetToDefault);
+		Object->PreEditChange(ChangeEvent);
+		Object->PostEditChangeChainProperty(ChangeEvent);
+#endif
+	}
+
+	void SendPreOverridePropertyEvent(TNotNull<UObject*> Object, const FPropertyVisitorPath& PropertyPath)
+	{
+#if UE_EDITOR
+		ProxyChangeEvent ChangeEvent(Object, PropertyPath, EPropertyChangeType::Unspecified);
+		Object->PreEditChange(ChangeEvent);
+#endif
+	}
+
+	void SendPreOverridePropertyEvent(TNotNull<UObject*> Object, const FEditPropertyChain& PropertyChain)
+	{
+#if UE_EDITOR
+		Object->PreEditChange(const_cast<FEditPropertyChain&>(PropertyChain));
+#endif
+	}
+
+	void SendPostOverridePropertyEvent(TNotNull<UObject*> Object, const FPropertyVisitorPath& PropertyPath, const EPropertyChangeType::Type ChangeType)
+	{
+#if UE_EDITOR
+		ProxyChangeEvent ChangeEvent(Object, PropertyPath, ChangeType);
+		Object->PostEditChangeChainProperty(ChangeEvent);
+#endif
+	}
+
+	void SendPostOverridePropertyEvent(TNotNull<UObject*> Object, const FPropertyChangedEvent& PropertyEvent, const FEditPropertyChain& PropertyChain)
+	{
+#if UE_EDITOR
+		ProxyChangeEvent ChangeEvent(Object, PropertyEvent, PropertyChain, PropertyEvent.ChangeType);
+		Object->PostEditChangeChainProperty(ChangeEvent);
+#endif
+	}
+}
diff --git a/Engine/Source/Runtime/CoreUObject/Private/UObject/PropertyBagRepository.cpp b/Engine/Source/Runtime/CoreUObject/Private/UObject/PropertyBagRepository.cpp
--- a/Engine/Source/Runtime/CoreUObject/Private/UObject/PropertyBagRepository.cpp
+++ b/Engine/Source/Runtime/CoreUObject/Private/UObject/PropertyBagRepository.cpp
@@ -350,11 +350,11 @@ public:
 				check(SourceChainNode);
 				
 				SourceField = SourceChainNode->GetValue();
-				SourceMemoryPtr = SourceField.Get<FProperty>()->ContainerPtrToValuePtr<void>(SourceMemoryPtr);
+				SourceMemoryPtr = SourceMemoryPtr ? SourceField.Get<FProperty>()->ContainerPtrToValuePtr<void>(SourceMemoryPtr) : nullptr;
 				
 				ResolvedChainNode = ResolveCurChainNode(ResolvedStruct, SourceField.Get<FProperty>());
 				ResolvedField = ResolvedChainNode->GetValue();
-				ResolvedMemoryPtr = ResolvedField.Get<FProperty>()->ContainerPtrToValuePtr<void>(ResolvedMemoryPtr);
+				ResolvedMemoryPtr = ResolvedMemoryPtr ? ResolvedField.Get<FProperty>()->ContainerPtrToValuePtr<void>(ResolvedMemoryPtr) : nullptr;
 			}
 			else if (UObject* ResolvedObject = ResolvedField.Get<UObject>())
 			{
@@ -375,15 +375,21 @@ public:
 				{
 					FScriptArrayHelper SourceArrayHelper(SourceArrayProperty, SourceMemoryPtr);
 	                FScriptArrayHelper ResolvedArrayHelper(ResolvedArrayProperty, ResolvedMemoryPtr);
-	                if(!ResolvedArrayHelper.IsValidIndex(ArrayIndex))
+	                if(ResolvedMemoryPtr && !ResolvedArrayHelper.IsValidIndex(ArrayIndex))
 	                {
-                		check(ChangeType & EPropertyChangeType::ArrayAdd);
-                		check(ResolvedArrayHelper.Num() == ArrayIndex);
-                		ResolvedArrayHelper.Resize(ArrayIndex + 1);
+						if (ChangeType & EPropertyChangeType::ArrayAdd)
+						{
+							check(ResolvedArrayHelper.Num() == ArrayIndex);
+							ResolvedArrayHelper.Resize(ArrayIndex + 1);
+						}
+						else
+						{
+							ResolvedMemoryPtr = nullptr;
+						}
 	                }
 	                ResolvedArrayIndices[SourceEvent.ObjectIteratorIndex].Add(ResolvedArrayProperty->GetName(), ArrayIndex);
-	                SourceMemoryPtr = SourceArrayHelper.GetRawPtr(ArrayIndex);
-	                ResolvedMemoryPtr = ResolvedArrayHelper.GetRawPtr(ArrayIndex);
+	                SourceMemoryPtr = SourceMemoryPtr ? SourceArrayHelper.GetRawPtr(ArrayIndex) : nullptr;
+	                ResolvedMemoryPtr = ResolvedMemoryPtr ? ResolvedArrayHelper.GetRawPtr(ArrayIndex) : nullptr;
 	                SourceField = SourceArrayProperty->Inner;
 	                ResolvedField = ResolvedArrayProperty->Inner;
 				}
@@ -403,16 +409,22 @@ public:
                 {
                 	FScriptSetHelper SourceSetHelper(SourceSetProperty, SourceMemoryPtr);
                 	FScriptSetHelper ResolvedSetHelper(ResolvedSetProperty, ResolvedMemoryPtr);
-                	if(!ResolvedSetHelper.IsValidIndex(SetIndex))
+                	if(ResolvedMemoryPtr && !ResolvedSetHelper.IsValidIndex(SetIndex))
                 	{
-                		check(ChangeType & EPropertyChangeType::ArrayAdd);
-                		check(ResolvedSetHelper.Num() == SetIndex);
-                		int32 AddedIndex = ResolvedSetHelper.AddUninitializedValue();
-                		check(AddedIndex == SetIndex);
+						if (ChangeType & EPropertyChangeType::ArrayAdd)
+						{
+							check(ResolvedSetHelper.Num() == SetIndex);
+							int32 AddedIndex = ResolvedSetHelper.AddUninitializedValue();
+							check(AddedIndex == SetIndex);
+						}
+						else
+						{
+							ResolvedMemoryPtr = nullptr;
+						}
                 	}
                     ResolvedArrayIndices[SourceEvent.ObjectIteratorIndex].Add(ResolvedSetProperty->GetName(), SetIndex);
-                	SourceMemoryPtr = SourceSetHelper.GetElementPtr(SetIndex);
-                	ResolvedMemoryPtr = ResolvedSetHelper.GetElementPtr(SetIndex);
+                	SourceMemoryPtr = SourceMemoryPtr ? SourceSetHelper.GetElementPtr(SetIndex) : nullptr;
+                	ResolvedMemoryPtr = ResolvedMemoryPtr ? ResolvedSetHelper.GetElementPtr(SetIndex) : nullptr;
                 	SourceField = SourceSetProperty->ElementProp;
                 	ResolvedField = ResolvedSetProperty->ElementProp;
                 }
@@ -432,49 +444,24 @@ public:
 				{
 					FScriptMapHelper SourceMapHelper(SourceMapProperty, SourceMemoryPtr);
 					FScriptMapHelper ResolvedMapHelper(ResolvedMapProperty, ResolvedMemoryPtr);
-					if(!ResolvedMapHelper.IsValidIndex(MapIndex))
+					if(ResolvedMemoryPtr && !ResolvedMapHelper.IsValidIndex(MapIndex))
 					{
-						check(ChangeType & EPropertyChangeType::ArrayAdd);
-						check(ResolvedMapHelper.Num() == MapIndex);
-						int32 AddedIndex = ResolvedMapHelper.AddUninitializedValue();
-						check(AddedIndex == MapIndex);
-					}
-	                ResolvedArrayIndices[SourceEvent.ObjectIteratorIndex].Add(ResolvedMapProperty->GetName(), MapIndex);
-
-					// try to figure out whether we should be traversing into the key or the value
-					auto TryEnterMapWithHint = [&](const FProperty* Hint)
-					{
-						if (Hint == SourceMapProperty->KeyProp)
+						if (ChangeType & EPropertyChangeType::ArrayAdd)
 						{
-							SourceMemoryPtr = SourceMapHelper.GetKeyPtr(MapIndex);
-							ResolvedMemoryPtr = ResolvedMapHelper.GetKeyPtr(MapIndex);
-							SourceField = SourceMapProperty->KeyProp;
-							ResolvedField = ResolvedMapProperty->KeyProp;
-							return true;
+							check(ResolvedMapHelper.Num() == MapIndex);
+							int32 AddedIndex = ResolvedMapHelper.AddUninitializedValue();
+							check(AddedIndex == MapIndex);
 						}
-						else if (Hint == SourceMapProperty->ValueProp)
+						else
 						{
-							SourceMemoryPtr = SourceMapHelper.GetValuePtr(MapIndex);
-							ResolvedMemoryPtr = ResolvedMapHelper.GetValuePtr(MapIndex);
-							SourceField = SourceMapProperty->ValueProp;
-							ResolvedField = ResolvedMapProperty->ValueProp;
-							return true;
+							ResolvedMemoryPtr = nullptr;
 						}
-						return false;
-					};
-
-					// this process isn't well-defined we need to check several places for the property to see if it's a key or a value prop
-					if (FEditPropertyChain::TDoubleLinkedListNode* NextSourceChainNode = SourceChainNode->GetNextNode())
-					{
-						bool bSuccess = TryEnterMapWithHint(NextSourceChainNode->GetValue());
-						check(bSuccess);
-					}
-					else if (TryEnterMapWithHint(SourceEvent.Property)) {}
-					else if (TryEnterMapWithHint(SourceEvent.MemberProperty)) {}
-					else
-					{
-						checkf(false, TEXT("An index was provided for the map but it's ambiguous whether the key or value was modified"));
 					}
+	                ResolvedArrayIndices[SourceEvent.ObjectIteratorIndex].Add(ResolvedMapProperty->GetName(), MapIndex);
+					SourceMemoryPtr = SourceMemoryPtr ? SourceMapHelper.GetValuePtr(MapIndex) : nullptr;
+					ResolvedMemoryPtr = ResolvedMemoryPtr ? ResolvedMapHelper.GetValuePtr(MapIndex) : nullptr;
+					SourceField = SourceMapProperty->ValueProp;
+					ResolvedField = ResolvedMapProperty->ValueProp;
 				}
 			}
 			else if (FStructProperty* ResolvedStructProperty = ResolvedField.Get<FStructProperty>())
@@ -485,33 +472,40 @@ public:
 			}
 			else if (FObjectProperty* ResolvedObjectProperty = ResolvedField.Get<FObjectProperty>())
 			{
-				UObject* ResolvedSubObject = GetObjectRefFromProperty(ResolvedObjectProperty, ResolvedMemoryPtr, ResolvedObject);
-				ResolvedField = ResolvedSubObject->GetClass();
+				UObject* ResolvedSubObject = GetObjectRefFromProperty(ResolvedObjectProperty, ResolvedMemoryPtr, ObjectToResolve);
+				ResolvedField = ResolvedSubObject ? ResolvedSubObject->GetClass() : ResolvedObjectProperty->PropertyClass.Get();
 				ResolvedMemoryPtr = ResolvedSubObject;
 				
 				FObjectProperty* SourceObjectProperty = SourceField.Get<FObjectProperty>();
 				UObject* SourceSubObject = GetObjectRefFromProperty(SourceObjectProperty, SourceMemoryPtr, SourceObject);
-				SourceField = SourceSubObject->GetClass();
+				SourceField = SourceSubObject ? SourceSubObject->GetClass() : SourceObjectProperty->PropertyClass.Get();
 				SourceMemoryPtr = SourceSubObject;
 			}
 			else if (FOptionalProperty* ResolvedOptionalProperty = ResolvedField.Get<FOptionalProperty>())
 			{
 				FOptionalProperty* SourceOptionalProperty = SourceField.Get<FOptionalProperty>();
 				
-				check(SourceOptionalProperty->IsSet(SourceMemoryPtr))
 				SourceField = SourceOptionalProperty->GetValueProperty();
-				SourceMemoryPtr = SourceOptionalProperty->GetValuePointerForReadOrReplace(SourceMemoryPtr);
+            	ResolvedField = ResolvedOptionalProperty->GetValueProperty();
+				if (SourceMemoryPtr)
+				{
+					if (SourceOptionalProperty->IsSet(SourceMemoryPtr))
+					{
+						SourceMemoryPtr = SourceOptionalProperty->GetValuePointerForReadOrReplace(SourceMemoryPtr);
+					}
 				
-	            if (!ResolvedOptionalProperty->IsSet(ResolvedMemoryPtr))
-	            {
-            		ResolvedField = ResolvedOptionalProperty->GetValueProperty();
-            		ResolvedMemoryPtr = ResolvedOptionalProperty->MarkSetAndGetInitializedValuePointerToReplace(ResolvedMemoryPtr);
-	            }
-	            else
-	            {
-            		ResolvedField = ResolvedOptionalProperty->GetValueProperty();
-            		ResolvedMemoryPtr = ResolvedOptionalProperty->GetValuePointerForReadOrReplace(ResolvedMemoryPtr);
-	            }
+					if (ResolvedMemoryPtr)
+					{
+						if (!ResolvedOptionalProperty->IsSet(ResolvedMemoryPtr))
+						{
+							ResolvedMemoryPtr = ResolvedOptionalProperty->MarkSetAndGetInitializedValuePointerToReplace(ResolvedMemoryPtr);
+						}
+						else
+						{
+							ResolvedMemoryPtr = ResolvedOptionalProperty->GetValuePointerForReadOrReplace(ResolvedMemoryPtr);
+						}
+					}
+				}
 			}
 			else
 			{
@@ -546,10 +540,6 @@ public:
 		{
 			ResolvedArrayIndices[SourceEvent.ObjectIteratorIndex].Add(ResolvedField.GetName(), AddChangeIndex);
 		}
-		
-		// make sure the Property and MemberProperty members were remapped as well
-		check(SourceEvent.Property != Property);
-		check(SourceEvent.MemberProperty != MemberProperty);
 	}
 
 	bool ConstructionComplete() const
@@ -564,7 +554,7 @@ public:
 		}
 		check(SourceField.IsValid());
 		check(ResolvedField.IsValid());
-		return false;
+		return SourceEvent.PropertyChain.GetActiveMemberNode() == nullptr;
 	}
 
 	FEditPropertyChain::TDoubleLinkedListNode* ResolveCurChainNode(const UStruct* ResolvedStruct, const FProperty* SourceProperty)
@@ -578,6 +568,10 @@ public:
 
 	static UObject* GetObjectRefFromProperty(const FObjectProperty* Property, const void* Memory, const UObject* OwningObject)
 	{
+		if (!Memory)
+		{
+			return nullptr;
+		}
 		UObject* SubObject = Property->GetObjectPropertyValue(Memory);
 		if (SubObject && UE::IsClassOfInstanceDataObjectClass(OwningObject->GetClass()))
 		{
@@ -621,11 +615,17 @@ void FPropertyBagRepository::PostEditChangeChainProperty(const UObject* Object,
 		
 		Dest->PreEditChange(RemappedChainEvent.PropertyChain);
 		
-		CopyProperty(
-			RemappedChainEvent.SourceField.Get<FProperty>(),
-			RemappedChainEvent.SourceMemoryPtr,
-			RemappedChainEvent.ResolvedField.Get<FProperty>(),
-			RemappedChainEvent.ResolvedMemoryPtr);
+		FProperty* SourceProperty = RemappedChainEvent.SourceField.Get<FProperty>();
+		FProperty* DestProperty = RemappedChainEvent.ResolvedField.Get<FProperty>();
+		if (SourceProperty && DestProperty && RemappedChainEvent.SourceMemoryPtr && RemappedChainEvent.ResolvedMemoryPtr)
+		{
+			CopyProperty(
+				SourceProperty,
+				RemappedChainEvent.SourceMemoryPtr,
+				DestProperty,
+				RemappedChainEvent.ResolvedMemoryPtr);
+		}
+
             
 		Dest->PostEditChangeChainProperty(RemappedChainEvent);
 	};
diff --git a/Engine/Source/Runtime/CoreUObject/Private/UObject/UObjectGlobals.cpp b/Engine/Source/Runtime/CoreUObject/Private/UObject/UObjectGlobals.cpp
--- a/Engine/Source/Runtime/CoreUObject/Private/UObject/UObjectGlobals.cpp
+++ b/Engine/Source/Runtime/CoreUObject/Private/UObject/UObjectGlobals.cpp
@@ -3341,7 +3341,7 @@ UObject* StaticDuplicateObjectEx( FObjectDuplicationParameters& Parameters )
 
 #if WITH_EDITORONLY_DATA
 	// if Source has an IDO, make one for Dest and copy it
-	if (UE::IsInstanceDataObjectSupportEnabled(Parameters.SourceObject))
+	if (UE::IsInstanceDataObjectSupportEnabled(Parameters.SourceObject) && UE::IsInstanceDataObjectImpersonationEnabledOnSave())
 	{
 		UE::FPropertyBagRepository& PropertyBagRepository = UE::FPropertyBagRepository::Get();
 		PropertyBagRepository.DuplicateInstanceDataObject(Parameters.SourceObject, DupRootObject);
diff --git a/Engine/Source/Runtime/CoreUObject/Public/UObject/OverrideEventHelper.h b/Engine/Source/Runtime/CoreUObject/Public/UObject/OverrideEventHelper.h
new file mode 100644
--- /dev/null
+++ b/Engine/Source/Runtime/CoreUObject/Public/UObject/OverrideEventHelper.h
@@ -0,0 +1,57 @@
+// Copyright Epic Games, Inc. All Rights Reserved.
+
+#pragma once
+#include "PropertyVisitor.h"
+#include "UnrealType.h"
+#include "Object.h"
+
+namespace UE
+{
+	/**
+	 * Helper method that calls the pre and post change notification responsible for overriding an object.
+	 * This is the preferred way to override because it will notify listeners
+	 * @param Object to override all it's properties */
+	COREUOBJECT_API void SendOverrideObjectEvent(TNotNull<UObject*> Object);
+
+	/**
+	 * Helper method that calls the pre and post change notification responsible for clearing the overrides of an object.
+	 * This is the preferred way to clear overrides because it will notify listeners
+	 * @param Object to clear overrides on */
+	COREUOBJECT_API void SendClearOverridesEvent(TNotNull<UObject*> Object);
+
+	/**
+	 * Helper method that calls the pre and post change notification responsible for overriding an object property.
+	 * This is the preferred way to override properties because it will notify listeners
+	 * Note: If you're also changing property values (especially if it's in a container) use SendPreOverridePropertyEvent and SendPostOverridePropertyEvent instead!
+	 * @param Object owning the property
+	 * @param PropertyPath leading to the property that is about to be overridden
+	 * @param ChangeType of the current operation */
+	COREUOBJECT_API void SendOverridePropertyEvent(TNotNull<UObject*> Object, const FPropertyVisitorPath& PropertyPath, EPropertyChangeType::Type ChangeType = EPropertyChangeType::Unspecified);
+	COREUOBJECT_API void SendOverridePropertyEvent(TNotNull<UObject*> Object, const FPropertyChangedEvent& PropertyEvent, const FEditPropertyChain& PropertyChain);
+
+	/**
+	 * Helper method that calls the pre and post change notification responsible for clearing a property override.
+	 * This is the preferred way to clear property overrides because it will notify listeners
+	 * @param Object owning the property to clear
+	 * @param PropertyPath to the property to clear from the root of the specified object
+	 * @return true if the property was successfully cleared. */
+	COREUOBJECT_API void SendClearOverriddenPropertyEvent(TNotNull<UObject*> Object, const FPropertyVisitorPath& PropertyPath);
+	COREUOBJECT_API void SendClearOverriddenPropertyEvent(TNotNull<UObject*> Object, const FPropertyChangedEvent& PropertyEvent, const FEditPropertyChain& PropertyChain);
+
+	/**
+	 * Helper method that calls the pre-change notification responsible for overriding an object property.
+	 * This is the preferred way to override properties because it will notify listeners
+	 * @param Object owning the property
+	 * @param PropertyPath leading to the property about to be overridden */
+	COREUOBJECT_API void SendPreOverridePropertyEvent(TNotNull<UObject*> Object, const FPropertyVisitorPath& PropertyPath);
+	COREUOBJECT_API void SendPreOverridePropertyEvent(TNotNull<UObject*> Object, const FEditPropertyChain& PropertyChain);
+
+	/**
+	 * Helper method that calls the post-change notification responsible for overriding an object property.
+	 * This is the preferred way to override properties because it will notify listeners
+	 * @param Object owning the property
+	 * @param PropertyPath leading to the property that was overridden
+	 * @param ChangeType of the current operation */
+	COREUOBJECT_API void SendPostOverridePropertyEvent(TNotNull<UObject*> Object, const FPropertyVisitorPath& PropertyPath, EPropertyChangeType::Type ChangeType = EPropertyChangeType::Unspecified);
+	COREUOBJECT_API void SendPostOverridePropertyEvent(TNotNull<UObject*> Object, const FPropertyChangedEvent& PropertyEvent, const FEditPropertyChain& PropertyChain);
+}
diff --git a/Engine/Source/Runtime/CoreUObject/Public/UObject/UnrealType.h b/Engine/Source/Runtime/CoreUObject/Public/UObject/UnrealType.h
--- a/Engine/Source/Runtime/CoreUObject/Public/UObject/UnrealType.h
+++ b/Engine/Source/Runtime/CoreUObject/Public/UObject/UnrealType.h
@@ -6574,7 +6574,7 @@ private:
  */
 struct FPropertyChangedChainEvent : public FPropertyChangedEvent
 {
-	FPropertyChangedChainEvent(FEditPropertyChain& InPropertyChain, FPropertyChangedEvent& SrcChangeEvent) :
+	FPropertyChangedChainEvent(FEditPropertyChain& InPropertyChain, const FPropertyChangedEvent& SrcChangeEvent) :
 		FPropertyChangedEvent(SrcChangeEvent),
 		PropertyChain(InPropertyChain)
 	{
