[CISFix] Static analysis caught a missing null check. This fixes that
#pf skip

diff --git a/Engine/Source/Runtime/CoreUObject/Private/UObject/OverriddenPropertySet.cpp b/Engine/Source/Runtime/CoreUObject/Private/UObject/OverriddenPropertySet.cpp
--- a/Engine/Source/Runtime/CoreUObject/Private/UObject/OverriddenPropertySet.cpp
+++ b/Engine/Source/Runtime/CoreUObject/Private/UObject/OverriddenPropertySet.cpp
@@ -810,7 +810,7 @@ void FOverriddenPropertySet::NotifyPropertyChange(FOverriddenPropertyNode* Paren
 
 	if (ChangeType == EPropertyChangeType::ResetToDefault)
 	{
-		if (Notification == EPropertyNotificationType::PostEdit)
+		if (ParentPropertyNode && Notification == EPropertyNotificationType::PostEdit)
 		{
 			this->ClearOverriddenProperty(*ParentPropertyNode, PropertyIterator, Data);
 		}
