[CISFix] Added missing copyright boilerplate

diff --git a/Engine/Source/Runtime/CoreUObject/Private/UObject/OverrideEventHelper.cpp b/Engine/Source/Runtime/CoreUObject/Private/UObject/OverrideEventHelper.cpp
--- a/Engine/Source/Runtime/CoreUObject/Private/UObject/OverrideEventHelper.cpp
+++ b/Engine/Source/Runtime/CoreUObject/Private/UObject/OverrideEventHelper.cpp
@@ -1,3 +1,5 @@
+// Copyright Epic Games, Inc. All Rights Reserved.
+
 #include "UObject/OverrideEventHelper.h"
 #include "UObject/OverridableManager.h"
 
