Fix package FMetaData object included in cooked builds platforms with AllowsEditorObjects=false.
The original UMetaData object was considered editor-only because it was included in the ClassesExcludedOnDedicatedServer/ClassesExcludedOnDedicatedClient config sections but this wasn't properly understood during the conversion to FMetaData.

diff --git a/Engine/Source/Runtime/CoreUObject/Private/UObject/LinkerLoad.cpp b/Engine/Source/Runtime/CoreUObject/Private/UObject/LinkerLoad.cpp
--- a/Engine/Source/Runtime/CoreUObject/Private/UObject/LinkerLoad.cpp
+++ b/Engine/Source/Runtime/CoreUObject/Private/UObject/LinkerLoad.cpp
@@ -2512,7 +2512,7 @@ FLinkerLoad::ELinkerStatus FLinkerLoad::SerializeMetaData()
 
 	if ( Summary.MetaDataOffset == 0 )
 	{
-		// This package was saved before meta data stored in the package summary
+		// This package was saved before meta data stored in the package summary or the meta data object wasn't saved
 		return LINKER_Loaded;
 	}
 
diff --git a/Engine/Source/Runtime/CoreUObject/Private/UObject/SavePackage2.cpp b/Engine/Source/Runtime/CoreUObject/Private/UObject/SavePackage2.cpp
--- a/Engine/Source/Runtime/CoreUObject/Private/UObject/SavePackage2.cpp
+++ b/Engine/Source/Runtime/CoreUObject/Private/UObject/SavePackage2.cpp
@@ -480,8 +480,11 @@ ESavePackageResult HarvestPackage(FSaveContext& SaveContext)
 
 #if WITH_METADATA
 		// Harvest package MetaData
-		FMetaData& PackageMetaData = SaveContext.GetPackage()->GetMetaData();
-		Harvester << PackageMetaData;
+		if (!SaveContext.IsCooking() || SaveContext.GetTargetPlatform()->AllowsEditorObjects())
+		{
+			FMetaData& PackageMetaData = SaveContext.GetPackage()->GetMetaData();
+			Harvester << PackageMetaData;
+		}
 #endif
 
 		// if we have a WorldTileInfo, we need to harvest its dependencies as well, i.e. Custom Version
@@ -2191,7 +2194,7 @@ ESavePackageResult WritePackageHeader(FStructuredArchive::FRecord& StructuredArc
 
 #if WITH_METADATA
 	// Save MetaData
-	if (!SaveContext.IsSaveOptional())
+	if (!SaveContext.IsCooking() || SaveContext.GetTargetPlatform()->AllowsEditorObjects())
 	{
 		SCOPED_SAVETIMER(UPackage_Save_WriteMetaData);
 		UE::SavePackageUtilities::SaveMetaData(SaveContext.GetPackage(), Linker, StructuredArchiveRoot);
